More coding standard fixes

diff --git a/Engine/Source/Programs/Shared/EpicGames.UHT/Parsers/UhtParsingScope.cs b/Engine/Source/Programs/Shared/EpicGames.UHT/Parsers/UhtParsingScope.cs
--- a/Engine/Source/Programs/Shared/EpicGames.UHT/Parsers/UhtParsingScope.cs
+++ b/Engine/Source/Programs/Shared/EpicGames.UHT/Parsers/UhtParsingScope.cs
@@ -1,12 +1,12 @@
 // Copyright Epic Games, Inc. All Rights Reserved.
 
+using System;
+using System.Text;
 using EpicGames.Core;
 using EpicGames.UHT.Tables;
 using EpicGames.UHT.Tokenizer;
 using EpicGames.UHT.Types;
 using EpicGames.UHT.Utils;
-using System;
-using System.Text;
 
 namespace EpicGames.UHT.Parsers
 {
@@ -59,14 +59,14 @@ namespace EpicGames.UHT.Parsers
 		{
 			get
 			{
-				UhtParsingScope? CurrentScope = this;
-				while (CurrentScope != null)
+				UhtParsingScope? currentScope = this;
+				while (currentScope != null)
 				{
-					if (CurrentScope.ScopeType is UhtClass)
+					if (currentScope.ScopeType is UhtClass)
 					{
-						return CurrentScope;
+						return currentScope;
 					}
-					CurrentScope = CurrentScope.ParentScope;
+					currentScope = currentScope.ParentScope;
 				}
 				throw new UhtIceException("Attempt to fetch the current class when a class isn't currently being parsed");
 			}
@@ -80,34 +80,34 @@ namespace EpicGames.UHT.Parsers
 		/// <summary>
 		/// Construct a root/global scope
 		/// </summary>
-		/// <param name="HeaderParser">Header parser</param>
-		/// <param name="ScopeType">Type being parsed</param>
-		/// <param name="KeywordTable">Keyword table</param>
-		public UhtParsingScope(UhtHeaderFileParser HeaderParser, UhtType ScopeType, UhtKeywordTable KeywordTable)
+		/// <param name="headerParser">Header parser</param>
+		/// <param name="scopeType">Type being parsed</param>
+		/// <param name="keywordTable">Keyword table</param>
+		public UhtParsingScope(UhtHeaderFileParser headerParser, UhtType scopeType, UhtKeywordTable keywordTable)
 		{
-			this.HeaderParser = HeaderParser;
-			this.TokenReader = HeaderParser.TokenReader;
+			this.HeaderParser = headerParser;
+			this.TokenReader = headerParser.TokenReader;
 			this.ParentScope = null;
-			this.ScopeType = ScopeType;
-			this.ScopeKeywordTable = KeywordTable;
+			this.ScopeType = scopeType;
+			this.ScopeKeywordTable = keywordTable;
 			this.HeaderParser.PushScope(this);
 		}
 
 		/// <summary>
 		/// Construct a scope for a type
 		/// </summary>
-		/// <param name="ParentScope">Parent scope</param>
-		/// <param name="ScopeType">Type being parsed</param>
-		/// <param name="KeywordTable">Keyword table</param>
-		/// <param name="AccessSpecifier">Current access specifier</param>
-		public UhtParsingScope(UhtParsingScope ParentScope, UhtType ScopeType, UhtKeywordTable KeywordTable, UhtAccessSpecifier AccessSpecifier)
+		/// <param name="parentScope">Parent scope</param>
+		/// <param name="scopeType">Type being parsed</param>
+		/// <param name="keywordTable">Keyword table</param>
+		/// <param name="accessSpecifier">Current access specifier</param>
+		public UhtParsingScope(UhtParsingScope parentScope, UhtType scopeType, UhtKeywordTable keywordTable, UhtAccessSpecifier accessSpecifier)
 		{
-			this.HeaderParser = ParentScope.HeaderParser;
-			this.TokenReader = ParentScope.TokenReader;
-			this.ParentScope = ParentScope;
-			this.ScopeType = ScopeType;
-			this.ScopeKeywordTable = KeywordTable;
-			this.AccessSpecifier = AccessSpecifier;
+			this.HeaderParser = parentScope.HeaderParser;
+			this.TokenReader = parentScope.TokenReader;
+			this.ParentScope = parentScope;
+			this.ScopeType = scopeType;
+			this.ScopeKeywordTable = keywordTable;
+			this.AccessSpecifier = accessSpecifier;
 			this.HeaderParser.PushScope(this);
 		}
 
@@ -143,71 +143,71 @@ namespace EpicGames.UHT.Parsers
 		/// <summary>
 		/// Add the module's relative path to the meta data
 		/// </summary>
-		/// <param name="MetaData">The meta data to add the information to</param>
-		/// <param name="HeaderFile">The header file currently being parsed</param>
-		public static void AddModuleRelativePathToMetaData(UhtMetaData MetaData, UhtHeaderFile HeaderFile)
+		/// <param name="metaData">The meta data to add the information to</param>
+		/// <param name="headerFile">The header file currently being parsed</param>
+		public static void AddModuleRelativePathToMetaData(UhtMetaData metaData, UhtHeaderFile headerFile)
 		{
-			MetaData.Add(UhtNames.ModuleRelativePath, HeaderFile.ModuleRelativeFilePath);
+			metaData.Add(UhtNames.ModuleRelativePath, headerFile.ModuleRelativeFilePath);
 		}
 
 		/// <summary>
 		/// Format the current token reader comments and add it as meta data
 		/// </summary>
-		/// <param name="MetaNameIndex">Index for the meta data key.  This is used for enum values</param>
-		public void AddFormattedCommentsAsTooltipMetaData(int MetaNameIndex = UhtMetaData.IndexNone)
+		/// <param name="metaNameIndex">Index for the meta data key.  This is used for enum values</param>
+		public void AddFormattedCommentsAsTooltipMetaData(int metaNameIndex = UhtMetaData.IndexNone)
 		{
-			AddFormattedCommentsAsTooltipMetaData(this.ScopeType, MetaNameIndex);
+			AddFormattedCommentsAsTooltipMetaData(this.ScopeType, metaNameIndex);
 		}
 
 		/// <summary>
 		/// Format the current token reader comments and add it as meta data
 		/// </summary>
-		/// <param name="Type">The type to add the meta data to</param>
-		/// <param name="MetaNameIndex">Index for the meta data key.  This is used for enum values</param>
-		public void AddFormattedCommentsAsTooltipMetaData(UhtType Type, int MetaNameIndex = UhtMetaData.IndexNone)
+		/// <param name="type">The type to add the meta data to</param>
+		/// <param name="metaNameIndex">Index for the meta data key.  This is used for enum values</param>
+		public void AddFormattedCommentsAsTooltipMetaData(UhtType type, int metaNameIndex = UhtMetaData.IndexNone)
 		{
 
 			// Don't add a tooltip if one already exists.
-			if (Type.MetaData.ContainsKey(UhtNames.ToolTip, MetaNameIndex))
+			if (type.MetaData.ContainsKey(UhtNames.ToolTip, metaNameIndex))
 			{
 				return;
 			}
 
 			// Fetch the comments
-			ReadOnlySpan<StringView> Comments = this.TokenReader.Comments;
+			ReadOnlySpan<StringView> comments = this.TokenReader.Comments;
 
 			// If we don't have any comments, just return
-			string MergedString = String.Empty;
-			if (Comments.Length == 0)
+			string mergedString = String.Empty;
+			if (comments.Length == 0)
 			{
 				return;
 			}
 
 			// Set the comment as just a simple concatenation of all the strings
-			if (Comments.Length == 1)
+			if (comments.Length == 1)
 			{
-				MergedString = Comments[0].ToString();
-				Type.MetaData.Add(UhtNames.Comment, MetaNameIndex, Comments[0].ToString());
+				mergedString = comments[0].ToString();
+				type.MetaData.Add(UhtNames.Comment, metaNameIndex, comments[0].ToString());
 			}
 			else
 			{
-				using (BorrowStringBuilder Borrower = new BorrowStringBuilder(StringBuilderCache.Small))
+				using (BorrowStringBuilder borrower = new BorrowStringBuilder(StringBuilderCache.Small))
 				{
-					StringBuilder Builder = Borrower.StringBuilder;
-					foreach (StringView Comment in Comments)
+					StringBuilder builder = borrower.StringBuilder;
+					foreach (StringView comment in comments)
 					{
-						Builder.Append(Comment);
+						builder.Append(comment);
 					}
-					MergedString = Builder.ToString();
-					Type.MetaData.Add(UhtNames.Comment, MetaNameIndex, MergedString);
+					mergedString = builder.ToString();
+					type.MetaData.Add(UhtNames.Comment, metaNameIndex, mergedString);
 				}
 			}
 
 			// Format the tooltip and set the metadata
-			StringView ToolTip = FormatCommentForToolTip(MergedString);
-			if (ToolTip.Span.Length > 0)
+			StringView toolTip = FormatCommentForToolTip(mergedString);
+			if (toolTip.Span.Length > 0)
 			{
-				Type.MetaData.Add(UhtNames.ToolTip, MetaNameIndex, ToolTip.ToString());
+				type.MetaData.Add(UhtNames.ToolTip, metaNameIndex, toolTip.ToString());
 
 				//COMPATIBILITY-TODO - Old UHT would only clear the comments if there was some form of a tooltip
 				this.TokenReader.ClearComments();
@@ -222,13 +222,13 @@ namespace EpicGames.UHT.Parsers
 		/// <summary>
 		/// Given a list of comments, check to see if any have alpha, numeric, or unicode code points with a value larger than 0xFF.
 		/// </summary>
-		/// <param name="Comments">Comments to search</param>
+		/// <param name="comments">Comments to search</param>
 		/// <returns>True is a character in question was found</returns>
-		private static bool HasValidCommentChar(ReadOnlySpan<char> Comments)
+		private static bool HasValidCommentChar(ReadOnlySpan<char> comments)
 		{
-			foreach (char C in Comments)
+			foreach (char c in comments)
 			{
-				if (UhtFCString.IsAlnum(C) || C > 0xFF)
+				if (UhtFCString.IsAlnum(c) || c > 0xFF)
 				{
 					return true;
 				}
@@ -252,28 +252,28 @@ namespace EpicGames.UHT.Parsers
 		/// are cleared.  However, if a C++ style comment follows a C style comment (regardless of any intermediate blank lines), then both blocks of comments will be present.
 		/// If any blank lines are encountered between blocks of C++ style comments, then any prior comments are cleared.
 		/// </summary>
-		/// <param name="Comments">Comments to be parsed</param>
+		/// <param name="comments">Comments to be parsed</param>
 		/// <returns>The generated tooltip</returns>
-		private static StringView FormatCommentForToolTip(string Comments)
+		private static StringView FormatCommentForToolTip(string comments)
 		{
-			if (!HasValidCommentChar(Comments))
+			if (!HasValidCommentChar(comments))
 			{
 				return new StringView();
 			}
 
-			string Result = Comments;
+			string result = comments;
 
 			// Sweep out comments marked to be ignored.
 			{
-				int CommentStart, CommentEnd;
+				int commentStart, commentEnd;
 
 				// Block comments go first
-				while ((CommentStart = Result.IndexOf("/*~", StringComparison.Ordinal)) != -1)
+				while ((commentStart = result.IndexOf("/*~", StringComparison.Ordinal)) != -1)
 				{
-					CommentEnd = Result.IndexOf("*/", CommentStart, StringComparison.Ordinal);
-					if (CommentEnd != -1)
+					commentEnd = result.IndexOf("*/", commentStart, StringComparison.Ordinal);
+					if (commentEnd != -1)
 					{
-						Result = Result.Remove(CommentStart, CommentEnd - CommentStart + 2);
+						result = result.Remove(commentStart, commentEnd - commentStart + 2);
 					}
 					else
 					{
@@ -283,198 +283,198 @@ namespace EpicGames.UHT.Parsers
 				}
 
 				// Leftover line comments go next
-				while ((CommentStart = Result.IndexOf("//~", StringComparison.Ordinal)) != -1)
+				while ((commentStart = result.IndexOf("//~", StringComparison.Ordinal)) != -1)
 				{
-					CommentEnd = Result.IndexOf("\n", CommentStart, StringComparison.Ordinal);
-					if (CommentEnd != -1)
+					commentEnd = result.IndexOf("\n", commentStart, StringComparison.Ordinal);
+					if (commentEnd != -1)
 					{
-						Result = Result.Remove(CommentStart, CommentEnd - CommentStart + 1);
+						result = result.Remove(commentStart, commentEnd - commentStart + 1);
 					}
 					else
 					{
-						Result = Result.Remove(CommentStart);
+						result = result.Remove(commentStart);
 						break;
 					}
 				}
 			}
 
 			// Check for known commenting styles.
-			bool bJavaDocStyle = Result.Contains("/**", StringComparison.Ordinal);
-			bool bCStyle = Result.Contains("/*", StringComparison.Ordinal);
-			bool bCPPStyle = Result.StartsWith("//", StringComparison.Ordinal);
+			bool javaDocStyle = result.Contains("/**", StringComparison.Ordinal);
+			bool cStyle = result.Contains("/*", StringComparison.Ordinal);
+			bool cppStyle = result.StartsWith("//", StringComparison.Ordinal);
 
-			if (bJavaDocStyle || bCStyle)
+			if (javaDocStyle || cStyle)
 			{
 				// Remove beginning and end markers.
-				if (bJavaDocStyle)
+				if (javaDocStyle)
 				{
-					Result = Result.Replace("/**", "", StringComparison.Ordinal);
+					result = result.Replace("/**", "", StringComparison.Ordinal);
 				}
-				if (bCStyle)
+				if (cStyle)
 				{
-					Result = Result.Replace("/*", "", StringComparison.Ordinal);
+					result = result.Replace("/*", "", StringComparison.Ordinal);
 				}
-				Result = Result.Replace("*/", "", StringComparison.Ordinal);
+				result = result.Replace("*/", "", StringComparison.Ordinal);
 			}
 
-			if (bCPPStyle)
+			if (cppStyle)
 			{
 				// Remove c++-style comment markers.  Also handle javadoc-style comments 
-				Result = Result.Replace("///", "", StringComparison.Ordinal);
-				Result = Result.Replace("//", "", StringComparison.Ordinal);
+				result = result.Replace("///", "", StringComparison.Ordinal);
+				result = result.Replace("//", "", StringComparison.Ordinal);
 
 				// Parser strips cpptext and replaces it with "// (cpptext)" -- prevent
 				// this from being treated as a comment on variables declared below the
 				// cpptext section
-				Result = Result.Replace("(cpptext)", "", StringComparison.Ordinal);
+				result = result.Replace("(cpptext)", "", StringComparison.Ordinal);
 			}
 
 			// Get rid of carriage return or tab characters, which mess up tooltips.
-			Result = Result.Replace("\r", "", StringComparison.Ordinal);
+			result = result.Replace("\r", "", StringComparison.Ordinal);
 
 			//wx widgets has a hard coded tab size of 8
 			{
-				int SpacesPerTab = 8;
-				Result = UhtFCString.TabsToSpaces(Result, SpacesPerTab, true);
+				const int SpacesPerTab = 8;
+				result = UhtFCString.TabsToSpaces(result, SpacesPerTab, true);
 			}
 
 			// get rid of uniform leading whitespace and all trailing whitespace, on each line
-			string[] Lines = Result.Split('\n');
+			string[] lines = result.Split('\n');
 
-			for (int Index = 0; Index < Lines.Length; ++Index)
+			for (int index = 0; index < lines.Length; ++index)
 			{
 				// Remove trailing whitespace
-				string Line = Lines[Index].TrimEnd();
+				string line = lines[index].TrimEnd();
 
 				// Remove leading "*" and "* " in javadoc comments.
-				if (bJavaDocStyle)
+				if (javaDocStyle)
 				{
 					// Find first non-whitespace character
-					int Pos = 0;
-					while (Pos < Line.Length && UhtFCString.IsWhitespace(Line[Pos]))
+					int pos = 0;
+					while (pos < line.Length && UhtFCString.IsWhitespace(line[pos]))
 					{
-						++Pos;
+						++pos;
 					}
 
 					// Is it a *?
-					if (Pos < Line.Length && Line[Pos] == '*')
+					if (pos < line.Length && line[pos] == '*')
 					{
 						// Eat next space as well
-						if (Pos + 1 < Line.Length && UhtFCString.IsWhitespace(Line[Pos + 1]))
+						if (pos + 1 < line.Length && UhtFCString.IsWhitespace(line[pos + 1]))
 						{
-							++Pos;
+							++pos;
 						}
 
-						Line = Line.Substring(Pos + 1);
+						line = line.Substring(pos + 1);
 					}
 				}
-				Lines[Index] = Line;
+				lines[index] = line;
 			}
 
-			Func<string, int, char, bool> IsAllSameChar = (string Line, int StartIndex, char TestChar) =>
+			static bool IsAllSameChar(string line, int startIndex, char testChar)
 			{
-				for (int Index = StartIndex, End = Line.Length; Index < End; ++Index)
+				for (int index = startIndex, end = line.Length; index < end; ++index)
 				{
-					if (Line[Index] != TestChar)
+					if (line[index] != testChar)
 					{
 						return false;
 					}
 				}
 				return true;
-			};
+			}
 
-			Func<string, bool> IsWhitespaceOrLineSeparator = (string Line) =>
+			static bool IsWhitespaceOrLineSeparator(string line)
 			{
 				// Skip any leading spaces
-				int Index = 0;
-				int EndPos = Line.Length;
-				for (; Index < EndPos && UhtFCString.IsWhitespace(Line[Index]); ++Index) 
+				int index = 0;
+				int endPos = line.Length;
+				for (; index < endPos && UhtFCString.IsWhitespace(line[index]); ++index)
 				{
 				}
-				if (Index == EndPos)
+				if (index == endPos)
 				{
 					return true;
 				}
 
 				// Check for the same character
-				return IsAllSameChar(Line, Index, '-') || IsAllSameChar(Line, Index, '=') || IsAllSameChar(Line, Index, '*');
-			};
+				return IsAllSameChar(line, index, '-') || IsAllSameChar(line, index, '=') || IsAllSameChar(line, index, '*');
+			}
 
 			// Find first meaningful line
-			int FirstIndex = 0;
-			for (; FirstIndex < Lines.Length && IsWhitespaceOrLineSeparator(Lines[FirstIndex]); ++FirstIndex)
+			int firstIndex = 0;
+			for (; firstIndex < lines.Length && IsWhitespaceOrLineSeparator(lines[firstIndex]); ++firstIndex)
 			{ }
 
-			int LastIndex = Lines.Length;
-			for (; LastIndex > FirstIndex && IsWhitespaceOrLineSeparator(Lines[LastIndex - 1]); --LastIndex)
+			int lastIndex = lines.Length;
+			for (; lastIndex > firstIndex && IsWhitespaceOrLineSeparator(lines[lastIndex - 1]); --lastIndex)
 			{ }
 
-			Result = String.Empty;
-			if (FirstIndex != LastIndex)
+			result = String.Empty;
+			if (firstIndex != lastIndex)
 			{
-				string FirstLine = Lines[FirstIndex];
+				string firstLine = lines[firstIndex];
 
 				// Figure out how much whitespace is on the first line
-				int MaxNumWhitespaceToRemove = 0;
-				for (; MaxNumWhitespaceToRemove < FirstLine.Length; MaxNumWhitespaceToRemove++)
+				int maxNumWhitespaceToRemove = 0;
+				for (; maxNumWhitespaceToRemove < firstLine.Length; maxNumWhitespaceToRemove++)
 				{
-					if (/*!UhtFCString.IsLinebreak(FirstLine[MaxNumWhitespaceToRemove]) && */!UhtFCString.IsWhitespace(FirstLine[MaxNumWhitespaceToRemove]))
+					if (/*!UhtFCString.IsLinebreak(FirstLine[MaxNumWhitespaceToRemove]) && */!UhtFCString.IsWhitespace(firstLine[maxNumWhitespaceToRemove]))
 					{
 						break;
 					}
 				}
 
-				for (int Index = FirstIndex; Index != LastIndex; ++Index)
+				for (int index = firstIndex; index != lastIndex; ++index)
 				{
-					string Line = Lines[Index];
+					string line = lines[index];
 
-					int TemporaryMaxWhitespace = MaxNumWhitespaceToRemove;
+					int temporaryMaxWhitespace = maxNumWhitespaceToRemove;
 
 					// Allow eating an extra tab on subsequent lines if it's present
-					if ((Index > FirstIndex) && (Line.Length > 0) && (Line[0] == '\t'))
+					if ((index > firstIndex) && (line.Length > 0) && (line[0] == '\t'))
 					{
-						TemporaryMaxWhitespace++;
+						temporaryMaxWhitespace++;
 					}
 
 					// Advance past whitespace
-					int Pos = 0;
-					while (Pos < TemporaryMaxWhitespace && Pos < Line.Length && UhtFCString.IsWhitespace(Line[Pos]))
+					int pos = 0;
+					while (pos < temporaryMaxWhitespace && pos < line.Length && UhtFCString.IsWhitespace(line[pos]))
 					{
-						++Pos;
+						++pos;
 					}
 
-					if (Pos > 0)
+					if (pos > 0)
 					{
-						Line = Line.Substring(Pos);
+						line = line.Substring(pos);
 					}
 
-					if (Index > FirstIndex)
+					if (index > firstIndex)
 					{
-						Result += "\n";
+						result += "\n";
 					}
 
-					if (Line.Length > 0 && !IsAllSameChar(Line, 0, '='))
+					if (line.Length > 0 && !IsAllSameChar(line, 0, '='))
 					{
-						Result += Line;
+						result += line;
 					}
 				}
 			}
 
 			//@TODO: UCREMOVAL: Really want to trim an arbitrary number of newlines above and below, but keep multiple newlines internally
 			// Make sure it doesn't start with a newline
-			if (Result.Length != 0 && UhtFCString.IsLinebreak(Result[0]))
+			if (result.Length != 0 && UhtFCString.IsLinebreak(result[0]))
 			{
-				Result = Result.Substring(1);
+				result = result.Substring(1);
 			}
 
 			// Make sure it doesn't end with a dead newline
-			if (Result.Length != 0 && UhtFCString.IsLinebreak(Result[Result.Length - 1]))
+			if (result.Length != 0 && UhtFCString.IsLinebreak(result[^1]))
 			{
-				Result = Result.Substring(0, Result.Length - 1);
+				result = result.Substring(0, result.Length - 1);
 			}
 
 			// Done.
-			return Result;
+			return result;
 		}
 	}
 
@@ -483,43 +483,43 @@ namespace EpicGames.UHT.Parsers
 	/// </summary>
 	public struct UhtTokenRecorder : IDisposable
 	{
-		private readonly UhtCompilerDirective CompilerDirective;
-		private readonly UhtParsingScope Scope;
-		private readonly UhtFunction? Function;
-		private bool bFlushed;
+		private readonly UhtCompilerDirective _compilerDirective;
+		private readonly UhtParsingScope _scope;
+		private readonly UhtFunction? _function;
+		private bool _flushed;
 
 		/// <summary>
 		/// Construct a new recorder
 		/// </summary>
-		/// <param name="Scope">Scope being parsed</param>
-		/// <param name="InitialToken">Initial toke nto add to the recorder</param>
-		public UhtTokenRecorder(UhtParsingScope Scope, ref UhtToken InitialToken)
+		/// <param name="scope">Scope being parsed</param>
+		/// <param name="initialToken">Initial toke nto add to the recorder</param>
+		public UhtTokenRecorder(UhtParsingScope scope, ref UhtToken initialToken)
 		{
-			this.Scope = Scope;
-			this.CompilerDirective = this.Scope.HeaderParser.GetCurrentCompositeCompilerDirective();
-			this.Function = null;
-			this.bFlushed = false;
-			if (this.Scope.ScopeType is UhtClass)
+			this._scope = scope;
+			this._compilerDirective = this._scope.HeaderParser.GetCurrentCompositeCompilerDirective();
+			this._function = null;
+			this._flushed = false;
+			if (this._scope.ScopeType is UhtClass)
 			{
-				this.Scope.TokenReader.EnableRecording();
-				this.Scope.TokenReader.RecordToken(ref InitialToken);
+				this._scope.TokenReader.EnableRecording();
+				this._scope.TokenReader.RecordToken(ref initialToken);
 			}
 		}
 
 		/// <summary>
 		/// Create a new recorder
 		/// </summary>
-		/// <param name="Scope">Scope being parsed</param>
-		/// <param name="Function">Function associated with the recorder</param>
-		public UhtTokenRecorder(UhtParsingScope Scope, UhtFunction Function)
+		/// <param name="scope">Scope being parsed</param>
+		/// <param name="function">Function associated with the recorder</param>
+		public UhtTokenRecorder(UhtParsingScope scope, UhtFunction function)
 		{
-			this.Scope = Scope;
-			this.CompilerDirective = this.Scope.HeaderParser.GetCurrentCompositeCompilerDirective();
-			this.Function = Function;
-			this.bFlushed = false;
-			if (this.Scope.ScopeType is UhtClass)
+			this._scope = scope;
+			this._compilerDirective = this._scope.HeaderParser.GetCurrentCompositeCompilerDirective();
+			this._function = function;
+			this._flushed = false;
+			if (this._scope.ScopeType is UhtClass)
 			{
-				this.Scope.TokenReader.EnableRecording();
+				this._scope.TokenReader.EnableRecording();
 			}
 		}
 
@@ -537,13 +537,13 @@ namespace EpicGames.UHT.Parsers
 		/// <returns>True if the recorded content was added to a class</returns>
 		public bool Stop()
 		{
-			if (!bFlushed)
+			if (!_flushed)
 			{
-				bFlushed = true;
-				if (this.Scope.ScopeType is UhtClass Class)
+				_flushed = true;
+				if (this._scope.ScopeType is UhtClass classObj)
 				{
-					Class.AddDeclaration(this.CompilerDirective, this.Scope.TokenReader.RecordedTokens, this.Function);
-					this.Scope.TokenReader.DisableRecording();
+					classObj.AddDeclaration(this._compilerDirective, this._scope.TokenReader.RecordedTokens, this._function);
+					this._scope.TokenReader.DisableRecording();
 					return true;
 				}
 			}
