More coding standard fixes

diff --git a/Engine/Source/Programs/Shared/EpicGames.UHT/Parsers/UhtPropertyParser.cs b/Engine/Source/Programs/Shared/EpicGames.UHT/Parsers/UhtPropertyParser.cs
--- a/Engine/Source/Programs/Shared/EpicGames.UHT/Parsers/UhtPropertyParser.cs
+++ b/Engine/Source/Programs/Shared/EpicGames.UHT/Parsers/UhtPropertyParser.cs
@@ -1,16 +1,16 @@
 // Copyright Epic Games, Inc. All Rights Reserved.
 
+using System;
+using System.Collections.Generic;
+using System.Diagnostics.CodeAnalysis;
+using System.Text;
+using System.Threading;
 using EpicGames.Core;
 using EpicGames.UHT.Parsers;
 using EpicGames.UHT.Tables;
 using EpicGames.UHT.Tokenizer;
 using EpicGames.UHT.Types;
 using EpicGames.UHT.Utils;
-using System;
-using System.Collections.Generic;
-using System.Diagnostics.CodeAnalysis;
-using System.Text;
-using System.Threading;
 
 namespace EpicGames.UHT.Parsers
 {
@@ -77,35 +77,35 @@ namespace EpicGames.UHT.Parsers
 		/// <summary>
 		/// Test to see if any of the specified flags are set
 		/// </summary>
-		/// <param name="InFlags">Current flags</param>
-		/// <param name="TestFlags">Flags to test for</param>
+		/// <param name="inFlags">Current flags</param>
+		/// <param name="testFlags">Flags to test for</param>
 		/// <returns>True if any of the flags are set</returns>
-		public static bool HasAnyFlags(this UhtPropertyParseOptions InFlags, UhtPropertyParseOptions TestFlags)
+		public static bool HasAnyFlags(this UhtPropertyParseOptions inFlags, UhtPropertyParseOptions testFlags)
 		{
-			return (InFlags & TestFlags) != 0;
+			return (inFlags & testFlags) != 0;
 		}
 
 		/// <summary>
 		/// Test to see if all of the specified flags are set
 		/// </summary>
-		/// <param name="InFlags">Current flags</param>
-		/// <param name="TestFlags">Flags to test for</param>
+		/// <param name="inFlags">Current flags</param>
+		/// <param name="testFlags">Flags to test for</param>
 		/// <returns>True if all the flags are set</returns>
-		public static bool HasAllFlags(this UhtPropertyParseOptions InFlags, UhtPropertyParseOptions TestFlags)
+		public static bool HasAllFlags(this UhtPropertyParseOptions inFlags, UhtPropertyParseOptions testFlags)
 		{
-			return (InFlags & TestFlags) == TestFlags;
+			return (inFlags & testFlags) == testFlags;
 		}
 
 		/// <summary>
 		/// Test to see if a specific set of flags have a specific value.
 		/// </summary>
-		/// <param name="InFlags">Current flags</param>
-		/// <param name="TestFlags">Flags to test for</param>
-		/// <param name="MatchFlags">Expected value of the tested flags</param>
+		/// <param name="inFlags">Current flags</param>
+		/// <param name="testFlags">Flags to test for</param>
+		/// <param name="matchFlags">Expected value of the tested flags</param>
 		/// <returns>True if the given flags have a specific value.</returns>
-		public static bool HasExactFlags(this UhtPropertyParseOptions InFlags, UhtPropertyParseOptions TestFlags, UhtPropertyParseOptions MatchFlags)
+		public static bool HasExactFlags(this UhtPropertyParseOptions inFlags, UhtPropertyParseOptions testFlags, UhtPropertyParseOptions matchFlags)
 		{
-			return (InFlags & TestFlags) == MatchFlags;
+			return (inFlags & testFlags) == matchFlags;
 		}
 	}
 
@@ -136,7 +136,7 @@ namespace EpicGames.UHT.Parsers
 		/// None found
 		/// </summary>
 		None,
-		
+
 		/// <summary>
 		/// Array
 		/// </summary>
@@ -182,11 +182,11 @@ namespace EpicGames.UHT.Parsers
 		/// <summary>
 		/// Return true if the type is editor only
 		/// </summary>
-		/// <param name="LayoutMacroType">Layout macro type</param>
+		/// <param name="layoutMacroType">Layout macro type</param>
 		/// <returns>True if editor only</returns>
-		public static bool IsEditorOnly(this UhtLayoutMacroType LayoutMacroType)
+		public static bool IsEditorOnly(this UhtLayoutMacroType layoutMacroType)
 		{
-			switch (LayoutMacroType)
+			switch (layoutMacroType)
 			{
 				case UhtLayoutMacroType.ArrayEditorOnly:
 				case UhtLayoutMacroType.BitfieldEditorOnly:
@@ -201,11 +201,11 @@ namespace EpicGames.UHT.Parsers
 		/// <summary>
 		/// Return true if the type is a bit field
 		/// </summary>
-		/// <param name="LayoutMacroType">Layout macro type</param>
+		/// <param name="layoutMacroType">Layout macro type</param>
 		/// <returns>True if bit field</returns>
-		public static bool IsBitfield(this UhtLayoutMacroType LayoutMacroType)
+		public static bool IsBitfield(this UhtLayoutMacroType layoutMacroType)
 		{
-			switch (LayoutMacroType)
+			switch (layoutMacroType)
 			{
 				case UhtLayoutMacroType.Bitfield:
 				case UhtLayoutMacroType.BitfieldEditorOnly:
@@ -219,11 +219,11 @@ namespace EpicGames.UHT.Parsers
 		/// <summary>
 		/// Return true if the type is an array
 		/// </summary>
-		/// <param name="LayoutMacroType">Layout macro type</param>
+		/// <param name="layoutMacroType">Layout macro type</param>
 		/// <returns>True if array</returns>
-		public static bool IsArray(this UhtLayoutMacroType LayoutMacroType)
+		public static bool IsArray(this UhtLayoutMacroType layoutMacroType)
 		{
-			switch (LayoutMacroType)
+			switch (layoutMacroType)
 			{
 				case UhtLayoutMacroType.Array:
 				case UhtLayoutMacroType.ArrayEditorOnly:
@@ -237,11 +237,11 @@ namespace EpicGames.UHT.Parsers
 		/// <summary>
 		/// Return true if the type has an initializer
 		/// </summary>
-		/// <param name="LayoutMacroType">Layout macro type</param>
+		/// <param name="layoutMacroType">Layout macro type</param>
 		/// <returns>True if it has an initializer</returns>
-		public static bool HasInitializer(this UhtLayoutMacroType LayoutMacroType)
+		public static bool HasInitializer(this UhtLayoutMacroType layoutMacroType)
 		{
-			switch (LayoutMacroType)
+			switch (layoutMacroType)
 			{
 				case UhtLayoutMacroType.FieldInitialized:
 					return true;
@@ -254,12 +254,12 @@ namespace EpicGames.UHT.Parsers
 		/// <summary>
 		/// Return the layout macro name
 		/// </summary>
-		/// <param name="LayoutMacroType">Type in question</param>
+		/// <param name="layoutMacroType">Type in question</param>
 		/// <returns>Macro name</returns>
 		/// <exception cref="UhtIceException">Thrown if the macro type is none or invalid</exception>
-		public static StringView MacroName(this UhtLayoutMacroType LayoutMacroType)
+		public static StringView MacroName(this UhtLayoutMacroType layoutMacroType)
 		{
-			switch (LayoutMacroType)
+			switch (layoutMacroType)
 			{
 				default:
 				case UhtLayoutMacroType.None:
@@ -291,22 +291,22 @@ namespace EpicGames.UHT.Parsers
 		/// <summary>
 		/// Return the macro name and value
 		/// </summary>
-		/// <param name="LayoutMacroType">Macro name</param>
+		/// <param name="layoutMacroType">Macro name</param>
 		/// <returns>Name and type</returns>
-		public static KeyValuePair<StringView, UhtLayoutMacroType> MacroNameAndValue(this UhtLayoutMacroType LayoutMacroType)
+		public static KeyValuePair<StringView, UhtLayoutMacroType> MacroNameAndValue(this UhtLayoutMacroType layoutMacroType)
 		{
-			return new KeyValuePair<StringView, UhtLayoutMacroType>(LayoutMacroType.MacroName(), LayoutMacroType);
+			return new KeyValuePair<StringView, UhtLayoutMacroType>(layoutMacroType.MacroName(), layoutMacroType);
 		}
 	}
 
 	/// <summary>
 	/// Delegate invoked to handle a parsed property 
 	/// </summary>
-	/// <param name="TopScope">Scope being parsed</param>
-	/// <param name="Property">Property just parsed</param>
-	/// <param name="NameToken">Name of the property</param>
-	/// <param name="LayoutMacroType">Layout macro type</param>
-	public delegate void UhtPropertyDelegate(UhtParsingScope TopScope, UhtProperty Property, ref UhtToken NameToken, UhtLayoutMacroType LayoutMacroType);
+	/// <param name="topScope">Scope being parsed</param>
+	/// <param name="property">Property just parsed</param>
+	/// <param name="nameToken">Name of the property</param>
+	/// <param name="layoutMacroType">Layout macro type</param>
+	public delegate void UhtPropertyDelegate(UhtParsingScope topScope, UhtProperty property, ref UhtToken nameToken, UhtLayoutMacroType layoutMacroType);
 
 	/// <summary>
 	/// Context for property specifier parsing
@@ -321,32 +321,32 @@ namespace EpicGames.UHT.Parsers
 		/// <summary>
 		/// If true, editor specifier seen
 		/// </summary>
-		public bool bSeenEditSpecifier { get; set; } = false;
+		public bool SeenEditSpecifier { get; set; } = false;
 
 		/// <summary>
 		/// If true, blueprint write specifier seen
 		/// </summary>
-		public bool bSeenBlueprintWriteSpecifier { get; set; } = false;
+		public bool SeenBlueprintWriteSpecifier { get; set; } = false;
 
 		/// <summary>
 		/// If true, blueprint readonly specifier seen
 		/// </summary>
-		public bool bSeenBlueprintReadOnlySpecifier { get; set; } = false;
+		public bool SeenBlueprintReadOnlySpecifier { get; set; } = false;
 
 		/// <summary>
 		/// If true, blueprint getter specifier seen
 		/// </summary>
-		public bool bSeenBlueprintGetterSpecifier { get; set; } = false;
+		public bool SeenBlueprintGetterSpecifier { get; set; } = false;
 
 		/// <summary>
 		/// Construct a new property specifier context
 		/// </summary>
-		/// <param name="Scope">Scope being parsed</param>
-		/// <param name="MessageSite">Message site</param>
-		public UhtPropertySpecifierContext(UhtParsingScope Scope, IUhtMessageSite MessageSite) : base(Scope, MessageSite, UhtMetaData.Empty)
+		/// <param name="scope">Scope being parsed</param>
+		/// <param name="messageSite">Message site</param>
+		public UhtPropertySpecifierContext(UhtParsingScope scope, IUhtMessageSite messageSite) : base(scope, messageSite, UhtMetaData.Empty)
 		{
 			this.PropertySettings = new UhtPropertySettings();
-			this.PropertySettings.Reset(Scope.ScopeType, 1, UhtPropertyCategory.Member, 0);
+			this.PropertySettings.Reset(scope.ScopeType, 1, UhtPropertyCategory.Member, 0);
 		}
 	}
 
@@ -357,13 +357,13 @@ namespace EpicGames.UHT.Parsers
 	public class UhtPreResolveProperty : UhtProperty
 	{
 		/// <inheritdoc/>
-		public override string EngineClassName { get => "UHTParsedProperty"; }
+		public override string EngineClassName => "UHTParsedProperty";
 
 		/// <inheritdoc/>
-		protected override string CppTypeText { get => "invalid"; }
+		protected override string CppTypeText => "invalid";
 
 		/// <inheritdoc/>
-		protected override string PGetMacroText { get => "invalid"; }
+		protected override string PGetMacroText => "invalid";
 
 		/// <summary>
 		/// Collection of type tokens 
@@ -378,52 +378,52 @@ namespace EpicGames.UHT.Parsers
 		/// <summary>
 		/// Construct a new property to be resolved
 		/// </summary>
-		/// <param name="PropertySettings">Property settings</param>
-		/// <param name="TypeTokens">Type tokens</param>
-		public UhtPreResolveProperty(UhtPropertySettings PropertySettings, ReadOnlyMemory<UhtToken> TypeTokens) : base(PropertySettings)
+		/// <param name="propertySettings">Property settings</param>
+		/// <param name="typeTokens">Type tokens</param>
+		public UhtPreResolveProperty(UhtPropertySettings propertySettings, ReadOnlyMemory<UhtToken> typeTokens) : base(propertySettings)
 		{
-			this.TypeTokens = TypeTokens;
-			this.PropertySettings = PropertySettings;
+			this.TypeTokens = typeTokens;
+			this.PropertySettings = propertySettings;
 		}
 
 		/// <inheritdoc/>
-		public override bool IsSameType(UhtProperty Other)
+		public override bool IsSameType(UhtProperty other)
 		{
 			throw new NotImplementedException();
 		}
 
 		/// <inheritdoc/>
-		public override StringBuilder AppendText(StringBuilder Builder, UhtPropertyTextType TextType, bool bIsTemplateArgument)
+		public override StringBuilder AppendText(StringBuilder builder, UhtPropertyTextType textType, bool isTemplateArgument)
 		{
 			throw new NotImplementedException();
 		}
 
 		/// <inheritdoc/>
-		public override StringBuilder AppendMemberDecl(StringBuilder Builder, IUhtPropertyMemberContext Context, string Name, string NameSuffix, int Tabs)
+		public override StringBuilder AppendMemberDecl(StringBuilder builder, IUhtPropertyMemberContext context, string name, string nameSuffix, int tabs)
 		{
 			throw new NotImplementedException();
 		}
 
 		/// <inheritdoc/>
-		public override StringBuilder AppendMemberDef(StringBuilder Builder, IUhtPropertyMemberContext Context, string Name, string NameSuffix, string? Offset, int Tabs)
+		public override StringBuilder AppendMemberDef(StringBuilder builder, IUhtPropertyMemberContext context, string name, string nameSuffix, string? offset, int tabs)
 		{
 			throw new NotImplementedException();
 		}
 
 		/// <inheritdoc/>
-		public override StringBuilder AppendFunctionThunkParameterGet(StringBuilder Builder)
+		public override StringBuilder AppendFunctionThunkParameterGet(StringBuilder builder)
 		{
 			throw new NotImplementedException();
 		}
 
 		/// <inheritdoc/>
-		public override StringBuilder AppendNullConstructorArg(StringBuilder Builder, bool bIsInitializer)
+		public override StringBuilder AppendNullConstructorArg(StringBuilder builder, bool isInitializer)
 		{
 			throw new NotImplementedException();
 		}
 
 		/// <inheritdoc/>
-		public override bool SanitizeDefaultValue(IUhtTokenReader DefaultValueReader, StringBuilder InnerDefaultValue)
+		public override bool SanitizeDefaultValue(IUhtTokenReader defaultValueReader, StringBuilder innerDefaultValue)
 		{
 			throw new NotImplementedException();
 		}
@@ -434,21 +434,21 @@ namespace EpicGames.UHT.Parsers
 	/// </summary>
 	public class UhtPropertyParser : IUhtMessageExtraContext
 	{
-		private readonly UhtPropertySpecifierContext SpecifierContext;
-		private UhtPropertyParseOptions Options;
-		private UhtParsePropertyDeclarationStyle DeclarationStyle;
+		private readonly UhtPropertySpecifierContext _specifierContext;
+		private UhtPropertyParseOptions _options;
+		private UhtParsePropertyDeclarationStyle _declarationStyle;
 
-		private IUhtTokenReader TokenReader => this.SpecifierContext.Scope.TokenReader;
-		private UhtParsingScope TopScope => this.SpecifierContext.Scope;
+		private IUhtTokenReader TokenReader => this._specifierContext.Scope.TokenReader;
+		private UhtParsingScope TopScope => this._specifierContext.Scope;
 
 		// Scratch pad variables used by actions
-		private List<UhtToken> CurrentTypeTokens = new List<UhtToken>();
-		private int CurrentTemplateDepth = 0;
+		private List<UhtToken> _currentTypeTokens = new List<UhtToken>();
+		private int _currentTemplateDepth = 0;
 
 		// Actions/Delegates
-		private readonly UhtTokensUntilDelegate GatherTypeTokensDelegate;
+		private readonly UhtTokensUntilDelegate _gatherTypeTokensDelegate;
 
-		private static readonly Dictionary<StringView, UhtLayoutMacroType> LayoutMacroTypes = new Dictionary<StringView, UhtLayoutMacroType>(new[]
+		private static readonly Dictionary<StringView, UhtLayoutMacroType> s_layoutMacroTypes = new Dictionary<StringView, UhtLayoutMacroType>(new[]
 		{
 			UhtLayoutMacroType.Array.MacroNameAndValue(),
 			UhtLayoutMacroType.ArrayEditorOnly.MacroNameAndValue(),
@@ -462,56 +462,56 @@ namespace EpicGames.UHT.Parsers
 		/// <summary>
 		/// Construct a new property parser
 		/// </summary>
-		/// <param name="Scope">Current scope being parsed</param>
-		/// <param name="MessageSite">Message site</param>
-		public UhtPropertyParser(UhtParsingScope Scope, IUhtMessageSite MessageSite)
+		/// <param name="scope">Current scope being parsed</param>
+		/// <param name="messageSite">Message site</param>
+		public UhtPropertyParser(UhtParsingScope scope, IUhtMessageSite messageSite)
 		{
-			this.SpecifierContext = new UhtPropertySpecifierContext(Scope, MessageSite);
-			this.GatherTypeTokensDelegate = GatherTypeTokens;
+			this._specifierContext = new UhtPropertySpecifierContext(scope, messageSite);
+			this._gatherTypeTokensDelegate = GatherTypeTokens;
 		}
 
 		/// <summary>
 		/// Reset the property parser for a new parse
 		/// </summary>
-		/// <param name="Scope">Current scope being parsed</param>
-		/// <param name="MessageSite">Message site</param>
-		public void Reset(UhtParsingScope Scope, IUhtMessageSite MessageSite)
+		/// <param name="scope">Current scope being parsed</param>
+		/// <param name="messageSite">Message site</param>
+		public void Reset(UhtParsingScope scope, IUhtMessageSite messageSite)
 		{
-			this.SpecifierContext.Scope = Scope;
-			this.SpecifierContext.MessageSite = MessageSite;
+			this._specifierContext.Scope = scope;
+			this._specifierContext.MessageSite = messageSite;
 		}
 
 		/// <summary>
 		/// Parse the property
 		/// </summary>
-		/// <param name="DisallowPropertyFlags">Flags to be disallowed</param>
-		/// <param name="Options">Parsing options</param>
-		/// <param name="DeclarationStyle">Style of declaration</param>
-		/// <param name="Category">Property category</param>
-		/// <param name="Delegate">Delegate to be invoked after property has been parsed</param>
+		/// <param name="disallowPropertyFlags">Flags to be disallowed</param>
+		/// <param name="options">Parsing options</param>
+		/// <param name="declarationStyle">Style of declaration</param>
+		/// <param name="category">Property category</param>
+		/// <param name="propertyDelegate">Delegate to be invoked after property has been parsed</param>
 		/// <returns>The property parser</returns>
-		public UhtPropertyParser Parse(EPropertyFlags DisallowPropertyFlags, UhtPropertyParseOptions Options, UhtParsePropertyDeclarationStyle DeclarationStyle, UhtPropertyCategory Category, UhtPropertyDelegate Delegate)
+		public UhtPropertyParser Parse(EPropertyFlags disallowPropertyFlags, UhtPropertyParseOptions options, UhtParsePropertyDeclarationStyle declarationStyle, UhtPropertyCategory category, UhtPropertyDelegate propertyDelegate)
 		{
 			// Reset the context and create the property
-			this.SpecifierContext.PropertySettings.Reset(this.SpecifierContext.Scope.ScopeType, 0, Category, DisallowPropertyFlags);
-			this.SpecifierContext.MetaData = this.SpecifierContext.PropertySettings.MetaData;
-			this.SpecifierContext.MetaNameIndex = UhtMetaData.IndexNone;
-			this.SpecifierContext.bSeenEditSpecifier = false;
-			this.SpecifierContext.bSeenBlueprintWriteSpecifier = false;
-			this.SpecifierContext.bSeenBlueprintReadOnlySpecifier = false;
-			this.SpecifierContext.bSeenBlueprintGetterSpecifier = false;
+			this._specifierContext.PropertySettings.Reset(this._specifierContext.Scope.ScopeType, 0, category, disallowPropertyFlags);
+			this._specifierContext.MetaData = this._specifierContext.PropertySettings.MetaData;
+			this._specifierContext.MetaNameIndex = UhtMetaData.IndexNone;
+			this._specifierContext.SeenEditSpecifier = false;
+			this._specifierContext.SeenBlueprintWriteSpecifier = false;
+			this._specifierContext.SeenBlueprintReadOnlySpecifier = false;
+			this._specifierContext.SeenBlueprintGetterSpecifier = false;
 
 			// Initialize the property
 
-			this.Options = Options;
-			this.DeclarationStyle = DeclarationStyle;
+			this._options = options;
+			this._declarationStyle = declarationStyle;
 
-			this.CurrentTypeTokens = new List<UhtToken>();
-			this.CurrentTemplateDepth = 0;
+			this._currentTypeTokens = new List<UhtToken>();
+			this._currentTemplateDepth = 0;
 
-			using (var TokenContext = new UhtMessageContext(this))
+			using (UhtMessageContext tokenContext = new UhtMessageContext(this))
 			{
-				ParseInternal(Delegate);
+				ParseInternal(propertyDelegate);
 			}
 			return this;
 		}
@@ -522,9 +522,9 @@ namespace EpicGames.UHT.Parsers
 		{
 			get
 			{
-				Stack<object?> ExtraContext = new Stack<object?>(1);
-				ExtraContext.Push(this.SpecifierContext.PropertySettings.PropertyCategory.GetHintText());
-				return ExtraContext;
+				Stack<object?> extraContext = new Stack<object?>(1);
+				extraContext.Push(this._specifierContext.PropertySettings.PropertyCategory.GetHintText());
+				return extraContext;
 			}
 		}
 		#endregion
@@ -533,83 +533,83 @@ namespace EpicGames.UHT.Parsers
 		/// Resolve the given property.  This method will resolve any immediate property during the parsing phase or 
 		/// resolve any previously parsed property to the final version.
 		/// </summary>
-		/// <param name="ResolvePhase">Used to detect if the property is being parsed or resolved</param>
-		/// <param name="Data">Character buffer for header file</param>
-		/// <param name="PropertySettings">The property settings.</param>
-		/// <param name="TypeTokens">The tokens that represent the type</param>
+		/// <param name="resolvePhase">Used to detect if the property is being parsed or resolved</param>
+		/// <param name="data">Character buffer for header file</param>
+		/// <param name="propertySettings">The property settings.</param>
+		/// <param name="typeTokens">The tokens that represent the type</param>
 		/// <returns></returns>
-		public static UhtProperty? ResolveProperty(UhtPropertyResolvePhase ResolvePhase, UhtPropertySettings PropertySettings, ReadOnlyMemory<char> Data, ReadOnlyMemory<UhtToken> TypeTokens)
+		public static UhtProperty? ResolveProperty(UhtPropertyResolvePhase resolvePhase, UhtPropertySettings propertySettings, ReadOnlyMemory<char> data, ReadOnlyMemory<UhtToken> typeTokens)
 		{
-			UhtSession Session = PropertySettings.Outer.Session;
-			IUhtTokenReader ReplayReader = UhtTokenReplayReader.GetThreadInstance(PropertySettings.Outer, Data, TypeTokens, UhtTokenType.EndOfType);
+			UhtSession session = propertySettings.Outer.Session;
+			IUhtTokenReader replayReader = UhtTokenReplayReader.GetThreadInstance(propertySettings.Outer, data, typeTokens, UhtTokenType.EndOfType);
 
 			// Loop through the tokens until we find a known property type or the start of a template argument list
-			UhtPropertyType PropertyType;
-			for (int Index = 0; Index < TypeTokens.Length; ++Index)
+			UhtPropertyType propertyType;
+			for (int index = 0; index < typeTokens.Length; ++index)
 			{
-				if (TypeTokens.Span[Index].IsSymbol())
+				if (typeTokens.Span[index].IsSymbol())
 				{
-					ReadOnlySpan<char> Span = TypeTokens.Span[Index].Value.Span;
-					if (Span.Length == 1 && (Span[0] == '<' || Span[0] == '>' || Span[0] == ','))
+					ReadOnlySpan<char> span = typeTokens.Span[index].Value.Span;
+					if (span.Length == 1 && (span[0] == '<' || span[0] == '>' || span[0] == ','))
 					{
 						break;
 					}
 				}
 				else
 				{
-					if (Session.TryGetPropertyType(TypeTokens.Span[Index].Value, out PropertyType))
+					if (session.TryGetPropertyType(typeTokens.Span[index].Value, out propertyType))
 					{
-						if (ResolvePhase == UhtPropertyResolvePhase.Resolving || PropertyType.Options.HasAnyFlags(UhtPropertyTypeOptions.Immediate))
+						if (resolvePhase == UhtPropertyResolvePhase.Resolving || propertyType.Options.HasAnyFlags(UhtPropertyTypeOptions.Immediate))
 						{
-							return ResolvePropertyType(ResolvePhase, PropertySettings, ReplayReader, PropertyType, TypeTokens.Span[Index], false);
+							return ResolvePropertyType(resolvePhase, propertySettings, replayReader, propertyType, typeTokens.Span[index], false);
 						}
 					}
 				}
 			}
 
 			// Try the default processor.  We only use the default processor when trying to resolve something post parsing phase.
-			if (ResolvePhase == UhtPropertyResolvePhase.Resolving)
+			if (resolvePhase == UhtPropertyResolvePhase.Resolving)
 			{
-				return ResolvePropertyType(ResolvePhase, PropertySettings, ReplayReader, Session.DefaultPropertyType, new UhtToken(), false);
+				return ResolvePropertyType(resolvePhase, propertySettings, replayReader, session.DefaultPropertyType, new UhtToken(), false);
 			}
 			return null;
 		}
 
-		private static readonly ThreadLocal<UhtPropertySettings> TlsPropertySettings = new ThreadLocal<UhtPropertySettings>(() => { return new UhtPropertySettings(); });
+		private static readonly ThreadLocal<UhtPropertySettings> s_tlsPropertySettings = new ThreadLocal<UhtPropertySettings>(() => { return new UhtPropertySettings(); });
 
 		/// <summary>
 		/// Given a type with children, resolve any children that couldn't be resolved during the parsing phase.
 		/// </summary>
-		/// <param name="Type">The type with children</param>
-		/// <param name="Options">Parsing options</param>
-		public static void ResolveChildren(UhtType Type, UhtPropertyParseOptions Options)
+		/// <param name="type">The type with children</param>
+		/// <param name="options">Parsing options</param>
+		public static void ResolveChildren(UhtType type, UhtPropertyParseOptions options)
 		{
-			UhtPropertyOptions PropertyOptions = UhtPropertyOptions.None;
-			if (Options.HasAnyFlags(UhtPropertyParseOptions.NoAutoConst))
+			UhtPropertyOptions propertyOptions = UhtPropertyOptions.None;
+			if (options.HasAnyFlags(UhtPropertyParseOptions.NoAutoConst))
 			{
-				PropertyOptions |= UhtPropertyOptions.NoAutoConst;
+				propertyOptions |= UhtPropertyOptions.NoAutoConst;
 			}
-			bool bInSymbolTable = Type.EngineType.AddChildrenToSymbolTable();
+			bool inSymbolTable = type.EngineType.AddChildrenToSymbolTable();
 
-			UhtPropertySettings? PropertySettings = TlsPropertySettings.Value;
-			if (PropertySettings == null)
+			UhtPropertySettings? propertySettings = s_tlsPropertySettings.Value;
+			if (propertySettings == null)
 			{
 				throw new UhtIceException("Unable to acquire threaded property settings");
 			}
 
-			for (int Index = 0; Index < Type.Children.Count; ++Index)
+			for (int index = 0; index < type.Children.Count; ++index)
 			{
-				if (Type.Children[Index] is UhtPreResolveProperty Property)
+				if (type.Children[index] is UhtPreResolveProperty property)
 				{
-					PropertySettings.Reset(Property, PropertyOptions);
-					UhtProperty? Resolved = UhtPropertyParser.ResolveProperty(UhtPropertyResolvePhase.Resolving, PropertySettings, Property.HeaderFile.Data.Memory, Property.TypeTokens);
-					if (Resolved != null)
+					propertySettings.Reset(property, propertyOptions);
+					UhtProperty? resolved = UhtPropertyParser.ResolveProperty(UhtPropertyResolvePhase.Resolving, propertySettings, property.HeaderFile.Data.Memory, property.TypeTokens);
+					if (resolved != null)
 					{
-						if (bInSymbolTable && Resolved != Property)
+						if (inSymbolTable && resolved != property)
 						{
-							Type.Session.ReplaceTypeInSymbolTable(Property, Resolved);
+							type.Session.ReplaceTypeInSymbolTable(property, resolved);
 						}
-						Type.Children[Index] = Resolved;
+						type.Children[index] = resolved;
 					}
 				}
 			}
@@ -618,243 +618,243 @@ namespace EpicGames.UHT.Parsers
 		/// <summary>
 		/// Parse a template parameter
 		/// </summary>
-		/// <param name="ResolvePhase">Resolution phase</param>
-		/// <param name="ParentPropertySettings">Parent property (container) settings</param>
-		/// <param name="ParamName">Name of the template parameter</param>
-		/// <param name="TokenReader">Token type</param>
+		/// <param name="resolvePhase">Resolution phase</param>
+		/// <param name="parentPropertySettings">Parent property (container) settings</param>
+		/// <param name="paramName">Name of the template parameter</param>
+		/// <param name="tokenReader">Token type</param>
 		/// <returns>Parsed property</returns>
-		public static UhtProperty? ParseTemplateParam(UhtPropertyResolvePhase ResolvePhase, UhtPropertySettings ParentPropertySettings, StringView ParamName, IUhtTokenReader TokenReader)
+		public static UhtProperty? ParseTemplateParam(UhtPropertyResolvePhase resolvePhase, UhtPropertySettings parentPropertySettings, StringView paramName, IUhtTokenReader tokenReader)
 		{
-			UhtSession Session = ParentPropertySettings.Outer.Session;
+			UhtSession session = parentPropertySettings.Outer.Session;
 
 			// Save the token reader state.  We need this to restore back to the start when invoking the resolve methods.
-			using (var SaveState = new UhtTokenSaveState(TokenReader))
+			using (UhtTokenSaveState saveState = new UhtTokenSaveState(tokenReader))
 			{
-				UhtPropertySettings PropertySettings = new UhtPropertySettings(ParentPropertySettings, ParamName.ToString(), TokenReader);
+				UhtPropertySettings propertySettings = new UhtPropertySettings(parentPropertySettings, paramName.ToString(), tokenReader);
 
-				UhtPropertyType PropertyType = new UhtPropertyType();
-				while (!TokenReader.bIsEOF)
+				UhtPropertyType propertyType = new UhtPropertyType();
+				while (!tokenReader.IsEOF)
 				{
-					UhtToken Token = TokenReader.GetToken();
+					UhtToken token = tokenReader.GetToken();
 
-					if (Token.IsSymbol('<') || Token.IsSymbol('>') || Token.IsSymbol(','))
+					if (token.IsSymbol('<') || token.IsSymbol('>') || token.IsSymbol(','))
 					{
 						break;
 					}
 
-					if (Session.TryGetPropertyType(Token.Value, out PropertyType))
+					if (session.TryGetPropertyType(token.Value, out propertyType))
 					{
-						if (ResolvePhase == UhtPropertyResolvePhase.Resolving || PropertyType.Options.HasAnyFlags(UhtPropertyTypeOptions.Immediate))
+						if (resolvePhase == UhtPropertyResolvePhase.Resolving || propertyType.Options.HasAnyFlags(UhtPropertyTypeOptions.Immediate))
 						{
-							SaveState.RestoreState();
-							return ResolvePropertyType(ResolvePhase, PropertySettings, TokenReader, PropertyType, Token, true);
+							saveState.RestoreState();
+							return ResolvePropertyType(resolvePhase, propertySettings, tokenReader, propertyType, token, true);
 						}
 					}
 				}
 
 				// Try the default processor.  We only use the default processor when trying to resolve something post parsing phase.
-				if (ResolvePhase == UhtPropertyResolvePhase.Resolving)
+				if (resolvePhase == UhtPropertyResolvePhase.Resolving)
 				{
-					SaveState.RestoreState();
-					return ResolvePropertyType(ResolvePhase, PropertySettings, TokenReader, Session.DefaultPropertyType, new UhtToken(), true);
+					saveState.RestoreState();
+					return ResolvePropertyType(resolvePhase, propertySettings, tokenReader, session.DefaultPropertyType, new UhtToken(), true);
 				}
 				return null;
 			}
 		}
 
-		private static UhtProperty? ResolvePropertyType(UhtPropertyResolvePhase ResolvePhase,
-			UhtPropertySettings PropertySettings, IUhtTokenReader TokenReader, UhtPropertyType PropertyType, UhtToken MatchedToken, bool bIsTemplate)
+		private static UhtProperty? ResolvePropertyType(UhtPropertyResolvePhase resolvePhase,
+			UhtPropertySettings propertySettings, IUhtTokenReader tokenReader, UhtPropertyType propertyType, UhtToken matchedToken, bool isTemplate)
 		{
-			UhtProperty? Out = PropertyType.Delegate(ResolvePhase, PropertySettings, TokenReader, MatchedToken);
-			if (Out == null)
+			UhtProperty? outProperty = propertyType.Delegate(resolvePhase, propertySettings, tokenReader, matchedToken);
+			if (outProperty == null)
 			{
 				return null;
 			}
 
 			// If this is a simple type, skip the type
-			if (PropertyType.Options.HasAnyFlags(UhtPropertyTypeOptions.Simple))
+			if (propertyType.Options.HasAnyFlags(UhtPropertyTypeOptions.Simple))
 			{
-				if (!TokenReader.SkipExpectedType(MatchedToken.Value, PropertySettings.PropertyCategory == UhtPropertyCategory.Member))
+				if (!tokenReader.SkipExpectedType(matchedToken.Value, propertySettings.PropertyCategory == UhtPropertyCategory.Member))
 				{
 					return null;
 				}
 			}
 
 			// Handle any trailing const
-			if (Out.PropertyCategory == UhtPropertyCategory.Member)
+			if (outProperty.PropertyCategory == UhtPropertyCategory.Member)
 			{
 				//@TODO: UCREMOVAL: 'const' member variables that will get written post-construction by defaultproperties
-				UhtClass? OuterClass = Out.Outer as UhtClass;
-				if (OuterClass != null && OuterClass.ClassFlags.HasAnyFlags(EClassFlags.Const))
+				UhtClass? outerClass = outProperty.Outer as UhtClass;
+				if (outerClass != null && outerClass.ClassFlags.HasAnyFlags(EClassFlags.Const))
 				{
 					// Eat a 'not quite truthful' const after the type; autogenerated for member variables of const classes.
-					if (TokenReader.TryOptional("const"))
+					if (tokenReader.TryOptional("const"))
 					{
-						Out.MetaData.Add(UhtNames.NativeConst, "");
+						outProperty.MetaData.Add(UhtNames.NativeConst, "");
 					}
 				}
 			}
 			else
 			{
-				if (TokenReader.TryOptional("const"))
+				if (tokenReader.TryOptional("const"))
 				{
-					Out.MetaData.Add(UhtNames.NativeConst, "");
-					Out.PropertyFlags |= EPropertyFlags.ConstParm;
+					outProperty.MetaData.Add(UhtNames.NativeConst, "");
+					outProperty.PropertyFlags |= EPropertyFlags.ConstParm;
 				}
 			}
 
 			// Check for unexpected '*'
-			if (TokenReader.TryOptional('*'))
+			if (tokenReader.TryOptional('*'))
 			{
-				TokenReader.LogError($"Inappropriate '*' on variable of type '{Out.GetUserFacingDecl()}', cannot have an exposed pointer to this type.");
+				tokenReader.LogError($"Inappropriate '*' on variable of type '{outProperty.GetUserFacingDecl()}', cannot have an exposed pointer to this type.");
 			}
 
 			// Arrays are passed by reference but are only implicitly so; setting it explicitly could cause a problem with replicated functions
-			if (TokenReader.TryOptional('&'))
+			if (tokenReader.TryOptional('&'))
 			{
-				switch (Out.PropertyCategory)
+				switch (outProperty.PropertyCategory)
 				{
 					case UhtPropertyCategory.RegularParameter:
 					case UhtPropertyCategory.Return:
-						Out.PropertyFlags |= EPropertyFlags.OutParm;
+						outProperty.PropertyFlags |= EPropertyFlags.OutParm;
 
 						//@TODO: UCREMOVAL: How to determine if we have a ref param?
-						if (Out.PropertyFlags.HasAnyFlags(EPropertyFlags.ConstParm))
+						if (outProperty.PropertyFlags.HasAnyFlags(EPropertyFlags.ConstParm))
 						{
-							Out.PropertyFlags |= EPropertyFlags.ReferenceParm;
+							outProperty.PropertyFlags |= EPropertyFlags.ReferenceParm;
 						}
 						break;
 
 					case UhtPropertyCategory.ReplicatedParameter:
-						Out.PropertyFlags |= EPropertyFlags.ReferenceParm;
+						outProperty.PropertyFlags |= EPropertyFlags.ReferenceParm;
 						break;
 
 					default:
 						break;
 				}
 
-				if (Out.PropertyFlags.HasAnyFlags(EPropertyFlags.ConstParm))
+				if (outProperty.PropertyFlags.HasAnyFlags(EPropertyFlags.ConstParm))
 				{
-					Out.RefQualifier = UhtPropertyRefQualifier.ConstRef;
+					outProperty.RefQualifier = UhtPropertyRefQualifier.ConstRef;
 				}
 				else
 				{
-					Out.RefQualifier = UhtPropertyRefQualifier.NonConstRef;
+					outProperty.RefQualifier = UhtPropertyRefQualifier.NonConstRef;
 				}
 			}
 
-			if (!bIsTemplate)
+			if (!isTemplate)
 			{
-				if (!TokenReader.bIsEOF)
+				if (!tokenReader.IsEOF)
 				{
-					throw new UhtTokenException(TokenReader, TokenReader.PeekToken(), "end of type declaration");
+					throw new UhtTokenException(tokenReader, tokenReader.PeekToken(), "end of type declaration");
 				}
 			}
 			else
 			{
-				ref UhtToken Token = ref TokenReader.PeekToken();
-				if (!Token.IsSymbol(',') && !Token.IsSymbol('>'))
+				ref UhtToken token = ref tokenReader.PeekToken();
+				if (!token.IsSymbol(',') && !token.IsSymbol('>'))
 				{
-					throw new UhtTokenException(TokenReader, TokenReader.PeekToken(), "end of type declaration");
+					throw new UhtTokenException(tokenReader, tokenReader.PeekToken(), "end of type declaration");
 				}
 			}
-			return Out;
+			return outProperty;
 		}
 
-		private void ParseInternal(UhtPropertyDelegate Delegate)
+		private void ParseInternal(UhtPropertyDelegate propertyDelegate)
 		{
-			this.SpecifierContext.PropertySettings.LineNumber = this.TokenReader.InputLine;
+			this._specifierContext.PropertySettings.LineNumber = this.TokenReader.InputLine;
 
-			bool bIsParamList = this.SpecifierContext.PropertySettings.PropertyCategory != UhtPropertyCategory.Member && this.TokenReader.TryOptional("UPARAM");
+			bool isParamList = this._specifierContext.PropertySettings.PropertyCategory != UhtPropertyCategory.Member && this.TokenReader.TryOptional("UPARAM");
 
-			UhtSpecifierParser Specifiers = this.TopScope.HeaderParser.GetCachedSpecifierParser(this.SpecifierContext, "Variable", 
-				bIsParamList ? this.TopScope.Session.GetSpecifierTable(UhtTableNames.PropertyArgument) : this.TopScope.Session.GetSpecifierTable(UhtTableNames.PropertyMember));
-			if (DeclarationStyle == UhtParsePropertyDeclarationStyle.UPROPERTY || bIsParamList)
+			UhtSpecifierParser specifiers = this.TopScope.HeaderParser.GetCachedSpecifierParser(this._specifierContext, "Variable",
+				isParamList ? this.TopScope.Session.GetSpecifierTable(UhtTableNames.PropertyArgument) : this.TopScope.Session.GetSpecifierTable(UhtTableNames.PropertyMember));
+			if (_declarationStyle == UhtParsePropertyDeclarationStyle.UPROPERTY || isParamList)
 			{
-				Specifiers.ParseSpecifiers();
+				specifiers.ParseSpecifiers();
 			}
 
-			if (this.SpecifierContext.PropertySettings.PropertyCategory != UhtPropertyCategory.Member)
+			if (this._specifierContext.PropertySettings.PropertyCategory != UhtPropertyCategory.Member)
 			{
 				// const before the variable type support (only for params)
 				if (this.TokenReader.TryOptional("const"))
 				{
-					this.SpecifierContext.PropertySettings.PropertyFlags |= EPropertyFlags.ConstParm;
-					this.SpecifierContext.PropertySettings.MetaData.Add(UhtNames.NativeConst, "");
+					this._specifierContext.PropertySettings.PropertyFlags |= EPropertyFlags.ConstParm;
+					this._specifierContext.PropertySettings.MetaData.Add(UhtNames.NativeConst, "");
 				}
 			}
 
-			UhtCompilerDirective CompilerDirective = this.TopScope.HeaderParser.GetCurrentCompositeCompilerDirective();
-			if (CompilerDirective.HasAnyFlags(UhtCompilerDirective.WithEditorOnlyData))
+			UhtCompilerDirective compilerDirective = this.TopScope.HeaderParser.GetCurrentCompositeCompilerDirective();
+			if (compilerDirective.HasAnyFlags(UhtCompilerDirective.WithEditorOnlyData))
 			{
-				this.SpecifierContext.PropertySettings.PropertyFlags |= EPropertyFlags.EditorOnly;
+				this._specifierContext.PropertySettings.PropertyFlags |= EPropertyFlags.EditorOnly;
 			}
-			else if (CompilerDirective.HasAnyFlags(UhtCompilerDirective.WithEditor))
+			else if (compilerDirective.HasAnyFlags(UhtCompilerDirective.WithEditor))
 			{
 				// Checking for this error is a bit tricky given legacy code.  
 				// 1) If already wrapped in WITH_EDITORONLY_DATA (see above), then we ignore the error via the else 
 				// 2) Ignore any module that is an editor module
-				UhtPackage Package = this.TopScope.HeaderParser.HeaderFile.Package;
-				UHTManifest.Module Module = Package.Module;
-				bool bIsEditorModule =
-					Module.ModuleType == UHTModuleType.EngineEditor ||
-					Module.ModuleType == UHTModuleType.GameEditor ||
-					Module.ModuleType == UHTModuleType.EngineUncooked ||
-					Module.ModuleType == UHTModuleType.GameUncooked;
-				if (this.SpecifierContext.PropertySettings.PropertyCategory == UhtPropertyCategory.Member && !bIsEditorModule)
+				UhtPackage package = this.TopScope.HeaderParser.HeaderFile.Package;
+				UHTManifest.Module module = package.Module;
+				bool isEditorModule =
+					module.ModuleType == UHTModuleType.EngineEditor ||
+					module.ModuleType == UHTModuleType.GameEditor ||
+					module.ModuleType == UHTModuleType.EngineUncooked ||
+					module.ModuleType == UHTModuleType.GameUncooked;
+				if (this._specifierContext.PropertySettings.PropertyCategory == UhtPropertyCategory.Member && !isEditorModule)
 				{
 					this.TokenReader.LogError("UProperties should not be wrapped by WITH_EDITOR, use WITH_EDITORONLY_DATA instead.");
 				}
 			}
 
 			// Store the start and end positions of the parsed type
-			int StartPos = this.TokenReader.InputPos;
+			int startPos = this.TokenReader.InputPos;
 
 			// Process the specifiers
-			if (DeclarationStyle == UhtParsePropertyDeclarationStyle.UPROPERTY || bIsParamList)
+			if (_declarationStyle == UhtParsePropertyDeclarationStyle.UPROPERTY || isParamList)
 			{
-				Specifiers.ParseDeferred();
+				specifiers.ParseDeferred();
 			}
 
 			// If we saw a BlueprintGetter but did not see BlueprintSetter or 
 			// or BlueprintReadWrite then treat as BlueprintReadOnly
-			if (this.SpecifierContext.bSeenBlueprintGetterSpecifier && !this.SpecifierContext.bSeenBlueprintWriteSpecifier)
+			if (this._specifierContext.SeenBlueprintGetterSpecifier && !this._specifierContext.SeenBlueprintWriteSpecifier)
 			{
-				this.SpecifierContext.PropertySettings.PropertyFlags |= EPropertyFlags.BlueprintReadOnly;
+				this._specifierContext.PropertySettings.PropertyFlags |= EPropertyFlags.BlueprintReadOnly;
 			}
 
-			if (this.SpecifierContext.PropertySettings.MetaData.ContainsKey(UhtNames.ExposeOnSpawn))
+			if (this._specifierContext.PropertySettings.MetaData.ContainsKey(UhtNames.ExposeOnSpawn))
 			{
-				this.SpecifierContext.PropertySettings.PropertyFlags |= EPropertyFlags.ExposeOnSpawn;
+				this._specifierContext.PropertySettings.PropertyFlags |= EPropertyFlags.ExposeOnSpawn;
 			}
 
-			UhtAccessSpecifier AccessSpecifier = TopScope.AccessSpecifier;
-			if (AccessSpecifier == UhtAccessSpecifier.Public || this.SpecifierContext.PropertySettings.PropertyCategory != UhtPropertyCategory.Member)
+			UhtAccessSpecifier accessSpecifier = TopScope.AccessSpecifier;
+			if (accessSpecifier == UhtAccessSpecifier.Public || this._specifierContext.PropertySettings.PropertyCategory != UhtPropertyCategory.Member)
 			{
-				this.SpecifierContext.PropertySettings.PropertyFlags &= ~EPropertyFlags.Protected;
-				this.SpecifierContext.PropertySettings.PropertyExportFlags |= UhtPropertyExportFlags.Public;
-				this.SpecifierContext.PropertySettings.PropertyExportFlags &= ~(UhtPropertyExportFlags.Private | UhtPropertyExportFlags.Protected);
+				this._specifierContext.PropertySettings.PropertyFlags &= ~EPropertyFlags.Protected;
+				this._specifierContext.PropertySettings.PropertyExportFlags |= UhtPropertyExportFlags.Public;
+				this._specifierContext.PropertySettings.PropertyExportFlags &= ~(UhtPropertyExportFlags.Private | UhtPropertyExportFlags.Protected);
 
-				this.SpecifierContext.PropertySettings.PropertyFlags &= ~EPropertyFlags.NativeAccessSpecifiers;
-				this.SpecifierContext.PropertySettings.PropertyFlags |= EPropertyFlags.NativeAccessSpecifierPublic;
+				this._specifierContext.PropertySettings.PropertyFlags &= ~EPropertyFlags.NativeAccessSpecifiers;
+				this._specifierContext.PropertySettings.PropertyFlags |= EPropertyFlags.NativeAccessSpecifierPublic;
 			}
-			else if (AccessSpecifier == UhtAccessSpecifier.Protected)
+			else if (accessSpecifier == UhtAccessSpecifier.Protected)
 			{
-				this.SpecifierContext.PropertySettings.PropertyFlags |= EPropertyFlags.Protected;
-				this.SpecifierContext.PropertySettings.PropertyExportFlags |= UhtPropertyExportFlags.Protected;
-				this.SpecifierContext.PropertySettings.PropertyExportFlags &= ~(UhtPropertyExportFlags.Public | UhtPropertyExportFlags.Private);
+				this._specifierContext.PropertySettings.PropertyFlags |= EPropertyFlags.Protected;
+				this._specifierContext.PropertySettings.PropertyExportFlags |= UhtPropertyExportFlags.Protected;
+				this._specifierContext.PropertySettings.PropertyExportFlags &= ~(UhtPropertyExportFlags.Public | UhtPropertyExportFlags.Private);
 
-				this.SpecifierContext.PropertySettings.PropertyFlags &= ~EPropertyFlags.NativeAccessSpecifiers;
-				this.SpecifierContext.PropertySettings.PropertyFlags |= EPropertyFlags.NativeAccessSpecifierProtected;
+				this._specifierContext.PropertySettings.PropertyFlags &= ~EPropertyFlags.NativeAccessSpecifiers;
+				this._specifierContext.PropertySettings.PropertyFlags |= EPropertyFlags.NativeAccessSpecifierProtected;
 			}
-			else if (AccessSpecifier == UhtAccessSpecifier.Private)
+			else if (accessSpecifier == UhtAccessSpecifier.Private)
 			{
-				this.SpecifierContext.PropertySettings.PropertyFlags &= ~EPropertyFlags.Protected;
-				this.SpecifierContext.PropertySettings.PropertyExportFlags |= UhtPropertyExportFlags.Private;
-				this.SpecifierContext.PropertySettings.PropertyExportFlags &= ~(UhtPropertyExportFlags.Public | UhtPropertyExportFlags.Protected);
+				this._specifierContext.PropertySettings.PropertyFlags &= ~EPropertyFlags.Protected;
+				this._specifierContext.PropertySettings.PropertyExportFlags |= UhtPropertyExportFlags.Private;
+				this._specifierContext.PropertySettings.PropertyExportFlags &= ~(UhtPropertyExportFlags.Public | UhtPropertyExportFlags.Protected);
 
-				this.SpecifierContext.PropertySettings.PropertyFlags &= ~EPropertyFlags.NativeAccessSpecifiers;
-				this.SpecifierContext.PropertySettings.PropertyFlags |= EPropertyFlags.NativeAccessSpecifierPrivate;
+				this._specifierContext.PropertySettings.PropertyFlags &= ~EPropertyFlags.NativeAccessSpecifiers;
+				this._specifierContext.PropertySettings.PropertyFlags |= EPropertyFlags.NativeAccessSpecifierPrivate;
 			}
 			else
 			{
@@ -862,7 +862,7 @@ namespace EpicGames.UHT.Parsers
 			}
 
 			// Swallow inline keywords
-			if (this.SpecifierContext.PropertySettings.PropertyCategory == UhtPropertyCategory.Return)
+			if (this._specifierContext.PropertySettings.PropertyCategory == UhtPropertyCategory.Return)
 			{
 				this.TokenReader
 					.Optional("inline")
@@ -871,21 +871,21 @@ namespace EpicGames.UHT.Parsers
 			}
 
 			// Handle MemoryLayout.h macros
-			bool bHasWrapperBrackets = false;
-			UhtLayoutMacroType LayoutMacroType = UhtLayoutMacroType.None;
-			if (this.Options.HasAnyFlags(UhtPropertyParseOptions.ParseLayoutMacro))
+			bool hasWrapperBrackets = false;
+			UhtLayoutMacroType layoutMacroType = UhtLayoutMacroType.None;
+			if (this._options.HasAnyFlags(UhtPropertyParseOptions.ParseLayoutMacro))
 			{
-				ref UhtToken LayoutToken = ref this.TokenReader.PeekToken();
-				if (LayoutToken.IsIdentifier())
+				ref UhtToken layoutToken = ref this.TokenReader.PeekToken();
+				if (layoutToken.IsIdentifier())
 				{
-					if (LayoutMacroTypes.TryGetValue(LayoutToken.Value, out LayoutMacroType))
+					if (s_layoutMacroTypes.TryGetValue(layoutToken.Value, out layoutMacroType))
 					{
 						this.TokenReader.ConsumeToken();
 						this.TokenReader.Require('(');
-						bHasWrapperBrackets = this.TokenReader.TryOptional('(');
-						if (LayoutMacroType.IsEditorOnly())
+						hasWrapperBrackets = this.TokenReader.TryOptional('(');
+						if (layoutMacroType.IsEditorOnly())
 						{
-							this.SpecifierContext.PropertySettings.PropertyFlags |= EPropertyFlags.EditorOnly;
+							this._specifierContext.PropertySettings.PropertyFlags |= EPropertyFlags.EditorOnly;
 						}
 					}
 				}
@@ -904,70 +904,70 @@ namespace EpicGames.UHT.Parsers
 			this.TokenReader.Optional("mutable");
 
 			// Gather the type tokens and possibly the property name.
-			this.TokenReader.While(GatherTypeTokensDelegate);
+			this.TokenReader.While(_gatherTypeTokensDelegate);
 
 			// Verify we at least have one type
-			if (this.CurrentTypeTokens.Count < 1)
+			if (this._currentTypeTokens.Count < 1)
 			{
-				throw new UhtException(this.TokenReader, $"{this.SpecifierContext.PropertySettings.PropertyCategory.GetHintText()}: Missing variable type or name");
+				throw new UhtException(this.TokenReader, $"{this._specifierContext.PropertySettings.PropertyCategory.GetHintText()}: Missing variable type or name");
 			}
 
 			// Consume the wrapper brackets.  This is just an extra set
-			if (bHasWrapperBrackets)
+			if (hasWrapperBrackets)
 			{
 				this.TokenReader.Require(')');
 			}
 
 			// Check for any disallowed flags
-			if (this.SpecifierContext.PropertySettings.PropertyFlags.HasAnyFlags(this.SpecifierContext.PropertySettings.DisallowPropertyFlags))
+			if (this._specifierContext.PropertySettings.PropertyFlags.HasAnyFlags(this._specifierContext.PropertySettings.DisallowPropertyFlags))
 			{
 				this.TokenReader.LogError("Specified type modifiers not allowed here");
 			}
 
-			if (this.Options.HasAnyFlags(UhtPropertyParseOptions.AddModuleRelativePath))
+			if (this._options.HasAnyFlags(UhtPropertyParseOptions.AddModuleRelativePath))
 			{
-				UhtParsingScope.AddModuleRelativePathToMetaData(this.SpecifierContext.PropertySettings.MetaData, this.SpecifierContext.Scope.ScopeType.HeaderFile);
+				UhtParsingScope.AddModuleRelativePathToMetaData(this._specifierContext.PropertySettings.MetaData, this._specifierContext.Scope.ScopeType.HeaderFile);
 			}
 
 			// Fetch the name of the property, bitfield and array size
-			if (LayoutMacroType != UhtLayoutMacroType.None)
+			if (layoutMacroType != UhtLayoutMacroType.None)
 			{
 				this.TokenReader.Require(',');
-				UhtToken NameToken = this.TokenReader.GetIdentifier();
-				if (LayoutMacroType.IsArray())
+				UhtToken nameToken = this.TokenReader.GetIdentifier();
+				if (layoutMacroType.IsArray())
 				{
 					this.TokenReader.Require(',');
-					RequireArray(this.SpecifierContext.PropertySettings, ref NameToken, ')');
+					RequireArray(this._specifierContext.PropertySettings, ref nameToken, ')');
 				}
-				else if (LayoutMacroType.IsBitfield())
+				else if (layoutMacroType.IsBitfield())
 				{
 					this.TokenReader.Require(',');
-					RequireBitfield(this.SpecifierContext.PropertySettings, ref NameToken);
+					RequireBitfield(this._specifierContext.PropertySettings, ref nameToken);
 				}
-				else if (LayoutMacroType.HasInitializer())
+				else if (layoutMacroType.HasInitializer())
 				{
 					this.TokenReader.SkipBrackets('(', ')', 1);
 				}
 				this.TokenReader.Require(')');
 
-				Finalize(ref NameToken, new ReadOnlyMemory<UhtToken>(this.CurrentTypeTokens.ToArray()), LayoutMacroType, Delegate);
+				Finalize(ref nameToken, new ReadOnlyMemory<UhtToken>(this._currentTypeTokens.ToArray()), layoutMacroType, propertyDelegate);
 			}
-			else if (this.Options.HasAnyFlags(UhtPropertyParseOptions.List))
+			else if (this._options.HasAnyFlags(UhtPropertyParseOptions.List))
 			{
 				// Extract the property name from the types
-				if (this.CurrentTypeTokens.Count < 2 || !this.CurrentTypeTokens[this.CurrentTypeTokens.Count - 1].IsIdentifier())
+				if (this._currentTypeTokens.Count < 2 || !this._currentTypeTokens[^1].IsIdentifier())
 				{
-					throw new UhtException(this.TokenReader, $"{this.SpecifierContext.PropertySettings.PropertyCategory.GetHintText()}: Expected name");
+					throw new UhtException(this.TokenReader, $"{this._specifierContext.PropertySettings.PropertyCategory.GetHintText()}: Expected name");
 				}
-				UhtToken NameToken = this.CurrentTypeTokens[this.CurrentTypeTokens.Count - 1];
-				this.CurrentTypeTokens.RemoveAt(this.CurrentTypeTokens.Count - 1);
-				CheckForOptionalParts(this.SpecifierContext.PropertySettings, ref NameToken);
+				UhtToken nameToken = this._currentTypeTokens[^1];
+				this._currentTypeTokens.RemoveAt(this._currentTypeTokens.Count - 1);
+				CheckForOptionalParts(this._specifierContext.PropertySettings, ref nameToken);
 
-				ReadOnlyMemory<UhtToken> TypeTokens = new ReadOnlyMemory<UhtToken>(this.CurrentTypeTokens.ToArray());
+				ReadOnlyMemory<UhtToken> typeTokens = new ReadOnlyMemory<UhtToken>(this._currentTypeTokens.ToArray());
 
 				while (true)
 				{
-					UhtProperty _ = Finalize(ref NameToken, TypeTokens, LayoutMacroType, Delegate);
+					UhtProperty _ = Finalize(ref nameToken, typeTokens, layoutMacroType, propertyDelegate);
 
 					// If we have reached the end
 					if (!this.TokenReader.TryOptional(','))
@@ -987,197 +987,197 @@ namespace EpicGames.UHT.Parsers
 					//CheckForOptionalParts(this.SpecifierContext.PropertySettings, ref NameToken);
 				}
 			}
-			else if (this.Options.HasAnyFlags(UhtPropertyParseOptions.CommaSeparatedName))
+			else if (this._options.HasAnyFlags(UhtPropertyParseOptions.CommaSeparatedName))
 			{
 				this.TokenReader.Require(',');
-				UhtToken NameToken = this.TokenReader.GetIdentifier();
-				CheckForOptionalParts(this.SpecifierContext.PropertySettings, ref NameToken);
-				Finalize(ref NameToken, new ReadOnlyMemory<UhtToken>(this.CurrentTypeTokens.ToArray()), LayoutMacroType, Delegate);
+				UhtToken nameToken = this.TokenReader.GetIdentifier();
+				CheckForOptionalParts(this._specifierContext.PropertySettings, ref nameToken);
+				Finalize(ref nameToken, new ReadOnlyMemory<UhtToken>(this._currentTypeTokens.ToArray()), layoutMacroType, propertyDelegate);
 			}
-			else if (this.Options.HasAnyFlags(UhtPropertyParseOptions.FunctionNameIncluded))
+			else if (this._options.HasAnyFlags(UhtPropertyParseOptions.FunctionNameIncluded))
 			{
-				UhtToken NameToken = this.CurrentTypeTokens[this.CurrentTypeTokens.Count - 1];
-				NameToken.Value = new StringView("Function");
-				if (CheckForOptionalParts(this.SpecifierContext.PropertySettings, ref NameToken))
+				UhtToken nameToken = this._currentTypeTokens[^1];
+				nameToken.Value = new StringView("Function");
+				if (CheckForOptionalParts(this._specifierContext.PropertySettings, ref nameToken))
 				{
-					NameToken = this.TokenReader.GetIdentifier("function name");
+					nameToken = this.TokenReader.GetIdentifier("function name");
 				}
 				else
 				{
-					if (this.CurrentTypeTokens.Count < 2 || !this.CurrentTypeTokens[this.CurrentTypeTokens.Count - 1].IsIdentifier())
+					if (this._currentTypeTokens.Count < 2 || !this._currentTypeTokens[^1].IsIdentifier())
 					{
-						throw new UhtException(this.TokenReader, $"{this.SpecifierContext.PropertySettings.PropertyCategory.GetHintText()}: Expected name");
+						throw new UhtException(this.TokenReader, $"{this._specifierContext.PropertySettings.PropertyCategory.GetHintText()}: Expected name");
 					}
-					NameToken = this.CurrentTypeTokens[this.CurrentTypeTokens.Count - 1];
-					this.CurrentTypeTokens.RemoveAt(this.CurrentTypeTokens.Count - 1);
+					nameToken = this._currentTypeTokens[^1];
+					this._currentTypeTokens.RemoveAt(this._currentTypeTokens.Count - 1);
 				}
-				Finalize(ref NameToken, new ReadOnlyMemory<UhtToken>(this.CurrentTypeTokens.ToArray()), LayoutMacroType, Delegate);
+				Finalize(ref nameToken, new ReadOnlyMemory<UhtToken>(this._currentTypeTokens.ToArray()), layoutMacroType, propertyDelegate);
 			}
-			else if (this.Options.HasAnyFlags(UhtPropertyParseOptions.NameIncluded))
+			else if (this._options.HasAnyFlags(UhtPropertyParseOptions.NameIncluded))
 			{
-				if (this.CurrentTypeTokens.Count < 2 || !this.CurrentTypeTokens[this.CurrentTypeTokens.Count - 1].IsIdentifier())
+				if (this._currentTypeTokens.Count < 2 || !this._currentTypeTokens[^1].IsIdentifier())
 				{
-					throw new UhtException(this.TokenReader, $"{this.SpecifierContext.PropertySettings.PropertyCategory.GetHintText()}: Expected name");
+					throw new UhtException(this.TokenReader, $"{this._specifierContext.PropertySettings.PropertyCategory.GetHintText()}: Expected name");
 				}
-				UhtToken NameToken = this.CurrentTypeTokens[this.CurrentTypeTokens.Count - 1];
-				this.CurrentTypeTokens.RemoveAt(this.CurrentTypeTokens.Count - 1);
-				CheckForOptionalParts(this.SpecifierContext.PropertySettings, ref NameToken);
-				Finalize(ref NameToken, new ReadOnlyMemory<UhtToken>(this.CurrentTypeTokens.ToArray()), LayoutMacroType, Delegate);
+				UhtToken nameToken = this._currentTypeTokens[^1];
+				this._currentTypeTokens.RemoveAt(this._currentTypeTokens.Count - 1);
+				CheckForOptionalParts(this._specifierContext.PropertySettings, ref nameToken);
+				Finalize(ref nameToken, new ReadOnlyMemory<UhtToken>(this._currentTypeTokens.ToArray()), layoutMacroType, propertyDelegate);
 			}
 			else
 			{
-				UhtToken NameToken = new UhtToken();
-				CheckForOptionalParts(this.SpecifierContext.PropertySettings, ref NameToken);
-				Finalize(ref NameToken, new ReadOnlyMemory<UhtToken>(this.CurrentTypeTokens.ToArray()), LayoutMacroType, Delegate);
+				UhtToken nameToken = new UhtToken();
+				CheckForOptionalParts(this._specifierContext.PropertySettings, ref nameToken);
+				Finalize(ref nameToken, new ReadOnlyMemory<UhtToken>(this._currentTypeTokens.ToArray()), layoutMacroType, propertyDelegate);
 			}
 		}
 
 		/// <summary>
 		/// Finish creating the property
 		/// </summary>
-		/// <param name="NameToken">The name of the property</param>
-		/// <param name="TypeTokens">Series of tokens that represent the type</param>
-		/// <param name="LayoutMacroType">Optional layout macro type being parsed</param>
-		/// <param name="Delegate">Delegate to invoke when processing has been completed</param>
+		/// <param name="nameToken">The name of the property</param>
+		/// <param name="typeTokens">Series of tokens that represent the type</param>
+		/// <param name="layoutMacroType">Optional layout macro type being parsed</param>
+		/// <param name="propertyDelegate">Delegate to invoke when processing has been completed</param>
 		/// <returns>The newly created property.  During the parsing phase, this will often be a temporary property if the type references engine types.</returns>
-		private UhtProperty Finalize(ref UhtToken NameToken, ReadOnlyMemory<UhtToken> TypeTokens, UhtLayoutMacroType LayoutMacroType, UhtPropertyDelegate Delegate)
+		private UhtProperty Finalize(ref UhtToken nameToken, ReadOnlyMemory<UhtToken> typeTokens, UhtLayoutMacroType layoutMacroType, UhtPropertyDelegate propertyDelegate)
 		{
-			this.SpecifierContext.PropertySettings.SourceName = this.SpecifierContext.PropertySettings.PropertyCategory == UhtPropertyCategory.Return ? "ReturnValue" : NameToken.Value.ToString();
+			this._specifierContext.PropertySettings.SourceName = this._specifierContext.PropertySettings.PropertyCategory == UhtPropertyCategory.Return ? "ReturnValue" : nameToken.Value.ToString();
 
 			// Try to resolve the property using any immediate mode property types
-			UhtProperty? NewProperty = ResolveProperty(UhtPropertyResolvePhase.Parsing, this.SpecifierContext.PropertySettings, this.SpecifierContext.PropertySettings.Outer.HeaderFile.Data.Memory, TypeTokens);
-			if (NewProperty == null)
+			UhtProperty? newProperty = ResolveProperty(UhtPropertyResolvePhase.Parsing, this._specifierContext.PropertySettings, this._specifierContext.PropertySettings.Outer.HeaderFile.Data.Memory, typeTokens);
+			if (newProperty == null)
 			{
-				NewProperty = new UhtPreResolveProperty(this.SpecifierContext.PropertySettings, TypeTokens);
+				newProperty = new UhtPreResolveProperty(this._specifierContext.PropertySettings, typeTokens);
 			}
 
 			// Force the category in non-engine projects
-			if (NewProperty.PropertyCategory == UhtPropertyCategory.Member)
+			if (newProperty.PropertyCategory == UhtPropertyCategory.Member)
 			{
-				if (!NewProperty.Package.bIsPartOfEngine && 
-					NewProperty.PropertyFlags.HasAnyFlags(EPropertyFlags.Edit | EPropertyFlags.BlueprintVisible) &&
-					!NewProperty.MetaData.ContainsKey(UhtNames.Category))
+				if (!newProperty.Package.IsPartOfEngine &&
+					newProperty.PropertyFlags.HasAnyFlags(EPropertyFlags.Edit | EPropertyFlags.BlueprintVisible) &&
+					!newProperty.MetaData.ContainsKey(UhtNames.Category))
 				{
-					NewProperty.MetaData.Add(UhtNames.Category, NewProperty.Outer!.EngineName);
+					newProperty.MetaData.Add(UhtNames.Category, newProperty.Outer!.EngineName);
 				}
 			}
 
 			// Check to see if the variable is deprecated, and if so set the flag
 			{
-				int DeprecatedIndex = NewProperty.SourceName.IndexOf("_DEPRECATED", StringComparison.Ordinal);
-				int NativizedPropertyPostfixIndex = NewProperty.SourceName.IndexOf("__pf", StringComparison.Ordinal); //@TODO: check OverrideNativeName in Meta Data, to be sure it's not a random occurrence of the "__pf" string.
-				bool bIgnoreDeprecatedWord = (NativizedPropertyPostfixIndex != -1) && (NativizedPropertyPostfixIndex > DeprecatedIndex);
-				if ((DeprecatedIndex != -1) && !bIgnoreDeprecatedWord)
+				int deprecatedIndex = newProperty.SourceName.IndexOf("_DEPRECATED", StringComparison.Ordinal);
+				int nativizedPropertyPostfixIndex = newProperty.SourceName.IndexOf("__pf", StringComparison.Ordinal); //@TODO: check OverrideNativeName in Meta Data, to be sure it's not a random occurrence of the "__pf" string.
+				bool ignoreDeprecatedWord = (nativizedPropertyPostfixIndex != -1) && (nativizedPropertyPostfixIndex > deprecatedIndex);
+				if ((deprecatedIndex != -1) && !ignoreDeprecatedWord)
 				{
-					if (DeprecatedIndex != NewProperty.SourceName.Length - 11)
+					if (deprecatedIndex != newProperty.SourceName.Length - 11)
 					{
 						this.TokenReader.LogError("Deprecated variables must end with _DEPRECATED");
 					}
 
 					// We allow deprecated properties in blueprints that have getters and setters assigned as they may be part of a backwards compatibility path
-					bool bBlueprintVisible = NewProperty.PropertyFlags.HasAnyFlags(EPropertyFlags.BlueprintVisible);
-					bool bWarnOnGetter = bBlueprintVisible && !NewProperty.MetaData.ContainsKey(UhtNames.BlueprintGetter);
-					bool bWarnOnSetter = bBlueprintVisible && !NewProperty.PropertyFlags.HasAnyFlags(EPropertyFlags.BlueprintReadOnly) && !NewProperty.MetaData.ContainsKey(UhtNames.BlueprintSetter);
+					bool blueprintVisible = newProperty.PropertyFlags.HasAnyFlags(EPropertyFlags.BlueprintVisible);
+					bool warnOnGetter = blueprintVisible && !newProperty.MetaData.ContainsKey(UhtNames.BlueprintGetter);
+					bool warnOnSetter = blueprintVisible && !newProperty.PropertyFlags.HasAnyFlags(EPropertyFlags.BlueprintReadOnly) && !newProperty.MetaData.ContainsKey(UhtNames.BlueprintSetter);
 
-					if (bWarnOnGetter)
+					if (warnOnGetter)
 					{
-						this.TokenReader.LogWarning($"{NewProperty.PropertyCategory.GetHintText()}: Deprecated property '{NewProperty.SourceName}' should not be marked as blueprint visible without having a BlueprintGetter");
+						this.TokenReader.LogWarning($"{newProperty.PropertyCategory.GetHintText()}: Deprecated property '{newProperty.SourceName}' should not be marked as blueprint visible without having a BlueprintGetter");
 					}
 
-					if (bWarnOnSetter)
+					if (warnOnSetter)
 					{
-						this.TokenReader.LogWarning($"{NewProperty.PropertyCategory.GetHintText()}: Deprecated property '{NewProperty.SourceName}' should not be marked as blueprint writable without having a BlueprintSetter");
+						this.TokenReader.LogWarning($"{newProperty.PropertyCategory.GetHintText()}: Deprecated property '{newProperty.SourceName}' should not be marked as blueprint writable without having a BlueprintSetter");
 					}
 
 					// Warn if a deprecated property is visible
-					if (NewProperty.PropertyFlags.HasAnyFlags(EPropertyFlags.Edit | EPropertyFlags.EditConst) || // Property is marked as editable
-						(!bBlueprintVisible && NewProperty.PropertyFlags.HasAnyFlags(EPropertyFlags.BlueprintReadOnly) && 
-						!NewProperty.PropertyExportFlags.HasAnyFlags(UhtPropertyExportFlags.ImpliedBlueprintPure))) // Is BPRO, but not via Implied Flags and not caught by Getter/Setter path above
+					if (newProperty.PropertyFlags.HasAnyFlags(EPropertyFlags.Edit | EPropertyFlags.EditConst) || // Property is marked as editable
+						(!blueprintVisible && newProperty.PropertyFlags.HasAnyFlags(EPropertyFlags.BlueprintReadOnly) &&
+						!newProperty.PropertyExportFlags.HasAnyFlags(UhtPropertyExportFlags.ImpliedBlueprintPure))) // Is BPRO, but not via Implied Flags and not caught by Getter/Setter path above
 					{
-						this.TokenReader.LogWarning($"{NewProperty.PropertyCategory.GetHintText()}: Deprecated property '{NewProperty.SourceName}' should not be marked as visible or editable");
+						this.TokenReader.LogWarning($"{newProperty.PropertyCategory.GetHintText()}: Deprecated property '{newProperty.SourceName}' should not be marked as visible or editable");
 					}
 
-					NewProperty.PropertyFlags |= EPropertyFlags.Deprecated;
-					NewProperty.EngineName = NewProperty.SourceName.Substring(0, DeprecatedIndex);
+					newProperty.PropertyFlags |= EPropertyFlags.Deprecated;
+					newProperty.EngineName = newProperty.SourceName.Substring(0, deprecatedIndex);
 				}
 			}
 
 			// Try gathering metadata for member fields
-			if (NewProperty.PropertyCategory == UhtPropertyCategory.Member)
+			if (newProperty.PropertyCategory == UhtPropertyCategory.Member)
 			{
-				UhtSpecifierParser Specifiers = this.TopScope.HeaderParser.GetCachedSpecifierParser(this.SpecifierContext, NewProperty.SourceName, this.TopScope.Session.GetSpecifierTable(UhtTableNames.PropertyMember));
-				Specifiers.ParseFieldMetaData();
-				this.TopScope.AddFormattedCommentsAsTooltipMetaData(NewProperty);
+				UhtSpecifierParser specifiers = this.TopScope.HeaderParser.GetCachedSpecifierParser(this._specifierContext, newProperty.SourceName, this.TopScope.Session.GetSpecifierTable(UhtTableNames.PropertyMember));
+				specifiers.ParseFieldMetaData();
+				this.TopScope.AddFormattedCommentsAsTooltipMetaData(newProperty);
 			}
 
-			Delegate(this.TopScope, NewProperty, ref NameToken, LayoutMacroType);
+			propertyDelegate(this.TopScope, newProperty, ref nameToken, layoutMacroType);
 
 			// Void properties don't get added when they are the return value
-			if (NewProperty.PropertyCategory != UhtPropertyCategory.Return || !this.Options.HasAnyFlags(UhtPropertyParseOptions.DontAddReturn))
+			if (newProperty.PropertyCategory != UhtPropertyCategory.Return || !this._options.HasAnyFlags(UhtPropertyParseOptions.DontAddReturn))
 			{
-				this.TopScope.ScopeType.AddChild(NewProperty);
+				this.TopScope.ScopeType.AddChild(newProperty);
 			}
-			return NewProperty;
+			return newProperty;
 		}
 
-		private bool CheckForOptionalParts(UhtPropertySettings PropertySettings, ref UhtToken NameToken)
+		private bool CheckForOptionalParts(UhtPropertySettings propertySettings, ref UhtToken nameToken)
 		{
-			bool bGotOptionalParts = false;
+			bool gotOptionalParts = false;
 
 			if (this.TokenReader.TryOptional('['))
 			{
-				RequireArray(PropertySettings, ref NameToken, ']');
+				RequireArray(propertySettings, ref nameToken, ']');
 				this.TokenReader.Require(']');
-				bGotOptionalParts = true;
+				gotOptionalParts = true;
 			}
 
 			if (this.TokenReader.TryOptional(':'))
 			{
-				RequireBitfield(PropertySettings, ref NameToken);
-				bGotOptionalParts = true;
+				RequireBitfield(propertySettings, ref nameToken);
+				gotOptionalParts = true;
 			}
-			return bGotOptionalParts;
+			return gotOptionalParts;
 		}
 
-		private void RequireBitfield(UhtPropertySettings PropertySettings, ref UhtToken NameToken)
+		private void RequireBitfield(UhtPropertySettings propertySettings, ref UhtToken nameToken)
 		{
-			int BitfieldSize;
-			if (!this.TokenReader.TryOptionalConstInt(out BitfieldSize) || BitfieldSize != 1)
+			int bitfieldSize;
+			if (!this.TokenReader.TryOptionalConstInt(out bitfieldSize) || bitfieldSize != 1)
 			{
-				throw new UhtException(this.TokenReader, $"Bad or missing bit field size for '{NameToken.Value}', must be 1.");
+				throw new UhtException(this.TokenReader, $"Bad or missing bit field size for '{nameToken.Value}', must be 1.");
 			}
-			PropertySettings.bIsBitfield = true;
+			propertySettings.IsBitfield = true;
 		}
 
-		private void RequireArray(UhtPropertySettings PropertySettings, ref UhtToken NameToken, char Terminator)
+		private void RequireArray(UhtPropertySettings propertySettings, ref UhtToken nameToken, char terminator)
 		{
 			// Ignore how the actual array dimensions are actually defined - we'll calculate those with the compiler anyway.
-			PropertySettings.ArrayDimensions = this.TokenReader.GetRawString(Terminator, UhtRawStringOptions.DontConsumeTerminator).ToString();
-			if (PropertySettings.ArrayDimensions.Length == 0)
+			propertySettings.ArrayDimensions = this.TokenReader.GetRawString(terminator, UhtRawStringOptions.DontConsumeTerminator).ToString();
+			if (propertySettings.ArrayDimensions.Length == 0)
 			{
-				throw new UhtException(this.TokenReader, $"{PropertySettings.PropertyCategory.GetHintText()} {NameToken.Value}: Missing array dimensions or terminating '{Terminator}'");
+				throw new UhtException(this.TokenReader, $"{propertySettings.PropertyCategory.GetHintText()} {nameToken.Value}: Missing array dimensions or terminating '{terminator}'");
 			}
 		}
 
-		private bool GatherTypeTokens(ref UhtToken Token)
+		private bool GatherTypeTokens(ref UhtToken token)
 		{
-			if (this.CurrentTemplateDepth == 0 && Token.IsSymbol() && (Token.IsValue(',') || Token.IsValue('(') || Token.IsValue(')') || Token.IsValue(';') || Token.IsValue('[') || Token.IsValue(':') || Token.IsValue('=') || Token.IsValue('{')))
+			if (this._currentTemplateDepth == 0 && token.IsSymbol() && (token.IsValue(',') || token.IsValue('(') || token.IsValue(')') || token.IsValue(';') || token.IsValue('[') || token.IsValue(':') || token.IsValue('=') || token.IsValue('{')))
 			{
 				return false;
 			}
 
-			this.CurrentTypeTokens.Add(Token);
-			if (Token.IsSymbol('<'))
+			this._currentTypeTokens.Add(token);
+			if (token.IsSymbol('<'))
 			{
-				++this.CurrentTemplateDepth;
+				++this._currentTemplateDepth;
 			}
-			else if (Token.IsSymbol('>'))
+			else if (token.IsSymbol('>'))
 			{
-				if (this.CurrentTemplateDepth == 0)
+				if (this._currentTemplateDepth == 0)
 				{
-					throw new UhtTokenException(TokenReader, Token, "',' or ')'");
+					throw new UhtTokenException(TokenReader, token, "',' or ')'");
 				}
-				--this.CurrentTemplateDepth;
+				--this._currentTemplateDepth;
 			}
 			return true;
 		}
@@ -1191,29 +1191,29 @@ namespace EpicGames.UHT.Parsers
 		[UhtKeyword(Extends = UhtTableNames.ScriptStruct)]
 		[SuppressMessage("CodeQuality", "IDE0051:Remove unused private members", Justification = "Attribute accessed method")]
 		[SuppressMessage("Style", "IDE0060:Remove unused parameter", Justification = "Attribute accessed method")]
-		private static UhtParseResult UPROPERTYKeyword(UhtParsingScope TopScope, UhtParsingScope ActionScope, ref UhtToken Token)
+		private static UhtParseResult UPROPERTYKeyword(UhtParsingScope topScope, UhtParsingScope actionScope, ref UhtToken token)
 		{
-			UhtPropertyParseOptions Options = UhtPropertyParseOptions.ParseLayoutMacro | UhtPropertyParseOptions.List | UhtPropertyParseOptions.AddModuleRelativePath;
-			TopScope.HeaderParser.GetCachedPropertyParser(TopScope).Parse(EPropertyFlags.ParmFlags, Options, UhtParsePropertyDeclarationStyle.UPROPERTY, UhtPropertyCategory.Member, PropertyDelegate);
-			TopScope.TokenReader.Require(';');
+			UhtPropertyParseOptions options = UhtPropertyParseOptions.ParseLayoutMacro | UhtPropertyParseOptions.List | UhtPropertyParseOptions.AddModuleRelativePath;
+			topScope.HeaderParser.GetCachedPropertyParser(topScope).Parse(EPropertyFlags.ParmFlags, options, UhtParsePropertyDeclarationStyle.UPROPERTY, UhtPropertyCategory.Member, s_propertyDelegate);
+			topScope.TokenReader.Require(';');
 			return UhtParseResult.Handled;
 		}
 		#endregion
 
-		private static readonly UhtPropertyDelegate PropertyDelegate = PropertyParsed;
+		private static readonly UhtPropertyDelegate s_propertyDelegate = PropertyParsed;
 
-		private static void PropertyParsed(UhtParsingScope TopScope, UhtProperty Property, ref UhtToken NameToken, UhtLayoutMacroType LayoutMacroType)
+		private static void PropertyParsed(UhtParsingScope topScope, UhtProperty property, ref UhtToken nameToken, UhtLayoutMacroType layoutMacroType)
 		{
-			IUhtTokenReader TokenReader = TopScope.TokenReader;
+			IUhtTokenReader tokenReader = topScope.TokenReader;
 
 			// Skip any initialization
-			if (TokenReader.TryOptional('='))
+			if (tokenReader.TryOptional('='))
 			{
-				TokenReader.SkipUntil(';');
+				tokenReader.SkipUntil(';');
 			}
-			else if (TokenReader.TryOptional('{'))
+			else if (tokenReader.TryOptional('{'))
 			{
-				TokenReader.SkipBrackets('{', '}', 1);
+				tokenReader.SkipBrackets('{', '}', 1);
 			}
 		}
 	}
@@ -1224,105 +1224,105 @@ namespace EpicGames.UHT.Parsers
 		[UhtPropertyType(Options = UhtPropertyTypeOptions.Default)]
 		[SuppressMessage("CodeQuality", "IDE0051:Remove unused private members", Justification = "Attribute accessed method")]
 		[SuppressMessage("Style", "IDE0060:Remove unused parameter", Justification = "Attribute accessed method")]
-		private static UhtProperty? DefaultProperty(UhtPropertyResolvePhase ResolvePhase, UhtPropertySettings PropertySettings, IUhtTokenReader TokenReader, UhtToken MatchedToken)
+		private static UhtProperty? DefaultProperty(UhtPropertyResolvePhase resolvePhase, UhtPropertySettings propertySettings, IUhtTokenReader tokenReader, UhtToken matchedToken)
 		{
-			UhtSession Session = PropertySettings.Outer.Session;
-			int TypeStartPos = TokenReader.PeekToken().InputStartPos;
+			UhtSession session = propertySettings.Outer.Session;
+			int typeStartPos = tokenReader.PeekToken().InputStartPos;
 
-			if (TokenReader.TryOptional("const"))
+			if (tokenReader.TryOptional("const"))
 			{
-				PropertySettings.MetaData.Add(UhtNames.NativeConst, "");
+				propertySettings.MetaData.Add(UhtNames.NativeConst, "");
 			}
 
-			UhtFindOptions FindOptions = UhtFindOptions.DelegateFunction | UhtFindOptions.Enum | UhtFindOptions.Class | UhtFindOptions.ScriptStruct;
-			if (TokenReader.TryOptional("enum"))
+			UhtFindOptions findOptions = UhtFindOptions.DelegateFunction | UhtFindOptions.Enum | UhtFindOptions.Class | UhtFindOptions.ScriptStruct;
+			if (tokenReader.TryOptional("enum"))
 			{
-				if (PropertySettings.PropertyCategory == UhtPropertyCategory.Member)
+				if (propertySettings.PropertyCategory == UhtPropertyCategory.Member)
 				{
-					TokenReader.LogError($"Cannot declare enum at variable declaration");
+					tokenReader.LogError($"Cannot declare enum at variable declaration");
 				}
-				FindOptions = UhtFindOptions.Enum;
+				findOptions = UhtFindOptions.Enum;
 			}
-			else if (TokenReader.TryOptional("class"))
+			else if (tokenReader.TryOptional("class"))
 			{
-				FindOptions = UhtFindOptions.Class;
+				findOptions = UhtFindOptions.Class;
 			}
-			else if (TokenReader.TryOptional("struct"))
+			else if (tokenReader.TryOptional("struct"))
 			{
-				FindOptions = UhtFindOptions.ScriptStruct;
+				findOptions = UhtFindOptions.ScriptStruct;
 			}
 
-			UhtTokenList Identifiers = TokenReader.GetCppIdentifier();
-			UhtType? Type = PropertySettings.Outer.FindType(UhtFindOptions.SourceName | FindOptions, Identifiers, TokenReader);
-			if (Type == null)
+			UhtTokenList identifiers = tokenReader.GetCppIdentifier();
+			UhtType? type = propertySettings.Outer.FindType(UhtFindOptions.SourceName | findOptions, identifiers, tokenReader);
+			if (type == null)
 			{
 				return null;
 			}
-			UhtTokenListCache.Return(Identifiers);
+			UhtTokenListCache.Return(identifiers);
 
-			if (Type is UhtEnum Enum)
+			if (type is UhtEnum enumObj)
 			{
-				if (PropertySettings.PropertyCategory == UhtPropertyCategory.Member)
+				if (propertySettings.PropertyCategory == UhtPropertyCategory.Member)
 				{
-					if (Enum.CppForm != UhtEnumCppForm.EnumClass)
+					if (enumObj.CppForm != UhtEnumCppForm.EnumClass)
 					{
-						TokenReader.LogError("You cannot use the raw enum name as a type for member variables, instead use TEnumAsByte or a C++11 enum class with an explicit underlying type.");
+						tokenReader.LogError("You cannot use the raw enum name as a type for member variables, instead use TEnumAsByte or a C++11 enum class with an explicit underlying type.");
 					}
 				}
 
-				return new UhtEnumProperty(PropertySettings, Enum);
+				return new UhtEnumProperty(propertySettings, enumObj);
 			}
-			else if (Type is UhtScriptStruct ScriptStruct)
+			else if (type is UhtScriptStruct scriptStruct)
 			{
-				return new UhtStructProperty(PropertySettings, ScriptStruct);
+				return new UhtStructProperty(propertySettings, scriptStruct);
 			}
-			else if (Type is UhtFunction Function)
+			else if (type is UhtFunction function)
 			{
-				if (!Function.FunctionFlags.HasAnyFlags(EFunctionFlags.MulticastDelegate))
+				if (!function.FunctionFlags.HasAnyFlags(EFunctionFlags.MulticastDelegate))
 				{
-					return new UhtDelegateProperty(PropertySettings, Function);
+					return new UhtDelegateProperty(propertySettings, function);
 				}
-				else if (Function.FunctionType == UhtFunctionType.SparseDelegate)
+				else if (function.FunctionType == UhtFunctionType.SparseDelegate)
 				{
-					return new UhtMulticastSparseDelegateProperty(PropertySettings, Function);
+					return new UhtMulticastSparseDelegateProperty(propertySettings, function);
 				}
 				else
 				{
-					return new UhtMulticastInlineDelegateProperty(PropertySettings, Function);
+					return new UhtMulticastInlineDelegateProperty(propertySettings, function);
 				}
 			}
-			else if (Type is UhtClass Class)
+			else if (type is UhtClass classObj)
 			{
 				// Const after variable type but before pointer symbol
-				if (TokenReader.TryOptional("const"))
+				if (tokenReader.TryOptional("const"))
 				{
-					PropertySettings.MetaData.Add(UhtNames.NativeConst, "");
+					propertySettings.MetaData.Add(UhtNames.NativeConst, "");
 				}
-				TokenReader.Require('*');
+				tokenReader.Require('*');
 
 				// Optionally emit messages about native pointer members and swallow trailing 'const' after pointer properties
-				UhtObjectPropertyBase.ConditionalLogPointerUsage(PropertySettings, Session.Config!.EngineNativePointerMemberBehavior,
-					Session.Config!.NonEngineNativePointerMemberBehavior, "Native pointer", TokenReader, TypeStartPos, "TObjectPtr");
+				UhtObjectPropertyBase.ConditionalLogPointerUsage(propertySettings, session.Config!.EngineNativePointerMemberBehavior,
+					session.Config!.NonEngineNativePointerMemberBehavior, "Native pointer", tokenReader, typeStartPos, "TObjectPtr");
 
-				if (PropertySettings.PropertyCategory == UhtPropertyCategory.Member)
+				if (propertySettings.PropertyCategory == UhtPropertyCategory.Member)
 				{
-					TokenReader.TryOptional("const");
+					tokenReader.TryOptional("const");
 				}
 
-				PropertySettings.PointerType = UhtPointerType.Native;
+				propertySettings.PointerType = UhtPointerType.Native;
 
-				if (Class.ClassFlags.HasAnyFlags(EClassFlags.Interface))
+				if (classObj.ClassFlags.HasAnyFlags(EClassFlags.Interface))
 				{
-					return new UhtInterfaceProperty(PropertySettings, Class);
+					return new UhtInterfaceProperty(propertySettings, classObj);
 				}
-				else if (Class.IsChildOf(Session.UClass))
+				else if (classObj.IsChildOf(session.UClass))
 				{
 					// UObject here specifies that there is no limiter
-					return new UhtClassProperty(PropertySettings, Class, Session.UObject);
+					return new UhtClassProperty(propertySettings, classObj, session.UObject);
 				}
 				else
 				{
-					return new UhtObjectProperty(PropertySettings, Class);
+					return new UhtObjectProperty(propertySettings, classObj);
 				}
 			}
 			else
diff --git a/Engine/Source/Programs/Shared/EpicGames.UHT/Parsers/UhtScriptStructParser.cs b/Engine/Source/Programs/Shared/EpicGames.UHT/Parsers/UhtScriptStructParser.cs
--- a/Engine/Source/Programs/Shared/EpicGames.UHT/Parsers/UhtScriptStructParser.cs
+++ b/Engine/Source/Programs/Shared/EpicGames.UHT/Parsers/UhtScriptStructParser.cs
@@ -1,13 +1,13 @@
 // Copyright Epic Games, Inc. All Rights Reserved.
 
+using System.Collections.Generic;
+using System.Diagnostics.CodeAnalysis;
+using System.Text.Json.Serialization;
 using EpicGames.Core;
 using EpicGames.UHT.Tables;
 using EpicGames.UHT.Tokenizer;
 using EpicGames.UHT.Types;
 using EpicGames.UHT.Utils;
-using System.Collections.Generic;
-using System.Diagnostics.CodeAnalysis;
-using System.Text.Json.Serialization;
 
 namespace EpicGames.UHT.Parsers
 {
@@ -34,248 +34,247 @@ namespace EpicGames.UHT.Parsers
 		/// <summary>
 		/// Construct a new instance
 		/// </summary>
-		/// <param name="Outer">Output type</param>
-		/// <param name="LineNumber">Line number of declaration</param>
-		public UhtScriptStructParser(UhtType Outer, int LineNumber) : base(Outer, LineNumber)
+		/// <param name="outer">Output type</param>
+		/// <param name="lineNumber">Line number of declaration</param>
+		public UhtScriptStructParser(UhtType outer, int lineNumber) : base(outer, lineNumber)
 		{
 		}
 
 		/// <inheritdoc/>
-		protected override void ResolveSuper(UhtResolvePhase ResolvePhase)
+		protected override void ResolveSuper(UhtResolvePhase resolvePhase)
 		{
-			base.ResolveSuper(ResolvePhase);
+			base.ResolveSuper(resolvePhase);
 
-			switch (ResolvePhase)
+			switch (resolvePhase)
 			{
 				case UhtResolvePhase.Bases:
 					BindAndResolveSuper(this.SuperIdentifier, UhtFindOptions.ScriptStruct);
 
 					// if we have a base struct, propagate inherited struct flags now
-					UhtScriptStruct? SuperScriptStruct = this.SuperScriptStruct;
-					if (SuperScriptStruct != null)
+					UhtScriptStruct? superScriptStruct = this.SuperScriptStruct;
+					if (superScriptStruct != null)
 					{
-						this.ScriptStructFlags |= SuperScriptStruct.ScriptStructFlags & EStructFlags.Inherit;
+						this.ScriptStructFlags |= superScriptStruct.ScriptStructFlags & EStructFlags.Inherit;
 					}
 					break;
 			}
 		}
 
 		/// <inheritdoc/>
-		protected override bool ResolveSelf(UhtResolvePhase ResolvePhase)
+		protected override bool ResolveSelf(UhtResolvePhase resolvePhase)
 		{
-			bool bResult = base.ResolveSelf(ResolvePhase);
+			bool result = base.ResolveSelf(resolvePhase);
 
-			switch (ResolvePhase)
+			switch (resolvePhase)
 			{
 				case UhtResolvePhase.Properties:
 					UhtPropertyParser.ResolveChildren(this, UhtPropertyParseOptions.AddModuleRelativePath);
 					break;
 			}
-			return bResult;
+			return result;
 		}
 
 		#region Keywords
 		[UhtKeyword(Extends = UhtTableNames.Global)]
 		[SuppressMessage("CodeQuality", "IDE0051:Remove unused private members", Justification = "Attribute accessed method")]
 		[SuppressMessage("Style", "IDE0060:Remove unused parameter", Justification = "Attribute accessed method")]
-		private static UhtParseResult USTRUCTKeyword(UhtParsingScope TopScope, UhtParsingScope ActionScope, ref UhtToken Token)
+		private static UhtParseResult USTRUCTKeyword(UhtParsingScope topScope, UhtParsingScope actionScope, ref UhtToken token)
 		{
-			return ParseUScriptStruct(TopScope, Token);
+			return ParseUScriptStruct(topScope, token);
 		}
 
 		[UhtKeyword(Extends = UhtTableNames.ScriptStruct)]
 		[UhtKeyword(Extends = UhtTableNames.ScriptStruct, Keyword = "GENERATED_BODY")]
 		[SuppressMessage("CodeQuality", "IDE0051:Remove unused private members", Justification = "Attribute accessed method")]
 		[SuppressMessage("Style", "IDE0060:Remove unused parameter", Justification = "Attribute accessed method")]
-		private static UhtParseResult GENERATED_USTRUCT_BODYKeyword(UhtParsingScope TopScope, UhtParsingScope ActionScope, ref UhtToken Token)
+		private static UhtParseResult GENERATED_USTRUCT_BODYKeyword(UhtParsingScope topScope, UhtParsingScope actionScope, ref UhtToken token)
 		{
-			UhtScriptStruct ScriptStruct = (UhtScriptStruct)TopScope.ScopeType;
+			UhtScriptStruct scriptStruct = (UhtScriptStruct)topScope.ScopeType;
 
-			if (TopScope.AccessSpecifier != UhtAccessSpecifier.Public)
+			if (topScope.AccessSpecifier != UhtAccessSpecifier.Public)
 			{
-				TopScope.TokenReader.LogError($"{Token.Value} must be in the public scope of '{ScriptStruct.SourceName}', not private or protected.");
+				topScope.TokenReader.LogError($"{token.Value} must be in the public scope of '{scriptStruct.SourceName}', not private or protected.");
 			}
 
-			if (ScriptStruct.MacroDeclaredLineNumber != -1)
+			if (scriptStruct.MacroDeclaredLineNumber != -1)
 			{
-				TopScope.TokenReader.LogError($"Multiple {Token.Value} declarations found in '{ScriptStruct.SourceName}'");
+				topScope.TokenReader.LogError($"Multiple {token.Value} declarations found in '{scriptStruct.SourceName}'");
 			}
 
-			ScriptStruct.MacroDeclaredLineNumber = TopScope.TokenReader.InputLine;
+			scriptStruct.MacroDeclaredLineNumber = topScope.TokenReader.InputLine;
 
-			UhtParserHelpers.ParseCompileVersionDeclaration(TopScope.TokenReader, TopScope.Session.Config!, ScriptStruct);
+			UhtParserHelpers.ParseCompileVersionDeclaration(topScope.TokenReader, topScope.Session.Config!, scriptStruct);
 
-			TopScope.TokenReader.Optional(';');
+			topScope.TokenReader.Optional(';');
 			return UhtParseResult.Handled;
 		}
 
 		[UhtKeyword(Extends = UhtTableNames.ScriptStruct)]
 		[SuppressMessage("CodeQuality", "IDE0051:Remove unused private members", Justification = "Attribute accessed method")]
 		[SuppressMessage("Style", "IDE0060:Remove unused parameter", Justification = "Attribute accessed method")]
-		private static UhtParseResult RIGVM_METHODKeyword(UhtParsingScope TopScope, UhtParsingScope ActionScope, ref UhtToken Token)
+		private static UhtParseResult RIGVM_METHODKeyword(UhtParsingScope topScope, UhtParsingScope actionScope, ref UhtToken token)
 		{
-			ParseRigVM(TopScope);
+			ParseRigVM(topScope);
 			return UhtParseResult.Handled;
 		}
 		#endregion
 
-		private static UhtParseResult ParseUScriptStruct(UhtParsingScope ParentScope, UhtToken KeywordToken)
+		private static UhtParseResult ParseUScriptStruct(UhtParsingScope parentScope, UhtToken keywordToken)
 		{
-			UhtScriptStructParser ScriptStruct = new UhtScriptStructParser(ParentScope.ScopeType, KeywordToken.InputLine);
-			using (var TopScope = new UhtParsingScope(ParentScope, ScriptStruct, ParentScope.Session.GetKeywordTable(UhtTableNames.ScriptStruct), UhtAccessSpecifier.Public))
+			UhtScriptStructParser scriptStruct = new UhtScriptStructParser(parentScope.ScopeType, keywordToken.InputLine);
+			using (UhtParsingScope topScope = new UhtParsingScope(parentScope, scriptStruct, parentScope.Session.GetKeywordTable(UhtTableNames.ScriptStruct), UhtAccessSpecifier.Public))
 			{
 				const string ScopeName = "struct";
 
-				using (var TokenContext = new UhtMessageContext(ScopeName))
+				using (UhtMessageContext tokenContext = new UhtMessageContext(ScopeName))
 				{
 					// Parse the specifiers
-					UhtSpecifierContext SpecifierContext = new UhtSpecifierContext(TopScope, TopScope.TokenReader, ScriptStruct.MetaData);
-					UhtSpecifierParser Specifiers = TopScope.HeaderParser.GetCachedSpecifierParser(SpecifierContext, ScopeName, ParentScope.Session.GetSpecifierTable(UhtTableNames.ScriptStruct));
-					Specifiers.ParseSpecifiers();
+					UhtSpecifierContext specifierContext = new UhtSpecifierContext(topScope, topScope.TokenReader, scriptStruct.MetaData);
+					UhtSpecifierParser specifiers = topScope.HeaderParser.GetCachedSpecifierParser(specifierContext, ScopeName, parentScope.Session.GetSpecifierTable(UhtTableNames.ScriptStruct));
+					specifiers.ParseSpecifiers();
 
 					// Consume the struct specifier
-					TopScope.TokenReader.Require("struct");
+					topScope.TokenReader.Require("struct");
 
-					TopScope.TokenReader.SkipAlignasAndDeprecatedMacroIfNecessary();
+					topScope.TokenReader.SkipAlignasAndDeprecatedMacroIfNecessary();
 
 					// Read the struct name and possible API macro name
-					UhtToken APIMacroToken;
-					TopScope.TokenReader.TryOptionalAPIMacro(out APIMacroToken);
-					ScriptStruct.SourceName = TopScope.TokenReader.GetIdentifier().Value.ToString();
+					UhtToken apiMacroToken;
+					topScope.TokenReader.TryOptionalAPIMacro(out apiMacroToken);
+					scriptStruct.SourceName = topScope.TokenReader.GetIdentifier().Value.ToString();
 
-					TopScope.AddModuleRelativePathToMetaData();
+					topScope.AddModuleRelativePathToMetaData();
 
 					// Strip the name
-					if (ScriptStruct.SourceName[0] == 'T' || ScriptStruct.SourceName[0] == 'F')
+					if (scriptStruct.SourceName[0] == 'T' || scriptStruct.SourceName[0] == 'F')
 					{
-						ScriptStruct.EngineName = ScriptStruct.SourceName.Substring(1);
+						scriptStruct.EngineName = scriptStruct.SourceName.Substring(1);
 					}
 					else
 					{
 						// This will be flagged later in the validation phase
-						ScriptStruct.EngineName = ScriptStruct.SourceName;
+						scriptStruct.EngineName = scriptStruct.SourceName;
 					}
 
 					// Check for an empty engine name
-					if (ScriptStruct.EngineName.Length == 0)
+					if (scriptStruct.EngineName.Length == 0)
 					{
-						TopScope.TokenReader.LogError($"When compiling struct definition for '{ScriptStruct.SourceName}', attempting to strip prefix results in an empty name. Did you leave off a prefix?");
+						topScope.TokenReader.LogError($"When compiling struct definition for '{scriptStruct.SourceName}', attempting to strip prefix results in an empty name. Did you leave off a prefix?");
 					}
 
 					// Parse the inheritance
-					UhtToken SuperIdentifier;
-					List<UhtToken[]>? BaseIdentifiers;
-					UhtParserHelpers.ParseInheritance(TopScope.TokenReader, TopScope.Session.Config!, out SuperIdentifier, out BaseIdentifiers);
-					ScriptStruct.SuperIdentifier = SuperIdentifier;
-					ScriptStruct.BaseIdentifiers = BaseIdentifiers;
+					UhtToken superIdentifier;
+					List<UhtToken[]>? baseIdentifiers;
+					UhtParserHelpers.ParseInheritance(topScope.TokenReader, topScope.Session.Config!, out superIdentifier, out baseIdentifiers);
+					scriptStruct.SuperIdentifier = superIdentifier;
+					scriptStruct.BaseIdentifiers = baseIdentifiers;
 
 					// Add the comments here for compatibility with old UHT
 					//COMPATIBILITY-TODO - Move this back to where the AddModuleRelativePathToMetaData is called.
-					TopScope.TokenReader.PeekToken();
-					TopScope.TokenReader.CommitPendingComments();
-					TopScope.AddFormattedCommentsAsTooltipMetaData();
+					topScope.TokenReader.PeekToken();
+					topScope.TokenReader.CommitPendingComments();
+					topScope.AddFormattedCommentsAsTooltipMetaData();
 
 					// Initialize the structure flags
-					ScriptStruct.ScriptStructFlags |= EStructFlags.Native;
+					scriptStruct.ScriptStructFlags |= EStructFlags.Native;
 
 					// Record that this struct is RequiredAPI if the CORE_API style macro was present
-					if (APIMacroToken)
+					if (apiMacroToken)
 					{
-						ScriptStruct.ScriptStructFlags |= EStructFlags.RequiredAPI;
+						scriptStruct.ScriptStructFlags |= EStructFlags.RequiredAPI;
 					}
 
 					// Process the deferred specifiers
-					Specifiers.ParseDeferred();
+					specifiers.ParseDeferred();
 
-					if (ScriptStruct.Outer != null)
+					if (scriptStruct.Outer != null)
 					{
-						ScriptStruct.Outer.AddChild(ScriptStruct);
+						scriptStruct.Outer.AddChild(scriptStruct);
 					}
 
-					TopScope.HeaderParser.ParseStatements('{', '}', true);
+					topScope.HeaderParser.ParseStatements('{', '}', true);
 
-					TopScope.TokenReader.Require(';');
+					topScope.TokenReader.Require(';');
 
-					if (ScriptStruct.MacroDeclaredLineNumber == -1 && ScriptStruct.ScriptStructFlags.HasAnyFlags(EStructFlags.Native))
+					if (scriptStruct.MacroDeclaredLineNumber == -1 && scriptStruct.ScriptStructFlags.HasAnyFlags(EStructFlags.Native))
 					{
-						TopScope.TokenReader.LogError("Expected a GENERATED_BODY() at the start of the struct");
+						topScope.TokenReader.LogError("Expected a GENERATED_BODY() at the start of the struct");
 					}
 				}
 				return UhtParseResult.Handled;
 			}
 		}
 
-		private static void ParseRigVM(UhtParsingScope TopScope)
+		private static void ParseRigVM(UhtParsingScope topScope)
 		{
-			UhtScriptStruct ScriptStruct = (UhtScriptStruct)TopScope.ScopeType;
+			UhtScriptStruct scriptStruct = (UhtScriptStruct)topScope.ScopeType;
 
-			using (var TokenContext = new UhtMessageContext("RIGVM_METHOD"))
+			using (UhtMessageContext tokenContext = new UhtMessageContext("RIGVM_METHOD"))
 			{
 
 				// Create the RigVM information if it doesn't already exist
-				UhtRigVMStructInfo? StructInfo = ScriptStruct.RigVMStructInfo;
-				if (StructInfo == null)
+				UhtRigVMStructInfo? structInfo = scriptStruct.RigVMStructInfo;
+				if (structInfo == null)
 				{
-					StructInfo = new UhtRigVMStructInfo();
-					StructInfo.Name = ScriptStruct.EngineName;
-					ScriptStruct.RigVMStructInfo = StructInfo;
+					structInfo = new UhtRigVMStructInfo();
+					structInfo.Name = scriptStruct.EngineName;
+					scriptStruct.RigVMStructInfo = structInfo;
 				}
 
 				// Create a new method information and add it
-				UhtRigVMMethodInfo MethodInfo = new UhtRigVMMethodInfo();
+				UhtRigVMMethodInfo methodInfo = new UhtRigVMMethodInfo();
 
 				// NOTE: The argument list reader doesn't handle templates with multiple arguments (i.e. the ',' issue)
-				TopScope.TokenReader
+				topScope.TokenReader
 					.RequireList('(', ')')
 					.Optional("virtual")
-					.RequireIdentifier((ref UhtToken Identifier) => MethodInfo.ReturnType = Identifier.Value.ToString())
-					.RequireIdentifier((ref UhtToken Identifier) => MethodInfo.Name = Identifier.Value.ToString());
+					.RequireIdentifier((ref UhtToken identifier) => methodInfo.ReturnType = identifier.Value.ToString())
+					.RequireIdentifier((ref UhtToken identifier) => methodInfo.Name = identifier.Value.ToString());
 
-				bool bIsGetUpgradeInfo = false;
-				bool bIsGetNextAggregateName = false;
-				if (MethodInfo.ReturnType == "FRigVMStructUpgradeInfo" && MethodInfo.Name == "GetUpgradeInfo")
+				bool isGetUpgradeInfo = false;
+				bool isGetNextAggregateName = false;
+				if (methodInfo.ReturnType == "FRigVMStructUpgradeInfo" && methodInfo.Name == "GetUpgradeInfo")
 				{
-					StructInfo.bHasGetUpgradeInfoMethod = true;
-					bIsGetUpgradeInfo = true;
+					structInfo.HasGetUpgradeInfoMethod = true;
+					isGetUpgradeInfo = true;
 				}
-				if (MethodInfo.Name == "GetNextAggregateName")
+				if (methodInfo.Name == "GetNextAggregateName")
 				{
-					StructInfo.bHasGetNextAggregateNameMethod = true;
-					bIsGetNextAggregateName = true;
+					structInfo.HasGetNextAggregateNameMethod = true;
+					isGetNextAggregateName = true;
 				}
-				
 
-				TopScope.TokenReader
-					.RequireList('(', ')', ',', false, (IEnumerable<UhtToken> Tokens) =>
+				topScope.TokenReader
+					.RequireList('(', ')', ',', false, (IEnumerable<UhtToken> tokens) =>
 					{
-						if (!bIsGetUpgradeInfo && !bIsGetNextAggregateName)
+						if (!isGetUpgradeInfo && !isGetNextAggregateName)
 						{
-							StringViewBuilder Builder = new StringViewBuilder();
-							UhtToken LastToken = new UhtToken();
-							foreach (UhtToken Token in Tokens)
+							StringViewBuilder builder = new StringViewBuilder();
+							UhtToken lastToken = new UhtToken();
+							foreach (UhtToken token in tokens)
 							{
-								if (Token.IsSymbol('='))
+								if (token.IsSymbol('='))
 								{
 									break;
 								}
-								if (LastToken)
+								if (lastToken)
 								{
-									if (Builder.Length != 0)
+									if (builder.Length != 0)
 									{
-										Builder.Append(' ');
+										builder.Append(' ');
 									}
-									Builder.Append(LastToken.Value);
+									builder.Append(lastToken.Value);
 								}
-								LastToken = Token;
+								lastToken = token;
 							}
-							MethodInfo.Parameters.Add(new UhtRigVMParameter(LastToken.Value.ToString(), Builder.ToString()));
+							methodInfo.Parameters.Add(new UhtRigVMParameter(lastToken.Value.ToString(), builder.ToString()));
 						}
 					})
 					.ConsumeUntil(';');
 
-				if (!bIsGetUpgradeInfo && !bIsGetNextAggregateName)
+				if (!isGetUpgradeInfo && !isGetNextAggregateName)
 				{
-					StructInfo.Methods.Add(MethodInfo);
+					structInfo.Methods.Add(methodInfo);
 				}
 			}
 		}
