More coding standard fixes

diff --git a/Engine/Source/Programs/Shared/EpicGames.UHT/Parsers/UhtSpecifierParser.cs b/Engine/Source/Programs/Shared/EpicGames.UHT/Parsers/UhtSpecifierParser.cs
--- a/Engine/Source/Programs/Shared/EpicGames.UHT/Parsers/UhtSpecifierParser.cs
+++ b/Engine/Source/Programs/Shared/EpicGames.UHT/Parsers/UhtSpecifierParser.cs
@@ -1,11 +1,11 @@
 // Copyright Epic Games, Inc. All Rights Reserved.
 
-using EpicGames.UHT.Utils;
 using System;
 using System.Collections.Generic;
-using EpicGames.UHT.Tables;
 using EpicGames.Core;
+using EpicGames.UHT.Tables;
 using EpicGames.UHT.Tokenizer;
+using EpicGames.UHT.Utils;
 
 namespace EpicGames.UHT.Parsers
 {
@@ -19,73 +19,73 @@ namespace EpicGames.UHT.Parsers
 	{
 		struct DeferredSpecifier
 		{
-			public UhtSpecifier Specifier;
-			public object? Value;
+			public UhtSpecifier _specifier;
+			public object? _value;
 		}
 
-		private static readonly List<KeyValuePair<StringView, StringView>> EmptyKVPValues = new List<KeyValuePair<StringView, StringView>>();
+		private static readonly List<KeyValuePair<StringView, StringView>> s_emptyKVPValues = new List<KeyValuePair<StringView, StringView>>();
 
-		private UhtSpecifierContext SpecifierContext;
-		private IUhtTokenReader TokenReader;
-		private UhtSpecifierTable Table;
-		private StringView Context;
-		private StringView CurrentSpecifier = new StringView();
-		private List<DeferredSpecifier> DeferredSpecifiers = new List<DeferredSpecifier>();
-		private bool bIsParsingFieldMetaData = false;
-		private List<KeyValuePair<StringView, StringView>>? CurrentKVPValues = null;
-		private List<StringView>? CurrentStringValues = null;
-		private int UMetaElementsParsed = 0;
+		private UhtSpecifierContext _specifierContext;
+		private IUhtTokenReader _tokenReader;
+		private UhtSpecifierTable _table;
+		private StringView _context;
+		private StringView _currentSpecifier = new StringView();
+		private List<DeferredSpecifier> _deferredSpecifiers = new List<DeferredSpecifier>();
+		private bool _isParsingFieldMetaData = false;
+		private List<KeyValuePair<StringView, StringView>>? _currentKVPValues = null;
+		private List<StringView>? _currentStringValues = null;
+		private int _umetaElementsParsed = 0;
 
-		private readonly Action ParseAction;
-		private readonly Action ParseFieldMetaDataAction;
-		private readonly Action ParseKVPValueAction;
-		private readonly Action ParseStringViewListAction;
+		private readonly Action _parseAction;
+		private readonly Action _parseFieldMetaDataAction;
+		private readonly Action _parseKVPValueAction;
+		private readonly Action _parseStringViewListAction;
 
 		/// <summary>
 		/// Construct a new specifier parser
 		/// </summary>
-		/// <param name="SpecifierContext">Specifier context</param>
-		/// <param name="Context">User facing context added to messages</param>
-		/// <param name="Table">Specifier table</param>
-		public UhtSpecifierParser(UhtSpecifierContext SpecifierContext, StringView Context, UhtSpecifierTable Table)
+		/// <param name="specifierContext">Specifier context</param>
+		/// <param name="context">User facing context added to messages</param>
+		/// <param name="table">Specifier table</param>
+		public UhtSpecifierParser(UhtSpecifierContext specifierContext, StringView context, UhtSpecifierTable table)
 		{
-			this.SpecifierContext = SpecifierContext;
-			this.TokenReader = SpecifierContext.Scope.TokenReader;
-			this.Context = Context;
-			this.Table = Table;
-
-			this.ParseAction = ParseInternal;
-			this.ParseFieldMetaDataAction = ParseFieldMetaDataInternal;
-			this.ParseKVPValueAction = () => 
+			this._specifierContext = specifierContext;
+			this._tokenReader = specifierContext.Scope.TokenReader;
+			this._context = context;
+			this._table = table;
+
+			this._parseAction = ParseInternal;
+			this._parseFieldMetaDataAction = ParseFieldMetaDataInternal;
+			this._parseKVPValueAction = () =>
 			{
-				if (this.CurrentKVPValues == null)
+				if (this._currentKVPValues == null)
 				{
-					this.CurrentKVPValues = new List<KeyValuePair<StringView, StringView>>();
+					this._currentKVPValues = new List<KeyValuePair<StringView, StringView>>();
 				}
-				this.CurrentKVPValues.Add(ReadKVP()); 
+				this._currentKVPValues.Add(ReadKVP());
 			};
-			this.ParseStringViewListAction = () => 
-			{ 
-				if (this.CurrentStringValues == null)
+			this._parseStringViewListAction = () =>
+			{
+				if (this._currentStringValues == null)
 				{
-					this.CurrentStringValues = new List<StringView>();
+					this._currentStringValues = new List<StringView>();
 				}
-				this.CurrentStringValues.Add(ReadValue()); 
+				this._currentStringValues.Add(ReadValue());
 			};
 		}
 
 		/// <summary>
 		/// Reset an existing parser to parse a new specifier block
 		/// </summary>
-		/// <param name="SpecifierContext">Specifier context</param>
-		/// <param name="Context">User facing context added to messages</param>
-		/// <param name="Table">Specifier table</param>
-		public void Reset(UhtSpecifierContext SpecifierContext, StringView Context, UhtSpecifierTable Table)
+		/// <param name="specifierContext">Specifier context</param>
+		/// <param name="context">User facing context added to messages</param>
+		/// <param name="table">Specifier table</param>
+		public void Reset(UhtSpecifierContext specifierContext, StringView context, UhtSpecifierTable table)
 		{
-			this.SpecifierContext = SpecifierContext;
-			this.TokenReader = SpecifierContext.Scope.TokenReader;
-			this.Context = Context;
-			this.Table = Table;
+			this._specifierContext = specifierContext;
+			this._tokenReader = specifierContext.Scope.TokenReader;
+			this._context = context;
+			this._table = table;
 		}
 
 		/// <summary>
@@ -94,12 +94,12 @@ namespace EpicGames.UHT.Parsers
 		/// <returns>The parser</returns>
 		public UhtSpecifierParser ParseSpecifiers()
 		{
-			this.bIsParsingFieldMetaData = false;
-			this.SpecifierContext.MetaData.LineNumber = TokenReader.InputLine;
+			this._isParsingFieldMetaData = false;
+			this._specifierContext.MetaData.LineNumber = _tokenReader.InputLine;
 
-			using (var TokenContext = new UhtMessageContext(this))
+			using (UhtMessageContext tokenContext = new UhtMessageContext(this))
 			{
-				this.TokenReader.RequireList('(', ')', ',', false, this.ParseAction);
+				this._tokenReader.RequireList('(', ')', ',', false, this._parseAction);
 			}
 			return this;
 		}
@@ -110,18 +110,18 @@ namespace EpicGames.UHT.Parsers
 		/// <returns>Specifier parser</returns>
 		public UhtSpecifierParser ParseFieldMetaData()
 		{
-			this.TokenReader = this.SpecifierContext.Scope.TokenReader;
-			this.bIsParsingFieldMetaData = true;
+			this._tokenReader = this._specifierContext.Scope.TokenReader;
+			this._isParsingFieldMetaData = true;
 
-			using (var TokenContext = new UhtMessageContext(this))
+			using (UhtMessageContext tokenContext = new UhtMessageContext(this))
 			{
-				if (TokenReader.TryOptional("UMETA"))
+				if (_tokenReader.TryOptional("UMETA"))
 				{
-					this.UMetaElementsParsed = 0;
-					TokenReader.RequireList('(', ')', ',', false, this.ParseFieldMetaDataAction);
-					if (this.UMetaElementsParsed == 0)
+					this._umetaElementsParsed = 0;
+					_tokenReader.RequireList('(', ')', ',', false, this._parseFieldMetaDataAction);
+					if (this._umetaElementsParsed == 0)
 					{
-						this.TokenReader.LogError($"No metadata specified while parsing {UhtMessage.FormatContext(this)}");
+						this._tokenReader.LogError($"No metadata specified while parsing {UhtMessage.FormatContext(this)}");
 					}
 				}
 			}
@@ -133,11 +133,11 @@ namespace EpicGames.UHT.Parsers
 		/// </summary>
 		public void ParseDeferred()
 		{
-			foreach (var Deferred in this.DeferredSpecifiers)
+			foreach (DeferredSpecifier deferred in this._deferredSpecifiers)
 			{
-				Dispatch(Deferred.Specifier, Deferred.Value);
+				Dispatch(deferred._specifier, deferred._value);
 			}
-			this.DeferredSpecifiers.Clear();
+			this._deferredSpecifiers.Clear();
 		}
 
 		#region IMessageExtraContext implementation
@@ -147,180 +147,180 @@ namespace EpicGames.UHT.Parsers
 		{
 			get
 			{
-				Stack<object?> ExtraContext = new Stack<object?>(1);
-				string What = this.bIsParsingFieldMetaData ? "metadata" : "specifiers";
-				if (this.Context.Span.Length > 0)
+				Stack<object?> extraContext = new Stack<object?>(1);
+				string what = this._isParsingFieldMetaData ? "metadata" : "specifiers";
+				if (this._context.Span.Length > 0)
 				{
-					ExtraContext.Push($"{this.Context} {What}");
+					extraContext.Push($"{this._context} {what}");
 				}
 				else
 				{
-					ExtraContext.Push(What);
+					extraContext.Push(what);
 				}
-				return ExtraContext;
+				return extraContext;
 			}
 		}
 		#endregion
 
 		private void ParseInternal()
 		{
-			UhtToken Identifier = this.TokenReader.GetIdentifier();
+			UhtToken identifier = this._tokenReader.GetIdentifier();
 
-			this.CurrentSpecifier = Identifier.Value;
-			UhtSpecifier? Specifier;
-			if (this.Table.TryGetValue(CurrentSpecifier, out Specifier))
+			this._currentSpecifier = identifier.Value;
+			UhtSpecifier? specifier;
+			if (this._table.TryGetValue(_currentSpecifier, out specifier))
 			{
-				if (TryParseValue(Specifier.ValueType, out object? Value))
+				if (TryParseValue(specifier.ValueType, out object? value))
 				{
-					if (Specifier.When == UhtSpecifierWhen.Deferred)
+					if (specifier.When == UhtSpecifierWhen.Deferred)
 					{
-						if (this.DeferredSpecifiers == null)
+						if (this._deferredSpecifiers == null)
 						{
-							this.DeferredSpecifiers = new List<DeferredSpecifier>();
+							this._deferredSpecifiers = new List<DeferredSpecifier>();
 						}
-						this.DeferredSpecifiers.Add(new DeferredSpecifier { Specifier = Specifier, Value = Value });
+						this._deferredSpecifiers.Add(new DeferredSpecifier { _specifier = specifier, _value = value });
 					}
 					else
 					{
-						Dispatch(Specifier, Value);
+						Dispatch(specifier, value);
 					}
 				}
 			}
 			else
 			{
-				this.TokenReader.LogError($"Unknown specifier '{CurrentSpecifier}' found while parsing {UhtMessage.FormatContext(this)}");
+				this._tokenReader.LogError($"Unknown specifier '{_currentSpecifier}' found while parsing {UhtMessage.FormatContext(this)}");
 			}
 		}
 
 		private void ParseFieldMetaDataInternal()
 		{
-			UhtToken Key;
-			if (!TokenReader.TryOptionalIdentifier(out Key))
+			UhtToken key;
+			if (!_tokenReader.TryOptionalIdentifier(out key))
 			{
-				throw new UhtException(this.TokenReader, $"UMETA expects a key and optional value", this);
+				throw new UhtException(this._tokenReader, $"UMETA expects a key and optional value", this);
 			}
 
-			StringViewBuilder SVB = new StringViewBuilder();
-			if (TokenReader.TryOptional('='))
+			StringViewBuilder builder = new StringViewBuilder();
+			if (_tokenReader.TryOptional('='))
 			{
-				if (!ReadValue(TokenReader, SVB, true))
+				if (!ReadValue(_tokenReader, builder, true))
 				{
-					throw new UhtException(this.TokenReader, $"UMETA key '{Key.Value}' expects a value", this);
+					throw new UhtException(this._tokenReader, $"UMETA key '{key.Value}' expects a value", this);
 				}
 			}
-			
-			++this.UMetaElementsParsed;
-			this.SpecifierContext.MetaData.Add(Key.Value.ToString(), this.SpecifierContext.MetaNameIndex, SVB.ToString());
+
+			++this._umetaElementsParsed;
+			this._specifierContext.MetaData.Add(key.Value.ToString(), this._specifierContext.MetaNameIndex, builder.ToString());
 		}
 
-		private void Dispatch(UhtSpecifier Specifier, object? Value)
+		private void Dispatch(UhtSpecifier specifier, object? value)
 		{
-			UhtSpecifierDispatchResults Results = Specifier.Dispatch(this.SpecifierContext, Value);
-			if (Results == UhtSpecifierDispatchResults.Unknown)
+			UhtSpecifierDispatchResults results = specifier.Dispatch(this._specifierContext, value);
+			if (results == UhtSpecifierDispatchResults.Unknown)
 			{
-				this.TokenReader.LogError($"Unknown specifier '{Specifier.Name}' found while parsing {UhtMessage.FormatContext(this)}");
+				this._tokenReader.LogError($"Unknown specifier '{specifier.Name}' found while parsing {UhtMessage.FormatContext(this)}");
 			}
 		}
 
-		private bool TryParseValue(UhtSpecifierValueType ValueType, out object? Value)
+		private bool TryParseValue(UhtSpecifierValueType valueType, out object? value)
 		{
-			Value = null;
+			value = null;
 
-			switch (ValueType)
+			switch (valueType)
 			{
 				case UhtSpecifierValueType.NotSet:
 					throw new UhtIceException("NotSet is an invalid value for value types");
 
 				case UhtSpecifierValueType.None:
-					if (this.TokenReader.TryOptional('='))
+					if (this._tokenReader.TryOptional('='))
 					{
 						ReadValue(); // consume the value;
-						this.TokenReader.LogError($"The specifier '{this.CurrentSpecifier}' found a value when none was expected", this);
+						this._tokenReader.LogError($"The specifier '{this._currentSpecifier}' found a value when none was expected", this);
 						return false;
 					}
 					return true;
 
 				case UhtSpecifierValueType.String:
-					if (!this.TokenReader.TryOptional('='))
+					if (!this._tokenReader.TryOptional('='))
 					{
-						this.TokenReader.LogError($"The specifier '{this.CurrentSpecifier}' expects a value", this);
+						this._tokenReader.LogError($"The specifier '{this._currentSpecifier}' expects a value", this);
 						return false;
 					}
-					Value = ReadValue();
+					value = ReadValue();
 					return true;
 
 				case UhtSpecifierValueType.OptionalString:
 					{
-						List<StringView>? StringList = ReadValueList();
-						if (StringList != null && StringList.Count > 0)
+						List<StringView>? stringList = ReadValueList();
+						if (stringList != null && stringList.Count > 0)
 						{
-							Value = StringList[0];
+							value = stringList[0];
 						}
 						return true;
 					}
 
 				case UhtSpecifierValueType.SingleString:
 					{
-						List<StringView>? StringList = ReadValueList();
-						if (StringList == null || StringList.Count != 1)
+						List<StringView>? stringList = ReadValueList();
+						if (stringList == null || stringList.Count != 1)
 						{
-							this.TokenReader.LogError($"The specifier '{this.CurrentSpecifier}' expects a single value", this);
+							this._tokenReader.LogError($"The specifier '{this._currentSpecifier}' expects a single value", this);
 							return false;
 						}
-						Value = StringList[0];
+						value = stringList[0];
 						return true;
 					}
 
 				case UhtSpecifierValueType.StringList:
-					Value = ReadValueList();
+					value = ReadValueList();
 					return true;
 
 				case UhtSpecifierValueType.Legacy:
-					Value = ReadValueList();
+					value = ReadValueList();
 					return true;
 
 				case UhtSpecifierValueType.NonEmptyStringList:
 					{
-						List<StringView>? StringList = ReadValueList();
-						if (StringList == null || StringList.Count == 0)
+						List<StringView>? stringList = ReadValueList();
+						if (stringList == null || stringList.Count == 0)
 						{
-							this.TokenReader.LogError($"The specifier '{this.CurrentSpecifier}' expects at list one value", this);
+							this._tokenReader.LogError($"The specifier '{this._currentSpecifier}' expects at list one value", this);
 							return false;
 						}
-						Value = StringList;
+						value = stringList;
 						return true;
 					}
 
 				case UhtSpecifierValueType.KeyValuePairList:
 					{
-						this.CurrentKVPValues = null;
-						this.TokenReader
+						this._currentKVPValues = null;
+						this._tokenReader
 							.Require('=')
-							.RequireList('(', ')', ',', false, this.ParseKVPValueAction);
-						List<KeyValuePair<StringView, StringView>> Return = this.CurrentKVPValues ?? EmptyKVPValues;
-						this.CurrentKVPValues = null;
-						Value = Return;
+							.RequireList('(', ')', ',', false, this._parseKVPValueAction);
+						List<KeyValuePair<StringView, StringView>> kvps = this._currentKVPValues ?? s_emptyKVPValues;
+						this._currentKVPValues = null;
+						value = kvps;
 						return true;
 					}
 
 				case UhtSpecifierValueType.OptionalEqualsKeyValuePairList:
 					{
-						this.CurrentKVPValues = null;
+						this._currentKVPValues = null;
 						// This parser isn't as strict as the other parsers...
-						if (this.TokenReader.TryOptional('='))
+						if (this._tokenReader.TryOptional('='))
 						{
-							if (!this.TokenReader.TryOptionalList('(', ')', ',', false, this.ParseKVPValueAction))
+							if (!this._tokenReader.TryOptionalList('(', ')', ',', false, this._parseKVPValueAction))
 							{
-								this.ParseKVPValueAction();
+								this._parseKVPValueAction();
 							}
 						}
 						else
 						{
-							this.TokenReader.TryOptionalList('(', ')', ',', false, this.ParseKVPValueAction);
+							this._tokenReader.TryOptionalList('(', ')', ',', false, this._parseKVPValueAction);
 						}
-						List<KeyValuePair<StringView, StringView>> Return = this.CurrentKVPValues ?? EmptyKVPValues;
-						this.CurrentKVPValues = null;
-						Value = Return;
+						List<KeyValuePair<StringView, StringView>> kvps = this._currentKVPValues ?? s_emptyKVPValues;
+						this._currentKVPValues = null;
+						value = kvps;
 						return true;
 					}
 
@@ -331,62 +331,62 @@ namespace EpicGames.UHT.Parsers
 
 		private KeyValuePair<StringView, StringView> ReadKVP()
 		{
-			UhtToken Key;
-			if (!this.TokenReader.TryOptionalIdentifier(out Key))
+			UhtToken key;
+			if (!this._tokenReader.TryOptionalIdentifier(out key))
 			{
-				throw new UhtException(this.TokenReader, $"The specifier '{this.CurrentSpecifier}' expects a key and optional value", this);
+				throw new UhtException(this._tokenReader, $"The specifier '{this._currentSpecifier}' expects a key and optional value", this);
 			}
 
-			StringView Value = "";
-			if (this.TokenReader.TryOptional('='))
+			StringView value = "";
+			if (this._tokenReader.TryOptional('='))
 			{
-				Value = ReadValue();
+				value = ReadValue();
 			}
-			return new KeyValuePair<StringView, StringView>(Key.Value, Value);
+			return new KeyValuePair<StringView, StringView>(key.Value, value);
 		}
 
 		private List<StringView>? ReadValueList()
 		{
-			this.CurrentStringValues = null;
+			this._currentStringValues = null;
 
 			// This parser isn't as strict as the other parsers...
-			if (this.TokenReader.TryOptional('='))
+			if (this._tokenReader.TryOptional('='))
 			{
-				if (!this.TokenReader.TryOptionalList('(', ')', ',', false, this.ParseStringViewListAction))
+				if (!this._tokenReader.TryOptionalList('(', ')', ',', false, this._parseStringViewListAction))
 				{
-					this.ParseStringViewListAction();
+					this._parseStringViewListAction();
 				}
 			}
 			else
 			{
-				this.TokenReader.TryOptionalList('(', ')', ',', false, this.ParseStringViewListAction);
+				this._tokenReader.TryOptionalList('(', ')', ',', false, this._parseStringViewListAction);
 			}
-			List<StringView>? StringValues = this.CurrentStringValues;
-			this.CurrentStringValues = null;
-			return StringValues;
+			List<StringView>? stringValues = this._currentStringValues;
+			this._currentStringValues = null;
+			return stringValues;
 		}
 
 		private StringView ReadValue()
 		{
-			StringViewBuilder SVB = new StringViewBuilder();
-			if (!ReadValue(this.TokenReader, SVB, false))
+			StringViewBuilder builder = new StringViewBuilder();
+			if (!ReadValue(this._tokenReader, builder, false))
 			{
-				throw new UhtException(this.TokenReader, $"The specifier '{this.CurrentSpecifier}' expects a value", this);
+				throw new UhtException(this._tokenReader, $"The specifier '{this._currentSpecifier}' expects a value", this);
 			}
-			return SVB.ToStringView();
+			return builder.ToStringView();
 		}
 
 		/// <summary>
 		/// Parse the sequence of meta data
 		/// </summary>
-		/// <param name="TokenReader">Input token reader</param>
-		/// <param name="SVB">Output string builder</param>
-		/// <param name="bRespectQuotes">If true, do not convert \" to " in string constants.  This is required for UMETA data</param>
+		/// <param name="tokenReader">Input token reader</param>
+		/// <param name="builder">Output string builder</param>
+		/// <param name="respectQuotes">If true, do not convert \" to " in string constants.  This is required for UMETA data</param>
 		/// <returns>True if data was read</returns>
-		private static bool ReadValue(IUhtTokenReader TokenReader, StringViewBuilder SVB, bool bRespectQuotes)
+		private static bool ReadValue(IUhtTokenReader tokenReader, StringViewBuilder builder, bool respectQuotes)
 		{
-			UhtToken Token = TokenReader.GetToken();
-			switch (Token.TokenType)
+			UhtToken token = tokenReader.GetToken();
+			switch (token.TokenType)
 			{
 				case UhtTokenType.EndOfFile:
 				case UhtTokenType.EndOfDefault:
@@ -396,22 +396,22 @@ namespace EpicGames.UHT.Parsers
 
 				case UhtTokenType.Identifier:
 					// We handle true/false differently for compatibility with old UHT
-					if (Token.IsValue("true", true))
+					if (token.IsValue("true", true))
 					{
-						SVB.Append("TRUE");
+						builder.Append("TRUE");
 					}
-					else if (Token.IsValue("false", true))
+					else if (token.IsValue("false", true))
 					{
-						SVB.Append("FALSE");
+						builder.Append("FALSE");
 					}
 					else
 					{
-						SVB.Append(Token.Value);
+						builder.Append(token.Value);
 					}
-					if (TokenReader.TryOptional('='))
+					if (tokenReader.TryOptional('='))
 					{
-						SVB.Append('=');
-						if (!ReadValue(TokenReader, SVB, bRespectQuotes))
+						builder.Append('=');
+						if (!ReadValue(tokenReader, builder, respectQuotes))
 						{
 							return false;
 						}
@@ -419,11 +419,11 @@ namespace EpicGames.UHT.Parsers
 					break;
 
 				case UhtTokenType.Symbol:
-					SVB.Append(Token.Value);
-					if (TokenReader.TryOptional('='))
+					builder.Append(token.Value);
+					if (tokenReader.TryOptional('='))
 					{
-						SVB.Append('=');
-						if (!ReadValue(TokenReader, SVB, bRespectQuotes))
+						builder.Append('=');
+						if (!ReadValue(tokenReader, builder, respectQuotes))
 						{
 							return false;
 						}
@@ -431,7 +431,7 @@ namespace EpicGames.UHT.Parsers
 					break;
 
 				default:
-					SVB.Append(Token.GetConstantValue(bRespectQuotes));
+					builder.Append(token.GetConstantValue(respectQuotes));
 					break;
 			}
 
diff --git a/Engine/Source/Programs/Shared/EpicGames.UHT/Specifiers/UhtClassSpecifiers.cs b/Engine/Source/Programs/Shared/EpicGames.UHT/Specifiers/UhtClassSpecifiers.cs
--- a/Engine/Source/Programs/Shared/EpicGames.UHT/Specifiers/UhtClassSpecifiers.cs
+++ b/Engine/Source/Programs/Shared/EpicGames.UHT/Specifiers/UhtClassSpecifiers.cs
@@ -1,11 +1,11 @@
 // Copyright Epic Games, Inc. All Rights Reserved.
 
+using System.Collections.Generic;
+using System.Diagnostics.CodeAnalysis;
 using EpicGames.Core;
 using EpicGames.UHT.Tables;
 using EpicGames.UHT.Types;
 using EpicGames.UHT.Utils;
-using System.Collections.Generic;
-using System.Diagnostics.CodeAnalysis;
 
 namespace EpicGames.UHT.Parsers
 {
@@ -17,381 +17,381 @@ namespace EpicGames.UHT.Parsers
 	{
 		[UhtSpecifier(Extends = UhtTableNames.Class, ValueType = UhtSpecifierValueType.Legacy)]
 		[SuppressMessage("CodeQuality", "IDE0051:Remove unused private members", Justification = "Attribute accessed method")]
-		private static void NoExportSpecifier(UhtSpecifierContext SpecifierContext)
+		private static void NoExportSpecifier(UhtSpecifierContext specifierContext)
 		{
-			UhtClassParser Class = (UhtClassParser)SpecifierContext.Scope.ScopeType;
-			Class.ClassExportFlags |= UhtClassExportFlags.NoExport;
+			UhtClassParser classObj = (UhtClassParser)specifierContext.Scope.ScopeType;
+			classObj.ClassExportFlags |= UhtClassExportFlags.NoExport;
 		}
 
 		[UhtSpecifier(Extends = UhtTableNames.Class, ValueType = UhtSpecifierValueType.Legacy)]
 		[SuppressMessage("CodeQuality", "IDE0051:Remove unused private members", Justification = "Attribute accessed method")]
-		private static void IntrinsicSpecifier(UhtSpecifierContext SpecifierContext)
+		private static void IntrinsicSpecifier(UhtSpecifierContext specifierContext)
 		{
-			UhtClassParser Class = (UhtClassParser)SpecifierContext.Scope.ScopeType;
-			Class.AddClassFlags(EClassFlags.Intrinsic);
+			UhtClassParser classObj = (UhtClassParser)specifierContext.Scope.ScopeType;
+			classObj.AddClassFlags(EClassFlags.Intrinsic);
 		}
 
 		[UhtSpecifier(Extends = UhtTableNames.Class, ValueType = UhtSpecifierValueType.Legacy)]
 		[SuppressMessage("CodeQuality", "IDE0051:Remove unused private members", Justification = "Attribute accessed method")]
-		private static void ComponentWrapperClassSpecifier(UhtSpecifierContext SpecifierContext)
+		private static void ComponentWrapperClassSpecifier(UhtSpecifierContext specifierContext)
 		{
-			SpecifierContext.MetaData.Add(UhtNames.IgnoreCategoryKeywordsInSubclasses, true);
+			specifierContext.MetaData.Add(UhtNames.IgnoreCategoryKeywordsInSubclasses, true);
 		}
 
 		[UhtSpecifier(Extends = UhtTableNames.Class, ValueType = UhtSpecifierValueType.SingleString)]
 		[SuppressMessage("CodeQuality", "IDE0051:Remove unused private members", Justification = "Attribute accessed method")]
-		private static void WithinSpecifier(UhtSpecifierContext SpecifierContext, StringView Value)
+		private static void WithinSpecifier(UhtSpecifierContext specifierContext, StringView value)
 		{
-			UhtClassParser Class = (UhtClassParser)SpecifierContext.Scope.ScopeType;
-			Class.ClassWithinIdentifier = Value.ToString();
+			UhtClassParser classObj = (UhtClassParser)specifierContext.Scope.ScopeType;
+			classObj.ClassWithinIdentifier = value.ToString();
 		}
 
 		[UhtSpecifier(Extends = UhtTableNames.Class, ValueType = UhtSpecifierValueType.Legacy)]
 		[SuppressMessage("CodeQuality", "IDE0051:Remove unused private members", Justification = "Attribute accessed method")]
-		private static void EditInlineNewSpecifier(UhtSpecifierContext SpecifierContext)
+		private static void EditInlineNewSpecifier(UhtSpecifierContext specifierContext)
 		{
-			UhtClassParser Class = (UhtClassParser)SpecifierContext.Scope.ScopeType;
-			Class.AddClassFlags(EClassFlags.EditInlineNew);
+			UhtClassParser classObj = (UhtClassParser)specifierContext.Scope.ScopeType;
+			classObj.AddClassFlags(EClassFlags.EditInlineNew);
 		}
 
 		[UhtSpecifier(Extends = UhtTableNames.Class, ValueType = UhtSpecifierValueType.Legacy)]
 		[SuppressMessage("CodeQuality", "IDE0051:Remove unused private members", Justification = "Attribute accessed method")]
-		private static void NotEditInlineNewSpecifier(UhtSpecifierContext SpecifierContext)
+		private static void NotEditInlineNewSpecifier(UhtSpecifierContext specifierContext)
 		{
-			UhtClassParser Class = (UhtClassParser)SpecifierContext.Scope.ScopeType;
-			Class.RemoveClassFlags(EClassFlags.EditInlineNew);
+			UhtClassParser classObj = (UhtClassParser)specifierContext.Scope.ScopeType;
+			classObj.RemoveClassFlags(EClassFlags.EditInlineNew);
 		}
 
 		[UhtSpecifier(Extends = UhtTableNames.Class, ValueType = UhtSpecifierValueType.Legacy)]
 		[SuppressMessage("CodeQuality", "IDE0051:Remove unused private members", Justification = "Attribute accessed method")]
-		private static void PlaceableSpecifier(UhtSpecifierContext SpecifierContext)
+		private static void PlaceableSpecifier(UhtSpecifierContext specifierContext)
 		{
-			UhtClassParser Class = (UhtClassParser)SpecifierContext.Scope.ScopeType;
-			Class.RemoveClassFlags(EClassFlags.NotPlaceable);
+			UhtClassParser classObj = (UhtClassParser)specifierContext.Scope.ScopeType;
+			classObj.RemoveClassFlags(EClassFlags.NotPlaceable);
 		}
 
 		[UhtSpecifier(Extends = UhtTableNames.Class, ValueType = UhtSpecifierValueType.Legacy)]
 		[SuppressMessage("CodeQuality", "IDE0051:Remove unused private members", Justification = "Attribute accessed method")]
-		private static void NotPlaceableSpecifier(UhtSpecifierContext SpecifierContext)
+		private static void NotPlaceableSpecifier(UhtSpecifierContext specifierContext)
 		{
-			UhtClassParser Class = (UhtClassParser)SpecifierContext.Scope.ScopeType;
-			Class.AddClassFlags(EClassFlags.NotPlaceable);
+			UhtClassParser classObj = (UhtClassParser)specifierContext.Scope.ScopeType;
+			classObj.AddClassFlags(EClassFlags.NotPlaceable);
 		}
 
 		[UhtSpecifier(Extends = UhtTableNames.Class, ValueType = UhtSpecifierValueType.Legacy)]
 		[SuppressMessage("CodeQuality", "IDE0051:Remove unused private members", Justification = "Attribute accessed method")]
-		private static void DefaultToInstancedSpecifier(UhtSpecifierContext SpecifierContext)
+		private static void DefaultToInstancedSpecifier(UhtSpecifierContext specifierContext)
 		{
-			UhtClassParser Class = (UhtClassParser)SpecifierContext.Scope.ScopeType;
-			Class.AddClassFlags(EClassFlags.DefaultToInstanced);
+			UhtClassParser classObj = (UhtClassParser)specifierContext.Scope.ScopeType;
+			classObj.AddClassFlags(EClassFlags.DefaultToInstanced);
 		}
 
 		[UhtSpecifier(Extends = UhtTableNames.Class, ValueType = UhtSpecifierValueType.Legacy)]
 		[SuppressMessage("CodeQuality", "IDE0051:Remove unused private members", Justification = "Attribute accessed method")]
-		private static void HideDropdownSpecifier(UhtSpecifierContext SpecifierContext)
+		private static void HideDropdownSpecifier(UhtSpecifierContext specifierContext)
 		{
-			UhtClassParser Class = (UhtClassParser)SpecifierContext.Scope.ScopeType;
-			Class.AddClassFlags(EClassFlags.HideDropDown);
+			UhtClassParser classObj = (UhtClassParser)specifierContext.Scope.ScopeType;
+			classObj.AddClassFlags(EClassFlags.HideDropDown);
 		}
 
 		[UhtSpecifier(Extends = UhtTableNames.Class, ValueType = UhtSpecifierValueType.Legacy)]
 		[SuppressMessage("CodeQuality", "IDE0051:Remove unused private members", Justification = "Attribute accessed method")]
-		private static void HiddenSpecifier(UhtSpecifierContext SpecifierContext)
+		private static void HiddenSpecifier(UhtSpecifierContext specifierContext)
 		{
 			// Prevents class from appearing in the editor class browser and edit inline menus.
-			UhtClassParser Class = (UhtClassParser)SpecifierContext.Scope.ScopeType;
-			Class.AddClassFlags(EClassFlags.Hidden);
+			UhtClassParser classObj = (UhtClassParser)specifierContext.Scope.ScopeType;
+			classObj.AddClassFlags(EClassFlags.Hidden);
 		}
 
 		[UhtSpecifier(Extends = UhtTableNames.Class, ValueType = UhtSpecifierValueType.Legacy)]
 		[SuppressMessage("CodeQuality", "IDE0051:Remove unused private members", Justification = "Attribute accessed method")]
-		private static void DependsOnSpecifier(UhtSpecifierContext SpecifierContext)
+		private static void DependsOnSpecifier(UhtSpecifierContext specifierContext)
 		{
-			SpecifierContext.MessageSite.LogError("The dependsOn specifier is deprecated. Please use #include \"ClassHeaderFilename.h\" instead.");
+			specifierContext.MessageSite.LogError("The dependsOn specifier is deprecated. Please use #include \"ClassHeaderFilename.h\" instead.");
 		}
 
 		[UhtSpecifier(Extends = UhtTableNames.Class, ValueType = UhtSpecifierValueType.Legacy)]
 		[SuppressMessage("CodeQuality", "IDE0051:Remove unused private members", Justification = "Attribute accessed method")]
-		private static void MinimalAPISpecifier(UhtSpecifierContext SpecifierContext)
+		private static void MinimalAPISpecifier(UhtSpecifierContext specifierContext)
 		{
-			UhtClassParser Class = (UhtClassParser)SpecifierContext.Scope.ScopeType;
-			Class.AddClassFlags(EClassFlags.MinimalAPI);
+			UhtClassParser classObj = (UhtClassParser)specifierContext.Scope.ScopeType;
+			classObj.AddClassFlags(EClassFlags.MinimalAPI);
 		}
 
 		[UhtSpecifier(Extends = UhtTableNames.Class, ValueType = UhtSpecifierValueType.Legacy)]
 		[SuppressMessage("CodeQuality", "IDE0051:Remove unused private members", Justification = "Attribute accessed method")]
-		private static void ConstSpecifier(UhtSpecifierContext SpecifierContext)
+		private static void ConstSpecifier(UhtSpecifierContext specifierContext)
 		{
-			UhtClassParser Class = (UhtClassParser)SpecifierContext.Scope.ScopeType;
-			Class.AddClassFlags(EClassFlags.Const);
+			UhtClassParser classObj = (UhtClassParser)specifierContext.Scope.ScopeType;
+			classObj.AddClassFlags(EClassFlags.Const);
 		}
 
 		[UhtSpecifier(Extends = UhtTableNames.Class, ValueType = UhtSpecifierValueType.Legacy)]
 		[SuppressMessage("CodeQuality", "IDE0051:Remove unused private members", Justification = "Attribute accessed method")]
-		private static void PerObjectConfigSpecifier(UhtSpecifierContext SpecifierContext)
+		private static void PerObjectConfigSpecifier(UhtSpecifierContext specifierContext)
 		{
-			UhtClassParser Class = (UhtClassParser)SpecifierContext.Scope.ScopeType;
-			Class.AddClassFlags(EClassFlags.PerObjectConfig);
+			UhtClassParser classObj = (UhtClassParser)specifierContext.Scope.ScopeType;
+			classObj.AddClassFlags(EClassFlags.PerObjectConfig);
 		}
 
 		[UhtSpecifier(Extends = UhtTableNames.Class, ValueType = UhtSpecifierValueType.Legacy)]
 		[SuppressMessage("CodeQuality", "IDE0051:Remove unused private members", Justification = "Attribute accessed method")]
-		private static void ConfigDoNotCheckDefaultsSpecifier(UhtSpecifierContext SpecifierContext)
+		private static void ConfigDoNotCheckDefaultsSpecifier(UhtSpecifierContext specifierContext)
 		{
-			UhtClassParser Class = (UhtClassParser)SpecifierContext.Scope.ScopeType;
-			Class.AddClassFlags(EClassFlags.ConfigDoNotCheckDefaults);
+			UhtClassParser classObj = (UhtClassParser)specifierContext.Scope.ScopeType;
+			classObj.AddClassFlags(EClassFlags.ConfigDoNotCheckDefaults);
 		}
 
 		[UhtSpecifier(Extends = UhtTableNames.Class, ValueType = UhtSpecifierValueType.Legacy)]
 		[SuppressMessage("CodeQuality", "IDE0051:Remove unused private members", Justification = "Attribute accessed method")]
-		private static void AbstractSpecifier(UhtSpecifierContext SpecifierContext)
+		private static void AbstractSpecifier(UhtSpecifierContext specifierContext)
 		{
-			UhtClassParser Class = (UhtClassParser)SpecifierContext.Scope.ScopeType;
-			Class.AddClassFlags(EClassFlags.Abstract);
+			UhtClassParser classObj = (UhtClassParser)specifierContext.Scope.ScopeType;
+			classObj.AddClassFlags(EClassFlags.Abstract);
 		}
 
 		[UhtSpecifier(Extends = UhtTableNames.Class, ValueType = UhtSpecifierValueType.Legacy)]
 		[SuppressMessage("CodeQuality", "IDE0051:Remove unused private members", Justification = "Attribute accessed method")]
-		private static void DeprecatedSpecifier(UhtSpecifierContext SpecifierContext)
+		private static void DeprecatedSpecifier(UhtSpecifierContext specifierContext)
 		{
-			UhtClassParser Class = (UhtClassParser)SpecifierContext.Scope.ScopeType;
-			Class.AddClassFlags(EClassFlags.Deprecated | EClassFlags.NotPlaceable);
+			UhtClassParser classObj = (UhtClassParser)specifierContext.Scope.ScopeType;
+			classObj.AddClassFlags(EClassFlags.Deprecated | EClassFlags.NotPlaceable);
 		}
 
 		[UhtSpecifier(Extends = UhtTableNames.Class, ValueType = UhtSpecifierValueType.Legacy)]
 		[SuppressMessage("CodeQuality", "IDE0051:Remove unused private members", Justification = "Attribute accessed method")]
-		private static void TransientSpecifier(UhtSpecifierContext SpecifierContext)
+		private static void TransientSpecifier(UhtSpecifierContext specifierContext)
 		{
-			UhtClassParser Class = (UhtClassParser)SpecifierContext.Scope.ScopeType;
-			Class.AddClassFlags(EClassFlags.Transient);
+			UhtClassParser classObj = (UhtClassParser)specifierContext.Scope.ScopeType;
+			classObj.AddClassFlags(EClassFlags.Transient);
 		}
 
 		[UhtSpecifier(Extends = UhtTableNames.Class, ValueType = UhtSpecifierValueType.Legacy)]
 		[SuppressMessage("CodeQuality", "IDE0051:Remove unused private members", Justification = "Attribute accessed method")]
-		private static void NonTransientSpecifier(UhtSpecifierContext SpecifierContext)
+		private static void NonTransientSpecifier(UhtSpecifierContext specifierContext)
 		{
-			UhtClassParser Class = (UhtClassParser)SpecifierContext.Scope.ScopeType;
-			Class.RemoveClassFlags(EClassFlags.Transient);
+			UhtClassParser classObj = (UhtClassParser)specifierContext.Scope.ScopeType;
+			classObj.RemoveClassFlags(EClassFlags.Transient);
 		}
 
 		[UhtSpecifier(Extends = UhtTableNames.Class, ValueType = UhtSpecifierValueType.Legacy)]
 		[SuppressMessage("CodeQuality", "IDE0051:Remove unused private members", Justification = "Attribute accessed method")]
-		private static void OptionalSpecifier(UhtSpecifierContext SpecifierContext)
+		private static void OptionalSpecifier(UhtSpecifierContext specifierContext)
 		{
 			// Optional class
-			UhtClassParser Class = (UhtClassParser)SpecifierContext.Scope.ScopeType;
-			Class.AddClassFlags(EClassFlags.Optional);
+			UhtClassParser classObj = (UhtClassParser)specifierContext.Scope.ScopeType;
+			classObj.AddClassFlags(EClassFlags.Optional);
 		}
 
 		[UhtSpecifier(Extends = UhtTableNames.Class, ValueType = UhtSpecifierValueType.Legacy)]
 		[SuppressMessage("CodeQuality", "IDE0051:Remove unused private members", Justification = "Attribute accessed method")]
-		private static void CustomConstructorSpecifier(UhtSpecifierContext SpecifierContext)
+		private static void CustomConstructorSpecifier(UhtSpecifierContext specifierContext)
 		{
 			// we will not export a constructor for this class, assuming it is in the CPP block
-			UhtClass Class = (UhtClass)SpecifierContext.Scope.ScopeType;
-			Class.ClassExportFlags |= UhtClassExportFlags.HasCustomConstructor;
+			UhtClass classObj = (UhtClass)specifierContext.Scope.ScopeType;
+			classObj.ClassExportFlags |= UhtClassExportFlags.HasCustomConstructor;
 		}
 
 		[UhtSpecifier(Extends = UhtTableNames.Class, ValueType = UhtSpecifierValueType.SingleString)]
 		[SuppressMessage("CodeQuality", "IDE0051:Remove unused private members", Justification = "Attribute accessed method")]
-		private static void ConfigSpecifier(UhtSpecifierContext SpecifierContext, StringView Value)
+		private static void ConfigSpecifier(UhtSpecifierContext specifierContext, StringView value)
 		{
-			UhtClass Class = (UhtClass)SpecifierContext.Scope.ScopeType;
-			Class.Config = Value.ToString();
+			UhtClass classObj = (UhtClass)specifierContext.Scope.ScopeType;
+			classObj.Config = value.ToString();
 		}
 
 		[UhtSpecifier(Extends = UhtTableNames.Class, ValueType = UhtSpecifierValueType.Legacy)]
 		[SuppressMessage("CodeQuality", "IDE0051:Remove unused private members", Justification = "Attribute accessed method")]
-		private static void DefaultConfigSpecifier(UhtSpecifierContext SpecifierContext)
+		private static void DefaultConfigSpecifier(UhtSpecifierContext specifierContext)
 		{
 			// Save object config only to Default INIs, never to local INIs.
-			UhtClassParser Class = (UhtClassParser)SpecifierContext.Scope.ScopeType;
-			Class.AddClassFlags(EClassFlags.DefaultConfig);
+			UhtClassParser classObj = (UhtClassParser)specifierContext.Scope.ScopeType;
+			classObj.AddClassFlags(EClassFlags.DefaultConfig);
 		}
 
 		[UhtSpecifier(Extends = UhtTableNames.Class, ValueType = UhtSpecifierValueType.Legacy)]
 		[SuppressMessage("CodeQuality", "IDE0051:Remove unused private members", Justification = "Attribute accessed method")]
-		private static void GlobalUserConfigSpecifier(UhtSpecifierContext SpecifierContext)
+		private static void GlobalUserConfigSpecifier(UhtSpecifierContext specifierContext)
 		{
 			// Save object config only to global user overrides, never to local INIs
-			UhtClassParser Class = (UhtClassParser)SpecifierContext.Scope.ScopeType;
-			Class.AddClassFlags(EClassFlags.GlobalUserConfig);
+			UhtClassParser classObj = (UhtClassParser)specifierContext.Scope.ScopeType;
+			classObj.AddClassFlags(EClassFlags.GlobalUserConfig);
 		}
 
 		[UhtSpecifier(Extends = UhtTableNames.Class, ValueType = UhtSpecifierValueType.Legacy)]
 		[SuppressMessage("CodeQuality", "IDE0051:Remove unused private members", Justification = "Attribute accessed method")]
-		private static void ProjectUserConfigSpecifier(UhtSpecifierContext SpecifierContext)
+		private static void ProjectUserConfigSpecifier(UhtSpecifierContext specifierContext)
 		{
 			// Save object config only to project user overrides, never to INIs that are checked in
-			UhtClassParser Class = (UhtClassParser)SpecifierContext.Scope.ScopeType;
-			Class.AddClassFlags(EClassFlags.ProjectUserConfig);
+			UhtClassParser classObj = (UhtClassParser)specifierContext.Scope.ScopeType;
+			classObj.AddClassFlags(EClassFlags.ProjectUserConfig);
 		}
 
 		[UhtSpecifier(Extends = UhtTableNames.Class, ValueType = UhtSpecifierValueType.SingleString)]
 		[SuppressMessage("CodeQuality", "IDE0051:Remove unused private members", Justification = "Attribute accessed method")]
-		private static void EditorConfigSpecifier(UhtSpecifierContext SpecifierContext, StringView Value)
+		private static void EditorConfigSpecifier(UhtSpecifierContext specifierContext, StringView value)
 		{
-			SpecifierContext.MetaData.Add(UhtNames.EditorConfig, Value.ToString());
+			specifierContext.MetaData.Add(UhtNames.EditorConfig, value.ToString());
 		}
 
 		[UhtSpecifier(Extends = UhtTableNames.Class, ValueType = UhtSpecifierValueType.NonEmptyStringList)]
 		[SuppressMessage("CodeQuality", "IDE0051:Remove unused private members", Justification = "Attribute accessed method")]
-		private static void ShowCategoriesSpecifier(UhtSpecifierContext SpecifierContext, List<StringView> Value)
+		private static void ShowCategoriesSpecifier(UhtSpecifierContext specifierContext, List<StringView> value)
 		{
-			UhtClassParser Class = (UhtClassParser)SpecifierContext.Scope.ScopeType;
-			Class.ShowCategories.AddUniqueRange(Value);
+			UhtClassParser classObj = (UhtClassParser)specifierContext.Scope.ScopeType;
+			classObj.ShowCategories.AddUniqueRange(value);
 		}
 
 		[UhtSpecifier(Extends = UhtTableNames.Class, ValueType = UhtSpecifierValueType.NonEmptyStringList)]
 		[SuppressMessage("CodeQuality", "IDE0051:Remove unused private members", Justification = "Attribute accessed method")]
-		private static void HideCategoriesSpecifier(UhtSpecifierContext SpecifierContext, List<StringView> Value)
+		private static void HideCategoriesSpecifier(UhtSpecifierContext specifierContext, List<StringView> value)
 		{
-			UhtClassParser Class = (UhtClassParser)SpecifierContext.Scope.ScopeType;
-			Class.HideCategories.AddUniqueRange(Value);
+			UhtClassParser classObj = (UhtClassParser)specifierContext.Scope.ScopeType;
+			classObj.HideCategories.AddUniqueRange(value);
 		}
 
 		[UhtSpecifier(Extends = UhtTableNames.Class, ValueType = UhtSpecifierValueType.NonEmptyStringList)]
 		[SuppressMessage("CodeQuality", "IDE0051:Remove unused private members", Justification = "Attribute accessed method")]
-		private static void ShowFunctionsSpecifier(UhtSpecifierContext SpecifierContext, List<StringView> Value)
+		private static void ShowFunctionsSpecifier(UhtSpecifierContext specifierContext, List<StringView> value)
 		{
-			UhtClassParser Class = (UhtClassParser)SpecifierContext.Scope.ScopeType;
-			Class.ShowFunctions.AddUniqueRange(Value);
+			UhtClassParser classObj = (UhtClassParser)specifierContext.Scope.ScopeType;
+			classObj.ShowFunctions.AddUniqueRange(value);
 		}
 
 		[UhtSpecifier(Extends = UhtTableNames.Class, ValueType = UhtSpecifierValueType.NonEmptyStringList)]
 		[SuppressMessage("CodeQuality", "IDE0051:Remove unused private members", Justification = "Attribute accessed method")]
-		private static void HideFunctionsSpecifier(UhtSpecifierContext SpecifierContext, List<StringView> Value)
+		private static void HideFunctionsSpecifier(UhtSpecifierContext specifierContext, List<StringView> value)
 		{
-			UhtClassParser Class = (UhtClassParser)SpecifierContext.Scope.ScopeType;
-			Class.HideFunctions.AddUniqueRange(Value);
+			UhtClassParser classObj = (UhtClassParser)specifierContext.Scope.ScopeType;
+			classObj.HideFunctions.AddUniqueRange(value);
 		}
 
 		[UhtSpecifier(Extends = UhtTableNames.Class, ValueType = UhtSpecifierValueType.SingleString)]
 		[SuppressMessage("CodeQuality", "IDE0051:Remove unused private members", Justification = "Attribute accessed method")]
-		private static void SparseClassDataTypesSpecifier(UhtSpecifierContext SpecifierContext, StringView Value)
+		private static void SparseClassDataTypesSpecifier(UhtSpecifierContext specifierContext, StringView value)
 		{
-			UhtClassParser Class = (UhtClassParser)SpecifierContext.Scope.ScopeType;
-			Class.SparseClassDataTypes.AddUnique(Value.ToString());
+			UhtClassParser classObj = (UhtClassParser)specifierContext.Scope.ScopeType;
+			classObj.SparseClassDataTypes.AddUnique(value.ToString());
 		}
 
 		[UhtSpecifier(Extends = UhtTableNames.Class, ValueType = UhtSpecifierValueType.NonEmptyStringList)]
 		[SuppressMessage("CodeQuality", "IDE0051:Remove unused private members", Justification = "Attribute accessed method")]
-		private static void ClassGroupSpecifier(UhtSpecifierContext SpecifierContext, List<StringView> Value)
+		private static void ClassGroupSpecifier(UhtSpecifierContext specifierContext, List<StringView> value)
 		{
-			UhtClassParser Class = (UhtClassParser)SpecifierContext.Scope.ScopeType;
-			foreach (StringView Eleemnt in Value)
+			UhtClassParser classObj = (UhtClassParser)specifierContext.Scope.ScopeType;
+			foreach (StringView element in value)
 			{
-				Class.ClassGroupNames.Add(Eleemnt.ToString());
+				classObj.ClassGroupNames.Add(element.ToString());
 			}
 		}
 
 		[UhtSpecifier(Extends = UhtTableNames.Class, ValueType = UhtSpecifierValueType.NonEmptyStringList)]
 		[SuppressMessage("CodeQuality", "IDE0051:Remove unused private members", Justification = "Attribute accessed method")]
-		private static void AutoExpandCategoriesSpecifier(UhtSpecifierContext SpecifierContext, List<StringView> Value)
+		private static void AutoExpandCategoriesSpecifier(UhtSpecifierContext specifierContext, List<StringView> value)
 		{
-			UhtClassParser Class = (UhtClassParser)SpecifierContext.Scope.ScopeType;
-			Class.AutoExpandCategories.AddUniqueRange(Value);
-			Class.AutoCollapseCategories.RemoveSwapRange(Value);
+			UhtClassParser classObj = (UhtClassParser)specifierContext.Scope.ScopeType;
+			classObj.AutoExpandCategories.AddUniqueRange(value);
+			classObj.AutoCollapseCategories.RemoveSwapRange(value);
 		}
 
 		[UhtSpecifier(Extends = UhtTableNames.Class, ValueType = UhtSpecifierValueType.NonEmptyStringList)]
 		[SuppressMessage("CodeQuality", "IDE0051:Remove unused private members", Justification = "Attribute accessed method")]
-		private static void AutoCollapseCategoriesSpecifier(UhtSpecifierContext SpecifierContext, List<StringView> Value)
+		private static void AutoCollapseCategoriesSpecifier(UhtSpecifierContext specifierContext, List<StringView> value)
 		{
-			UhtClassParser Class = (UhtClassParser)SpecifierContext.Scope.ScopeType;
-			Class.AutoCollapseCategories.AddUniqueRange(Value);
-			Class.AutoExpandCategories.RemoveSwapRange(Value);
+			UhtClassParser classObj = (UhtClassParser)specifierContext.Scope.ScopeType;
+			classObj.AutoCollapseCategories.AddUniqueRange(value);
+			classObj.AutoExpandCategories.RemoveSwapRange(value);
 		}
 
 		[UhtSpecifier(Extends = UhtTableNames.Class, ValueType = UhtSpecifierValueType.NonEmptyStringList)]
 		[SuppressMessage("CodeQuality", "IDE0051:Remove unused private members", Justification = "Attribute accessed method")]
-		private static void PrioritizeCategoriesSpecifier(UhtSpecifierContext SpecifierContext, List<StringView> Value)
+		private static void PrioritizeCategoriesSpecifier(UhtSpecifierContext specifierContext, List<StringView> value)
 		{
-			UhtClassParser Class = (UhtClassParser)SpecifierContext.Scope.ScopeType;
-			Class.PrioritizeCategories.AddUniqueRange(Value);
+			UhtClassParser classObj = (UhtClassParser)specifierContext.Scope.ScopeType;
+			classObj.PrioritizeCategories.AddUniqueRange(value);
 		}
 
 		[UhtSpecifier(Extends = UhtTableNames.Class, ValueType = UhtSpecifierValueType.NonEmptyStringList)]
 		[SuppressMessage("CodeQuality", "IDE0051:Remove unused private members", Justification = "Attribute accessed method")]
-		private static void DontAutoCollapseCategoriesSpecifier(UhtSpecifierContext SpecifierContext, List<StringView> Value)
+		private static void DontAutoCollapseCategoriesSpecifier(UhtSpecifierContext specifierContext, List<StringView> value)
 		{
-			UhtClassParser Class = (UhtClassParser)SpecifierContext.Scope.ScopeType;
-			Class.AutoCollapseCategories.RemoveSwapRange(Value);
+			UhtClassParser classObj = (UhtClassParser)specifierContext.Scope.ScopeType;
+			classObj.AutoCollapseCategories.RemoveSwapRange(value);
 		}
 
 		[UhtSpecifier(Extends = UhtTableNames.Class, ValueType = UhtSpecifierValueType.Legacy)]
 		[SuppressMessage("CodeQuality", "IDE0051:Remove unused private members", Justification = "Attribute accessed method")]
-		private static void CollapseCategoriesSpecifier(UhtSpecifierContext SpecifierContext)
+		private static void CollapseCategoriesSpecifier(UhtSpecifierContext specifierContext)
 		{
 			// Class' properties should not be shown categorized in the editor.
-			UhtClassParser Class = (UhtClassParser)SpecifierContext.Scope.ScopeType;
-			Class.AddClassFlags(EClassFlags.CollapseCategories);
+			UhtClassParser classObj = (UhtClassParser)specifierContext.Scope.ScopeType;
+			classObj.AddClassFlags(EClassFlags.CollapseCategories);
 		}
 
 		[UhtSpecifier(Extends = UhtTableNames.Class, ValueType = UhtSpecifierValueType.Legacy)]
 		[SuppressMessage("CodeQuality", "IDE0051:Remove unused private members", Justification = "Attribute accessed method")]
-		private static void DontCollapseCategoriesSpecifier(UhtSpecifierContext SpecifierContext)
+		private static void DontCollapseCategoriesSpecifier(UhtSpecifierContext specifierContext)
 		{
 			// Class' properties should be shown categorized in the editor.
-			UhtClassParser Class = (UhtClassParser)SpecifierContext.Scope.ScopeType;
-			Class.RemoveClassFlags(EClassFlags.CollapseCategories);
+			UhtClassParser classObj = (UhtClassParser)specifierContext.Scope.ScopeType;
+			classObj.RemoveClassFlags(EClassFlags.CollapseCategories);
 		}
 
 		[UhtSpecifier(Extends = UhtTableNames.Class, ValueType = UhtSpecifierValueType.Legacy)]
 		[SuppressMessage("CodeQuality", "IDE0051:Remove unused private members", Justification = "Attribute accessed method")]
-		private static void AdvancedClassDisplaySpecifier(UhtSpecifierContext SpecifierContext)
+		private static void AdvancedClassDisplaySpecifier(UhtSpecifierContext specifierContext)
 		{
 			// By default the class properties are shown in advanced sections in UI
-			UhtClass Class = (UhtClass)SpecifierContext.Scope.ScopeType;
-			Class.MetaData.Add(UhtNames.AdvancedClassDisplay, "true");
+			UhtClass classObj = (UhtClass)specifierContext.Scope.ScopeType;
+			classObj.MetaData.Add(UhtNames.AdvancedClassDisplay, "true");
 		}
 
 		[UhtSpecifier(Extends = UhtTableNames.Class, ValueType = UhtSpecifierValueType.Legacy)]
 		[SuppressMessage("CodeQuality", "IDE0051:Remove unused private members", Justification = "Attribute accessed method")]
-		private static void ConversionRootSpecifier(UhtSpecifierContext SpecifierContext)
+		private static void ConversionRootSpecifier(UhtSpecifierContext specifierContext)
 		{
-			UhtClass Class = (UhtClass)SpecifierContext.Scope.ScopeType;
-			Class.MetaData.Add(UhtNames.IsConversionRoot, "true");
+			UhtClass classObj = (UhtClass)specifierContext.Scope.ScopeType;
+			classObj.MetaData.Add(UhtNames.IsConversionRoot, "true");
 		}
 
 		[UhtSpecifier(Extends = UhtTableNames.Class, ValueType = UhtSpecifierValueType.Legacy)]
 		[SuppressMessage("CodeQuality", "IDE0051:Remove unused private members", Justification = "Attribute accessed method")]
-		private static void NeedsDeferredDependencyLoadingSpecifier(UhtSpecifierContext SpecifierContext)
+		private static void NeedsDeferredDependencyLoadingSpecifier(UhtSpecifierContext specifierContext)
 		{
-			UhtClassParser Class = (UhtClassParser)SpecifierContext.Scope.ScopeType;
-			Class.AddClassFlags(EClassFlags.NeedsDeferredDependencyLoading);
+			UhtClassParser classObj = (UhtClassParser)specifierContext.Scope.ScopeType;
+			classObj.AddClassFlags(EClassFlags.NeedsDeferredDependencyLoading);
 		}
 
 		[UhtSpecifier(Extends = UhtTableNames.Class, ValueType = UhtSpecifierValueType.Legacy)]
 		[SuppressMessage("CodeQuality", "IDE0051:Remove unused private members", Justification = "Attribute accessed method")]
-		private static void MatchedSerializersSpecifier(UhtSpecifierContext SpecifierContext)
+		private static void MatchedSerializersSpecifier(UhtSpecifierContext specifierContext)
 		{
-			if (!SpecifierContext.Scope.HeaderParser.HeaderFile.bIsNoExportTypes)
+			if (!specifierContext.Scope.HeaderParser.HeaderFile.IsNoExportTypes)
 			{
-				SpecifierContext.MessageSite.LogError("The 'MatchedSerializers' class specifier is only valid in the NoExportTypes.h file");
+				specifierContext.MessageSite.LogError("The 'MatchedSerializers' class specifier is only valid in the NoExportTypes.h file");
 			}
-			UhtClassParser Class = (UhtClassParser)SpecifierContext.Scope.ScopeType;
-			Class.AddClassFlags(EClassFlags.MatchedSerializers);
+			UhtClassParser classObj = (UhtClassParser)specifierContext.Scope.ScopeType;
+			classObj.AddClassFlags(EClassFlags.MatchedSerializers);
 		}
 
 		[UhtSpecifier(Extends = UhtTableNames.Class, ValueType = UhtSpecifierValueType.Legacy)]
 		[SuppressMessage("CodeQuality", "IDE0051:Remove unused private members", Justification = "Attribute accessed method")]
-		private static void InterfaceSpecifier(UhtSpecifierContext SpecifierContext)
+		private static void InterfaceSpecifier(UhtSpecifierContext specifierContext)
 		{
-			if (!SpecifierContext.Scope.HeaderParser.HeaderFile.bIsNoExportTypes)
+			if (!specifierContext.Scope.HeaderParser.HeaderFile.IsNoExportTypes)
 			{
-				SpecifierContext.MessageSite.LogError("The 'Interface' class specifier is only valid in the NoExportTypes.h file");
+				specifierContext.MessageSite.LogError("The 'Interface' class specifier is only valid in the NoExportTypes.h file");
 			}
-			UhtClassParser Class = (UhtClassParser)SpecifierContext.Scope.ScopeType;
-			Class.AddClassFlags(EClassFlags.Interface);
+			UhtClassParser classObj = (UhtClassParser)specifierContext.Scope.ScopeType;
+			classObj.AddClassFlags(EClassFlags.Interface);
 		}
 
 		[UhtSpecifier(Extends = UhtTableNames.Class, ValueType = UhtSpecifierValueType.None)]
 		[SuppressMessage("CodeQuality", "IDE0051:Remove unused private members", Justification = "Attribute accessed method")]
-		private static void CustomFieldNotifySpecifier(UhtSpecifierContext SpecifierContext)
+		private static void CustomFieldNotifySpecifier(UhtSpecifierContext specifierContext)
 		{
-			UhtClassParser Class = (UhtClassParser)SpecifierContext.Scope.ScopeType;
-			Class.ClassExportFlags |= UhtClassExportFlags.HasCustomFieldNotify;
+			UhtClassParser classObj = (UhtClassParser)specifierContext.Scope.ScopeType;
+			classObj.ClassExportFlags |= UhtClassExportFlags.HasCustomFieldNotify;
 		}
 	}
 }
diff --git a/Engine/Source/Programs/Shared/EpicGames.UHT/Specifiers/UhtDefaultSpecifiers.cs b/Engine/Source/Programs/Shared/EpicGames.UHT/Specifiers/UhtDefaultSpecifiers.cs
--- a/Engine/Source/Programs/Shared/EpicGames.UHT/Specifiers/UhtDefaultSpecifiers.cs
+++ b/Engine/Source/Programs/Shared/EpicGames.UHT/Specifiers/UhtDefaultSpecifiers.cs
@@ -1,12 +1,12 @@
 // Copyright Epic Games, Inc. All Rights Reserved.
 
+using System;
+using System.Collections.Generic;
+using System.Diagnostics.CodeAnalysis;
 using EpicGames.Core;
 using EpicGames.UHT.Tables;
 using EpicGames.UHT.Types;
 using EpicGames.UHT.Utils;
-using System;
-using System.Collections.Generic;
-using System.Diagnostics.CodeAnalysis;
 
 namespace EpicGames.UHT.Parsers
 {
@@ -17,28 +17,28 @@ namespace EpicGames.UHT.Parsers
 	[UnrealHeaderTool]
 	public static class UhtDefaultSpecifiers
 	{
-		private static void SetMetaData(UhtSpecifierContext SpecifierContext, StringView Key, StringView Value)
+		private static void SetMetaData(UhtSpecifierContext specifierContext, StringView key, StringView value)
 		{
-			if (Key.Length == 0)
+			if (key.Length == 0)
 			{
-				SpecifierContext.MessageSite.LogError("Invalid metadata, name can not be blank");
+				specifierContext.MessageSite.LogError("Invalid metadata, name can not be blank");
 				return;
 			}
 
 			// Trim the leading and ending whitespace
-			ReadOnlySpan<char> Span = Value.Span;
-			int Start = 0;
-			int End = Value.Length;
-			for (; Start < End; ++Start)
+			ReadOnlySpan<char> span = value.Span;
+			int start = 0;
+			int end = value.Length;
+			for (; start < end; ++start)
 			{
-				if (!UhtFCString.IsWhitespace(Span[Start]))
+				if (!UhtFCString.IsWhitespace(span[start]))
 				{
 					break;
 				}
 			}
-			for (; Start < End; --End)
+			for (; start < end; --end)
 			{
-				if (!UhtFCString.IsWhitespace(Span[End - 1]))
+				if (!UhtFCString.IsWhitespace(span[end - 1]))
 				{
 					break;
 				}
@@ -47,167 +47,151 @@ namespace EpicGames.UHT.Parsers
 			// Trim any quotes 
 			//COMPATIBILITY-TODO - This doesn't handle strings that end in an escaped ".  This is a bug in old UHT.
 			// Only remove quotes if we have quotes at both ends and the last isn't escaped.
-			if (Start < End && Span[Start] == '"')
+			if (start < end && span[start] == '"')
 			{
-				++Start;
+				++start;
 			}
-			if (Start < End && Span[End - 1] == '"')
+			if (start < end && span[end - 1] == '"')
 			{
-				--End;
+				--end;
 			}
 
 			// Get the trimmed string
-			Value = new StringView(Value, Start, End - Start);
+			value = new StringView(value, start, end - start);
 
 			// Make sure this isn't a duplicate assignment
-			string KeyAsString = Key.ToString();
-			if (SpecifierContext.MetaData.TryGetValue(KeyAsString, out string? ExistingValue))
+			string keyAsString = key.ToString();
+			if (specifierContext.MetaData.TryGetValue(keyAsString, out string? existingValue))
 			{
-				if (StringViewComparer.OrdinalIgnoreCase.Compare(ExistingValue, Value) != 0)
+				if (StringViewComparer.OrdinalIgnoreCase.Compare(existingValue, value) != 0)
 				{
-					SpecifierContext.MessageSite.LogError($"Metadata key '{Key}' first seen with value '{ExistingValue}' then '{Value}'");
+					specifierContext.MessageSite.LogError($"Metadata key '{key}' first seen with value '{existingValue}' then '{value}'");
 				}
 			}
 			else
 			{
-				SpecifierContext.MetaData.Add(KeyAsString, SpecifierContext.MetaNameIndex, Value.ToString());
+				specifierContext.MetaData.Add(keyAsString, specifierContext.MetaNameIndex, value.ToString());
 			}
 		}
 
-		private static void SetMetaData(UhtSpecifierContext SpecifierContext, StringView Key, bool Value)
+		private static void SetMetaData(UhtSpecifierContext specifierContext, StringView key, bool value)
 		{
-			SetMetaData(SpecifierContext, Key, Value ? "true" : "false");
+			SetMetaData(specifierContext, key, value ? "true" : "false");
 		}
 
 		#region Specifiers
 		[UhtSpecifier(Extends = UhtTableNames.Default, ValueType = UhtSpecifierValueType.KeyValuePairList, When = UhtSpecifierWhen.Immediate)]
 		[SuppressMessage("CodeQuality", "IDE0051:Remove unused private members", Justification = "Attribute accessed method")]
-		[SuppressMessage("Style", "IDE0060:Remove unused parameter", Justification = "Attribute accessed method")]
-		private static void MetaSpecifier(UhtSpecifierContext SpecifierContext, List<KeyValuePair<StringView, StringView>> Value)
+		private static void MetaSpecifier(UhtSpecifierContext specifierContext, List<KeyValuePair<StringView, StringView>> value)
 		{
-			foreach (KeyValuePair<StringView, StringView> KVP in (List<KeyValuePair<StringView, StringView>>)Value)
+			foreach (KeyValuePair<StringView, StringView> kvp in (List<KeyValuePair<StringView, StringView>>)value)
 			{
-				SetMetaData(SpecifierContext, KVP.Key, KVP.Value);
+				SetMetaData(specifierContext, kvp.Key, kvp.Value);
 			}
 		}
 
 		[UhtSpecifier(Extends = UhtTableNames.Default, ValueType = UhtSpecifierValueType.String, When = UhtSpecifierWhen.Immediate)]
 		[SuppressMessage("CodeQuality", "IDE0051:Remove unused private members", Justification = "Attribute accessed method")]
-		[SuppressMessage("Style", "IDE0060:Remove unused parameter", Justification = "Attribute accessed method")]
-		[SuppressMessage("Style", "IDE0060:Remove unused parameter", Justification = "Attribute accessed method")]
-		private static void DisplayNameSpecifier(UhtSpecifierContext SpecifierContext, StringView Value)
+		private static void DisplayNameSpecifier(UhtSpecifierContext specifierContext, StringView value)
 		{
-			SetMetaData(SpecifierContext, "DisplayName", Value);
+			SetMetaData(specifierContext, "DisplayName", value);
 		}
 
 		[UhtSpecifier(Extends = UhtTableNames.Default, ValueType = UhtSpecifierValueType.String, When = UhtSpecifierWhen.Immediate)]
 		[SuppressMessage("CodeQuality", "IDE0051:Remove unused private members", Justification = "Attribute accessed method")]
-		[SuppressMessage("Style", "IDE0060:Remove unused parameter", Justification = "Attribute accessed method")]
-		private static void FriendlyNameSpecifier(UhtSpecifierContext SpecifierContext, StringView Value)
+		private static void FriendlyNameSpecifier(UhtSpecifierContext specifierContext, StringView value)
 		{
-			SetMetaData(SpecifierContext, "FriendlyName", Value);
+			SetMetaData(specifierContext, "FriendlyName", value);
 		}
 
 		[UhtSpecifier(Extends = UhtTableNames.Default, ValueType = UhtSpecifierValueType.None, When = UhtSpecifierWhen.Immediate)]
 		[SuppressMessage("CodeQuality", "IDE0051:Remove unused private members", Justification = "Attribute accessed method")]
-		[SuppressMessage("Style", "IDE0060:Remove unused parameter", Justification = "Attribute accessed method")]
-		private static void BlueprintInternalUseOnlySpecifier(UhtSpecifierContext SpecifierContext)
+		private static void BlueprintInternalUseOnlySpecifier(UhtSpecifierContext specifierContext)
 		{
-			SetMetaData(SpecifierContext, "BlueprintInternalUseOnly", true);
-			SetMetaData(SpecifierContext, "BlueprintType", true);
+			SetMetaData(specifierContext, "BlueprintInternalUseOnly", true);
+			SetMetaData(specifierContext, "BlueprintType", true);
 		}
 
 		[UhtSpecifier(Extends = UhtTableNames.Default, ValueType = UhtSpecifierValueType.None, When = UhtSpecifierWhen.Immediate)]
 		[SuppressMessage("CodeQuality", "IDE0051:Remove unused private members", Justification = "Attribute accessed method")]
-		[SuppressMessage("Style", "IDE0060:Remove unused parameter", Justification = "Attribute accessed method")]
-		private static void BlueprintInternalUseOnlyHierarchicalSpecifier(UhtSpecifierContext SpecifierContext)
+		private static void BlueprintInternalUseOnlyHierarchicalSpecifier(UhtSpecifierContext specifierContext)
 		{
-			SetMetaData(SpecifierContext, "BlueprintInternalUseOnlyHierarchical", true);
-			SetMetaData(SpecifierContext, "BlueprintInternalUseOnly", true);
-			SetMetaData(SpecifierContext, "BlueprintType", true);
+			SetMetaData(specifierContext, "BlueprintInternalUseOnlyHierarchical", true);
+			SetMetaData(specifierContext, "BlueprintInternalUseOnly", true);
+			SetMetaData(specifierContext, "BlueprintType", true);
 		}
 
 		[UhtSpecifier(Extends = UhtTableNames.Default, ValueType = UhtSpecifierValueType.None, When = UhtSpecifierWhen.Immediate)]
 		[SuppressMessage("CodeQuality", "IDE0051:Remove unused private members", Justification = "Attribute accessed method")]
-		[SuppressMessage("Style", "IDE0060:Remove unused parameter", Justification = "Attribute accessed method")]
-		private static void BlueprintTypeSpecifier(UhtSpecifierContext SpecifierContext)
+		private static void BlueprintTypeSpecifier(UhtSpecifierContext specifierContext)
 		{
-			SetMetaData(SpecifierContext, "BlueprintType", true);
+			SetMetaData(specifierContext, "BlueprintType", true);
 		}
 
 		[UhtSpecifier(Extends = UhtTableNames.Default, ValueType = UhtSpecifierValueType.None, When = UhtSpecifierWhen.Immediate)]
 		[SuppressMessage("CodeQuality", "IDE0051:Remove unused private members", Justification = "Attribute accessed method")]
-		[SuppressMessage("Style", "IDE0060:Remove unused parameter", Justification = "Attribute accessed method")]
-		private static void NotBlueprintTypeSpecifier(UhtSpecifierContext SpecifierContext)
+		private static void NotBlueprintTypeSpecifier(UhtSpecifierContext specifierContext)
 		{
-			SetMetaData(SpecifierContext, "NotBlueprintType", true);
-			SpecifierContext.MetaData.Remove("BlueprintType", SpecifierContext.MetaNameIndex);
+			SetMetaData(specifierContext, "NotBlueprintType", true);
+			specifierContext.MetaData.Remove("BlueprintType", specifierContext.MetaNameIndex);
 		}
 
 		[UhtSpecifier(Extends = UhtTableNames.Default, ValueType = UhtSpecifierValueType.None, When = UhtSpecifierWhen.Immediate)]
 		[SuppressMessage("CodeQuality", "IDE0051:Remove unused private members", Justification = "Attribute accessed method")]
-		[SuppressMessage("Style", "IDE0060:Remove unused parameter", Justification = "Attribute accessed method")]
-		private static void BlueprintableSpecifier(UhtSpecifierContext SpecifierContext)
+		private static void BlueprintableSpecifier(UhtSpecifierContext specifierContext)
 		{
-			SetMetaData(SpecifierContext, "IsBlueprintBase", true);
-			SetMetaData(SpecifierContext, "BlueprintType", true);
+			SetMetaData(specifierContext, "IsBlueprintBase", true);
+			SetMetaData(specifierContext, "BlueprintType", true);
 		}
 
 		[UhtSpecifier(Extends = UhtTableNames.Default, ValueType = UhtSpecifierValueType.None, When = UhtSpecifierWhen.Immediate)]
 		[SuppressMessage("CodeQuality", "IDE0051:Remove unused private members", Justification = "Attribute accessed method")]
-		[SuppressMessage("Style", "IDE0060:Remove unused parameter", Justification = "Attribute accessed method")]
-		private static void CallInEditorSpecifier(UhtSpecifierContext SpecifierContext)
+		private static void CallInEditorSpecifier(UhtSpecifierContext specifierContext)
 		{
-			SetMetaData(SpecifierContext, "CallInEditor", true);
+			SetMetaData(specifierContext, "CallInEditor", true);
 		}
 
 		[UhtSpecifier(Extends = UhtTableNames.Default, ValueType = UhtSpecifierValueType.None, When = UhtSpecifierWhen.Immediate)]
 		[SuppressMessage("CodeQuality", "IDE0051:Remove unused private members", Justification = "Attribute accessed method")]
-		[SuppressMessage("Style", "IDE0060:Remove unused parameter", Justification = "Attribute accessed method")]
-		private static void NotBlueprintableSpecifier(UhtSpecifierContext SpecifierContext)
+		private static void NotBlueprintableSpecifier(UhtSpecifierContext specifierContext)
 		{
-			SetMetaData(SpecifierContext, "IsBlueprintBase", false);
-			SpecifierContext.MetaData.Remove("BlueprintType", SpecifierContext.MetaNameIndex);
+			SetMetaData(specifierContext, "IsBlueprintBase", false);
+			specifierContext.MetaData.Remove("BlueprintType", specifierContext.MetaNameIndex);
 		}
 
 		[UhtSpecifier(Extends = UhtTableNames.Default, ValueType = UhtSpecifierValueType.String, When = UhtSpecifierWhen.Immediate)]
 		[SuppressMessage("CodeQuality", "IDE0051:Remove unused private members", Justification = "Attribute accessed method")]
-		[SuppressMessage("Style", "IDE0060:Remove unused parameter", Justification = "Attribute accessed method")]
-		private static void CategorySpecifier(UhtSpecifierContext SpecifierContext, StringView Value)
+		private static void CategorySpecifier(UhtSpecifierContext specifierContext, StringView value)
 		{
-			SetMetaData(SpecifierContext, "Category", Value);
+			SetMetaData(specifierContext, "Category", value);
 		}
 
 		[UhtSpecifier(Extends = UhtTableNames.Default, ValueType = UhtSpecifierValueType.None, When = UhtSpecifierWhen.Immediate)]
 		[SuppressMessage("CodeQuality", "IDE0051:Remove unused private members", Justification = "Attribute accessed method")]
-		[SuppressMessage("Style", "IDE0060:Remove unused parameter", Justification = "Attribute accessed method")]
-		private static void ExperimentalSpecifier(UhtSpecifierContext SpecifierContext)
+		private static void ExperimentalSpecifier(UhtSpecifierContext specifierContext)
 		{
-			SetMetaData(SpecifierContext, "DevelopmentStatus", "Experimental");
+			SetMetaData(specifierContext, "DevelopmentStatus", "Experimental");
 		}
 
 		[UhtSpecifier(Extends = UhtTableNames.Default, ValueType = UhtSpecifierValueType.None, When = UhtSpecifierWhen.Immediate)]
 		[SuppressMessage("CodeQuality", "IDE0051:Remove unused private members", Justification = "Attribute accessed method")]
-		[SuppressMessage("Style", "IDE0060:Remove unused parameter", Justification = "Attribute accessed method")]
-		private static void EarlyAccessPreviewSpecifier(UhtSpecifierContext SpecifierContext)
+		private static void EarlyAccessPreviewSpecifier(UhtSpecifierContext specifierContext)
 		{
-			SetMetaData(SpecifierContext, "DevelopmentStatus", "EarlyAccess");
+			SetMetaData(specifierContext, "DevelopmentStatus", "EarlyAccess");
 		}
 
 		[UhtSpecifier(Extends = UhtTableNames.Default, ValueType = UhtSpecifierValueType.None, When = UhtSpecifierWhen.Immediate)]
 		[SuppressMessage("CodeQuality", "IDE0051:Remove unused private members", Justification = "Attribute accessed method")]
-		[SuppressMessage("Style", "IDE0060:Remove unused parameter", Justification = "Attribute accessed method")]
-		private static void DocumentationPolicySpecifier(UhtSpecifierContext SpecifierContext)
+		private static void DocumentationPolicySpecifier(UhtSpecifierContext specifierContext)
 		{
-			SetMetaData(SpecifierContext, UhtNames.DocumentationPolicy, "Strict");
+			SetMetaData(specifierContext, UhtNames.DocumentationPolicy, "Strict");
 		}
 
 		[UhtSpecifier(Extends = UhtTableNames.Default, ValueType = UhtSpecifierValueType.String, When = UhtSpecifierWhen.Immediate)]
 		[SuppressMessage("CodeQuality", "IDE0051:Remove unused private members", Justification = "Attribute accessed method")]
-		[SuppressMessage("Style", "IDE0060:Remove unused parameter", Justification = "Attribute accessed method")]
-		private static void SparseClassDataTypeSpecifier(UhtSpecifierContext SpecifierContext, StringView Value)
+		private static void SparseClassDataTypeSpecifier(UhtSpecifierContext specifierContext, StringView value)
 		{
-			SetMetaData(SpecifierContext, "SparseClassDataType", Value);
+			SetMetaData(specifierContext, "SparseClassDataType", value);
 		}
 		#endregion
 
@@ -218,59 +202,58 @@ namespace EpicGames.UHT.Parsers
 		[UhtSpecifierValidator(Name = "ClampMax", Extends = UhtTableNames.Default)]
 		[SuppressMessage("CodeQuality", "IDE0051:Remove unused private members", Justification = "Attribute accessed method")]
 		[SuppressMessage("Style", "IDE0060:Remove unused parameter", Justification = "Attribute accessed method")]
-		private static void ValidateNumeric(UhtType Type, UhtMetaData MetaData, UhtMetaDataKey Key, StringView Value)
+		private static void ValidateNumeric(UhtType type, UhtMetaData metaData, UhtMetaDataKey key, StringView value)
 		{
-			if (!UhtFCString.IsNumeric(Value.Span))
+			if (!UhtFCString.IsNumeric(value.Span))
 			{
-				Type.LogError($"Metadata value for '{Key}' is non-numeric : '{Value}'");
+				type.LogError($"Metadata value for '{key}' is non-numeric : '{value}'");
 			}
 		}
 
 		[UhtSpecifierValidator(Extends = UhtTableNames.Default)]
 		[SuppressMessage("CodeQuality", "IDE0051:Remove unused private members", Justification = "Attribute accessed method")]
 		[SuppressMessage("Style", "IDE0060:Remove unused parameter", Justification = "Attribute accessed method")]
-		private static void DevelopmentStatusSpecifierValidator(UhtType Type, UhtMetaData MetaData, UhtMetaDataKey Key, StringView Value)
+		private static void DevelopmentStatusSpecifierValidator(UhtType type, UhtMetaData metaData, UhtMetaDataKey key, StringView value)
 		{
-			string[] AllowedValues = { "EarlyAccess", "Experimental" };
-			foreach (string AllowedValue in AllowedValues)
+			string[] allowedValues = { "EarlyAccess", "Experimental" };
+			foreach (string allowedValue in allowedValues)
 			{
-				if (Value.Equals(AllowedValue, StringComparison.OrdinalIgnoreCase))
+				if (value.Equals(allowedValue, StringComparison.OrdinalIgnoreCase))
 				{
 					return;
 				}
 			}
-			Type.LogError($"'{Key.Name}' metadata was '{Value}' but it must be {UhtUtilities.MergeTypeNames(AllowedValues, "or", false)}");
+			type.LogError($"'{key.Name}' metadata was '{value}' but it must be {UhtUtilities.MergeTypeNames(allowedValues, "or", false)}");
 		}
 
 		[UhtSpecifierValidator(Extends = UhtTableNames.Default)]
 		[SuppressMessage("CodeQuality", "IDE0051:Remove unused private members", Justification = "Attribute accessed method")]
-		[SuppressMessage("Style", "IDE0060:Remove unused parameter", Justification = "Attribute accessed method")]
-		private static void DocumentationPolicySpecifierValidator(UhtType Type, UhtMetaData MetaData, UhtMetaDataKey Key, StringView Value)
+		private static void DocumentationPolicySpecifierValidator(UhtType type, UhtMetaData metaData, UhtMetaDataKey key, StringView value)
 		{
 			const string StrictValue = "Strict";
-			if (!Value.Span.Equals(StrictValue, StringComparison.OrdinalIgnoreCase))
+			if (!value.Span.Equals(StrictValue, StringComparison.OrdinalIgnoreCase))
 			{
-				Type.LogError(MetaData.LineNumber, $"'{Key}' metadata was '{Value}' but it must be '{StrictValue}'");
+				type.LogError(metaData.LineNumber, $"'{key}' metadata was '{value}' but it must be '{StrictValue}'");
 			}
 		}
 
 		[UhtSpecifierValidator(Extends = UhtTableNames.Default)]
 		[SuppressMessage("CodeQuality", "IDE0051:Remove unused private members", Justification = "Attribute accessed method")]
 		[SuppressMessage("Style", "IDE0060:Remove unused parameter", Justification = "Attribute accessed method")]
-		private static void UnitsSpecifierValidator(UhtType Type, UhtMetaData MetaData, UhtMetaDataKey Key, StringView Value)
+		private static void UnitsSpecifierValidator(UhtType type, UhtMetaData metaData, UhtMetaDataKey key, StringView value)
 		{
 			// Check for numeric property
-			if (Type is UhtProperty)
+			if (type is UhtProperty)
 			{
-				if (!(Type is UhtNumericProperty) && !(Type is UhtStructProperty))
+				if (type is not UhtNumericProperty && type is not UhtStructProperty)
 				{
-					Type.LogError("'Units' meta data can only be applied to numeric and struct properties");
+					type.LogError("'Units' meta data can only be applied to numeric and struct properties");
 				}
 			}
 
-			if (!Type.Session.Config!.IsValidUnits(Value))
+			if (!type.Session.Config!.IsValidUnits(value))
 			{
-				Type.LogError($"Unrecognized units '{Value}' specified for '{Type.FullName}'");
+				type.LogError($"Unrecognized units '{value}' specified for '{type.FullName}'");
 			}
 		}
 		#endregion
diff --git a/Engine/Source/Programs/Shared/EpicGames.UHT/Specifiers/UhtEnumSpecifiers.cs b/Engine/Source/Programs/Shared/EpicGames.UHT/Specifiers/UhtEnumSpecifiers.cs
--- a/Engine/Source/Programs/Shared/EpicGames.UHT/Specifiers/UhtEnumSpecifiers.cs
+++ b/Engine/Source/Programs/Shared/EpicGames.UHT/Specifiers/UhtEnumSpecifiers.cs
@@ -17,10 +17,10 @@ namespace EpicGames.UHT.Parsers
 	{
 		[UhtSpecifier(Extends = UhtTableNames.Enum, ValueType = UhtSpecifierValueType.Legacy)]
 		[SuppressMessage("CodeQuality", "IDE0051:Remove unused private members", Justification = "Attribute accessed method")]
-		private static void FlagsSpecifier(UhtSpecifierContext SpecifierContext)
+		private static void FlagsSpecifier(UhtSpecifierContext specifierContext)
 		{
-			UhtEnum Enum = (UhtEnum)SpecifierContext.Scope.ScopeType;
-			Enum.EnumFlags |= EEnumFlags.Flags;
+			UhtEnum enumObj = (UhtEnum)specifierContext.Scope.ScopeType;
+			enumObj.EnumFlags |= EEnumFlags.Flags;
 		}
 	}
 }
