More coding standard fixes

diff --git a/Engine/Source/Programs/Shared/EpicGames.UHT/Specifiers/UhtFunctionSpecifiers.cs b/Engine/Source/Programs/Shared/EpicGames.UHT/Specifiers/UhtFunctionSpecifiers.cs
--- a/Engine/Source/Programs/Shared/EpicGames.UHT/Specifiers/UhtFunctionSpecifiers.cs
+++ b/Engine/Source/Programs/Shared/EpicGames.UHT/Specifiers/UhtFunctionSpecifiers.cs
@@ -1,13 +1,13 @@
 // Copyright Epic Games, Inc. All Rights Reserved.
 
+using System;
+using System.Collections.Generic;
+using System.Diagnostics.CodeAnalysis;
 using EpicGames.Core;
 using EpicGames.UHT.Tables;
 using EpicGames.UHT.Tokenizer;
 using EpicGames.UHT.Types;
 using EpicGames.UHT.Utils;
-using System;
-using System.Collections.Generic;
-using System.Diagnostics.CodeAnalysis;
 
 namespace EpicGames.UHT.Parsers
 {
@@ -19,358 +19,338 @@ namespace EpicGames.UHT.Parsers
 	{
 		[UhtSpecifier(Extends = UhtTableNames.Function, ValueType = UhtSpecifierValueType.Legacy)]
 		[SuppressMessage("CodeQuality", "IDE0051:Remove unused private members", Justification = "Attribute accessed method")]
-		[SuppressMessage("Style", "IDE0060:Remove unused parameter", Justification = "Attribute accessed method")]
-		private static void BlueprintAuthorityOnlySpecifier(UhtSpecifierContext SpecifierContext)
+		private static void BlueprintAuthorityOnlySpecifier(UhtSpecifierContext specifierContext)
 		{
-			UhtFunction Function = (UhtFunction)SpecifierContext.Scope.ScopeType;
-			Function.FunctionFlags |= EFunctionFlags.BlueprintAuthorityOnly;
+			UhtFunction function = (UhtFunction)specifierContext.Scope.ScopeType;
+			function.FunctionFlags |= EFunctionFlags.BlueprintAuthorityOnly;
 		}
 
 		[UhtSpecifier(Extends = UhtTableNames.Function, ValueType = UhtSpecifierValueType.Legacy)]
 		[SuppressMessage("CodeQuality", "IDE0051:Remove unused private members", Justification = "Attribute accessed method")]
-		[SuppressMessage("Style", "IDE0060:Remove unused parameter", Justification = "Attribute accessed method")]
-		private static void BlueprintCallableSpecifier(UhtSpecifierContext SpecifierContext)
+		private static void BlueprintCallableSpecifier(UhtSpecifierContext specifierContext)
 		{
-			UhtFunction Function = (UhtFunction)SpecifierContext.Scope.ScopeType;
-			Function.FunctionFlags |= EFunctionFlags.BlueprintCallable;
+			UhtFunction function = (UhtFunction)specifierContext.Scope.ScopeType;
+			function.FunctionFlags |= EFunctionFlags.BlueprintCallable;
 		}
 
 		[UhtSpecifier(Extends = UhtTableNames.Function, ValueType = UhtSpecifierValueType.Legacy)]
 		[SuppressMessage("CodeQuality", "IDE0051:Remove unused private members", Justification = "Attribute accessed method")]
-		[SuppressMessage("Style", "IDE0060:Remove unused parameter", Justification = "Attribute accessed method")]
-		private static void BlueprintCosmeticSpecifier(UhtSpecifierContext SpecifierContext)
+		private static void BlueprintCosmeticSpecifier(UhtSpecifierContext specifierContext)
 		{
-			UhtFunction Function = (UhtFunction)SpecifierContext.Scope.ScopeType;
-			Function.FunctionFlags |= EFunctionFlags.BlueprintCosmetic;
+			UhtFunction function = (UhtFunction)specifierContext.Scope.ScopeType;
+			function.FunctionFlags |= EFunctionFlags.BlueprintCosmetic;
 		}
 
 		[UhtSpecifier(Extends = UhtTableNames.Function, ValueType = UhtSpecifierValueType.Legacy)]
 		[SuppressMessage("CodeQuality", "IDE0051:Remove unused private members", Justification = "Attribute accessed method")]
-		[SuppressMessage("Style", "IDE0060:Remove unused parameter", Justification = "Attribute accessed method")]
-		private static void BlueprintGetterSpecifier(UhtSpecifierContext SpecifierContext)
+		private static void BlueprintGetterSpecifier(UhtSpecifierContext specifierContext)
 		{
-			UhtFunctionParser Function = (UhtFunctionParser)SpecifierContext.Scope.ScopeType;
+			UhtFunctionParser function = (UhtFunctionParser)specifierContext.Scope.ScopeType;
 
-			if (Function.FunctionFlags.HasAnyFlags(EFunctionFlags.Event))
+			if (function.FunctionFlags.HasAnyFlags(EFunctionFlags.Event))
 			{
-				SpecifierContext.MessageSite.LogError("Function cannot be a blueprint event and a blueprint getter.");
+				specifierContext.MessageSite.LogError("Function cannot be a blueprint event and a blueprint getter.");
 			}
 
-			Function.bSawPropertyAccessor = true;
-			Function.FunctionFlags |= EFunctionFlags.BlueprintCallable;
-			Function.FunctionFlags |= EFunctionFlags.BlueprintPure;
-			Function.MetaData.Add(UhtNames.BlueprintGetter, "");
+			function.SawPropertyAccessor = true;
+			function.FunctionFlags |= EFunctionFlags.BlueprintCallable;
+			function.FunctionFlags |= EFunctionFlags.BlueprintPure;
+			function.MetaData.Add(UhtNames.BlueprintGetter, "");
 		}
 
 		[UhtSpecifier(Extends = UhtTableNames.Function, ValueType = UhtSpecifierValueType.Legacy)]
 		[SuppressMessage("CodeQuality", "IDE0051:Remove unused private members", Justification = "Attribute accessed method")]
-		[SuppressMessage("Style", "IDE0060:Remove unused parameter", Justification = "Attribute accessed method")]
-		private static void BlueprintImplementableEventSpecifier(UhtSpecifierContext SpecifierContext)
+		private static void BlueprintImplementableEventSpecifier(UhtSpecifierContext specifierContext)
 		{
-			UhtFunctionParser Function = (UhtFunctionParser)SpecifierContext.Scope.ScopeType;
+			UhtFunctionParser function = (UhtFunctionParser)specifierContext.Scope.ScopeType;
 
-			if (Function.FunctionFlags.HasAnyFlags(EFunctionFlags.Net))
+			if (function.FunctionFlags.HasAnyFlags(EFunctionFlags.Net))
 			{
-				SpecifierContext.MessageSite.LogError("BlueprintImplementableEvent functions cannot be replicated!");
+				specifierContext.MessageSite.LogError("BlueprintImplementableEvent functions cannot be replicated!");
 			}
-			else if (Function.FunctionFlags.HasAnyFlags(EFunctionFlags.BlueprintEvent) && Function.FunctionFlags.HasAnyFlags(EFunctionFlags.Native))
+			else if (function.FunctionFlags.HasAnyFlags(EFunctionFlags.BlueprintEvent) && function.FunctionFlags.HasAnyFlags(EFunctionFlags.Native))
 			{
 				// already a BlueprintNativeEvent
-				SpecifierContext.MessageSite.LogError("A function cannot be both BlueprintNativeEvent and BlueprintImplementableEvent!");
+				specifierContext.MessageSite.LogError("A function cannot be both BlueprintNativeEvent and BlueprintImplementableEvent!");
 			}
-			else if (Function.bSawPropertyAccessor)
+			else if (function.SawPropertyAccessor)
 			{
-				SpecifierContext.MessageSite.LogError("A function cannot be both BlueprintImplementableEvent and a Blueprint Property accessor!");
+				specifierContext.MessageSite.LogError("A function cannot be both BlueprintImplementableEvent and a Blueprint Property accessor!");
 			}
-			else if (Function.FunctionFlags.HasAnyFlags(EFunctionFlags.Private))
+			else if (function.FunctionFlags.HasAnyFlags(EFunctionFlags.Private))
 			{
-				SpecifierContext.MessageSite.LogError("A Private function cannot be a BlueprintImplementableEvent!");
+				specifierContext.MessageSite.LogError("A Private function cannot be a BlueprintImplementableEvent!");
 			}
 
-			Function.FunctionFlags |= EFunctionFlags.Event;
-			Function.FunctionFlags |= EFunctionFlags.BlueprintEvent;
-			Function.FunctionFlags &= ~EFunctionFlags.Native;
+			function.FunctionFlags |= EFunctionFlags.Event;
+			function.FunctionFlags |= EFunctionFlags.BlueprintEvent;
+			function.FunctionFlags &= ~EFunctionFlags.Native;
 		}
 
 		[UhtSpecifier(Extends = UhtTableNames.Function, ValueType = UhtSpecifierValueType.Legacy)]
 		[SuppressMessage("CodeQuality", "IDE0051:Remove unused private members", Justification = "Attribute accessed method")]
-		[SuppressMessage("Style", "IDE0060:Remove unused parameter", Justification = "Attribute accessed method")]
-		private static void BlueprintNativeEventSpecifier(UhtSpecifierContext SpecifierContext)
+		private static void BlueprintNativeEventSpecifier(UhtSpecifierContext specifierContext)
 		{
-			UhtFunctionParser Function = (UhtFunctionParser)SpecifierContext.Scope.ScopeType;
+			UhtFunctionParser function = (UhtFunctionParser)specifierContext.Scope.ScopeType;
 
-			if (Function.FunctionFlags.HasAnyFlags(EFunctionFlags.Net))
+			if (function.FunctionFlags.HasAnyFlags(EFunctionFlags.Net))
 			{
-				SpecifierContext.MessageSite.LogError("BlueprintNativeEvent functions cannot be replicated!");
+				specifierContext.MessageSite.LogError("BlueprintNativeEvent functions cannot be replicated!");
 			}
-			else if (Function.FunctionFlags.HasAnyFlags(EFunctionFlags.BlueprintEvent) && !Function.FunctionFlags.HasAnyFlags(EFunctionFlags.Native))
+			else if (function.FunctionFlags.HasAnyFlags(EFunctionFlags.BlueprintEvent) && !function.FunctionFlags.HasAnyFlags(EFunctionFlags.Native))
 			{
 				// already a BlueprintImplementableEvent
-				SpecifierContext.MessageSite.LogError("A function cannot be both BlueprintNativeEvent and BlueprintImplementableEvent!");
+				specifierContext.MessageSite.LogError("A function cannot be both BlueprintNativeEvent and BlueprintImplementableEvent!");
 			}
-			else if (Function.bSawPropertyAccessor)
+			else if (function.SawPropertyAccessor)
 			{
-				SpecifierContext.MessageSite.LogError("A function cannot be both BlueprintNativeEvent and a Blueprint Property accessor!");
+				specifierContext.MessageSite.LogError("A function cannot be both BlueprintNativeEvent and a Blueprint Property accessor!");
 			}
-			else if (Function.FunctionFlags.HasAnyFlags(EFunctionFlags.Private))
+			else if (function.FunctionFlags.HasAnyFlags(EFunctionFlags.Private))
 			{
-				SpecifierContext.MessageSite.LogError("A Private function cannot be a BlueprintNativeEvent!");
+				specifierContext.MessageSite.LogError("A Private function cannot be a BlueprintNativeEvent!");
 			}
 
-			Function.FunctionFlags |= EFunctionFlags.Event;
-			Function.FunctionFlags |= EFunctionFlags.BlueprintEvent;
+			function.FunctionFlags |= EFunctionFlags.Event;
+			function.FunctionFlags |= EFunctionFlags.BlueprintEvent;
 		}
 
 		[UhtSpecifier(Extends = UhtTableNames.Function, ValueType = UhtSpecifierValueType.OptionalString)]
 		[SuppressMessage("CodeQuality", "IDE0051:Remove unused private members", Justification = "Attribute accessed method")]
-		[SuppressMessage("Style", "IDE0060:Remove unused parameter", Justification = "Attribute accessed method")]
-		private static void BlueprintPureSpecifier(UhtSpecifierContext SpecifierContext, StringView? Value)
+		private static void BlueprintPureSpecifier(UhtSpecifierContext specifierContext, StringView? value)
 		{
-			UhtFunction Function = (UhtFunction)SpecifierContext.Scope.ScopeType;
+			UhtFunction function = (UhtFunction)specifierContext.Scope.ScopeType;
 
 			// This function can be called, and is also pure.
-			Function.FunctionFlags |= EFunctionFlags.BlueprintCallable;
+			function.FunctionFlags |= EFunctionFlags.BlueprintCallable;
 
-			if (Value == null || UhtFCString.ToBool((StringView)Value))
+			if (value == null || UhtFCString.ToBool((StringView)value))
 			{
-				Function.FunctionFlags |= EFunctionFlags.BlueprintPure;
+				function.FunctionFlags |= EFunctionFlags.BlueprintPure;
 			}
 			else
 			{
-				Function.FunctionExportFlags |= UhtFunctionExportFlags.ForceBlueprintImpure;
+				function.FunctionExportFlags |= UhtFunctionExportFlags.ForceBlueprintImpure;
 			}
 		}
 
 		[UhtSpecifier(Extends = UhtTableNames.Function, ValueType = UhtSpecifierValueType.Legacy)]
 		[SuppressMessage("CodeQuality", "IDE0051:Remove unused private members", Justification = "Attribute accessed method")]
-		[SuppressMessage("Style", "IDE0060:Remove unused parameter", Justification = "Attribute accessed method")]
-		private static void BlueprintSetterSpecifier(UhtSpecifierContext SpecifierContext)
+		private static void BlueprintSetterSpecifier(UhtSpecifierContext specifierContext)
 		{
-			UhtFunctionParser Function = (UhtFunctionParser)SpecifierContext.Scope.ScopeType;
+			UhtFunctionParser function = (UhtFunctionParser)specifierContext.Scope.ScopeType;
 
-			if (Function.FunctionFlags.HasAnyFlags(EFunctionFlags.Event))
+			if (function.FunctionFlags.HasAnyFlags(EFunctionFlags.Event))
 			{
-				SpecifierContext.MessageSite.LogError("Function cannot be a blueprint event and a blueprint setter.");
+				specifierContext.MessageSite.LogError("Function cannot be a blueprint event and a blueprint setter.");
 			}
 
-			Function.bSawPropertyAccessor = true;
-			Function.FunctionFlags |= EFunctionFlags.BlueprintCallable;
-			Function.MetaData.Add(UhtNames.BlueprintSetter, "");
+			function.SawPropertyAccessor = true;
+			function.FunctionFlags |= EFunctionFlags.BlueprintCallable;
+			function.MetaData.Add(UhtNames.BlueprintSetter, "");
 		}
 
 		[UhtSpecifier(Extends = UhtTableNames.Function, ValueType = UhtSpecifierValueType.OptionalString)]
 		[SuppressMessage("CodeQuality", "IDE0051:Remove unused private members", Justification = "Attribute accessed method")]
-		[SuppressMessage("Style", "IDE0060:Remove unused parameter", Justification = "Attribute accessed method")]
-		private static void ClientSpecifier(UhtSpecifierContext SpecifierContext, StringView? Value)
+		private static void ClientSpecifier(UhtSpecifierContext specifierContext, StringView? value)
 		{
-			UhtFunction Function = (UhtFunction)SpecifierContext.Scope.ScopeType;
+			UhtFunction function = (UhtFunction)specifierContext.Scope.ScopeType;
 
-			if (Function.FunctionFlags.HasAnyFlags(EFunctionFlags.BlueprintEvent))
+			if (function.FunctionFlags.HasAnyFlags(EFunctionFlags.BlueprintEvent))
 			{
-				SpecifierContext.MessageSite.LogError("BlueprintImplementableEvent or BlueprintNativeEvent functions cannot be declared as Client or Server");
+				specifierContext.MessageSite.LogError("BlueprintImplementableEvent or BlueprintNativeEvent functions cannot be declared as Client or Server");
 			}
 
-			if (Function.FunctionFlags.HasAnyFlags(EFunctionFlags.Exec))
+			if (function.FunctionFlags.HasAnyFlags(EFunctionFlags.Exec))
 			{
-				SpecifierContext.MessageSite.LogError("Exec functions cannot be replicated!");
+				specifierContext.MessageSite.LogError("Exec functions cannot be replicated!");
 			}
 
-			Function.FunctionFlags |= EFunctionFlags.Net;
-			Function.FunctionFlags |= EFunctionFlags.NetClient;
+			function.FunctionFlags |= EFunctionFlags.Net;
+			function.FunctionFlags |= EFunctionFlags.NetClient;
 
-			if (Value != null)
+			if (value != null)
 			{
-				Function.CppImplName = ((StringView)Value).ToString();
+				function.CppImplName = ((StringView)value).ToString();
 			}
 		}
 
 		[UhtSpecifier(Extends = UhtTableNames.Function, ValueType = UhtSpecifierValueType.Legacy)]
 		[SuppressMessage("CodeQuality", "IDE0051:Remove unused private members", Justification = "Attribute accessed method")]
-		[SuppressMessage("Style", "IDE0060:Remove unused parameter", Justification = "Attribute accessed method")]
-		private static void CustomThunkSpecifier(UhtSpecifierContext SpecifierContext)
+		private static void CustomThunkSpecifier(UhtSpecifierContext specifierContext)
 		{
-			UhtFunction Function = (UhtFunction)SpecifierContext.Scope.ScopeType;
-			Function.FunctionExportFlags |= UhtFunctionExportFlags.CustomThunk;
+			UhtFunction function = (UhtFunction)specifierContext.Scope.ScopeType;
+			function.FunctionExportFlags |= UhtFunctionExportFlags.CustomThunk;
 		}
 
 		[UhtSpecifier(Extends = UhtTableNames.Function, ValueType = UhtSpecifierValueType.Legacy)]
 		[SuppressMessage("CodeQuality", "IDE0051:Remove unused private members", Justification = "Attribute accessed method")]
-		[SuppressMessage("Style", "IDE0060:Remove unused parameter", Justification = "Attribute accessed method")]
-		private static void ExecSpecifier(UhtSpecifierContext SpecifierContext)
+		private static void ExecSpecifier(UhtSpecifierContext specifierContext)
 		{
-			UhtFunction Function = (UhtFunction)SpecifierContext.Scope.ScopeType;
-			Function.FunctionFlags |= EFunctionFlags.Exec;
-			if (Function.FunctionFlags.HasAnyFlags(EFunctionFlags.Net))
+			UhtFunction function = (UhtFunction)specifierContext.Scope.ScopeType;
+			function.FunctionFlags |= EFunctionFlags.Exec;
+			if (function.FunctionFlags.HasAnyFlags(EFunctionFlags.Net))
 			{
-				SpecifierContext.MessageSite.LogError("Exec functions cannot be replicated!");
+				specifierContext.MessageSite.LogError("Exec functions cannot be replicated!");
 			}
 		}
 
 		[UhtSpecifier(Extends = UhtTableNames.Function, ValueType = UhtSpecifierValueType.Legacy)]
 		[SuppressMessage("CodeQuality", "IDE0051:Remove unused private members", Justification = "Attribute accessed method")]
-		[SuppressMessage("Style", "IDE0060:Remove unused parameter", Justification = "Attribute accessed method")]
-		private static void NetMulticastSpecifier(UhtSpecifierContext SpecifierContext)
+		private static void NetMulticastSpecifier(UhtSpecifierContext specifierContext)
 		{
-			UhtFunction Function = (UhtFunction)SpecifierContext.Scope.ScopeType;
-			if (Function.FunctionFlags.HasAnyFlags(EFunctionFlags.BlueprintEvent))
+			UhtFunction function = (UhtFunction)specifierContext.Scope.ScopeType;
+			if (function.FunctionFlags.HasAnyFlags(EFunctionFlags.BlueprintEvent))
 			{
-				SpecifierContext.MessageSite.LogError("BlueprintImplementableEvent or BlueprintNativeEvent functions cannot be declared as Multicast");
+				specifierContext.MessageSite.LogError("BlueprintImplementableEvent or BlueprintNativeEvent functions cannot be declared as Multicast");
 			}
 
-			Function.FunctionFlags |= EFunctionFlags.Net;
-			Function.FunctionFlags |= EFunctionFlags.NetMulticast;
+			function.FunctionFlags |= EFunctionFlags.Net;
+			function.FunctionFlags |= EFunctionFlags.NetMulticast;
 		}
 
 		[UhtSpecifier(Extends = UhtTableNames.Function, ValueType = UhtSpecifierValueType.Legacy)]
 		[SuppressMessage("CodeQuality", "IDE0051:Remove unused private members", Justification = "Attribute accessed method")]
-		[SuppressMessage("Style", "IDE0060:Remove unused parameter", Justification = "Attribute accessed method")]
-		private static void ReliableSpecifier(UhtSpecifierContext SpecifierContext)
+		private static void ReliableSpecifier(UhtSpecifierContext specifierContext)
 		{
-			UhtFunction Function = (UhtFunction)SpecifierContext.Scope.ScopeType;
-			Function.FunctionFlags |= EFunctionFlags.NetReliable;
+			UhtFunction function = (UhtFunction)specifierContext.Scope.ScopeType;
+			function.FunctionFlags |= EFunctionFlags.NetReliable;
 		}
 
 		[UhtSpecifier(Extends = UhtTableNames.Function, ValueType = UhtSpecifierValueType.Legacy)]
 		[SuppressMessage("CodeQuality", "IDE0051:Remove unused private members", Justification = "Attribute accessed method")]
-		[SuppressMessage("Style", "IDE0060:Remove unused parameter", Justification = "Attribute accessed method")]
-		private static void SealedEventSpecifier(UhtSpecifierContext SpecifierContext)
+		private static void SealedEventSpecifier(UhtSpecifierContext specifierContext)
 		{
-			UhtFunction Function = (UhtFunction)SpecifierContext.Scope.ScopeType;
-			Function.FunctionExportFlags |= UhtFunctionExportFlags.SealedEvent;
+			UhtFunction function = (UhtFunction)specifierContext.Scope.ScopeType;
+			function.FunctionExportFlags |= UhtFunctionExportFlags.SealedEvent;
 		}
 
 		[UhtSpecifier(Extends = UhtTableNames.Function, ValueType = UhtSpecifierValueType.OptionalString)]
 		[SuppressMessage("CodeQuality", "IDE0051:Remove unused private members", Justification = "Attribute accessed method")]
-		[SuppressMessage("Style", "IDE0060:Remove unused parameter", Justification = "Attribute accessed method")]
-		private static void ServerSpecifier(UhtSpecifierContext SpecifierContext, StringView? Value)
+		private static void ServerSpecifier(UhtSpecifierContext specifierContext, StringView? value)
 		{
-			UhtFunction Function = (UhtFunction)SpecifierContext.Scope.ScopeType;
+			UhtFunction function = (UhtFunction)specifierContext.Scope.ScopeType;
 
-			if (Function.FunctionFlags.HasAnyFlags(EFunctionFlags.BlueprintEvent))
+			if (function.FunctionFlags.HasAnyFlags(EFunctionFlags.BlueprintEvent))
 			{
-				SpecifierContext.MessageSite.LogError("BlueprintImplementableEvent or BlueprintNativeEvent functions cannot be declared as Client or Server");
+				specifierContext.MessageSite.LogError("BlueprintImplementableEvent or BlueprintNativeEvent functions cannot be declared as Client or Server");
 			}
 
-			if (Function.FunctionFlags.HasAnyFlags(EFunctionFlags.Exec))
+			if (function.FunctionFlags.HasAnyFlags(EFunctionFlags.Exec))
 			{
-				SpecifierContext.MessageSite.LogError("Exec functions cannot be replicated!");
+				specifierContext.MessageSite.LogError("Exec functions cannot be replicated!");
 			}
 
-			Function.FunctionFlags |= EFunctionFlags.Net;
-			Function.FunctionFlags |= EFunctionFlags.NetServer;
+			function.FunctionFlags |= EFunctionFlags.Net;
+			function.FunctionFlags |= EFunctionFlags.NetServer;
 
-			if (Value != null)
+			if (value != null)
 			{
-				Function.CppImplName = ((StringView)Value).ToString();
+				function.CppImplName = ((StringView)value).ToString();
 			}
 		}
 
 		[UhtSpecifier(Extends = UhtTableNames.Function, ValueType = UhtSpecifierValueType.OptionalEqualsKeyValuePairList)]
 		[SuppressMessage("CodeQuality", "IDE0051:Remove unused private members", Justification = "Attribute accessed method")]
-		[SuppressMessage("Style", "IDE0060:Remove unused parameter", Justification = "Attribute accessed method")]
-		private static void ServiceRequestSpecifier(UhtSpecifierContext SpecifierContext, List<KeyValuePair<StringView, StringView>> Value)
+		private static void ServiceRequestSpecifier(UhtSpecifierContext specifierContext, List<KeyValuePair<StringView, StringView>> value)
 		{
-			UhtFunctionParser Function = (UhtFunctionParser)SpecifierContext.Scope.ScopeType;
+			UhtFunctionParser function = (UhtFunctionParser)specifierContext.Scope.ScopeType;
 
-			if (Function.FunctionFlags.HasAnyFlags(EFunctionFlags.BlueprintEvent))
+			if (function.FunctionFlags.HasAnyFlags(EFunctionFlags.BlueprintEvent))
 			{
-				SpecifierContext.MessageSite.LogError("BlueprintImplementableEvent or BlueprintNativeEvent functions cannot be declared as a ServiceRequest");
+				specifierContext.MessageSite.LogError("BlueprintImplementableEvent or BlueprintNativeEvent functions cannot be declared as a ServiceRequest");
 			}
 
-			Function.FunctionFlags |= EFunctionFlags.Net;
-			Function.FunctionFlags |= EFunctionFlags.NetReliable;
-			Function.FunctionFlags |= EFunctionFlags.NetRequest;
-			Function.FunctionExportFlags |= UhtFunctionExportFlags.CustomThunk;
+			function.FunctionFlags |= EFunctionFlags.Net;
+			function.FunctionFlags |= EFunctionFlags.NetReliable;
+			function.FunctionFlags |= EFunctionFlags.NetRequest;
+			function.FunctionExportFlags |= UhtFunctionExportFlags.CustomThunk;
 
-			ParseNetServiceIdentifiers(SpecifierContext, Value);
+			ParseNetServiceIdentifiers(specifierContext, value);
 
-			if (string.IsNullOrEmpty(Function.EndpointName))
+			if (String.IsNullOrEmpty(function.EndpointName))
 			{
-				SpecifierContext.MessageSite.LogError("ServiceRequest needs to specify an endpoint name");
+				specifierContext.MessageSite.LogError("ServiceRequest needs to specify an endpoint name");
 			}
 		}
 
 		[UhtSpecifier(Extends = UhtTableNames.Function, ValueType = UhtSpecifierValueType.OptionalEqualsKeyValuePairList)]
 		[SuppressMessage("CodeQuality", "IDE0051:Remove unused private members", Justification = "Attribute accessed method")]
-		[SuppressMessage("Style", "IDE0060:Remove unused parameter", Justification = "Attribute accessed method")]
-		private static void ServiceResponseSpecifier(UhtSpecifierContext SpecifierContext, List<KeyValuePair<StringView, StringView>> Value)
+		private static void ServiceResponseSpecifier(UhtSpecifierContext specifierContext, List<KeyValuePair<StringView, StringView>> value)
 		{
-			UhtFunctionParser Function = (UhtFunctionParser)SpecifierContext.Scope.ScopeType;
+			UhtFunctionParser function = (UhtFunctionParser)specifierContext.Scope.ScopeType;
 
-			if (Function.FunctionFlags.HasAnyFlags(EFunctionFlags.BlueprintEvent))
+			if (function.FunctionFlags.HasAnyFlags(EFunctionFlags.BlueprintEvent))
 			{
-				SpecifierContext.MessageSite.LogError("BlueprintImplementableEvent or BlueprintNativeEvent functions cannot be declared as a ServiceResponse");
+				specifierContext.MessageSite.LogError("BlueprintImplementableEvent or BlueprintNativeEvent functions cannot be declared as a ServiceResponse");
 			}
 
-			Function.FunctionFlags |= EFunctionFlags.Net;
-			Function.FunctionFlags |= EFunctionFlags.NetReliable;
-			Function.FunctionFlags |= EFunctionFlags.NetResponse;
+			function.FunctionFlags |= EFunctionFlags.Net;
+			function.FunctionFlags |= EFunctionFlags.NetReliable;
+			function.FunctionFlags |= EFunctionFlags.NetResponse;
 			//Function.FunctionExportFlags |= EFunctionExportFlags.CustomThunk;
 
-			ParseNetServiceIdentifiers(SpecifierContext, Value);
+			ParseNetServiceIdentifiers(specifierContext, value);
 
-			if (string.IsNullOrEmpty(Function.EndpointName))
+			if (String.IsNullOrEmpty(function.EndpointName))
 			{
-				SpecifierContext.MessageSite.LogError("ServiceResponse needs to specify an endpoint name");
+				specifierContext.MessageSite.LogError("ServiceResponse needs to specify an endpoint name");
 			}
 		}
 
 		[UhtSpecifier(Extends = UhtTableNames.Function, ValueType = UhtSpecifierValueType.Legacy)]
 		[SuppressMessage("CodeQuality", "IDE0051:Remove unused private members", Justification = "Attribute accessed method")]
-		[SuppressMessage("Style", "IDE0060:Remove unused parameter", Justification = "Attribute accessed method")]
-		private static void UnreliableSpecifier(UhtSpecifierContext SpecifierContext)
+		private static void UnreliableSpecifier(UhtSpecifierContext specifierContext)
 		{
-			UhtFunction Function = (UhtFunction)SpecifierContext.Scope.ScopeType;
-			Function.FunctionExportFlags |= UhtFunctionExportFlags.Unreliable;
+			UhtFunction function = (UhtFunction)specifierContext.Scope.ScopeType;
+			function.FunctionExportFlags |= UhtFunctionExportFlags.Unreliable;
 		}
 
 		[UhtSpecifier(Extends = UhtTableNames.Function, ValueType = UhtSpecifierValueType.StringList)]
 		[SuppressMessage("CodeQuality", "IDE0051:Remove unused private members", Justification = "Attribute accessed method")]
-		[SuppressMessage("Style", "IDE0060:Remove unused parameter", Justification = "Attribute accessed method")]
-		private static void WithValidationSpecifier(UhtSpecifierContext SpecifierContext, List<StringView>? Value)
+		private static void WithValidationSpecifier(UhtSpecifierContext specifierContext, List<StringView>? value)
 		{
-			UhtFunction Function = (UhtFunction)SpecifierContext.Scope.ScopeType;
-			Function.FunctionFlags |= EFunctionFlags.NetValidate;
+			UhtFunction function = (UhtFunction)specifierContext.Scope.ScopeType;
+			function.FunctionFlags |= EFunctionFlags.NetValidate;
 
-			if (Value != null && Value.Count > 0)
+			if (value != null && value.Count > 0)
 			{
-				Function.CppValidationImplName = Value[0].ToString();
+				function.CppValidationImplName = value[0].ToString();
 			}
 		}
 
 		[UhtSpecifier(Extends = UhtTableNames.Function, ValueType = UhtSpecifierValueType.None)]
 		[SuppressMessage("CodeQuality", "IDE0051:Remove unused private members", Justification = "Attribute accessed method")]
-		[SuppressMessage("Style", "IDE0060:Remove unused parameter", Justification = "Attribute accessed method")]
-		private static void FieldNotifySpecifier(UhtSpecifierContext SpecifierContext)
+		private static void FieldNotifySpecifier(UhtSpecifierContext specifierContext)
 		{
-			UhtFunction Function = (UhtFunction)SpecifierContext.Scope.ScopeType;
-			Function.FunctionExportFlags |= UhtFunctionExportFlags.FieldNotify;
+			UhtFunction function = (UhtFunction)specifierContext.Scope.ScopeType;
+			function.FunctionExportFlags |= UhtFunctionExportFlags.FieldNotify;
 		}
 
 		[UhtSpecifierValidator(Extends = UhtTableNames.Function)]
 		[SuppressMessage("CodeQuality", "IDE0051:Remove unused private members", Justification = "Attribute accessed method")]
 		[SuppressMessage("Style", "IDE0060:Remove unused parameter", Justification = "Attribute accessed method")]
-		private static void BlueprintProtectedSpecifierValidator(UhtType Type, UhtMetaData MetaData, UhtMetaDataKey Key, StringView Value)
+		private static void BlueprintProtectedSpecifierValidator(UhtType type, UhtMetaData metaData, UhtMetaDataKey key, StringView value)
 		{
-			UhtFunction Function = (UhtFunction)Type;
-			if (Function.FunctionFlags.HasAnyFlags(EFunctionFlags.Static))
+			UhtFunction function = (UhtFunction)type;
+			if (function.FunctionFlags.HasAnyFlags(EFunctionFlags.Static))
 			{
 				// Given the owning class, locate the class that the owning derives from up to the point of UObject
-				if (Function.Outer is UhtClass OuterClass)
+				if (function.Outer is UhtClass outerClass)
 				{
-					UhtClass? ClassPriorToUObject = OuterClass;
-					for (; ClassPriorToUObject != null; ClassPriorToUObject = ClassPriorToUObject.SuperClass)
+					UhtClass? classPriorToUObject = outerClass;
+					for (; classPriorToUObject != null; classPriorToUObject = classPriorToUObject.SuperClass)
 					{
 						// If our super is UObject, then stop
-						if (ClassPriorToUObject.Super == Type.Session.UObject)
+						if (classPriorToUObject.Super == type.Session.UObject)
 						{
 							break;
 						}
 					}
 
-					if (ClassPriorToUObject != null && ClassPriorToUObject.SourceName == "UBlueprintFunctionLibrary")
+					if (classPriorToUObject != null && classPriorToUObject.SourceName == "UBlueprintFunctionLibrary")
 					{
-						Type.LogError(MetaData.LineNumber, $"'{Key.Name}' doesn't make sense on static method '{Type.SourceName}' in a blueprint function library");
+						type.LogError(metaData.LineNumber, $"'{key.Name}' doesn't make sense on static method '{type.SourceName}' in a blueprint function library");
 					}
 				}
 			}
@@ -379,15 +359,15 @@ namespace EpicGames.UHT.Parsers
 		[UhtSpecifierValidator(Extends = UhtTableNames.Function)]
 		[SuppressMessage("CodeQuality", "IDE0051:Remove unused private members", Justification = "Attribute accessed method")]
 		[SuppressMessage("Style", "IDE0060:Remove unused parameter", Justification = "Attribute accessed method")]
-		private static void CommutativeAssociativeBinaryOperatorSpecifierValidator(UhtType Type, UhtMetaData MetaData, UhtMetaDataKey Key, StringView Value)
+		private static void CommutativeAssociativeBinaryOperatorSpecifierValidator(UhtType type, UhtMetaData metaData, UhtMetaDataKey key, StringView value)
 		{
-			UhtFunction Function = (UhtFunction)Type;
-			UhtProperty? ReturnProperty = Function.ReturnProperty;
-			ReadOnlyMemory<UhtType> ParameterProperties = Function.ParameterProperties;
+			UhtFunction function = (UhtFunction)type;
+			UhtProperty? returnProperty = function.ReturnProperty;
+			ReadOnlyMemory<UhtType> parameterProperties = function.ParameterProperties;
 
-			if (ReturnProperty == null || ParameterProperties.Length != 2 || !((UhtProperty)ParameterProperties.Span[0]).IsSameType((UhtProperty)ParameterProperties.Span[1]))
+			if (returnProperty == null || parameterProperties.Length != 2 || !((UhtProperty)parameterProperties.Span[0]).IsSameType((UhtProperty)parameterProperties.Span[1]))
 			{
-				Function.LogError("Commutative associative binary operators must have exactly 2 parameters of the same type and a return value.");
+				function.LogError("Commutative associative binary operators must have exactly 2 parameters of the same type and a return value.");
 			}
 		}
 
@@ -395,35 +375,35 @@ namespace EpicGames.UHT.Parsers
 		[UhtSpecifierValidator(Name = "ExpandEnumAsExecs", Extends = UhtTableNames.Function)]
 		[SuppressMessage("CodeQuality", "IDE0051:Remove unused private members", Justification = "Attribute accessed method")]
 		[SuppressMessage("Style", "IDE0060:Remove unused parameter", Justification = "Attribute accessed method")]
-		private static void ExpandsSpecifierValidator(UhtType Type, UhtMetaData MetaData, UhtMetaDataKey Key, StringView Value)
+		private static void ExpandsSpecifierValidator(UhtType type, UhtMetaData metaData, UhtMetaDataKey key, StringView value)
 		{
 			// multiple entry parsing in the same format as eg SetParam.
-			UhtType? FirstInput = null;
-			foreach (string RawGroup in Value.ToString().Split(','))
+			UhtType? firstInput = null;
+			foreach (string rawGroup in value.ToString().Split(','))
 			{
-				foreach (string Entry in RawGroup.Split('|'))
+				foreach (string entry in rawGroup.Split('|'))
 				{
-					string Trimmed = Entry.Trim();
-					if (string.IsNullOrEmpty(Trimmed))
+					string trimmed = entry.Trim();
+					if (String.IsNullOrEmpty(trimmed))
 					{
 						continue;
 					}
 
-					UhtType? FoundField = Type.FindType(UhtFindOptions.SourceName | UhtFindOptions.SelfOnly | UhtFindOptions.Property, Trimmed, Type);
-					if (FoundField != null)
+					UhtType? foundField = type.FindType(UhtFindOptions.SourceName | UhtFindOptions.SelfOnly | UhtFindOptions.Property, trimmed, type);
+					if (foundField != null)
 					{
-						UhtProperty Property = (UhtProperty)FoundField;
-						if (!Property.PropertyFlags.HasAnyFlags(EPropertyFlags.ReturnParm) &&
-							(!Property.PropertyFlags.HasAnyFlags(EPropertyFlags.OutParm) ||
-							 Property.PropertyFlags.HasAnyFlags(EPropertyFlags.ReferenceParm)))
+						UhtProperty property = (UhtProperty)foundField;
+						if (!property.PropertyFlags.HasAnyFlags(EPropertyFlags.ReturnParm) &&
+							(!property.PropertyFlags.HasAnyFlags(EPropertyFlags.OutParm) ||
+							 property.PropertyFlags.HasAnyFlags(EPropertyFlags.ReferenceParm)))
 						{
-							if (FirstInput == null)
+							if (firstInput == null)
 							{
-								FirstInput = FoundField;
+								firstInput = foundField;
 							}
 							else
 							{
-								Type.LogError($"Function already specified an ExpandEnumAsExec input '{FirstInput.SourceName}', but '{Trimmed}' is also an input parameter. Only one is permitted.");
+								type.LogError($"Function already specified an ExpandEnumAsExec input '{firstInput.SourceName}', but '{trimmed}' is also an input parameter. Only one is permitted.");
 							}
 						}
 					}
@@ -431,37 +411,37 @@ namespace EpicGames.UHT.Parsers
 			}
 		}
 
-		private static void ParseNetServiceIdentifiers(UhtSpecifierContext SpecifierContext, List<KeyValuePair<StringView, StringView>> Identifiers)
+		private static void ParseNetServiceIdentifiers(UhtSpecifierContext specifierContext, List<KeyValuePair<StringView, StringView>> identifiers)
 		{
-			IUhtTokenReader TokenReader = SpecifierContext.Scope.TokenReader;
-			UhtFunctionParser Function = (UhtFunctionParser)SpecifierContext.Scope.ScopeType;
+			IUhtTokenReader tokenReader = specifierContext.Scope.TokenReader;
+			UhtFunctionParser function = (UhtFunctionParser)specifierContext.Scope.ScopeType;
 
-			foreach (var KVP in Identifiers)
+			foreach (KeyValuePair<StringView, StringView> kvp in identifiers)
 			{
-				if (KVP.Value.Span.Length > 0)
+				if (kvp.Value.Span.Length > 0)
 				{
-					if (KVP.Key.Span.StartsWith("Id", StringComparison.OrdinalIgnoreCase))
+					if (kvp.Key.Span.StartsWith("Id", StringComparison.OrdinalIgnoreCase))
 					{
-						int Id;
-						if (!int.TryParse(KVP.Value.Span, out Id) || Id <= 0 || Id > UInt16.MaxValue)
+						int id;
+						if (!Int32.TryParse(kvp.Value.Span, out id) || id <= 0 || id > UInt16.MaxValue)
 						{
-							TokenReader.LogError($"Invalid network identifier {KVP.Key} for function");
+							tokenReader.LogError($"Invalid network identifier {kvp.Key} for function");
 						}
 						else
 						{
-							Function.RPCId = (UInt16)Id;
+							function.RPCId = (ushort)id;
 						}
 					}
-					else if (KVP.Key.Span.StartsWith("ResponseId", StringComparison.OrdinalIgnoreCase) || KVP.Key.Span.StartsWith("Priority", StringComparison.OrdinalIgnoreCase))
+					else if (kvp.Key.Span.StartsWith("ResponseId", StringComparison.OrdinalIgnoreCase) || kvp.Key.Span.StartsWith("Priority", StringComparison.OrdinalIgnoreCase))
 					{
-						int Id;
-						if (!int.TryParse(KVP.Value.Span, out Id) || Id <= 0 || Id > UInt16.MaxValue)
+						int id;
+						if (!Int32.TryParse(kvp.Value.Span, out id) || id <= 0 || id > UInt16.MaxValue)
 						{
-							TokenReader.LogError($"Invalid network identifier {KVP.Key} for function");
+							tokenReader.LogError($"Invalid network identifier {kvp.Key} for function");
 						}
 						else
 						{
-							Function.RPCResponseId = (UInt16)Id;
+							function.RPCResponseId = (ushort)id;
 						}
 					}
 					else
@@ -473,13 +453,13 @@ namespace EpicGames.UHT.Parsers
 				// Assume it's an endpoint name
 				else
 				{
-					if (!string.IsNullOrEmpty(Function.EndpointName))
+					if (!String.IsNullOrEmpty(function.EndpointName))
 					{
-						TokenReader.LogError($"Function should not specify multiple endpoints - '{KVP.Key}' found but already using '{Function.EndpointName}'");
+						tokenReader.LogError($"Function should not specify multiple endpoints - '{kvp.Key}' found but already using '{function.EndpointName}'");
 					}
 					else
 					{
-						Function.EndpointName = KVP.Key.ToString();
+						function.EndpointName = kvp.Key.ToString();
 					}
 				}
 			}
diff --git a/Engine/Source/Programs/Shared/EpicGames.UHT/Specifiers/UhtInterfaceClassSpecifiers.cs b/Engine/Source/Programs/Shared/EpicGames.UHT/Specifiers/UhtInterfaceClassSpecifiers.cs
--- a/Engine/Source/Programs/Shared/EpicGames.UHT/Specifiers/UhtInterfaceClassSpecifiers.cs
+++ b/Engine/Source/Programs/Shared/EpicGames.UHT/Specifiers/UhtInterfaceClassSpecifiers.cs
@@ -13,25 +13,25 @@ namespace EpicGames.UHT.Parsers
 	{
 		[UhtSpecifier(Extends = UhtTableNames.Interface, ValueType = UhtSpecifierValueType.Legacy)]
 		[SuppressMessage("CodeQuality", "IDE0051:Remove unused private members", Justification = "Attribute accessed method")]
-		private static void DependsOnSpecifier(UhtSpecifierContext SpecifierContext)
+		private static void DependsOnSpecifier(UhtSpecifierContext specifierContext)
 		{
-			throw new UhtException(SpecifierContext.MessageSite, $"The dependsOn specifier is deprecated. Please use #include \"ClassHeaderFilename.h\" instead.");
+			throw new UhtException(specifierContext.MessageSite, $"The dependsOn specifier is deprecated. Please use #include \"ClassHeaderFilename.h\" instead.");
 		}
 
 		[UhtSpecifier(Extends = UhtTableNames.Interface, ValueType = UhtSpecifierValueType.Legacy)]
 		[SuppressMessage("CodeQuality", "IDE0051:Remove unused private members", Justification = "Attribute accessed method")]
-		private static void MinimalAPISpecifier(UhtSpecifierContext SpecifierContext)
+		private static void MinimalAPISpecifier(UhtSpecifierContext specifierContext)
 		{
-			UhtClass Class = (UhtClass)SpecifierContext.Scope.ScopeType;
-			Class.ClassFlags |= EClassFlags.MinimalAPI;
+			UhtClass clasObj = (UhtClass)specifierContext.Scope.ScopeType;
+			clasObj.ClassFlags |= EClassFlags.MinimalAPI;
 		}
 
 		[UhtSpecifier(Extends = UhtTableNames.Interface, ValueType = UhtSpecifierValueType.Legacy)]
 		[SuppressMessage("CodeQuality", "IDE0051:Remove unused private members", Justification = "Attribute accessed method")]
-		private static void ConversionRootSpecifier(UhtSpecifierContext SpecifierContext)
+		private static void ConversionRootSpecifier(UhtSpecifierContext specifierContext)
 		{
-			UhtClass Class = (UhtClass)SpecifierContext.Scope.ScopeType;
-			Class.MetaData.Add(UhtNames.IsConversionRoot, "true");
+			UhtClass classObj = (UhtClass)specifierContext.Scope.ScopeType;
+			classObj.MetaData.Add(UhtNames.IsConversionRoot, "true");
 		}
 	}
 }
diff --git a/Engine/Source/Programs/Shared/EpicGames.UHT/Specifiers/UhtPropertyArgumentSpecifiers.cs b/Engine/Source/Programs/Shared/EpicGames.UHT/Specifiers/UhtPropertyArgumentSpecifiers.cs
--- a/Engine/Source/Programs/Shared/EpicGames.UHT/Specifiers/UhtPropertyArgumentSpecifiers.cs
+++ b/Engine/Source/Programs/Shared/EpicGames.UHT/Specifiers/UhtPropertyArgumentSpecifiers.cs
@@ -18,33 +18,33 @@ namespace EpicGames.UHT.Parsers
 		#region Argument Property Specifiers
 		[UhtSpecifier(Extends = UhtTableNames.PropertyArgument, ValueType = UhtSpecifierValueType.Legacy)]
 		[SuppressMessage("CodeQuality", "IDE0051:Remove unused private members", Justification = "Attribute accessed method")]
-		private static void ConstSpecifier(UhtSpecifierContext SpecifierContext)
+		private static void ConstSpecifier(UhtSpecifierContext specifierContext)
 		{
-			UhtPropertySpecifierContext Context = (UhtPropertySpecifierContext)SpecifierContext;
-			Context.PropertySettings.PropertyFlags |= EPropertyFlags.ConstParm;
+			UhtPropertySpecifierContext context = (UhtPropertySpecifierContext)specifierContext;
+			context.PropertySettings.PropertyFlags |= EPropertyFlags.ConstParm;
 		}
 
 		[UhtSpecifier(Extends = UhtTableNames.PropertyArgument, ValueType = UhtSpecifierValueType.Legacy)]
 		[SuppressMessage("CodeQuality", "IDE0051:Remove unused private members", Justification = "Attribute accessed method")]
-		private static void RefSpecifier(UhtSpecifierContext SpecifierContext)
+		private static void RefSpecifier(UhtSpecifierContext specifierContext)
 		{
-			UhtPropertySpecifierContext Context = (UhtPropertySpecifierContext)SpecifierContext;
-			Context.PropertySettings.PropertyFlags |= EPropertyFlags.OutParm | EPropertyFlags.ReferenceParm;
+			UhtPropertySpecifierContext context = (UhtPropertySpecifierContext)specifierContext;
+			context.PropertySettings.PropertyFlags |= EPropertyFlags.OutParm | EPropertyFlags.ReferenceParm;
 		}
 
 		[UhtSpecifier(Extends = UhtTableNames.PropertyArgument, ValueType = UhtSpecifierValueType.Legacy)]
 		[SuppressMessage("CodeQuality", "IDE0051:Remove unused private members", Justification = "Attribute accessed method")]
-		private static void NotReplicatedSpecifier(UhtSpecifierContext SpecifierContext)
+		private static void NotReplicatedSpecifier(UhtSpecifierContext specifierContext)
 		{
-			UhtPropertySpecifierContext Context = (UhtPropertySpecifierContext)SpecifierContext;
-			if (Context.PropertySettings.PropertyCategory == UhtPropertyCategory.ReplicatedParameter)
+			UhtPropertySpecifierContext context = (UhtPropertySpecifierContext)specifierContext;
+			if (context.PropertySettings.PropertyCategory == UhtPropertyCategory.ReplicatedParameter)
 			{
-				Context.PropertySettings.PropertyCategory = UhtPropertyCategory.RegularParameter;
-				Context.PropertySettings.PropertyFlags |= EPropertyFlags.RepSkip;
+				context.PropertySettings.PropertyCategory = UhtPropertyCategory.RegularParameter;
+				context.PropertySettings.PropertyFlags |= EPropertyFlags.RepSkip;
 			}
 			else
 			{
-				Context.MessageSite.LogError("Only parameters in service request functions can be marked NotReplicated");
+				context.MessageSite.LogError("Only parameters in service request functions can be marked NotReplicated");
 			}
 		}
 		#endregion
diff --git a/Engine/Source/Programs/Shared/EpicGames.UHT/Specifiers/UhtPropertyMemberSpecifiers.cs b/Engine/Source/Programs/Shared/EpicGames.UHT/Specifiers/UhtPropertyMemberSpecifiers.cs
--- a/Engine/Source/Programs/Shared/EpicGames.UHT/Specifiers/UhtPropertyMemberSpecifiers.cs
+++ b/Engine/Source/Programs/Shared/EpicGames.UHT/Specifiers/UhtPropertyMemberSpecifiers.cs
@@ -1,11 +1,11 @@
 // Copyright Epic Games, Inc. All Rights Reserved.
 
+using System;
+using System.Diagnostics.CodeAnalysis;
 using EpicGames.Core;
 using EpicGames.UHT.Tables;
 using EpicGames.UHT.Types;
 using EpicGames.UHT.Utils;
-using System;
-using System.Diagnostics.CodeAnalysis;
 
 namespace EpicGames.UHT.Parsers
 {
@@ -18,430 +18,430 @@ namespace EpicGames.UHT.Parsers
 	{
 		[UhtSpecifier(Extends = UhtTableNames.PropertyMember, ValueType = UhtSpecifierValueType.Legacy)]
 		[SuppressMessage("CodeQuality", "IDE0051:Remove unused private members", Justification = "Attribute accessed method")]
-		private static void EditAnywhereSpecifier(UhtSpecifierContext SpecifierContext)
+		private static void EditAnywhereSpecifier(UhtSpecifierContext specifierContext)
 		{
-			UhtPropertySpecifierContext Context = (UhtPropertySpecifierContext)SpecifierContext;
-			if (Context.bSeenEditSpecifier)
+			UhtPropertySpecifierContext context = (UhtPropertySpecifierContext)specifierContext;
+			if (context.SeenEditSpecifier)
 			{
-				Context.MessageSite.LogError("Found more than one edit/visibility specifier (EditAnywhere), only one is allowed");
+				context.MessageSite.LogError("Found more than one edit/visibility specifier (EditAnywhere), only one is allowed");
 			}
-			Context.PropertySettings.PropertyFlags |= EPropertyFlags.Edit;
-			Context.bSeenEditSpecifier = true;
+			context.PropertySettings.PropertyFlags |= EPropertyFlags.Edit;
+			context.SeenEditSpecifier = true;
 		}
 
 		[UhtSpecifier(Extends = UhtTableNames.PropertyMember, ValueType = UhtSpecifierValueType.Legacy)]
 		[SuppressMessage("CodeQuality", "IDE0051:Remove unused private members", Justification = "Attribute accessed method")]
-		private static void EditInstanceOnlySpecifier(UhtSpecifierContext SpecifierContext)
+		private static void EditInstanceOnlySpecifier(UhtSpecifierContext specifierContext)
 		{
-			UhtPropertySpecifierContext Context = (UhtPropertySpecifierContext)SpecifierContext;
-			if (Context.bSeenEditSpecifier)
+			UhtPropertySpecifierContext context = (UhtPropertySpecifierContext)specifierContext;
+			if (context.SeenEditSpecifier)
 			{
-				Context.MessageSite.LogError("Found more than one edit/visibility specifier (EditInstanceOnly), only one is allowed");
+				context.MessageSite.LogError("Found more than one edit/visibility specifier (EditInstanceOnly), only one is allowed");
 			}
-			Context.PropertySettings.PropertyFlags |= EPropertyFlags.Edit | EPropertyFlags.DisableEditOnTemplate;
-			Context.bSeenEditSpecifier = true;
+			context.PropertySettings.PropertyFlags |= EPropertyFlags.Edit | EPropertyFlags.DisableEditOnTemplate;
+			context.SeenEditSpecifier = true;
 		}
 
 		[UhtSpecifier(Extends = UhtTableNames.PropertyMember, ValueType = UhtSpecifierValueType.Legacy)]
 		[SuppressMessage("CodeQuality", "IDE0051:Remove unused private members", Justification = "Attribute accessed method")]
-		private static void EditDefaultsOnlySpecifier(UhtSpecifierContext SpecifierContext)
+		private static void EditDefaultsOnlySpecifier(UhtSpecifierContext specifierContext)
 		{
-			UhtPropertySpecifierContext Context = (UhtPropertySpecifierContext)SpecifierContext;
-			if (Context.bSeenEditSpecifier)
+			UhtPropertySpecifierContext context = (UhtPropertySpecifierContext)specifierContext;
+			if (context.SeenEditSpecifier)
 			{
-				Context.MessageSite.LogError("Found more than one edit/visibility specifier (EditDefaultsOnly), only one is allowed");
+				context.MessageSite.LogError("Found more than one edit/visibility specifier (EditDefaultsOnly), only one is allowed");
 			}
-			Context.PropertySettings.PropertyFlags |= EPropertyFlags.Edit | EPropertyFlags.DisableEditOnInstance;
-			Context.bSeenEditSpecifier = true;
+			context.PropertySettings.PropertyFlags |= EPropertyFlags.Edit | EPropertyFlags.DisableEditOnInstance;
+			context.SeenEditSpecifier = true;
 		}
 
 		[UhtSpecifier(Extends = UhtTableNames.PropertyMember, ValueType = UhtSpecifierValueType.Legacy)]
 		[SuppressMessage("CodeQuality", "IDE0051:Remove unused private members", Justification = "Attribute accessed method")]
-		private static void VisibleAnywhereSpecifier(UhtSpecifierContext SpecifierContext)
+		private static void VisibleAnywhereSpecifier(UhtSpecifierContext specifierContext)
 		{
-			UhtPropertySpecifierContext Context = (UhtPropertySpecifierContext)SpecifierContext;
-			if (Context.bSeenEditSpecifier)
+			UhtPropertySpecifierContext context = (UhtPropertySpecifierContext)specifierContext;
+			if (context.SeenEditSpecifier)
 			{
-				Context.MessageSite.LogError("Found more than one edit/visibility specifier (VisibleAnywhere), only one is allowed");
+				context.MessageSite.LogError("Found more than one edit/visibility specifier (VisibleAnywhere), only one is allowed");
 			}
-			Context.PropertySettings.PropertyFlags |= EPropertyFlags.Edit | EPropertyFlags.EditConst;
-			Context.bSeenEditSpecifier = true;
+			context.PropertySettings.PropertyFlags |= EPropertyFlags.Edit | EPropertyFlags.EditConst;
+			context.SeenEditSpecifier = true;
 		}
 
 		[UhtSpecifier(Extends = UhtTableNames.PropertyMember, ValueType = UhtSpecifierValueType.Legacy)]
 		[SuppressMessage("CodeQuality", "IDE0051:Remove unused private members", Justification = "Attribute accessed method")]
-		private static void VisibleInstanceOnlySpecifier(UhtSpecifierContext SpecifierContext)
+		private static void VisibleInstanceOnlySpecifier(UhtSpecifierContext specifierContext)
 		{
-			UhtPropertySpecifierContext Context = (UhtPropertySpecifierContext)SpecifierContext;
-			if (Context.bSeenEditSpecifier)
+			UhtPropertySpecifierContext context = (UhtPropertySpecifierContext)specifierContext;
+			if (context.SeenEditSpecifier)
 			{
-				Context.MessageSite.LogError("Found more than one edit/visibility specifier (VisibleInstanceOnly), only one is allowed");
+				context.MessageSite.LogError("Found more than one edit/visibility specifier (VisibleInstanceOnly), only one is allowed");
 			}
-			Context.PropertySettings.PropertyFlags |= EPropertyFlags.Edit | EPropertyFlags.EditConst | EPropertyFlags.DisableEditOnTemplate;
-			Context.bSeenEditSpecifier = true;
+			context.PropertySettings.PropertyFlags |= EPropertyFlags.Edit | EPropertyFlags.EditConst | EPropertyFlags.DisableEditOnTemplate;
+			context.SeenEditSpecifier = true;
 		}
 
 		[UhtSpecifier(Extends = UhtTableNames.PropertyMember, ValueType = UhtSpecifierValueType.Legacy)]
 		[SuppressMessage("CodeQuality", "IDE0051:Remove unused private members", Justification = "Attribute accessed method")]
-		private static void VisibleDefaultsOnlySpecifier(UhtSpecifierContext SpecifierContext)
+		private static void VisibleDefaultsOnlySpecifier(UhtSpecifierContext specifierContext)
 		{
-			UhtPropertySpecifierContext Context = (UhtPropertySpecifierContext)SpecifierContext;
-			if (Context.bSeenEditSpecifier)
+			UhtPropertySpecifierContext context = (UhtPropertySpecifierContext)specifierContext;
+			if (context.SeenEditSpecifier)
 			{
-				Context.MessageSite.LogError("Found more than one edit/visibility specifier (VisibleDefaultsOnly), only one is allowed");
+				context.MessageSite.LogError("Found more than one edit/visibility specifier (VisibleDefaultsOnly), only one is allowed");
 			}
-			Context.PropertySettings.PropertyFlags |= EPropertyFlags.Edit | EPropertyFlags.EditConst | EPropertyFlags.DisableEditOnInstance;
-			Context.bSeenEditSpecifier = true;
+			context.PropertySettings.PropertyFlags |= EPropertyFlags.Edit | EPropertyFlags.EditConst | EPropertyFlags.DisableEditOnInstance;
+			context.SeenEditSpecifier = true;
 		}
 
 		[UhtSpecifier(Extends = UhtTableNames.PropertyMember, ValueType = UhtSpecifierValueType.Legacy)]
 		[SuppressMessage("CodeQuality", "IDE0051:Remove unused private members", Justification = "Attribute accessed method")]
-		private static void BlueprintReadWriteSpecifier(UhtSpecifierContext SpecifierContext)
+		private static void BlueprintReadWriteSpecifier(UhtSpecifierContext specifierContext)
 		{
-			UhtPropertySpecifierContext Context = (UhtPropertySpecifierContext)SpecifierContext;
-			if (Context.bSeenBlueprintReadOnlySpecifier)
+			UhtPropertySpecifierContext context = (UhtPropertySpecifierContext)specifierContext;
+			if (context.SeenBlueprintReadOnlySpecifier)
 			{
-				Context.MessageSite.LogError("Cannot specify a property as being both BlueprintReadOnly and BlueprintReadWrite.");
+				context.MessageSite.LogError("Cannot specify a property as being both BlueprintReadOnly and BlueprintReadWrite.");
 			}
 
-			string? PrivateAccessMD;
-			bool bAllowPrivateAccess = Context.MetaData.TryGetValue(UhtNames.AllowPrivateAccess, out PrivateAccessMD) && !PrivateAccessMD.Equals("false", StringComparison.OrdinalIgnoreCase);
-			if (SpecifierContext.Scope.AccessSpecifier == UhtAccessSpecifier.Private && !bAllowPrivateAccess)
+			string? privateAccessMD;
+			bool allowPrivateAccess = context.MetaData.TryGetValue(UhtNames.AllowPrivateAccess, out privateAccessMD) && !privateAccessMD.Equals("false", StringComparison.OrdinalIgnoreCase);
+			if (specifierContext.Scope.AccessSpecifier == UhtAccessSpecifier.Private && !allowPrivateAccess)
 			{
-				Context.MessageSite.LogError("BlueprintReadWrite should not be used on private members");
+				context.MessageSite.LogError("BlueprintReadWrite should not be used on private members");
 			}
 
-			if (Context.PropertySettings.PropertyFlags.HasAnyFlags(EPropertyFlags.EditorOnly) && Context.PropertySettings.Outer is UhtScriptStruct)
+			if (context.PropertySettings.PropertyFlags.HasAnyFlags(EPropertyFlags.EditorOnly) && context.PropertySettings.Outer is UhtScriptStruct)
 			{
-				Context.MessageSite.LogError("Blueprint exposed struct members cannot be editor only");
+				context.MessageSite.LogError("Blueprint exposed struct members cannot be editor only");
 			}
 
-			Context.PropertySettings.PropertyFlags |= EPropertyFlags.BlueprintVisible;
-			Context.bSeenBlueprintWriteSpecifier = true;
+			context.PropertySettings.PropertyFlags |= EPropertyFlags.BlueprintVisible;
+			context.SeenBlueprintWriteSpecifier = true;
 		}
 
 		[UhtSpecifier(Extends = UhtTableNames.PropertyMember, ValueType = UhtSpecifierValueType.Legacy)]
 		[SuppressMessage("CodeQuality", "IDE0051:Remove unused private members", Justification = "Attribute accessed method")]
-		private static void BlueprintReadOnlySpecifier(UhtSpecifierContext SpecifierContext)
+		private static void BlueprintReadOnlySpecifier(UhtSpecifierContext specifierContext)
 		{
-			UhtPropertySpecifierContext Context = (UhtPropertySpecifierContext)SpecifierContext;
-			if (Context.bSeenBlueprintWriteSpecifier)
+			UhtPropertySpecifierContext context = (UhtPropertySpecifierContext)specifierContext;
+			if (context.SeenBlueprintWriteSpecifier)
 			{
-				Context.MessageSite.LogError("Cannot specify both BlueprintReadOnly and BlueprintReadWrite or BlueprintSetter.");
+				context.MessageSite.LogError("Cannot specify both BlueprintReadOnly and BlueprintReadWrite or BlueprintSetter.");
 			}
 
-			string? PrivateAccessMD;
-			bool bAllowPrivateAccess = Context.MetaData.TryGetValue(UhtNames.AllowPrivateAccess, out PrivateAccessMD) && !PrivateAccessMD.Equals("false", StringComparison.OrdinalIgnoreCase);
-			if (SpecifierContext.Scope.AccessSpecifier == UhtAccessSpecifier.Private && !bAllowPrivateAccess)
+			string? privateAccessMD;
+			bool allowPrivateAccess = context.MetaData.TryGetValue(UhtNames.AllowPrivateAccess, out privateAccessMD) && !privateAccessMD.Equals("false", StringComparison.OrdinalIgnoreCase);
+			if (specifierContext.Scope.AccessSpecifier == UhtAccessSpecifier.Private && !allowPrivateAccess)
 			{
-				Context.MessageSite.LogError("BlueprintReadOnly should not be used on private members");
+				context.MessageSite.LogError("BlueprintReadOnly should not be used on private members");
 			}
 
-			if (Context.PropertySettings.PropertyFlags.HasAnyFlags(EPropertyFlags.EditorOnly) && Context.PropertySettings.Outer is UhtScriptStruct)
+			if (context.PropertySettings.PropertyFlags.HasAnyFlags(EPropertyFlags.EditorOnly) && context.PropertySettings.Outer is UhtScriptStruct)
 			{
-				Context.MessageSite.LogError("Blueprint exposed struct members cannot be editor only");
+				context.MessageSite.LogError("Blueprint exposed struct members cannot be editor only");
 			}
 
-			Context.PropertySettings.PropertyFlags |= EPropertyFlags.BlueprintVisible | EPropertyFlags.BlueprintReadOnly;
-			Context.bSeenBlueprintReadOnlySpecifier = true;
+			context.PropertySettings.PropertyFlags |= EPropertyFlags.BlueprintVisible | EPropertyFlags.BlueprintReadOnly;
+			context.SeenBlueprintReadOnlySpecifier = true;
 		}
 
 		[UhtSpecifier(Extends = UhtTableNames.PropertyMember, ValueType = UhtSpecifierValueType.SingleString)]
 		[SuppressMessage("CodeQuality", "IDE0051:Remove unused private members", Justification = "Attribute accessed method")]
-		private static void BlueprintSetterSpecifier(UhtSpecifierContext SpecifierContext, StringView Value)
+		private static void BlueprintSetterSpecifier(UhtSpecifierContext specifierContext, StringView value)
 		{
-			UhtPropertySpecifierContext Context = (UhtPropertySpecifierContext)SpecifierContext;
-			if (Context.bSeenBlueprintReadOnlySpecifier)
+			UhtPropertySpecifierContext context = (UhtPropertySpecifierContext)specifierContext;
+			if (context.SeenBlueprintReadOnlySpecifier)
 			{
-				Context.MessageSite.LogError("Cannot specify a property as being both BlueprintReadOnly and having a BlueprintSetter.");
+				context.MessageSite.LogError("Cannot specify a property as being both BlueprintReadOnly and having a BlueprintSetter.");
 			}
 
-			if (Context.PropertySettings.Outer is UhtScriptStruct)
+			if (context.PropertySettings.Outer is UhtScriptStruct)
 			{
-				Context.MessageSite.LogError("Cannot specify BlueprintSetter for a struct member.");
+				context.MessageSite.LogError("Cannot specify BlueprintSetter for a struct member.");
 			}
 
-			Context.MetaData.Add(UhtNames.BlueprintSetter, Value.ToString());
+			context.MetaData.Add(UhtNames.BlueprintSetter, value.ToString());
 
-			Context.PropertySettings.PropertyFlags |= EPropertyFlags.BlueprintVisible;
-			Context.bSeenBlueprintWriteSpecifier = true;
+			context.PropertySettings.PropertyFlags |= EPropertyFlags.BlueprintVisible;
+			context.SeenBlueprintWriteSpecifier = true;
 		}
 
 		[UhtSpecifier(Extends = UhtTableNames.PropertyMember, ValueType = UhtSpecifierValueType.SingleString)]
 		[SuppressMessage("CodeQuality", "IDE0051:Remove unused private members", Justification = "Attribute accessed method")]
-		private static void BlueprintGetterSpecifier(UhtSpecifierContext SpecifierContext, StringView Value)
+		private static void BlueprintGetterSpecifier(UhtSpecifierContext specifierContext, StringView value)
 		{
-			UhtPropertySpecifierContext Context = (UhtPropertySpecifierContext)SpecifierContext;
-			if (Context.PropertySettings.Outer is UhtScriptStruct)
+			UhtPropertySpecifierContext context = (UhtPropertySpecifierContext)specifierContext;
+			if (context.PropertySettings.Outer is UhtScriptStruct)
 			{
-				Context.MessageSite.LogError("Cannot specify BlueprintGetter for a struct member.");
+				context.MessageSite.LogError("Cannot specify BlueprintGetter for a struct member.");
 			}
 
-			Context.MetaData.Add(UhtNames.BlueprintGetter, Value.ToString());
+			context.MetaData.Add(UhtNames.BlueprintGetter, value.ToString());
 
-			Context.PropertySettings.PropertyFlags |= EPropertyFlags.BlueprintVisible;
-			Context.bSeenBlueprintGetterSpecifier = true;
+			context.PropertySettings.PropertyFlags |= EPropertyFlags.BlueprintVisible;
+			context.SeenBlueprintGetterSpecifier = true;
 		}
 
 		[UhtSpecifier(Extends = UhtTableNames.PropertyMember, ValueType = UhtSpecifierValueType.Legacy)]
 		[SuppressMessage("CodeQuality", "IDE0051:Remove unused private members", Justification = "Attribute accessed method")]
-		private static void ConfigSpecifier(UhtSpecifierContext SpecifierContext)
+		private static void ConfigSpecifier(UhtSpecifierContext specifierContext)
 		{
-			UhtPropertySpecifierContext Context = (UhtPropertySpecifierContext)SpecifierContext;
-			Context.PropertySettings.PropertyFlags |= EPropertyFlags.Config;
+			UhtPropertySpecifierContext context = (UhtPropertySpecifierContext)specifierContext;
+			context.PropertySettings.PropertyFlags |= EPropertyFlags.Config;
 		}
 
 		[UhtSpecifier(Extends = UhtTableNames.PropertyMember, ValueType = UhtSpecifierValueType.Legacy)]
 		[SuppressMessage("CodeQuality", "IDE0051:Remove unused private members", Justification = "Attribute accessed method")]
-		private static void GlobalConfigSpecifier(UhtSpecifierContext SpecifierContext)
+		private static void GlobalConfigSpecifier(UhtSpecifierContext specifierContext)
 		{
-			UhtPropertySpecifierContext Context = (UhtPropertySpecifierContext)SpecifierContext;
-			Context.PropertySettings.PropertyFlags |= EPropertyFlags.GlobalConfig | EPropertyFlags.Config;
+			UhtPropertySpecifierContext context = (UhtPropertySpecifierContext)specifierContext;
+			context.PropertySettings.PropertyFlags |= EPropertyFlags.GlobalConfig | EPropertyFlags.Config;
 		}
 
 		[UhtSpecifier(Extends = UhtTableNames.PropertyMember, ValueType = UhtSpecifierValueType.Legacy)]
 		[SuppressMessage("CodeQuality", "IDE0051:Remove unused private members", Justification = "Attribute accessed method")]
-		private static void LocalizedSpecifier(UhtSpecifierContext SpecifierContext)
+		private static void LocalizedSpecifier(UhtSpecifierContext specifierContext)
 		{
-			UhtPropertySpecifierContext Context = (UhtPropertySpecifierContext)SpecifierContext;
-			Context.MessageSite.LogError("The Localized specifier is deprecated");
+			UhtPropertySpecifierContext context = (UhtPropertySpecifierContext)specifierContext;
+			context.MessageSite.LogError("The Localized specifier is deprecated");
 		}
 
 		[UhtSpecifier(Extends = UhtTableNames.PropertyMember, ValueType = UhtSpecifierValueType.Legacy)]
 		[SuppressMessage("CodeQuality", "IDE0051:Remove unused private members", Justification = "Attribute accessed method")]
-		private static void TransientSpecifier(UhtSpecifierContext SpecifierContext)
+		private static void TransientSpecifier(UhtSpecifierContext specifierContext)
 		{
-			UhtPropertySpecifierContext Context = (UhtPropertySpecifierContext)SpecifierContext;
-			Context.PropertySettings.PropertyFlags |= EPropertyFlags.Transient;
+			UhtPropertySpecifierContext context = (UhtPropertySpecifierContext)specifierContext;
+			context.PropertySettings.PropertyFlags |= EPropertyFlags.Transient;
 		}
 
 		[UhtSpecifier(Extends = UhtTableNames.PropertyMember, ValueType = UhtSpecifierValueType.Legacy)]
 		[SuppressMessage("CodeQuality", "IDE0051:Remove unused private members", Justification = "Attribute accessed method")]
-		private static void DuplicateTransientSpecifier(UhtSpecifierContext SpecifierContext)
+		private static void DuplicateTransientSpecifier(UhtSpecifierContext specifierContext)
 		{
-			UhtPropertySpecifierContext Context = (UhtPropertySpecifierContext)SpecifierContext;
-			Context.PropertySettings.PropertyFlags |= EPropertyFlags.DuplicateTransient;
+			UhtPropertySpecifierContext context = (UhtPropertySpecifierContext)specifierContext;
+			context.PropertySettings.PropertyFlags |= EPropertyFlags.DuplicateTransient;
 		}
 
 		[UhtSpecifier(Extends = UhtTableNames.PropertyMember, ValueType = UhtSpecifierValueType.Legacy)]
 		[SuppressMessage("CodeQuality", "IDE0051:Remove unused private members", Justification = "Attribute accessed method")]
-		private static void TextExportTransientSpecifier(UhtSpecifierContext SpecifierContext)
+		private static void TextExportTransientSpecifier(UhtSpecifierContext specifierContext)
 		{
-			UhtPropertySpecifierContext Context = (UhtPropertySpecifierContext)SpecifierContext;
-			Context.PropertySettings.PropertyFlags |= EPropertyFlags.TextExportTransient;
+			UhtPropertySpecifierContext context = (UhtPropertySpecifierContext)specifierContext;
+			context.PropertySettings.PropertyFlags |= EPropertyFlags.TextExportTransient;
 		}
 
 		[UhtSpecifier(Extends = UhtTableNames.PropertyMember, ValueType = UhtSpecifierValueType.Legacy)]
 		[SuppressMessage("CodeQuality", "IDE0051:Remove unused private members", Justification = "Attribute accessed method")]
-		private static void NonPIETransientSpecifier(UhtSpecifierContext SpecifierContext)
+		private static void NonPIETransientSpecifier(UhtSpecifierContext specifierContext)
 		{
-			UhtPropertySpecifierContext Context = (UhtPropertySpecifierContext)SpecifierContext;
-			Context.MessageSite.LogWarning("NonPIETransient is deprecated - NonPIEDuplicateTransient should be used instead");
-			Context.PropertySettings.PropertyFlags |= EPropertyFlags.NonPIEDuplicateTransient;
+			UhtPropertySpecifierContext context = (UhtPropertySpecifierContext)specifierContext;
+			context.MessageSite.LogWarning("NonPIETransient is deprecated - NonPIEDuplicateTransient should be used instead");
+			context.PropertySettings.PropertyFlags |= EPropertyFlags.NonPIEDuplicateTransient;
 		}
 
 		[UhtSpecifier(Extends = UhtTableNames.PropertyMember, ValueType = UhtSpecifierValueType.Legacy)]
 		[SuppressMessage("CodeQuality", "IDE0051:Remove unused private members", Justification = "Attribute accessed method")]
-		private static void NonPIEDuplicateTransientSpecifier(UhtSpecifierContext SpecifierContext)
+		private static void NonPIEDuplicateTransientSpecifier(UhtSpecifierContext specifierContext)
 		{
-			UhtPropertySpecifierContext Context = (UhtPropertySpecifierContext)SpecifierContext;
-			Context.PropertySettings.PropertyFlags |= EPropertyFlags.NonPIEDuplicateTransient;
+			UhtPropertySpecifierContext context = (UhtPropertySpecifierContext)specifierContext;
+			context.PropertySettings.PropertyFlags |= EPropertyFlags.NonPIEDuplicateTransient;
 		}
 
 		[UhtSpecifier(Extends = UhtTableNames.PropertyMember, ValueType = UhtSpecifierValueType.Legacy)]
 		[SuppressMessage("CodeQuality", "IDE0051:Remove unused private members", Justification = "Attribute accessed method")]
-		private static void ExportSpecifier(UhtSpecifierContext SpecifierContext)
+		private static void ExportSpecifier(UhtSpecifierContext specifierContext)
 		{
-			UhtPropertySpecifierContext Context = (UhtPropertySpecifierContext)SpecifierContext;
-			Context.PropertySettings.PropertyFlags |= EPropertyFlags.ExportObject;
+			UhtPropertySpecifierContext context = (UhtPropertySpecifierContext)specifierContext;
+			context.PropertySettings.PropertyFlags |= EPropertyFlags.ExportObject;
 		}
 
 		[UhtSpecifier(Extends = UhtTableNames.PropertyMember, ValueType = UhtSpecifierValueType.Legacy)]
 		[SuppressMessage("CodeQuality", "IDE0051:Remove unused private members", Justification = "Attribute accessed method")]
-		private static void EditInlineSpecifier(UhtSpecifierContext SpecifierContext)
+		private static void EditInlineSpecifier(UhtSpecifierContext specifierContext)
 		{
-			UhtPropertySpecifierContext Context = (UhtPropertySpecifierContext)SpecifierContext;
-			Context.MessageSite.LogError("EditInline is deprecated. Remove it, or use Instanced instead.");
+			UhtPropertySpecifierContext context = (UhtPropertySpecifierContext)specifierContext;
+			context.MessageSite.LogError("EditInline is deprecated. Remove it, or use Instanced instead.");
 		}
 
 		[UhtSpecifier(Extends = UhtTableNames.PropertyMember, ValueType = UhtSpecifierValueType.Legacy)]
 		[SuppressMessage("CodeQuality", "IDE0051:Remove unused private members", Justification = "Attribute accessed method")]
-		private static void NoClearSpecifier(UhtSpecifierContext SpecifierContext)
+		private static void NoClearSpecifier(UhtSpecifierContext specifierContext)
 		{
-			UhtPropertySpecifierContext Context = (UhtPropertySpecifierContext)SpecifierContext;
-			Context.PropertySettings.PropertyFlags |= EPropertyFlags.NoClear;
+			UhtPropertySpecifierContext context = (UhtPropertySpecifierContext)specifierContext;
+			context.PropertySettings.PropertyFlags |= EPropertyFlags.NoClear;
 		}
 
 		[UhtSpecifier(Extends = UhtTableNames.PropertyMember, ValueType = UhtSpecifierValueType.Legacy)]
 		[SuppressMessage("CodeQuality", "IDE0051:Remove unused private members", Justification = "Attribute accessed method")]
-		private static void EditFixedSizeSpecifier(UhtSpecifierContext SpecifierContext)
+		private static void EditFixedSizeSpecifier(UhtSpecifierContext specifierContext)
 		{
-			UhtPropertySpecifierContext Context = (UhtPropertySpecifierContext)SpecifierContext;
-			Context.PropertySettings.PropertyFlags |= EPropertyFlags.EditFixedSize;
+			UhtPropertySpecifierContext context = (UhtPropertySpecifierContext)specifierContext;
+			context.PropertySettings.PropertyFlags |= EPropertyFlags.EditFixedSize;
 		}
 
 		[UhtSpecifier(Extends = UhtTableNames.PropertyMember, ValueType = UhtSpecifierValueType.Legacy)]
 		[SuppressMessage("CodeQuality", "IDE0051:Remove unused private members", Justification = "Attribute accessed method")]
-		private static void ReplicatedSpecifier(UhtSpecifierContext SpecifierContext)
+		private static void ReplicatedSpecifier(UhtSpecifierContext specifierContext)
 		{
-			UhtPropertySpecifierContext Context = (UhtPropertySpecifierContext)SpecifierContext;
-			if (Context.PropertySettings.Outer is UhtScriptStruct)
+			UhtPropertySpecifierContext context = (UhtPropertySpecifierContext)specifierContext;
+			if (context.PropertySettings.Outer is UhtScriptStruct)
 			{
-				Context.MessageSite.LogError("Struct members cannot be replicated");
+				context.MessageSite.LogError("Struct members cannot be replicated");
 			}
-			Context.PropertySettings.PropertyFlags |= EPropertyFlags.Net;
+			context.PropertySettings.PropertyFlags |= EPropertyFlags.Net;
 		}
 
 		[UhtSpecifier(Extends = UhtTableNames.PropertyMember, ValueType = UhtSpecifierValueType.SingleString)]
 		[SuppressMessage("CodeQuality", "IDE0051:Remove unused private members", Justification = "Attribute accessed method")]
-		private static void ReplicatedUsingSpecifier(UhtSpecifierContext SpecifierContext, StringView Value)
+		private static void ReplicatedUsingSpecifier(UhtSpecifierContext specifierContext, StringView value)
 		{
-			UhtPropertySpecifierContext Context = (UhtPropertySpecifierContext)SpecifierContext;
-			if (Context.PropertySettings.Outer is UhtScriptStruct)
+			UhtPropertySpecifierContext context = (UhtPropertySpecifierContext)specifierContext;
+			if (context.PropertySettings.Outer is UhtScriptStruct)
 			{
-				Context.MessageSite.LogError("Struct members cannot be replicated");
+				context.MessageSite.LogError("Struct members cannot be replicated");
 			}
-			Context.PropertySettings.PropertyFlags |= EPropertyFlags.Net;
+			context.PropertySettings.PropertyFlags |= EPropertyFlags.Net;
 
-			if (Value.Span.Length == 0)
+			if (value.Span.Length == 0)
 			{
-				Context.MessageSite.LogError("Must specify a valid function name for replication notifications");
+				context.MessageSite.LogError("Must specify a valid function name for replication notifications");
 			}
-			Context.PropertySettings.RepNotifyName = Value.ToString();
-			Context.PropertySettings.PropertyFlags |= EPropertyFlags.RepNotify;
+			context.PropertySettings.RepNotifyName = value.ToString();
+			context.PropertySettings.PropertyFlags |= EPropertyFlags.RepNotify;
 		}
 
 		[UhtSpecifier(Extends = UhtTableNames.PropertyMember, ValueType = UhtSpecifierValueType.Legacy)]
 		[SuppressMessage("CodeQuality", "IDE0051:Remove unused private members", Justification = "Attribute accessed method")]
-		private static void NotReplicatedSpecifier(UhtSpecifierContext SpecifierContext)
+		private static void NotReplicatedSpecifier(UhtSpecifierContext specifierContext)
 		{
-			UhtPropertySpecifierContext Context = (UhtPropertySpecifierContext)SpecifierContext;
-			if (!(Context.PropertySettings.Outer is UhtScriptStruct))
+			UhtPropertySpecifierContext context = (UhtPropertySpecifierContext)specifierContext;
+			if (context.PropertySettings.Outer is not UhtScriptStruct)
 			{
-				Context.MessageSite.LogError("Only Struct members can be marked NotReplicated");
+				context.MessageSite.LogError("Only Struct members can be marked NotReplicated");
 			}
-			Context.PropertySettings.PropertyFlags |= EPropertyFlags.RepSkip;
+			context.PropertySettings.PropertyFlags |= EPropertyFlags.RepSkip;
 		}
 
 		[UhtSpecifier(Extends = UhtTableNames.PropertyMember, ValueType = UhtSpecifierValueType.Legacy)]
 		[SuppressMessage("CodeQuality", "IDE0051:Remove unused private members", Justification = "Attribute accessed method")]
-		private static void RepRetrySpecifier(UhtSpecifierContext SpecifierContext)
+		private static void RepRetrySpecifier(UhtSpecifierContext specifierContext)
 		{
-			UhtPropertySpecifierContext Context = (UhtPropertySpecifierContext)SpecifierContext;
-			Context.MessageSite.LogError("'RepRetry' is deprecated.");
+			UhtPropertySpecifierContext context = (UhtPropertySpecifierContext)specifierContext;
+			context.MessageSite.LogError("'RepRetry' is deprecated.");
 		}
 
 		[UhtSpecifier(Extends = UhtTableNames.PropertyMember, ValueType = UhtSpecifierValueType.Legacy)]
 		[SuppressMessage("CodeQuality", "IDE0051:Remove unused private members", Justification = "Attribute accessed method")]
-		private static void InterpSpecifier(UhtSpecifierContext SpecifierContext)
+		private static void InterpSpecifier(UhtSpecifierContext specifierContext)
 		{
-			UhtPropertySpecifierContext Context = (UhtPropertySpecifierContext)SpecifierContext;
-			Context.PropertySettings.PropertyFlags |= EPropertyFlags.Edit | EPropertyFlags.BlueprintVisible | EPropertyFlags.Interp;
+			UhtPropertySpecifierContext context = (UhtPropertySpecifierContext)specifierContext;
+			context.PropertySettings.PropertyFlags |= EPropertyFlags.Edit | EPropertyFlags.BlueprintVisible | EPropertyFlags.Interp;
 		}
 
 		[UhtSpecifier(Extends = UhtTableNames.PropertyMember, ValueType = UhtSpecifierValueType.Legacy)]
 		[SuppressMessage("CodeQuality", "IDE0051:Remove unused private members", Justification = "Attribute accessed method")]
-		private static void NonTransactionalSpecifier(UhtSpecifierContext SpecifierContext)
+		private static void NonTransactionalSpecifier(UhtSpecifierContext specifierContext)
 		{
-			UhtPropertySpecifierContext Context = (UhtPropertySpecifierContext)SpecifierContext;
-			Context.PropertySettings.PropertyFlags |= EPropertyFlags.NonTransactional;
+			UhtPropertySpecifierContext context = (UhtPropertySpecifierContext)specifierContext;
+			context.PropertySettings.PropertyFlags |= EPropertyFlags.NonTransactional;
 		}
 
 		[UhtSpecifier(Extends = UhtTableNames.PropertyMember, ValueType = UhtSpecifierValueType.Legacy)]
 		[SuppressMessage("CodeQuality", "IDE0051:Remove unused private members", Justification = "Attribute accessed method")]
-		private static void InstancedSpecifier(UhtSpecifierContext SpecifierContext)
+		private static void InstancedSpecifier(UhtSpecifierContext specifierContext)
 		{
-			UhtPropertySpecifierContext Context = (UhtPropertySpecifierContext)SpecifierContext;
-			Context.PropertySettings.PropertyFlags |= EPropertyFlags.PersistentInstance | EPropertyFlags.ExportObject | EPropertyFlags.InstancedReference;
-			Context.MetaData.Add(UhtNames.EditInline, true);
+			UhtPropertySpecifierContext context = (UhtPropertySpecifierContext)specifierContext;
+			context.PropertySettings.PropertyFlags |= EPropertyFlags.PersistentInstance | EPropertyFlags.ExportObject | EPropertyFlags.InstancedReference;
+			context.MetaData.Add(UhtNames.EditInline, true);
 		}
 
 		[UhtSpecifier(Extends = UhtTableNames.PropertyMember, ValueType = UhtSpecifierValueType.Legacy)]
 		[SuppressMessage("CodeQuality", "IDE0051:Remove unused private members", Justification = "Attribute accessed method")]
-		private static void BlueprintAssignableSpecifier(UhtSpecifierContext SpecifierContext)
+		private static void BlueprintAssignableSpecifier(UhtSpecifierContext specifierContext)
 		{
-			UhtPropertySpecifierContext Context = (UhtPropertySpecifierContext)SpecifierContext;
-			Context.PropertySettings.PropertyFlags |= EPropertyFlags.BlueprintAssignable;
+			UhtPropertySpecifierContext context = (UhtPropertySpecifierContext)specifierContext;
+			context.PropertySettings.PropertyFlags |= EPropertyFlags.BlueprintAssignable;
 		}
 
 		[UhtSpecifier(Extends = UhtTableNames.PropertyMember, ValueType = UhtSpecifierValueType.Legacy)]
 		[SuppressMessage("CodeQuality", "IDE0051:Remove unused private members", Justification = "Attribute accessed method")]
-		private static void BlueprintCallableSpecifier(UhtSpecifierContext SpecifierContext)
+		private static void BlueprintCallableSpecifier(UhtSpecifierContext specifierContext)
 		{
-			UhtPropertySpecifierContext Context = (UhtPropertySpecifierContext)SpecifierContext;
-			Context.PropertySettings.PropertyFlags |= EPropertyFlags.BlueprintCallable;
+			UhtPropertySpecifierContext context = (UhtPropertySpecifierContext)specifierContext;
+			context.PropertySettings.PropertyFlags |= EPropertyFlags.BlueprintCallable;
 		}
 
 		[UhtSpecifier(Extends = UhtTableNames.PropertyMember, ValueType = UhtSpecifierValueType.Legacy)]
 		[SuppressMessage("CodeQuality", "IDE0051:Remove unused private members", Justification = "Attribute accessed method")]
-		private static void BlueprintAuthorityOnlySpecifier(UhtSpecifierContext SpecifierContext)
+		private static void BlueprintAuthorityOnlySpecifier(UhtSpecifierContext specifierContext)
 		{
-			UhtPropertySpecifierContext Context = (UhtPropertySpecifierContext)SpecifierContext;
-			Context.PropertySettings.PropertyFlags |= EPropertyFlags.BlueprintAuthorityOnly;
+			UhtPropertySpecifierContext context = (UhtPropertySpecifierContext)specifierContext;
+			context.PropertySettings.PropertyFlags |= EPropertyFlags.BlueprintAuthorityOnly;
 		}
 
 		[UhtSpecifier(Extends = UhtTableNames.PropertyMember, ValueType = UhtSpecifierValueType.Legacy)]
 		[SuppressMessage("CodeQuality", "IDE0051:Remove unused private members", Justification = "Attribute accessed method")]
-		private static void AssetRegistrySearchableSpecifier(UhtSpecifierContext SpecifierContext)
+		private static void AssetRegistrySearchableSpecifier(UhtSpecifierContext specifierContext)
 		{
-			UhtPropertySpecifierContext Context = (UhtPropertySpecifierContext)SpecifierContext;
-			Context.PropertySettings.PropertyFlags |= EPropertyFlags.AssetRegistrySearchable;
+			UhtPropertySpecifierContext context = (UhtPropertySpecifierContext)specifierContext;
+			context.PropertySettings.PropertyFlags |= EPropertyFlags.AssetRegistrySearchable;
 		}
 
 		[UhtSpecifier(Extends = UhtTableNames.PropertyMember, ValueType = UhtSpecifierValueType.Legacy)]
 		[SuppressMessage("CodeQuality", "IDE0051:Remove unused private members", Justification = "Attribute accessed method")]
-		private static void SimpleDisplaySpecifier(UhtSpecifierContext SpecifierContext)
+		private static void SimpleDisplaySpecifier(UhtSpecifierContext specifierContext)
 		{
-			UhtPropertySpecifierContext Context = (UhtPropertySpecifierContext)SpecifierContext;
-			Context.PropertySettings.PropertyFlags |= EPropertyFlags.SimpleDisplay;
+			UhtPropertySpecifierContext context = (UhtPropertySpecifierContext)specifierContext;
+			context.PropertySettings.PropertyFlags |= EPropertyFlags.SimpleDisplay;
 		}
 
 		[UhtSpecifier(Extends = UhtTableNames.PropertyMember, ValueType = UhtSpecifierValueType.Legacy)]
 		[SuppressMessage("CodeQuality", "IDE0051:Remove unused private members", Justification = "Attribute accessed method")]
-		private static void AdvancedDisplaySpecifier(UhtSpecifierContext SpecifierContext)
+		private static void AdvancedDisplaySpecifier(UhtSpecifierContext specifierContext)
 		{
-			UhtPropertySpecifierContext Context = (UhtPropertySpecifierContext)SpecifierContext;
-			Context.PropertySettings.PropertyFlags |= EPropertyFlags.AdvancedDisplay;
+			UhtPropertySpecifierContext context = (UhtPropertySpecifierContext)specifierContext;
+			context.PropertySettings.PropertyFlags |= EPropertyFlags.AdvancedDisplay;
 		}
 
 		[UhtSpecifier(Extends = UhtTableNames.PropertyMember, ValueType = UhtSpecifierValueType.Legacy)]
 		[SuppressMessage("CodeQuality", "IDE0051:Remove unused private members", Justification = "Attribute accessed method")]
-		private static void SaveGameSpecifier(UhtSpecifierContext SpecifierContext)
+		private static void SaveGameSpecifier(UhtSpecifierContext specifierContext)
 		{
-			UhtPropertySpecifierContext Context = (UhtPropertySpecifierContext)SpecifierContext;
-			Context.PropertySettings.PropertyFlags |= EPropertyFlags.SaveGame;
+			UhtPropertySpecifierContext context = (UhtPropertySpecifierContext)specifierContext;
+			context.PropertySettings.PropertyFlags |= EPropertyFlags.SaveGame;
 		}
 
 		[UhtSpecifier(Extends = UhtTableNames.PropertyMember, ValueType = UhtSpecifierValueType.Legacy)]
 		[SuppressMessage("CodeQuality", "IDE0051:Remove unused private members", Justification = "Attribute accessed method")]
-		private static void SkipSerializationSpecifier(UhtSpecifierContext SpecifierContext)
+		private static void SkipSerializationSpecifier(UhtSpecifierContext specifierContext)
 		{
-			UhtPropertySpecifierContext Context = (UhtPropertySpecifierContext)SpecifierContext;
-			Context.PropertySettings.PropertyFlags |= EPropertyFlags.SkipSerialization;
+			UhtPropertySpecifierContext context = (UhtPropertySpecifierContext)specifierContext;
+			context.PropertySettings.PropertyFlags |= EPropertyFlags.SkipSerialization;
 		}
 
 		[UhtSpecifier(Extends = UhtTableNames.PropertyMember, ValueType = UhtSpecifierValueType.OptionalString)]
 		[SuppressMessage("CodeQuality", "IDE0051:Remove unused private members", Justification = "Attribute accessed method")]
-		private static void GetterSpecifier(UhtSpecifierContext SpecifierContext, StringView? Value)
+		private static void GetterSpecifier(UhtSpecifierContext specifierContext, StringView? value)
 		{
-			UhtPropertySpecifierContext Context = (UhtPropertySpecifierContext)SpecifierContext;
-			if (!(Context.PropertySettings.Outer is UhtClass))
+			UhtPropertySpecifierContext context = (UhtPropertySpecifierContext)specifierContext;
+			if (context.PropertySettings.Outer is not UhtClass)
 			{
-				Context.MessageSite.LogError("Only class members can have Setters");
+				context.MessageSite.LogError("Only class members can have Setters");
 			}
-			Context.PropertySettings.PropertyExportFlags |= UhtPropertyExportFlags.GetterSpecified;
-			if (Value != null)
+			context.PropertySettings.PropertyExportFlags |= UhtPropertyExportFlags.GetterSpecified;
+			if (value != null)
 			{
-				StringView Temp = (StringView)Value;
-				if (Temp.Length > 0)
+				StringView temp = (StringView)value;
+				if (temp.Length > 0)
 				{
-					if (Temp.Equals("None", StringComparison.OrdinalIgnoreCase))
+					if (temp.Equals("None", StringComparison.OrdinalIgnoreCase))
 					{
-						Context.PropertySettings.PropertyExportFlags |= UhtPropertyExportFlags.GetterSpecifiedNone;
+						context.PropertySettings.PropertyExportFlags |= UhtPropertyExportFlags.GetterSpecifiedNone;
 					}
 					else
 					{
-						Context.PropertySettings.Getter = Value.ToString();
+						context.PropertySettings.Getter = value.ToString();
 					}
 				}
 			}
@@ -449,26 +449,26 @@ namespace EpicGames.UHT.Parsers
 
 		[UhtSpecifier(Extends = UhtTableNames.PropertyMember, ValueType = UhtSpecifierValueType.OptionalString)]
 		[SuppressMessage("CodeQuality", "IDE0051:Remove unused private members", Justification = "Attribute accessed method")]
-		private static void SetterSpecifier(UhtSpecifierContext SpecifierContext, StringView? Value)
+		private static void SetterSpecifier(UhtSpecifierContext specifierContext, StringView? value)
 		{
-			UhtPropertySpecifierContext Context = (UhtPropertySpecifierContext)SpecifierContext;
-			if (!(Context.PropertySettings.Outer is UhtClass))
+			UhtPropertySpecifierContext context = (UhtPropertySpecifierContext)specifierContext;
+			if (context.PropertySettings.Outer is not UhtClass)
 			{
-				Context.MessageSite.LogError("Only class members can have Getters");
+				context.MessageSite.LogError("Only class members can have Getters");
 			}
-			Context.PropertySettings.PropertyExportFlags |= UhtPropertyExportFlags.SetterSpecified;
-			if (Value != null)
+			context.PropertySettings.PropertyExportFlags |= UhtPropertyExportFlags.SetterSpecified;
+			if (value != null)
 			{
-				StringView Temp = (StringView)Value;
-				if (Temp.Length > 0)
+				StringView temp = (StringView)value;
+				if (temp.Length > 0)
 				{
-					if (Temp.Equals("None", StringComparison.OrdinalIgnoreCase))
+					if (temp.Equals("None", StringComparison.OrdinalIgnoreCase))
 					{
-						Context.PropertySettings.PropertyExportFlags |= UhtPropertyExportFlags.SetterSpecifiedNone;
+						context.PropertySettings.PropertyExportFlags |= UhtPropertyExportFlags.SetterSpecifiedNone;
 					}
 					else
 					{
-						Context.PropertySettings.Setter = Value.ToString();
+						context.PropertySettings.Setter = value.ToString();
 					}
 				}
 			}
@@ -476,10 +476,10 @@ namespace EpicGames.UHT.Parsers
 
 		[UhtSpecifier(Extends = UhtTableNames.PropertyMember, ValueType = UhtSpecifierValueType.None)]
 		[SuppressMessage("CodeQuality", "IDE0051:Remove unused private members", Justification = "Attribute accessed method")]
-		private static void FieldNotifySpecifier(UhtSpecifierContext SpecifierContext)
+		private static void FieldNotifySpecifier(UhtSpecifierContext specifierContext)
 		{
-			UhtPropertySpecifierContext Context = (UhtPropertySpecifierContext)SpecifierContext;
-			Context.PropertySettings.PropertyExportFlags |= UhtPropertyExportFlags.FieldNotify;
+			UhtPropertySpecifierContext context = (UhtPropertySpecifierContext)specifierContext;
+			context.PropertySettings.PropertyExportFlags |= UhtPropertyExportFlags.FieldNotify;
 		}
 	}
 }
diff --git a/Engine/Source/Programs/Shared/EpicGames.UHT/Specifiers/UhtScriptStructSpecifiers.cs b/Engine/Source/Programs/Shared/EpicGames.UHT/Specifiers/UhtScriptStructSpecifiers.cs
--- a/Engine/Source/Programs/Shared/EpicGames.UHT/Specifiers/UhtScriptStructSpecifiers.cs
+++ b/Engine/Source/Programs/Shared/EpicGames.UHT/Specifiers/UhtScriptStructSpecifiers.cs
@@ -16,75 +16,75 @@ namespace EpicGames.UHT.Parsers
 	{
 		[UhtSpecifier(Extends = UhtTableNames.ScriptStruct, ValueType = UhtSpecifierValueType.Legacy)]
 		[SuppressMessage("CodeQuality", "IDE0051:Remove unused private members", Justification = "Attribute accessed method")]
-		private static void NoExportSpecifier(UhtSpecifierContext SpecifierContext)
+		private static void NoExportSpecifier(UhtSpecifierContext specifierContext)
 		{
-			UhtScriptStruct ScriptStruct = (UhtScriptStruct)SpecifierContext.Scope.ScopeType;
-			ScriptStruct.ScriptStructFlags |= EStructFlags.NoExport;
-			ScriptStruct.ScriptStructFlags &= ~EStructFlags.Native;
+			UhtScriptStruct scriptStruct = (UhtScriptStruct)specifierContext.Scope.ScopeType;
+			scriptStruct.ScriptStructFlags |= EStructFlags.NoExport;
+			scriptStruct.ScriptStructFlags &= ~EStructFlags.Native;
 		}
 
 		[UhtSpecifier(Extends = UhtTableNames.ScriptStruct, ValueType = UhtSpecifierValueType.Legacy)]
 		[SuppressMessage("CodeQuality", "IDE0051:Remove unused private members", Justification = "Attribute accessed method")]
-		private static void AtomicSpecifier(UhtSpecifierContext SpecifierContext)
+		private static void AtomicSpecifier(UhtSpecifierContext specifierContext)
 		{
-			UhtScriptStruct ScriptStruct = (UhtScriptStruct)SpecifierContext.Scope.ScopeType;
-			ScriptStruct.ScriptStructFlags |= EStructFlags.Atomic;
+			UhtScriptStruct scriptStruct = (UhtScriptStruct)specifierContext.Scope.ScopeType;
+			scriptStruct.ScriptStructFlags |= EStructFlags.Atomic;
 		}
 
 		[UhtSpecifier(Extends = UhtTableNames.ScriptStruct, ValueType = UhtSpecifierValueType.Legacy)]
 		[SuppressMessage("CodeQuality", "IDE0051:Remove unused private members", Justification = "Attribute accessed method")]
-		private static void ImmutableSpecifier(UhtSpecifierContext SpecifierContext)
+		private static void ImmutableSpecifier(UhtSpecifierContext specifierContext)
 		{
-			UhtScriptStruct ScriptStruct = (UhtScriptStruct)SpecifierContext.Scope.ScopeType;
-			ScriptStruct.ScriptStructFlags |= EStructFlags.Atomic | EStructFlags.Immutable;
+			UhtScriptStruct scriptStruct = (UhtScriptStruct)specifierContext.Scope.ScopeType;
+			scriptStruct.ScriptStructFlags |= EStructFlags.Atomic | EStructFlags.Immutable;
 		}
 
 		[UhtSpecifier(Extends = UhtTableNames.ScriptStruct, ValueType = UhtSpecifierValueType.Legacy)]
 		[SuppressMessage("CodeQuality", "IDE0051:Remove unused private members", Justification = "Attribute accessed method")]
-		private static void HasDefaultsSpecifier(UhtSpecifierContext SpecifierContext)
+		private static void HasDefaultsSpecifier(UhtSpecifierContext specifierContext)
 		{
-			if (!SpecifierContext.Scope.HeaderParser.HeaderFile.bIsNoExportTypes)
+			if (!specifierContext.Scope.HeaderParser.HeaderFile.IsNoExportTypes)
 			{
-				SpecifierContext.MessageSite.LogError("The 'HasDefaults' struct specifier is only valid in the NoExportTypes.h file");
+				specifierContext.MessageSite.LogError("The 'HasDefaults' struct specifier is only valid in the NoExportTypes.h file");
 			}
-			UhtScriptStruct ScriptStruct = (UhtScriptStruct)SpecifierContext.Scope.ScopeType;
-			ScriptStruct.ScriptStructExportFlags |= UhtScriptStructExportFlags.HasDefaults;
+			UhtScriptStruct scriptStruct = (UhtScriptStruct)specifierContext.Scope.ScopeType;
+			scriptStruct.ScriptStructExportFlags |= UhtScriptStructExportFlags.HasDefaults;
 		}
 
 		[UhtSpecifier(Extends = UhtTableNames.ScriptStruct, ValueType = UhtSpecifierValueType.Legacy)]
 		[SuppressMessage("CodeQuality", "IDE0051:Remove unused private members", Justification = "Attribute accessed method")]
-		private static void HasNoOpConstructorSpecifier(UhtSpecifierContext SpecifierContext)
+		private static void HasNoOpConstructorSpecifier(UhtSpecifierContext specifierContext)
 		{
-			if (!SpecifierContext.Scope.HeaderParser.HeaderFile.bIsNoExportTypes)
+			if (!specifierContext.Scope.HeaderParser.HeaderFile.IsNoExportTypes)
 			{
-				SpecifierContext.MessageSite.LogError("The 'HasNoOpConstructor' struct specifier is only valid in the NoExportTypes.h file");
+				specifierContext.MessageSite.LogError("The 'HasNoOpConstructor' struct specifier is only valid in the NoExportTypes.h file");
 			}
-			UhtScriptStruct ScriptStruct = (UhtScriptStruct)SpecifierContext.Scope.ScopeType;
-			ScriptStruct.ScriptStructExportFlags |= UhtScriptStructExportFlags.HasNoOpConstructor;
+			UhtScriptStruct scriptStruct = (UhtScriptStruct)specifierContext.Scope.ScopeType;
+			scriptStruct.ScriptStructExportFlags |= UhtScriptStructExportFlags.HasNoOpConstructor;
 		}
 
 		[UhtSpecifier(Extends = UhtTableNames.ScriptStruct, ValueType = UhtSpecifierValueType.Legacy)]
 		[SuppressMessage("CodeQuality", "IDE0051:Remove unused private members", Justification = "Attribute accessed method")]
-		private static void IsAlwaysAccessibleSpecifier(UhtSpecifierContext SpecifierContext)
+		private static void IsAlwaysAccessibleSpecifier(UhtSpecifierContext specifierContext)
 		{
-			if (!SpecifierContext.Scope.HeaderParser.HeaderFile.bIsNoExportTypes)
+			if (!specifierContext.Scope.HeaderParser.HeaderFile.IsNoExportTypes)
 			{
-				SpecifierContext.MessageSite.LogError("The 'IsAlwaysAccessible' struct specifier is only valid in the NoExportTypes.h file");
+				specifierContext.MessageSite.LogError("The 'IsAlwaysAccessible' struct specifier is only valid in the NoExportTypes.h file");
 			}
-			UhtScriptStruct ScriptStruct = (UhtScriptStruct)SpecifierContext.Scope.ScopeType;
-			ScriptStruct.ScriptStructExportFlags |= UhtScriptStructExportFlags.IsAlwaysAccessible;
+			UhtScriptStruct scriptStruct = (UhtScriptStruct)specifierContext.Scope.ScopeType;
+			scriptStruct.ScriptStructExportFlags |= UhtScriptStructExportFlags.IsAlwaysAccessible;
 		}
 
 		[UhtSpecifier(Extends = UhtTableNames.ScriptStruct, ValueType = UhtSpecifierValueType.Legacy)]
 		[SuppressMessage("CodeQuality", "IDE0051:Remove unused private members", Justification = "Attribute accessed method")]
-		private static void IsCoreTypeSpecifier(UhtSpecifierContext SpecifierContext)
+		private static void IsCoreTypeSpecifier(UhtSpecifierContext specifierContext)
 		{
-			if (!SpecifierContext.Scope.HeaderParser.HeaderFile.bIsNoExportTypes)
+			if (!specifierContext.Scope.HeaderParser.HeaderFile.IsNoExportTypes)
 			{
-				SpecifierContext.MessageSite.LogError("The 'IsCoreType' struct specifier is only valid in the NoExportTypes.h file");
+				specifierContext.MessageSite.LogError("The 'IsCoreType' struct specifier is only valid in the NoExportTypes.h file");
 			}
-			UhtScriptStruct ScriptStruct = (UhtScriptStruct)SpecifierContext.Scope.ScopeType;
-			ScriptStruct.ScriptStructExportFlags |= UhtScriptStructExportFlags.IsCoreType;
+			UhtScriptStruct scriptStruct = (UhtScriptStruct)specifierContext.Scope.ScopeType;
+			scriptStruct.ScriptStructExportFlags |= UhtScriptStructExportFlags.IsCoreType;
 		}
 	}
 }
diff --git a/Engine/Source/Programs/Shared/EpicGames.UHT/Tables/UhtEngineClassTable.cs b/Engine/Source/Programs/Shared/EpicGames.UHT/Tables/UhtEngineClassTable.cs
--- a/Engine/Source/Programs/Shared/EpicGames.UHT/Tables/UhtEngineClassTable.cs
+++ b/Engine/Source/Programs/Shared/EpicGames.UHT/Tables/UhtEngineClassTable.cs
@@ -1,9 +1,9 @@
 // Copyright Epic Games, Inc. All Rights Reserved.
 
-using EpicGames.Core;
-using EpicGames.UHT.Utils;
 using System;
 using System.Collections.Generic;
+using EpicGames.Core;
+using EpicGames.UHT.Utils;
 
 namespace EpicGames.UHT.Tables
 {
@@ -39,7 +39,7 @@ namespace EpicGames.UHT.Tables
 		/// <summary>
 		/// If true, this class is a property
 		/// </summary>
-		public bool bIsProperty { get; set; }
+		public bool IsProperty { get; set; }
 	}
 
 	/// <summary>
@@ -50,18 +50,18 @@ namespace EpicGames.UHT.Tables
 		/// <summary>
 		/// Internal mapping from engine class name to information
 		/// </summary>
-		private readonly Dictionary<StringView, UhtEngineClass> EngineClasses = new Dictionary<StringView, UhtEngineClass>();
+		private readonly Dictionary<StringView, UhtEngineClass> _engineClasses = new Dictionary<StringView, UhtEngineClass>();
 
 		/// <summary>
 		/// Test to see if the given class name is a property
 		/// </summary>
-		/// <param name="Name">Name of the class without the prefix</param>
+		/// <param name="name">Name of the class without the prefix</param>
 		/// <returns>True if the class name is a property.  False if the class name isn't a property or isn't an engine class.</returns>
-		public bool IsValidPropertyTypeName(StringView Name)
+		public bool IsValidPropertyTypeName(StringView name)
 		{
-			if (this.EngineClasses.TryGetValue(Name, out UhtEngineClass EngineClass))
+			if (this._engineClasses.TryGetValue(name, out UhtEngineClass engineClass))
 			{
-				return EngineClass.bIsProperty;
+				return engineClass.IsProperty;
 			}
 			return false;
 		}
@@ -69,14 +69,14 @@ namespace EpicGames.UHT.Tables
 		/// <summary>
 		/// Add an entry to the table
 		/// </summary>
-		/// <param name="EngineClassAttribute">The located attribute</param>
-		public void OnEngineClassAttribute(UhtEngineClassAttribute EngineClassAttribute)
+		/// <param name="engineClassAttribute">The located attribute</param>
+		public void OnEngineClassAttribute(UhtEngineClassAttribute engineClassAttribute)
 		{
-			if (string.IsNullOrEmpty(EngineClassAttribute.Name))
+			if (String.IsNullOrEmpty(engineClassAttribute.Name))
 			{
 				throw new UhtIceException("EngineClassNames must have a name specified");
 			}
-			this.EngineClasses.Add(EngineClassAttribute.Name, new UhtEngineClass { Name = EngineClassAttribute.Name, bIsProperty = EngineClassAttribute.IsProperty });
+			this._engineClasses.Add(engineClassAttribute.Name, new UhtEngineClass { Name = engineClassAttribute.Name, IsProperty = engineClassAttribute.IsProperty });
 		}
 	}
 }
