More coding standard fixes

diff --git a/Engine/Source/Programs/Shared/EpicGames.UHT/Tables/UhtExporterTable.cs b/Engine/Source/Programs/Shared/EpicGames.UHT/Tables/UhtExporterTable.cs
--- a/Engine/Source/Programs/Shared/EpicGames.UHT/Tables/UhtExporterTable.cs
+++ b/Engine/Source/Programs/Shared/EpicGames.UHT/Tables/UhtExporterTable.cs
@@ -1,10 +1,10 @@
 // Copyright Epic Games, Inc. All Rights Reserved.
 
-using EpicGames.UHT.Utils;
 using System;
 using System.Collections;
 using System.Collections.Generic;
 using System.Reflection;
+using EpicGames.UHT.Utils;
 
 namespace EpicGames.UHT.Tables
 {
@@ -12,8 +12,8 @@ namespace EpicGames.UHT.Tables
 	/// <summary>
 	/// Delegate to invoke to run exporter
 	/// </summary>
-	/// <param name="Factory">Factory used to generate export tasks and outputs</param>
-	public delegate void UhtExporterDelegate(IUhtExportFactory Factory);
+	/// <param name="factory">Factory used to generate export tasks and outputs</param>
+	public delegate void UhtExporterDelegate(IUhtExportFactory factory);
 
 	/// <summary>
 	/// Export options
@@ -43,35 +43,35 @@ namespace EpicGames.UHT.Tables
 		/// <summary>
 		/// Test to see if any of the specified flags are set
 		/// </summary>
-		/// <param name="InFlags">Current flags</param>
-		/// <param name="TestFlags">Flags to test for</param>
+		/// <param name="inFlags">Current flags</param>
+		/// <param name="testFlags">Flags to test for</param>
 		/// <returns>True if any of the flags are set</returns>
-		public static bool HasAnyFlags(this UhtExporterOptions InFlags, UhtExporterOptions TestFlags)
+		public static bool HasAnyFlags(this UhtExporterOptions inFlags, UhtExporterOptions testFlags)
 		{
-			return (InFlags & TestFlags) != 0;
+			return (inFlags & testFlags) != 0;
 		}
 
 		/// <summary>
 		/// Test to see if all of the specified flags are set
 		/// </summary>
-		/// <param name="InFlags">Current flags</param>
-		/// <param name="TestFlags">Flags to test for</param>
+		/// <param name="inFlags">Current flags</param>
+		/// <param name="testFlags">Flags to test for</param>
 		/// <returns>True if all the flags are set</returns>
-		public static bool HasAllFlags(this UhtExporterOptions InFlags, UhtExporterOptions TestFlags)
+		public static bool HasAllFlags(this UhtExporterOptions inFlags, UhtExporterOptions testFlags)
 		{
-			return (InFlags & TestFlags) == TestFlags;
+			return (inFlags & testFlags) == testFlags;
 		}
 
 		/// <summary>
 		/// Test to see if a specific set of flags have a specific value.
 		/// </summary>
-		/// <param name="InFlags">Current flags</param>
-		/// <param name="TestFlags">Flags to test for</param>
-		/// <param name="MatchFlags">Expected value of the tested flags</param>
+		/// <param name="inFlags">Current flags</param>
+		/// <param name="testFlags">Flags to test for</param>
+		/// <param name="matchFlags">Expected value of the tested flags</param>
 		/// <returns>True if the given flags have a specific value.</returns>
-		public static bool HasExactFlags(this UhtExporterOptions InFlags, UhtExporterOptions TestFlags, UhtExporterOptions MatchFlags)
+		public static bool HasExactFlags(this UhtExporterOptions inFlags, UhtExporterOptions testFlags, UhtExporterOptions matchFlags)
 		{
-			return (InFlags & TestFlags) == MatchFlags;
+			return (inFlags & testFlags) == matchFlags;
 		}
 	}
 
@@ -85,17 +85,17 @@ namespace EpicGames.UHT.Tables
 		/// <summary>
 		/// Name of the exporter
 		/// </summary>
-		public string Name { get; set; } = string.Empty;
+		public string Name { get; set; } = String.Empty;
 
 		/// <summary>
 		/// Description of the export.  Used to display help
 		/// </summary>
-		public string Description { get; set; } = string.Empty;
+		public string Description { get; set; } = String.Empty;
 
 		/// <summary>
 		/// Exporters in plugins need to specify a module name
 		/// </summary>
-		public string ModuleName { get; set; } = string.Empty;
+		public string ModuleName { get; set; } = String.Empty;
 
 		/// <summary>
 		/// Exporter options
@@ -167,18 +167,18 @@ namespace EpicGames.UHT.Tables
 		/// <summary>
 		/// Construct an exporter table instance
 		/// </summary>
-		/// <param name="Attribute">Source attribute</param>
-		/// <param name="Delegate">Delegate to invoke</param>
-		public UhtExporter(UhtExporterAttribute Attribute, UhtExporterDelegate Delegate)
+		/// <param name="attribute">Source attribute</param>
+		/// <param name="exporterDelegate">Delegate to invoke</param>
+		public UhtExporter(UhtExporterAttribute attribute, UhtExporterDelegate exporterDelegate)
 		{
-			this.Name = Attribute.Name;
-			this.Description = Attribute.Description;
-			this.ModuleName = Attribute.ModuleName;
-			this.Options = Attribute.Options;
-			this.Delegate = Delegate;
-			this.CppFilters = Attribute.CppFilters != null ? new List<string>(Attribute.CppFilters) : new List<string>();
-			this.HeaderFilters = Attribute.HeaderFilters != null ? new List<string>(Attribute.HeaderFilters) : new List<string>();
-			this.OtherFilters = Attribute.OtherFilters != null ? new List<string>(Attribute.OtherFilters) : new List<string>();
+			this.Name = attribute.Name;
+			this.Description = attribute.Description;
+			this.ModuleName = attribute.ModuleName;
+			this.Options = attribute.Options;
+			this.Delegate = exporterDelegate;
+			this.CppFilters = attribute.CppFilters != null ? new List<string>(attribute.CppFilters) : new List<string>();
+			this.HeaderFilters = attribute.HeaderFilters != null ? new List<string>(attribute.HeaderFilters) : new List<string>();
+			this.OtherFilters = attribute.OtherFilters != null ? new List<string>(attribute.OtherFilters) : new List<string>();
 		}
 	}
 
@@ -188,43 +188,43 @@ namespace EpicGames.UHT.Tables
 	public class UhtExporterTable : IEnumerable<UhtExporter>
 	{
 
-		private readonly Dictionary<string, UhtExporter> ExporterValues = new Dictionary<string, UhtExporter>(StringComparer.OrdinalIgnoreCase);
+		private readonly Dictionary<string, UhtExporter> _exporterValues = new Dictionary<string, UhtExporter>(StringComparer.OrdinalIgnoreCase);
 
 		/// <summary>
 		/// Return the exporter associated with the given name
 		/// </summary>
-		/// <param name="Name"></param>
-		/// <param name="Value">Exporter associated with the name</param>
+		/// <param name="name"></param>
+		/// <param name="value">Exporter associated with the name</param>
 		/// <returns></returns>
-		public bool TryGet(string Name, out UhtExporter Value)
+		public bool TryGet(string name, out UhtExporter value)
 		{
-			return this.ExporterValues.TryGetValue(Name, out Value);
+			return this._exporterValues.TryGetValue(name, out value);
 		}
 
 		/// <summary>
 		/// Handle an exporter attribute
 		/// </summary>
-		/// <param name="Type">Containing type</param>
-		/// <param name="MethodInfo">Method info</param>
-		/// <param name="ExporterAttribute">Defining attribute</param>
+		/// <param name="type">Containing type</param>
+		/// <param name="methodInfo">Method info</param>
+		/// <param name="exporterAttribute">Defining attribute</param>
 		/// <exception cref="UhtIceException">Thrown if the attribute doesn't properly define an exporter.</exception>
-		public void OnExporterAttribute(Type Type, MethodInfo MethodInfo, UhtExporterAttribute ExporterAttribute)
+		public void OnExporterAttribute(Type type, MethodInfo methodInfo, UhtExporterAttribute exporterAttribute)
 		{
-			if (string.IsNullOrEmpty(ExporterAttribute.Name))
+			if (String.IsNullOrEmpty(exporterAttribute.Name))
 			{
 				throw new UhtIceException("An exporter must have a name");
 			}
 
-			if (Assembly.GetExecutingAssembly() != Type.Assembly)
+			if (Assembly.GetExecutingAssembly() != type.Assembly)
 			{
-				if (string.IsNullOrEmpty(ExporterAttribute.ModuleName))
+				if (String.IsNullOrEmpty(exporterAttribute.ModuleName))
 				{
 					throw new UhtIceException("An exporter in a UBT plugin must specify a ModuleName");
 				}
 			}
 
-			UhtExporter ExporterValue = new UhtExporter(ExporterAttribute, (UhtExporterDelegate)Delegate.CreateDelegate(typeof(UhtExporterDelegate), MethodInfo));
-			this.ExporterValues.Add(ExporterAttribute.Name, ExporterValue);
+			UhtExporter exporterValue = new UhtExporter(exporterAttribute, (UhtExporterDelegate)Delegate.CreateDelegate(typeof(UhtExporterDelegate), methodInfo));
+			this._exporterValues.Add(exporterAttribute.Name, exporterValue);
 		}
 
 		/// <summary>
@@ -233,9 +233,9 @@ namespace EpicGames.UHT.Tables
 		/// <returns>Enumerator</returns>
 		public IEnumerator<UhtExporter> GetEnumerator()
 		{
-			foreach (KeyValuePair<string, UhtExporter> KVP in this.ExporterValues)
+			foreach (KeyValuePair<string, UhtExporter> kvp in this._exporterValues)
 			{
-				yield return KVP.Value;
+				yield return kvp.Value;
 			}
 		}
 
@@ -245,9 +245,9 @@ namespace EpicGames.UHT.Tables
 		/// <returns>Enumerator</returns>
 		IEnumerator IEnumerable.GetEnumerator()
 		{
-			foreach (KeyValuePair<string, UhtExporter> KVP in this.ExporterValues)
+			foreach (KeyValuePair<string, UhtExporter> kvp in this._exporterValues)
 			{
-				yield return KVP.Value;
+				yield return kvp.Value;
 			}
 		}
 	}
diff --git a/Engine/Source/Programs/Shared/EpicGames.UHT/Tables/UhtKeywordTable.cs b/Engine/Source/Programs/Shared/EpicGames.UHT/Tables/UhtKeywordTable.cs
--- a/Engine/Source/Programs/Shared/EpicGames.UHT/Tables/UhtKeywordTable.cs
+++ b/Engine/Source/Programs/Shared/EpicGames.UHT/Tables/UhtKeywordTable.cs
@@ -60,19 +60,19 @@ namespace EpicGames.UHT.Tables
 	/// <summary>
 	/// Delegate to notify a keyword was parsed
 	/// </summary>
-	/// <param name="TopScope">Current scope being parsed</param>
-	/// <param name="ActionScope">The scope who's table was matched</param>
-	/// <param name="Token">Matching token</param>
+	/// <param name="topScope">Current scope being parsed</param>
+	/// <param name="actionScope">The scope who's table was matched</param>
+	/// <param name="token">Matching token</param>
 	/// <returns>Results of the parsing</returns>
-	public delegate UhtParseResult UhtKeywordDelegate(UhtParsingScope TopScope, UhtParsingScope ActionScope, ref UhtToken Token);
+	public delegate UhtParseResult UhtKeywordDelegate(UhtParsingScope topScope, UhtParsingScope actionScope, ref UhtToken token);
 
 	/// <summary>
 	/// Delegate to invoke as a last chance processor for a keyword
 	/// </summary>
-	/// <param name="TopScope">Current scope being parsed</param>
-	/// <param name="Token">Matching token</param>
+	/// <param name="topScope">Current scope being parsed</param>
+	/// <param name="token">Matching token</param>
 	/// <returns>Results of the parsing</returns>
-	public delegate UhtParseResult UhtKeywordCatchAllDelegate(UhtParsingScope TopScope, ref UhtToken Token);
+	public delegate UhtParseResult UhtKeywordCatchAllDelegate(UhtParsingScope topScope, ref UhtToken token);
 
 	/// <summary>
 	/// Defines a keyword
@@ -98,34 +98,34 @@ namespace EpicGames.UHT.Tables
 		/// <summary>
 		/// If true, this applies to all scopes
 		/// </summary>
-		public bool bAllScopes { get; }
+		public bool AllScopes { get; }
 
 		/// <summary>
 		/// If true, do not include in usage errors
 		/// </summary>
-		public bool bDisableUsageError { get; }
+		public bool DisableUsageError { get; }
 
 		/// <summary>
 		/// Construct a new keyword
 		/// </summary>
-		/// <param name="Name">Name of the keyword</param>
-		/// <param name="Delegate">Delegate to invoke</param>
-		/// <param name="Attribute">Defining attribute</param>
-		public UhtKeyword(string Name, UhtKeywordDelegate Delegate, UhtKeywordAttribute? Attribute)
+		/// <param name="name">Name of the keyword</param>
+		/// <param name="keywordDelegate">Delegate to invoke</param>
+		/// <param name="attribute">Defining attribute</param>
+		public UhtKeyword(string name, UhtKeywordDelegate keywordDelegate, UhtKeywordAttribute? attribute)
 		{
-			this.Name = Name;
-			this.Delegate = Delegate;
-			if (Attribute != null)
+			this.Name = name;
+			this.Delegate = keywordDelegate;
+			if (attribute != null)
 			{
-				this.AllowText = Attribute.AllowText;
-				this.bAllScopes = Attribute.AllScopes;
-				this.bDisableUsageError = Attribute.DisableUsageError;
+				this.AllowText = attribute.AllowText;
+				this.AllScopes = attribute.AllScopes;
+				this.DisableUsageError = attribute.DisableUsageError;
 			}
 			else
 			{
 				this.AllowText = null;
-				this.bAllScopes = false;
-				this.bDisableUsageError = false;
+				this.AllScopes = false;
+				this.DisableUsageError = false;
 			}
 		}
 	}
@@ -151,31 +151,31 @@ namespace EpicGames.UHT.Tables
 		/// <summary>
 		/// Add the given value to the lookup table.  It will throw an exception if it is a duplicate.
 		/// </summary>
-		/// <param name="Value">Value to be added</param>
-		public UhtKeywordTable Add(UhtKeyword Value)
+		/// <param name="value">Value to be added</param>
+		public UhtKeywordTable Add(UhtKeyword value)
 		{
-			base.Add(Value.Name, Value);
+			base.Add(value.Name, value);
 			return this;
 		}
 
 		/// <summary>
 		/// Add the given catch-all to the table.
 		/// </summary>
-		/// <param name="CatchAll">The catch-all to be added</param>
-		public UhtKeywordTable AddCatchAll(UhtKeywordCatchAllDelegate CatchAll)
+		/// <param name="catchAll">The catch-all to be added</param>
+		public UhtKeywordTable AddCatchAll(UhtKeywordCatchAllDelegate catchAll)
 		{
-			this.CatchAlls.Add(CatchAll);
+			this.CatchAlls.Add(catchAll);
 			return this;
 		}
 
 		/// <summary>
 		/// Merge the given keyword table.  Duplicates in the BaseTypeTable will be ignored.
 		/// </summary>
-		/// <param name="BaseTable">Base table being merged</param>
-		public override void Merge(UhtLookupTableBase BaseTable)
+		/// <param name="baseTable">Base table being merged</param>
+		public override void Merge(UhtLookupTableBase baseTable)
 		{
-			base.Merge(BaseTable);
-			this.CatchAlls.AddRange(((UhtKeywordTable)BaseTable).CatchAlls);
+			base.Merge(baseTable);
+			this.CatchAlls.AddRange(((UhtKeywordTable)baseTable).CatchAlls);
 		}
 	}
 
@@ -195,85 +195,85 @@ namespace EpicGames.UHT.Tables
 		/// <summary>
 		/// Handle a keyword attribute
 		/// </summary>
-		/// <param name="Type">Containing type</param>
-		/// <param name="MethodInfo">Method information</param>
-		/// <param name="KeywordCatchAllAttribute">Defining attribute</param>
+		/// <param name="type">Containing type</param>
+		/// <param name="methodInfo">Method information</param>
+		/// <param name="keywordCatchAllAttribute">Defining attribute</param>
 		/// <exception cref="UhtIceException">Thrown if the attribute isn't well defined</exception>
-		public void OnKeywordCatchAllAttribute(Type Type, MethodInfo MethodInfo, UhtKeywordCatchAllAttribute KeywordCatchAllAttribute)
+		public void OnKeywordCatchAllAttribute(Type type, MethodInfo methodInfo, UhtKeywordCatchAllAttribute keywordCatchAllAttribute)
 		{
-			if (string.IsNullOrEmpty(KeywordCatchAllAttribute.Extends))
+			if (String.IsNullOrEmpty(keywordCatchAllAttribute.Extends))
 			{
-				throw new UhtIceException($"The 'KeywordCatchAlll' attribute on the {Type.Name}.{MethodInfo.Name} method doesn't have a table specified.");
+				throw new UhtIceException($"The 'KeywordCatchAlll' attribute on the {type.Name}.{methodInfo.Name} method doesn't have a table specified.");
 			}
 
-			UhtKeywordTable Table = Get(KeywordCatchAllAttribute.Extends);
-			Table.AddCatchAll((UhtKeywordCatchAllDelegate)Delegate.CreateDelegate(typeof(UhtKeywordCatchAllDelegate), MethodInfo));
+			UhtKeywordTable table = Get(keywordCatchAllAttribute.Extends);
+			table.AddCatchAll((UhtKeywordCatchAllDelegate)Delegate.CreateDelegate(typeof(UhtKeywordCatchAllDelegate), methodInfo));
 		}
 
 		/// <summary>
 		/// Handle a keyword attribute
 		/// </summary>
-		/// <param name="Type">Containing type</param>
-		/// <param name="MethodInfo">Method information</param>
-		/// <param name="KeywordAttribute">Defining attribute</param>
+		/// <param name="type">Containing type</param>
+		/// <param name="methodInfo">Method information</param>
+		/// <param name="keywordAttribute">Defining attribute</param>
 		/// <exception cref="UhtIceException">Thrown if the attribute isn't well defined</exception>
-		public void OnKeywordAttribute(Type Type, MethodInfo MethodInfo, UhtKeywordAttribute KeywordAttribute)
+		public void OnKeywordAttribute(Type type, MethodInfo methodInfo, UhtKeywordAttribute keywordAttribute)
 		{
-			string Name = UhtLookupTableBase.GetSuffixedName(Type, MethodInfo, KeywordAttribute.Keyword, "Keyword");
+			string name = UhtLookupTableBase.GetSuffixedName(type, methodInfo, keywordAttribute.Keyword, "Keyword");
 
-			if (string.IsNullOrEmpty(KeywordAttribute.Extends))
+			if (String.IsNullOrEmpty(keywordAttribute.Extends))
 			{
-				throw new UhtIceException($"The 'Keyword' attribute on the {Type.Name}.{MethodInfo.Name} method doesn't have a table specified.");
+				throw new UhtIceException($"The 'Keyword' attribute on the {type.Name}.{methodInfo.Name} method doesn't have a table specified.");
 			}
 
-			UhtKeywordTable Table = Get(KeywordAttribute.Extends);
-			Table.Add(new UhtKeyword(Name, (UhtKeywordDelegate)Delegate.CreateDelegate(typeof(UhtKeywordDelegate), MethodInfo), KeywordAttribute));
+			UhtKeywordTable table = Get(keywordAttribute.Extends);
+			table.Add(new UhtKeyword(name, (UhtKeywordDelegate)Delegate.CreateDelegate(typeof(UhtKeywordDelegate), methodInfo), keywordAttribute));
 		}
 
 		/// <summary>
 		/// Log an unhandled error
 		/// </summary>
-		/// <param name="MessageSite">Destination message site</param>
-		/// <param name="Token">Keyword</param>
-		public void LogUnhandledError(IUhtMessageSite MessageSite, UhtToken Token)
+		/// <param name="messageSite">Destination message site</param>
+		/// <param name="token">Keyword</param>
+		public void LogUnhandledError(IUhtMessageSite messageSite, UhtToken token)
 		{
-			List<string>? Tables = null;
-			foreach (KeyValuePair<string, UhtKeywordTable> KVP in this.Tables)
+			List<string>? tables = null;
+			foreach (KeyValuePair<string, UhtKeywordTable> kvp in this.Tables)
 			{
-				UhtKeywordTable KeywordTable = KVP.Value;
-				if (KeywordTable.Internal)
+				UhtKeywordTable keywordTable = kvp.Value;
+				if (keywordTable.Internal)
 				{
 					continue;
 				}
-				UhtKeyword Info;
-				if (KeywordTable.TryGetValue(Token.Value, out Info))
+				UhtKeyword info;
+				if (keywordTable.TryGetValue(token.Value, out info))
 				{
-					if (Info.bDisableUsageError)
+					if (info.DisableUsageError)
 					{
 						// Do not log anything for this keyword in this table
 					}
 					else
 					{
-						if (Tables == null)
+						if (tables == null)
 						{
-							Tables = new List<string>();
+							tables = new List<string>();
 						}
-						if (Info.AllowText != null)
+						if (info.AllowText != null)
 						{
-							Tables.Add($"{KeywordTable.UserName} {Info.AllowText}");
+							tables.Add($"{keywordTable.UserName} {info.AllowText}");
 						}
 						else
 						{
-							Tables.Add(KeywordTable.UserName);
+							tables.Add(keywordTable.UserName);
 						}
 					}
 				}
 			}
 
-			if (Tables != null)
+			if (tables != null)
 			{
-				string Text = UhtUtilities.MergeTypeNames(Tables, "and");
-				MessageSite.LogError(Token.InputLine, $"Invalid use of keyword '{Token.Value}'.  It may only appear in {Text} scopes");
+				string text = UhtUtilities.MergeTypeNames(tables, "and");
+				messageSite.LogError(token.InputLine, $"Invalid use of keyword '{token.Value}'.  It may only appear in {text} scopes");
 			}
 		}
 	}
diff --git a/Engine/Source/Programs/Shared/EpicGames.UHT/Tables/UhtLocTextDefaultValueTable.cs b/Engine/Source/Programs/Shared/EpicGames.UHT/Tables/UhtLocTextDefaultValueTable.cs
--- a/Engine/Source/Programs/Shared/EpicGames.UHT/Tables/UhtLocTextDefaultValueTable.cs
+++ b/Engine/Source/Programs/Shared/EpicGames.UHT/Tables/UhtLocTextDefaultValueTable.cs
@@ -1,13 +1,13 @@
 // Copyright Epic Games, Inc. All Rights Reserved.
 
-using EpicGames.Core;
-using EpicGames.UHT.Tokenizer;
 using System;
 using System.Collections.Generic;
 using System.Reflection;
+using System.Text;
+using EpicGames.Core;
+using EpicGames.UHT.Tokenizer;
 using EpicGames.UHT.Types;
 using EpicGames.UHT.Utils;
-using System.Text;
 
 namespace EpicGames.UHT.Tables
 {
@@ -15,12 +15,12 @@ namespace EpicGames.UHT.Tables
 	/// <summary>
 	/// Delegate to invoke to sanitize a loctext default value
 	/// </summary>
-	/// <param name="Property">Property in question</param>
-	/// <param name="DefaultValueReader">The default value</param>
-	/// <param name="MacroToken">Token for the loctext type being parsed</param>
-	/// <param name="InnerDefaultValue">Output sanitized value.</param>
+	/// <param name="property">Property in question</param>
+	/// <param name="defaultValueReader">The default value</param>
+	/// <param name="macroToken">Token for the loctext type being parsed</param>
+	/// <param name="innerDefaultValue">Output sanitized value.</param>
 	/// <returns>True if sanitized, false if not.</returns>
-	public delegate bool UhtLocTextDefaultValueDelegate(UhtTextProperty Property, IUhtTokenReader DefaultValueReader, ref UhtToken MacroToken, StringBuilder InnerDefaultValue);
+	public delegate bool UhtLocTextDefaultValueDelegate(UhtTextProperty property, IUhtTokenReader defaultValueReader, ref UhtToken macroToken, StringBuilder innerDefaultValue);
 
 	/// <summary>
 	/// Attribute defining the loctext sanitizer
@@ -52,38 +52,38 @@ namespace EpicGames.UHT.Tables
 	public class UhtLocTextDefaultValueTable
 	{
 
-		private readonly Dictionary<StringView, UhtLocTextDefaultValue> LocTextDefaultValues = new Dictionary<StringView, UhtLocTextDefaultValue>();
+		private readonly Dictionary<StringView, UhtLocTextDefaultValue> _locTextDefaultValues = new Dictionary<StringView, UhtLocTextDefaultValue>();
 
 		/// <summary>
 		/// Return the loc text default value associated with the given name
 		/// </summary>
-		/// <param name="Name"></param>
-		/// <param name="LocTextDefaultValue">Loc text default value handler</param>
+		/// <param name="name"></param>
+		/// <param name="locTextDefaultValue">Loc text default value handler</param>
 		/// <returns></returns>
-		public bool TryGet(StringView Name, out UhtLocTextDefaultValue LocTextDefaultValue)
+		public bool TryGet(StringView name, out UhtLocTextDefaultValue locTextDefaultValue)
 		{
-			return this.LocTextDefaultValues.TryGetValue(Name, out LocTextDefaultValue);
+			return this._locTextDefaultValues.TryGetValue(name, out locTextDefaultValue);
 		}
 
 		/// <summary>
 		/// Handle a loctext default value attribute
 		/// </summary>
-		/// <param name="MethodInfo">Method info</param>
-		/// <param name="LocTextDefaultValueAttribute">Defining attribute</param>
+		/// <param name="methodInfo">Method info</param>
+		/// <param name="locTextDefaultValueAttribute">Defining attribute</param>
 		/// <exception cref="UhtIceException">Thrown if the attribute isn't properly defined</exception>
-		public void OnLocTextDefaultValueAttribute(MethodInfo MethodInfo, UhtLocTextDefaultValueAttribute LocTextDefaultValueAttribute)
+		public void OnLocTextDefaultValueAttribute(MethodInfo methodInfo, UhtLocTextDefaultValueAttribute locTextDefaultValueAttribute)
 		{
-			if (string.IsNullOrEmpty(LocTextDefaultValueAttribute.Name))
+			if (String.IsNullOrEmpty(locTextDefaultValueAttribute.Name))
 			{
 				throw new UhtIceException("A loc text default value attribute must have a name");
 			}
 
-			UhtLocTextDefaultValue LocTextDefaultValue = new UhtLocTextDefaultValue
+			UhtLocTextDefaultValue locTextDefaultValue = new UhtLocTextDefaultValue
 			{
-				Delegate = (UhtLocTextDefaultValueDelegate)Delegate.CreateDelegate(typeof(UhtLocTextDefaultValueDelegate), MethodInfo)
+				Delegate = (UhtLocTextDefaultValueDelegate)Delegate.CreateDelegate(typeof(UhtLocTextDefaultValueDelegate), methodInfo)
 			};
 
-			this.LocTextDefaultValues.Add(new StringView(LocTextDefaultValueAttribute.Name), LocTextDefaultValue);
+			this._locTextDefaultValues.Add(new StringView(locTextDefaultValueAttribute.Name), locTextDefaultValue);
 		}
 	}
 }
diff --git a/Engine/Source/Programs/Shared/EpicGames.UHT/Tables/UhtPropertyTypeTable.cs b/Engine/Source/Programs/Shared/EpicGames.UHT/Tables/UhtPropertyTypeTable.cs
--- a/Engine/Source/Programs/Shared/EpicGames.UHT/Tables/UhtPropertyTypeTable.cs
+++ b/Engine/Source/Programs/Shared/EpicGames.UHT/Tables/UhtPropertyTypeTable.cs
@@ -1,16 +1,16 @@
 // Copyright Epic Games, Inc. All Rights Reserved.
 
+using System;
+using System.Collections.Generic;
+using System.Reflection;
 using EpicGames.Core;
 using EpicGames.UHT.Tokenizer;
 using EpicGames.UHT.Types;
 using EpicGames.UHT.Utils;
-using System;
-using System.Collections.Generic;
-using System.Reflection;
 
 namespace EpicGames.UHT.Tables
 {
-	
+
 	/// <summary>
 	/// Property type options
 	/// </summary>
@@ -55,35 +55,35 @@ namespace EpicGames.UHT.Tables
 		/// <summary>
 		/// Test to see if any of the specified flags are set
 		/// </summary>
-		/// <param name="InFlags">Current flags</param>
-		/// <param name="TestFlags">Flags to test for</param>
+		/// <param name="inFlags">Current flags</param>
+		/// <param name="testFlags">Flags to test for</param>
 		/// <returns>True if any of the flags are set</returns>
-		public static bool HasAnyFlags(this UhtPropertyTypeOptions InFlags, UhtPropertyTypeOptions TestFlags)
+		public static bool HasAnyFlags(this UhtPropertyTypeOptions inFlags, UhtPropertyTypeOptions testFlags)
 		{
-			return (InFlags & TestFlags) != 0;
+			return (inFlags & testFlags) != 0;
 		}
 
 		/// <summary>
 		/// Test to see if all of the specified flags are set
 		/// </summary>
-		/// <param name="InFlags">Current flags</param>
-		/// <param name="TestFlags">Flags to test for</param>
+		/// <param name="inFlags">Current flags</param>
+		/// <param name="testFlags">Flags to test for</param>
 		/// <returns>True if all the flags are set</returns>
-		public static bool HasAllFlags(this UhtPropertyTypeOptions InFlags, UhtPropertyTypeOptions TestFlags)
+		public static bool HasAllFlags(this UhtPropertyTypeOptions inFlags, UhtPropertyTypeOptions testFlags)
 		{
-			return (InFlags & TestFlags) == TestFlags;
+			return (inFlags & testFlags) == testFlags;
 		}
 
 		/// <summary>
 		/// Test to see if a specific set of flags have a specific value.
 		/// </summary>
-		/// <param name="InFlags">Current flags</param>
-		/// <param name="TestFlags">Flags to test for</param>
-		/// <param name="MatchFlags">Expected value of the tested flags</param>
+		/// <param name="inFlags">Current flags</param>
+		/// <param name="testFlags">Flags to test for</param>
+		/// <param name="matchFlags">Expected value of the tested flags</param>
 		/// <returns>True if the given flags have a specific value.</returns>
-		public static bool HasExactFlags(this UhtPropertyTypeOptions InFlags, UhtPropertyTypeOptions TestFlags, UhtPropertyTypeOptions MatchFlags)
+		public static bool HasExactFlags(this UhtPropertyTypeOptions inFlags, UhtPropertyTypeOptions testFlags, UhtPropertyTypeOptions matchFlags)
 		{
-			return (InFlags & TestFlags) == MatchFlags;
+			return (inFlags & testFlags) == matchFlags;
 		}
 	}
 
@@ -107,12 +107,12 @@ namespace EpicGames.UHT.Tables
 	/// <summary>
 	/// Delegate invoked to resolve a tokenized type into a UHTProperty type
 	/// </summary>
-	/// <param name="ResolvePhase">Specifies if this is being resolved during the parsing phase or the resolution phase.  Type lookups can not happen during the parsing phase</param>
-	/// <param name="PropertySettings">The configuration of the property</param>
-	/// <param name="TokenReader">The token reader containing the type</param>
-	/// <param name="MatchedToken">The token that matched the delegate unless the delegate is the default resolver.</param>
+	/// <param name="resolvePhase">Specifies if this is being resolved during the parsing phase or the resolution phase.  Type lookups can not happen during the parsing phase</param>
+	/// <param name="propertySettings">The configuration of the property</param>
+	/// <param name="tokenReader">The token reader containing the type</param>
+	/// <param name="matchedToken">The token that matched the delegate unless the delegate is the default resolver.</param>
 	/// <returns></returns>
-	public delegate UhtProperty? UhtResolvePropertyDelegate(UhtPropertyResolvePhase ResolvePhase, UhtPropertySettings PropertySettings, IUhtTokenReader TokenReader, UhtToken MatchedToken);
+	public delegate UhtProperty? UhtResolvePropertyDelegate(UhtPropertyResolvePhase resolvePhase, UhtPropertySettings propertySettings, IUhtTokenReader tokenReader, UhtToken matchedToken);
 
 	/// <summary>
 	/// Property type attribute
@@ -154,9 +154,9 @@ namespace EpicGames.UHT.Tables
 	/// </summary>
 	public class UhtPropertyTypeTable
 	{
-		private readonly Dictionary<StringView, UhtPropertyType> CaseSensitive = new Dictionary<StringView, UhtPropertyType>();
-		private readonly Dictionary<StringView, UhtPropertyType> CaseInsensitive = new Dictionary<StringView, UhtPropertyType>(StringViewComparer.OrdinalIgnoreCase);
-		private UhtPropertyType? DefaultInternal = null;
+		private readonly Dictionary<StringView, UhtPropertyType> _caseSensitive = new Dictionary<StringView, UhtPropertyType>();
+		private readonly Dictionary<StringView, UhtPropertyType> _caseInsensitive = new Dictionary<StringView, UhtPropertyType>(StringViewComparer.OrdinalIgnoreCase);
+		private UhtPropertyType? _defaultInternal = null;
 
 		/// <summary>
 		/// Return the default processor
@@ -165,63 +165,63 @@ namespace EpicGames.UHT.Tables
 		{
 			get
 			{
-				if (this.DefaultInternal == null)
+				if (this._defaultInternal == null)
 				{
 					throw new UhtIceException("No property type has been marked as default");
 				}
-				return (UhtPropertyType)this.DefaultInternal;
+				return (UhtPropertyType)this._defaultInternal;
 			}
 		}
 
 		/// <summary>
 		/// Return the property type associated with the given name
 		/// </summary>
-		/// <param name="Name"></param>
-		/// <param name="PropertyType">Property type if matched</param>
+		/// <param name="name"></param>
+		/// <param name="propertyType">Property type if matched</param>
 		/// <returns></returns>
-		public bool TryGet(StringView Name, out UhtPropertyType PropertyType)
+		public bool TryGet(StringView name, out UhtPropertyType propertyType)
 		{
 			return
-				this.CaseSensitive.TryGetValue(Name, out PropertyType) ||
-				this.CaseInsensitive.TryGetValue(Name, out PropertyType);
+				this._caseSensitive.TryGetValue(name, out propertyType) ||
+				this._caseInsensitive.TryGetValue(name, out propertyType);
 		}
 
 		/// <summary>
 		/// Handle a property type attribute
 		/// </summary>
-		/// <param name="MethodInfo">Method info</param>
-		/// <param name="PropertyTypeAttribute">Attribute</param>
+		/// <param name="methodInfo">Method info</param>
+		/// <param name="propertyTypeAttribute">Attribute</param>
 		/// <exception cref="UhtIceException">Thrown if the property type isn't properly defined.</exception>
-		public void OnPropertyTypeAttribute(MethodInfo MethodInfo, UhtPropertyTypeAttribute PropertyTypeAttribute)
+		public void OnPropertyTypeAttribute(MethodInfo methodInfo, UhtPropertyTypeAttribute propertyTypeAttribute)
 		{
-			if (string.IsNullOrEmpty(PropertyTypeAttribute.Keyword) && !PropertyTypeAttribute.Options.HasAnyFlags(UhtPropertyTypeOptions.Default))
+			if (String.IsNullOrEmpty(propertyTypeAttribute.Keyword) && !propertyTypeAttribute.Options.HasAnyFlags(UhtPropertyTypeOptions.Default))
 			{
 				throw new UhtIceException("A property type must have a keyword or be marked as default");
 			}
 
-			UhtPropertyType PropertyType = new UhtPropertyType
+			UhtPropertyType propertyType = new UhtPropertyType
 			{
-				Delegate = (UhtResolvePropertyDelegate)Delegate.CreateDelegate(typeof(UhtResolvePropertyDelegate), MethodInfo),
-				Options = PropertyTypeAttribute.Options,
+				Delegate = (UhtResolvePropertyDelegate)Delegate.CreateDelegate(typeof(UhtResolvePropertyDelegate), methodInfo),
+				Options = propertyTypeAttribute.Options,
 			};
 
-			if (PropertyTypeAttribute.Options.HasAnyFlags(UhtPropertyTypeOptions.Default))
+			if (propertyTypeAttribute.Options.HasAnyFlags(UhtPropertyTypeOptions.Default))
 			{
-				if (this.DefaultInternal != null)
+				if (this._defaultInternal != null)
 				{
 					throw new UhtIceException("Only one property type dispatcher can be marked as default");
 				}
-				this.DefaultInternal = PropertyType;
+				this._defaultInternal = propertyType;
 			}
-			else if (!string.IsNullOrEmpty(PropertyTypeAttribute.Keyword))
+			else if (!String.IsNullOrEmpty(propertyTypeAttribute.Keyword))
 			{
-				if (PropertyTypeAttribute.Options.HasAnyFlags(UhtPropertyTypeOptions.CaseInsensitive))
+				if (propertyTypeAttribute.Options.HasAnyFlags(UhtPropertyTypeOptions.CaseInsensitive))
 				{
-					this.CaseInsensitive.Add(PropertyTypeAttribute.Keyword, PropertyType);
+					this._caseInsensitive.Add(propertyTypeAttribute.Keyword, propertyType);
 				}
 				else
 				{
-					this.CaseSensitive.Add(PropertyTypeAttribute.Keyword, PropertyType);
+					this._caseSensitive.Add(propertyTypeAttribute.Keyword, propertyType);
 				}
 			}
 		}
diff --git a/Engine/Source/Programs/Shared/EpicGames.UHT/Tables/UhtSpecifierTable.cs b/Engine/Source/Programs/Shared/EpicGames.UHT/Tables/UhtSpecifierTable.cs
--- a/Engine/Source/Programs/Shared/EpicGames.UHT/Tables/UhtSpecifierTable.cs
+++ b/Engine/Source/Programs/Shared/EpicGames.UHT/Tables/UhtSpecifierTable.cs
@@ -1,12 +1,12 @@
 // Copyright Epic Games, Inc. All Rights Reserved.
 
+using System;
+using System.Collections.Generic;
+using System.Reflection;
 using EpicGames.Core;
 using EpicGames.UHT.Parsers; // It would be nice if we didn't need this here
 using EpicGames.UHT.Types;
 using EpicGames.UHT.Utils;
-using System;
-using System.Collections.Generic;
-using System.Reflection;
 
 namespace EpicGames.UHT.Tables
 {
@@ -113,16 +113,16 @@ namespace EpicGames.UHT.Tables
 		/// <summary>
 		/// Construct a new specifier context
 		/// </summary>
-		/// <param name="Scope"></param>
-		/// <param name="MessageSite"></param>
-		/// <param name="MetaData"></param>
-		/// <param name="MetaNameIndex"></param>
-		public UhtSpecifierContext(UhtParsingScope Scope, IUhtMessageSite MessageSite, UhtMetaData MetaData, int MetaNameIndex = UhtMetaData.IndexNone)
+		/// <param name="scope"></param>
+		/// <param name="messageSite"></param>
+		/// <param name="metaData"></param>
+		/// <param name="metaNameIndex"></param>
+		public UhtSpecifierContext(UhtParsingScope scope, IUhtMessageSite messageSite, UhtMetaData metaData, int metaNameIndex = UhtMetaData.IndexNone)
 		{
-			this.Scope = Scope;
-			this.MessageSite = MessageSite;
-			this.MetaData = MetaData;
-			this.MetaNameIndex = MetaNameIndex;
+			this.Scope = scope;
+			this.MessageSite = messageSite;
+			this.MetaData = metaData;
+			this.MetaNameIndex = metaNameIndex;
 		}
 	}
 
@@ -167,43 +167,43 @@ namespace EpicGames.UHT.Tables
 		/// <summary>
 		/// Dispatch an instance of the specifier
 		/// </summary>
-		/// <param name="SpecifierContext">Current context</param>
-		/// <param name="Value">Specifier value</param>
+		/// <param name="specifierContext">Current context</param>
+		/// <param name="value">Specifier value</param>
 		/// <returns>Results of the dispatch</returns>
-		public abstract UhtSpecifierDispatchResults Dispatch(UhtSpecifierContext SpecifierContext, object? Value);
+		public abstract UhtSpecifierDispatchResults Dispatch(UhtSpecifierContext specifierContext, object? value);
 	}
 
 	/// <summary>
 	/// Delegate for a specifier with no value
 	/// </summary>
-	/// <param name="SpecifierContext"></param>
-	public delegate void UhtSpecifierNoneDelegate(UhtSpecifierContext SpecifierContext);
+	/// <param name="specifierContext"></param>
+	public delegate void UhtSpecifierNoneDelegate(UhtSpecifierContext specifierContext);
 
 	/// <summary>
 	/// Specifier with no value
 	/// </summary>
 	public class UhtSpecifierNone : UhtSpecifier
 	{
-		private readonly UhtSpecifierNoneDelegate Delegate;
+		private readonly UhtSpecifierNoneDelegate _delegate;
 
 		/// <summary>
 		/// Construct the specifier
 		/// </summary>
-		/// <param name="Name">Name of the specifier</param>
-		/// <param name="When">When the specifier is executed</param>
-		/// <param name="Delegate">Delegate to invoke</param>
-		public UhtSpecifierNone(string Name, UhtSpecifierWhen When, UhtSpecifierNoneDelegate Delegate)
+		/// <param name="name">Name of the specifier</param>
+		/// <param name="when">When the specifier is executed</param>
+		/// <param name="specifierDelegate">Delegate to invoke</param>
+		public UhtSpecifierNone(string name, UhtSpecifierWhen when, UhtSpecifierNoneDelegate specifierDelegate)
 		{
-			this.Name = Name;
+			this.Name = name;
 			this.ValueType = UhtSpecifierValueType.None;
-			this.When = When;
-			this.Delegate = Delegate;
+			this.When = when;
+			this._delegate = specifierDelegate;
 		}
 
 		/// <inheritdoc/>
-		public override UhtSpecifierDispatchResults Dispatch(UhtSpecifierContext SpecifierContext, object? Value)
+		public override UhtSpecifierDispatchResults Dispatch(UhtSpecifierContext specifierContext, object? value)
 		{
-			this.Delegate(SpecifierContext);
+			this._delegate(specifierContext);
 			return UhtSpecifierDispatchResults.Known;
 		}
 	}
@@ -211,39 +211,39 @@ namespace EpicGames.UHT.Tables
 	/// <summary>
 	/// Specifier delegate with a string value
 	/// </summary>
-	/// <param name="SpecifierContext">Specifier context</param>
-	/// <param name="Value">Specifier value</param>
-	public delegate void UhtSpecifierStringDelegate(UhtSpecifierContext SpecifierContext, StringView Value);
+	/// <param name="specifierContext">Specifier context</param>
+	/// <param name="value">Specifier value</param>
+	public delegate void UhtSpecifierStringDelegate(UhtSpecifierContext specifierContext, StringView value);
 
 	/// <summary>
 	/// Specifier with a string value
 	/// </summary>
 	public class UhtSpecifierString : UhtSpecifier
 	{
-		private readonly UhtSpecifierStringDelegate Delegate;
+		private readonly UhtSpecifierStringDelegate _delegate;
 
 		/// <summary>
 		/// Construct the specifier
 		/// </summary>
-		/// <param name="Name">Name of the specifier</param>
-		/// <param name="When">When the specifier is executed</param>
-		/// <param name="Delegate">Delegate to invoke</param>
-		public UhtSpecifierString(string Name, UhtSpecifierWhen When, UhtSpecifierStringDelegate Delegate)
+		/// <param name="name">Name of the specifier</param>
+		/// <param name="when">When the specifier is executed</param>
+		/// <param name="specifierDelegate">Delegate to invoke</param>
+		public UhtSpecifierString(string name, UhtSpecifierWhen when, UhtSpecifierStringDelegate specifierDelegate)
 		{
-			this.Name = Name;
+			this.Name = name;
 			this.ValueType = UhtSpecifierValueType.String;
-			this.When = When;
-			this.Delegate = Delegate;
+			this.When = when;
+			this._delegate = specifierDelegate;
 		}
 
 		/// <inheritdoc/>
-		public override UhtSpecifierDispatchResults Dispatch(UhtSpecifierContext SpecifierContext, object? Value)
+		public override UhtSpecifierDispatchResults Dispatch(UhtSpecifierContext specifierContext, object? value)
 		{
-			if (Value == null)
+			if (value == null)
 			{
 				throw new UhtIceException("Required value is null");
 			}
-			this.Delegate(SpecifierContext, (StringView)Value);
+			this._delegate(specifierContext, (StringView)value);
 			return UhtSpecifierDispatchResults.Known;
 		}
 	}
@@ -251,35 +251,35 @@ namespace EpicGames.UHT.Tables
 	/// <summary>
 	/// Specifier delegate with an optional string value
 	/// </summary>
-	/// <param name="SpecifierContext">Specifier context</param>
-	/// <param name="Value">Specifier value</param>
-	public delegate void UhtSpecifierOptionalStringDelegate(UhtSpecifierContext SpecifierContext, StringView? Value);
+	/// <param name="specifierContext">Specifier context</param>
+	/// <param name="value">Specifier value</param>
+	public delegate void UhtSpecifierOptionalStringDelegate(UhtSpecifierContext specifierContext, StringView? value);
 
 	/// <summary>
 	/// Specifier with an optional string value
 	/// </summary>
 	public class UhtSpecifierOptionalString : UhtSpecifier
 	{
-		private readonly UhtSpecifierOptionalStringDelegate Delegate;
+		private readonly UhtSpecifierOptionalStringDelegate _delegate;
 
 		/// <summary>
 		/// Construct the specifier
 		/// </summary>
-		/// <param name="Name">Name of the specifier</param>
-		/// <param name="When">When the specifier is executed</param>
-		/// <param name="Delegate">Delegate to invoke</param>
-		public UhtSpecifierOptionalString(string Name, UhtSpecifierWhen When, UhtSpecifierOptionalStringDelegate Delegate)
+		/// <param name="name">Name of the specifier</param>
+		/// <param name="when">When the specifier is executed</param>
+		/// <param name="specifierDelegate">Delegate to invoke</param>
+		public UhtSpecifierOptionalString(string name, UhtSpecifierWhen when, UhtSpecifierOptionalStringDelegate specifierDelegate)
 		{
-			this.Name = Name;
+			this.Name = name;
 			this.ValueType = UhtSpecifierValueType.OptionalString;
-			this.When = When;
-			this.Delegate = Delegate;
+			this.When = when;
+			this._delegate = specifierDelegate;
 		}
 
 		/// <inheritdoc/>
-		public override UhtSpecifierDispatchResults Dispatch(UhtSpecifierContext SpecifierContext, object? Value)
+		public override UhtSpecifierDispatchResults Dispatch(UhtSpecifierContext specifierContext, object? value)
 		{
-			this.Delegate(SpecifierContext, (StringView?)Value);
+			this._delegate(specifierContext, (StringView?)value);
 			return UhtSpecifierDispatchResults.Known;
 		}
 	}
@@ -287,39 +287,39 @@ namespace EpicGames.UHT.Tables
 	/// <summary>
 	/// Specifier delegate with a string value
 	/// </summary>
-	/// <param name="SpecifierContext">Specifier context</param>
-	/// <param name="Value">Specifier value</param>
-	public delegate void UhtSpecifierSingleStringDelegate(UhtSpecifierContext SpecifierContext, StringView Value);
+	/// <param name="specifierContext">Specifier context</param>
+	/// <param name="value">Specifier value</param>
+	public delegate void UhtSpecifierSingleStringDelegate(UhtSpecifierContext specifierContext, StringView value);
 
 	/// <summary>
 	/// Specifier with a string value
 	/// </summary>
 	public class UhtSpecifierSingleString : UhtSpecifier
 	{
-		private readonly UhtSpecifierSingleStringDelegate Delegate;
+		private readonly UhtSpecifierSingleStringDelegate _delegate;
 
 		/// <summary>
 		/// Construct the specifier
 		/// </summary>
-		/// <param name="Name">Name of the specifier</param>
-		/// <param name="When">When the specifier is executed</param>
-		/// <param name="Delegate">Delegate to invoke</param>
-		public UhtSpecifierSingleString(string Name, UhtSpecifierWhen When, UhtSpecifierSingleStringDelegate Delegate)
+		/// <param name="name">Name of the specifier</param>
+		/// <param name="when">When the specifier is executed</param>
+		/// <param name="specifierDelegate">Delegate to invoke</param>
+		public UhtSpecifierSingleString(string name, UhtSpecifierWhen when, UhtSpecifierSingleStringDelegate specifierDelegate)
 		{
-			this.Name = Name;
+			this.Name = name;
 			this.ValueType = UhtSpecifierValueType.SingleString;
-			this.When = When;
-			this.Delegate = Delegate;
+			this.When = when;
+			this._delegate = specifierDelegate;
 		}
 
 		/// <inheritdoc/>
-		public override UhtSpecifierDispatchResults Dispatch(UhtSpecifierContext SpecifierContext, object? Value)
+		public override UhtSpecifierDispatchResults Dispatch(UhtSpecifierContext specifierContext, object? value)
 		{
-			if (Value == null)
+			if (value == null)
 			{
 				throw new UhtIceException("Required value is null");
 			}
-			this.Delegate(SpecifierContext, (StringView)Value);
+			this._delegate(specifierContext, (StringView)value);
 			return UhtSpecifierDispatchResults.Known;
 		}
 	}
@@ -327,40 +327,40 @@ namespace EpicGames.UHT.Tables
 	/// <summary>
 	/// Specifier delegate with list of string keys and values 
 	/// </summary>
-	/// <param name="SpecifierContext">Specifier context</param>
-	/// <param name="Value">Specifier value</param>
-	public delegate void UhtSpecifierKeyValuePairListDelegate(UhtSpecifierContext SpecifierContext, List<KeyValuePair<StringView, StringView>> Value);
+	/// <param name="specifierContext">Specifier context</param>
+	/// <param name="value">Specifier value</param>
+	public delegate void UhtSpecifierKeyValuePairListDelegate(UhtSpecifierContext specifierContext, List<KeyValuePair<StringView, StringView>> value);
 
 	/// <summary>
 	/// Specifier with list of string keys and values 
 	/// </summary>
 	public class UhtSpecifierKeyValuePairList : UhtSpecifier
 	{
-		private readonly UhtSpecifierKeyValuePairListDelegate Delegate;
+		private readonly UhtSpecifierKeyValuePairListDelegate _delegate;
 
 		/// <summary>
 		/// Construct the specifier
 		/// </summary>
-		/// <param name="Name">Name of the specifier</param>
-		/// <param name="When">When the specifier is executed</param>
-		/// <param name="bEqualsOptional">If true this has an optional KVP list</param>
-		/// <param name="Delegate">Delegate to invoke</param>
-		public UhtSpecifierKeyValuePairList(string Name, UhtSpecifierWhen When, bool bEqualsOptional, UhtSpecifierKeyValuePairListDelegate Delegate)
+		/// <param name="name">Name of the specifier</param>
+		/// <param name="when">When the specifier is executed</param>
+		/// <param name="equalsOptional">If true this has an optional KVP list</param>
+		/// <param name="specifierDelegate">Delegate to invoke</param>
+		public UhtSpecifierKeyValuePairList(string name, UhtSpecifierWhen when, bool equalsOptional, UhtSpecifierKeyValuePairListDelegate specifierDelegate)
 		{
-			this.Name = Name;
-			this.ValueType = bEqualsOptional ? UhtSpecifierValueType.OptionalEqualsKeyValuePairList : UhtSpecifierValueType.KeyValuePairList;
-			this.When = When;
-			this.Delegate = Delegate;
+			this.Name = name;
+			this.ValueType = equalsOptional ? UhtSpecifierValueType.OptionalEqualsKeyValuePairList : UhtSpecifierValueType.KeyValuePairList;
+			this.When = when;
+			this._delegate = specifierDelegate;
 		}
 
 		/// <inheritdoc/>
-		public override UhtSpecifierDispatchResults Dispatch(UhtSpecifierContext SpecifierContext, object? Value)
+		public override UhtSpecifierDispatchResults Dispatch(UhtSpecifierContext specifierContext, object? value)
 		{
-			if (Value == null)
+			if (value == null)
 			{
 				throw new UhtIceException("Required value is null");
 			}
-			this.Delegate(SpecifierContext, (List<KeyValuePair<StringView, StringView>>)Value);
+			this._delegate(specifierContext, (List<KeyValuePair<StringView, StringView>>)value);
 			return UhtSpecifierDispatchResults.Known;
 		}
 	}
@@ -368,8 +368,8 @@ namespace EpicGames.UHT.Tables
 	/// <summary>
 	/// Specifier delegate with no value
 	/// </summary>
-	/// <param name="SpecifierContext">Specifier context</param>
-	public delegate void UhtSpecifierLegacyDelegate(UhtSpecifierContext SpecifierContext);
+	/// <param name="specifierContext">Specifier context</param>
+	public delegate void UhtSpecifierLegacyDelegate(UhtSpecifierContext specifierContext);
 
 	/// <summary>
 	/// Specifier delegate for legacy UHT specifiers with no value.  Will generate a information/deprecation message
@@ -377,29 +377,29 @@ namespace EpicGames.UHT.Tables
 	/// </summary>
 	public class UhtSpecifierLegacy : UhtSpecifier
 	{
-		private readonly UhtSpecifierLegacyDelegate Delegate;
+		private readonly UhtSpecifierLegacyDelegate _delegate;
 
 		/// <summary>
 		/// Construct the specifier
 		/// </summary>
-		/// <param name="Name">Name of the specifier</param>
-		/// <param name="Delegate">Delegate to invoke</param>
-		public UhtSpecifierLegacy(string Name, UhtSpecifierLegacyDelegate Delegate)
+		/// <param name="name">Name of the specifier</param>
+		/// <param name="specifierDelegate">Delegate to invoke</param>
+		public UhtSpecifierLegacy(string name, UhtSpecifierLegacyDelegate specifierDelegate)
 		{
-			this.Name = Name;
+			this.Name = name;
 			this.ValueType = UhtSpecifierValueType.StringList;
 			this.When = UhtSpecifierWhen.Deferred;
-			this.Delegate = Delegate;
+			this._delegate = specifierDelegate;
 		}
 
 		/// <inheritdoc/>
-		public override UhtSpecifierDispatchResults Dispatch(UhtSpecifierContext SpecifierContext, object? Value)
+		public override UhtSpecifierDispatchResults Dispatch(UhtSpecifierContext specifierContext, object? value)
 		{
-			if (Value != null)
+			if (value != null)
 			{
-				SpecifierContext.Scope.TokenReader.LogInfo($"Specifier '{this.Name}' has a value which is unused, future versions of UnrealHeaderTool will flag this as an error.");
+				specifierContext.Scope.TokenReader.LogInfo($"Specifier '{this.Name}' has a value which is unused, future versions of UnrealHeaderTool will flag this as an error.");
 			}
-			this.Delegate(SpecifierContext);
+			this._delegate(specifierContext);
 			return UhtSpecifierDispatchResults.Known;
 		}
 	}
@@ -407,35 +407,35 @@ namespace EpicGames.UHT.Tables
 	/// <summary>
 	/// Specifier delegate with an optional string list
 	/// </summary>
-	/// <param name="SpecifierContext">Specifier context</param>
-	/// <param name="Value">Specifier value</param>
-	public delegate void UhtSpecifierStringListDelegate(UhtSpecifierContext SpecifierContext, List<StringView>? Value);
+	/// <param name="specifierContext">Specifier context</param>
+	/// <param name="value">Specifier value</param>
+	public delegate void UhtSpecifierStringListDelegate(UhtSpecifierContext specifierContext, List<StringView>? value);
 
 	/// <summary>
 	/// Specifier with an optional string list
 	/// </summary>
 	public class UhtSpecifierStringList : UhtSpecifier
 	{
-		private readonly UhtSpecifierStringListDelegate Delegate;
+		private readonly UhtSpecifierStringListDelegate _delegate;
 
 		/// <summary>
 		/// Construct the specifier
 		/// </summary>
-		/// <param name="Name">Name of the specifier</param>
-		/// <param name="When">When the specifier is executed</param>
-		/// <param name="Delegate">Delegate to invoke</param>
-		public UhtSpecifierStringList(string Name, UhtSpecifierWhen When, UhtSpecifierStringListDelegate Delegate)
+		/// <param name="name">Name of the specifier</param>
+		/// <param name="when">When the specifier is executed</param>
+		/// <param name="specifierDelegate">Delegate to invoke</param>
+		public UhtSpecifierStringList(string name, UhtSpecifierWhen when, UhtSpecifierStringListDelegate specifierDelegate)
 		{
-			this.Name = Name;
+			this.Name = name;
 			this.ValueType = UhtSpecifierValueType.StringList;
-			this.When = When;
-			this.Delegate = Delegate;
+			this.When = when;
+			this._delegate = specifierDelegate;
 		}
 
 		/// <inheritdoc/>
-		public override UhtSpecifierDispatchResults Dispatch(UhtSpecifierContext SpecifierContext, object? Value)
+		public override UhtSpecifierDispatchResults Dispatch(UhtSpecifierContext specifierContext, object? value)
 		{
-			this.Delegate(SpecifierContext, (List<StringView>?)Value);
+			this._delegate(specifierContext, (List<StringView>?)value);
 			return UhtSpecifierDispatchResults.Known;
 		}
 	}
@@ -443,39 +443,39 @@ namespace EpicGames.UHT.Tables
 	/// <summary>
 	/// Specifier delegate with a list of string views
 	/// </summary>
-	/// <param name="SpecifierContext">Specifier context</param>
-	/// <param name="Value">Specifier value</param>
-	public delegate void UhtSpecifierNonEmptyStringListDelegate(UhtSpecifierContext SpecifierContext, List<StringView> Value);
+	/// <param name="specifierContext">Specifier context</param>
+	/// <param name="value">Specifier value</param>
+	public delegate void UhtSpecifierNonEmptyStringListDelegate(UhtSpecifierContext specifierContext, List<StringView> value);
 
 	/// <summary>
 	/// Specifier with a list of string views
 	/// </summary>
 	public class UhtSpecifierNonEmptyStringList : UhtSpecifier
 	{
-		private readonly UhtSpecifierNonEmptyStringListDelegate Delegate;
+		private readonly UhtSpecifierNonEmptyStringListDelegate _delegate;
 
 		/// <summary>
 		/// Construct the specifier
 		/// </summary>
-		/// <param name="Name">Name of the specifier</param>
-		/// <param name="When">When the specifier is executed</param>
-		/// <param name="Delegate">Delegate to invoke</param>
-		public UhtSpecifierNonEmptyStringList(string Name, UhtSpecifierWhen When, UhtSpecifierNonEmptyStringListDelegate Delegate)
+		/// <param name="name">Name of the specifier</param>
+		/// <param name="when">When the specifier is executed</param>
+		/// <param name="specifierDelegate">Delegate to invoke</param>
+		public UhtSpecifierNonEmptyStringList(string name, UhtSpecifierWhen when, UhtSpecifierNonEmptyStringListDelegate specifierDelegate)
 		{
-			this.Name = Name;
+			this.Name = name;
 			this.ValueType = UhtSpecifierValueType.NonEmptyStringList;
-			this.When = When;
-			this.Delegate = Delegate;
+			this.When = when;
+			this._delegate = specifierDelegate;
 		}
 
 		/// <inheritdoc/>
-		public override UhtSpecifierDispatchResults Dispatch(UhtSpecifierContext SpecifierContext, object? Value)
+		public override UhtSpecifierDispatchResults Dispatch(UhtSpecifierContext specifierContext, object? value)
 		{
-			if (Value == null)
+			if (value == null)
 			{
 				throw new UhtIceException("Required value is null");
 			}
-			this.Delegate(SpecifierContext, (List<StringView>)Value);
+			this._delegate(specifierContext, (List<StringView>)value);
 			return UhtSpecifierDispatchResults.Known;
 		}
 	}
@@ -523,10 +523,10 @@ namespace EpicGames.UHT.Tables
 		/// <summary>
 		/// Add the given value to the lookup table.  It will throw an exception if it is a duplicate.
 		/// </summary>
-		/// <param name="Specifier">The specifier to add</param>
-		public UhtSpecifierTable Add(UhtSpecifier Specifier)
+		/// <param name="specifier">The specifier to add</param>
+		public UhtSpecifierTable Add(UhtSpecifier specifier)
 		{
-			base.Add(Specifier.Name, Specifier);
+			base.Add(specifier.Name, specifier);
 			return this;
 		}
 	}
@@ -547,56 +547,56 @@ namespace EpicGames.UHT.Tables
 		/// <summary>
 		/// Invoke for a method that has the specifier attribute
 		/// </summary>
-		/// <param name="Type">Type containing the method</param>
-		/// <param name="MethodInfo">Method info</param>
-		/// <param name="SpecifierAttribute">Specified attributes</param>
+		/// <param name="type">Type containing the method</param>
+		/// <param name="methodInfo">Method info</param>
+		/// <param name="specifierAttribute">Specified attributes</param>
 		/// <exception cref="UhtIceException">Throw if the attribute isn't properly defined.</exception>
-		public void OnSpecifierAttribute(Type Type, MethodInfo MethodInfo, UhtSpecifierAttribute SpecifierAttribute)
+		public void OnSpecifierAttribute(Type type, MethodInfo methodInfo, UhtSpecifierAttribute specifierAttribute)
 		{
-			string Name = UhtLookupTableBase.GetSuffixedName(Type, MethodInfo, SpecifierAttribute.Name, "Specifier");
+			string name = UhtLookupTableBase.GetSuffixedName(type, methodInfo, specifierAttribute.Name, "Specifier");
 
-			if (string.IsNullOrEmpty(SpecifierAttribute.Extends))
+			if (String.IsNullOrEmpty(specifierAttribute.Extends))
 			{
-				throw new UhtIceException($"The 'Specifier' attribute on the {Type.Name}.{MethodInfo.Name} method doesn't have a table specified.");
+				throw new UhtIceException($"The 'Specifier' attribute on the {type.Name}.{methodInfo.Name} method doesn't have a table specified.");
 			}
 			else
 			{
 			}
 
-			if (SpecifierAttribute.ValueType == UhtSpecifierValueType.NotSet)
+			if (specifierAttribute.ValueType == UhtSpecifierValueType.NotSet)
 			{
-				throw new UhtIceException($"The 'Specifier' attribute on the {Type.Name}.{MethodInfo.Name} method doesn't have a value type specified.");
+				throw new UhtIceException($"The 'Specifier' attribute on the {type.Name}.{methodInfo.Name} method doesn't have a value type specified.");
 			}
 
-			UhtSpecifierTable Table = Get(SpecifierAttribute.Extends);
-			switch (SpecifierAttribute.ValueType)
+			UhtSpecifierTable table = Get(specifierAttribute.Extends);
+			switch (specifierAttribute.ValueType)
 			{
 				case UhtSpecifierValueType.None:
-					Table.Add(new UhtSpecifierNone(Name, SpecifierAttribute.When, (UhtSpecifierNoneDelegate)Delegate.CreateDelegate(typeof(UhtSpecifierNoneDelegate), MethodInfo)));
+					table.Add(new UhtSpecifierNone(name, specifierAttribute.When, (UhtSpecifierNoneDelegate)Delegate.CreateDelegate(typeof(UhtSpecifierNoneDelegate), methodInfo)));
 					break;
 				case UhtSpecifierValueType.String:
-					Table.Add(new UhtSpecifierString(Name, SpecifierAttribute.When, (UhtSpecifierStringDelegate)Delegate.CreateDelegate(typeof(UhtSpecifierStringDelegate), MethodInfo)));
+					table.Add(new UhtSpecifierString(name, specifierAttribute.When, (UhtSpecifierStringDelegate)Delegate.CreateDelegate(typeof(UhtSpecifierStringDelegate), methodInfo)));
 					break;
 				case UhtSpecifierValueType.OptionalString:
-					Table.Add(new UhtSpecifierOptionalString(Name, SpecifierAttribute.When, (UhtSpecifierOptionalStringDelegate)Delegate.CreateDelegate(typeof(UhtSpecifierOptionalStringDelegate), MethodInfo)));
+					table.Add(new UhtSpecifierOptionalString(name, specifierAttribute.When, (UhtSpecifierOptionalStringDelegate)Delegate.CreateDelegate(typeof(UhtSpecifierOptionalStringDelegate), methodInfo)));
 					break;
 				case UhtSpecifierValueType.SingleString:
-					Table.Add(new UhtSpecifierSingleString(Name, SpecifierAttribute.When, (UhtSpecifierSingleStringDelegate)Delegate.CreateDelegate(typeof(UhtSpecifierSingleStringDelegate), MethodInfo)));
+					table.Add(new UhtSpecifierSingleString(name, specifierAttribute.When, (UhtSpecifierSingleStringDelegate)Delegate.CreateDelegate(typeof(UhtSpecifierSingleStringDelegate), methodInfo)));
 					break;
 				case UhtSpecifierValueType.KeyValuePairList:
-					Table.Add(new UhtSpecifierKeyValuePairList(Name, SpecifierAttribute.When, false, (UhtSpecifierKeyValuePairListDelegate)Delegate.CreateDelegate(typeof(UhtSpecifierKeyValuePairListDelegate), MethodInfo)));
+					table.Add(new UhtSpecifierKeyValuePairList(name, specifierAttribute.When, false, (UhtSpecifierKeyValuePairListDelegate)Delegate.CreateDelegate(typeof(UhtSpecifierKeyValuePairListDelegate), methodInfo)));
 					break;
 				case UhtSpecifierValueType.OptionalEqualsKeyValuePairList:
-					Table.Add(new UhtSpecifierKeyValuePairList(Name, SpecifierAttribute.When, true, (UhtSpecifierKeyValuePairListDelegate)Delegate.CreateDelegate(typeof(UhtSpecifierKeyValuePairListDelegate), MethodInfo)));
+					table.Add(new UhtSpecifierKeyValuePairList(name, specifierAttribute.When, true, (UhtSpecifierKeyValuePairListDelegate)Delegate.CreateDelegate(typeof(UhtSpecifierKeyValuePairListDelegate), methodInfo)));
 					break;
 				case UhtSpecifierValueType.StringList:
-					Table.Add(new UhtSpecifierStringList(Name, SpecifierAttribute.When, (UhtSpecifierStringListDelegate)Delegate.CreateDelegate(typeof(UhtSpecifierStringListDelegate), MethodInfo)));
+					table.Add(new UhtSpecifierStringList(name, specifierAttribute.When, (UhtSpecifierStringListDelegate)Delegate.CreateDelegate(typeof(UhtSpecifierStringListDelegate), methodInfo)));
 					break;
 				case UhtSpecifierValueType.NonEmptyStringList:
-					Table.Add(new UhtSpecifierNonEmptyStringList(Name, SpecifierAttribute.When, (UhtSpecifierNonEmptyStringListDelegate)Delegate.CreateDelegate(typeof(UhtSpecifierNonEmptyStringListDelegate), MethodInfo)));
+					table.Add(new UhtSpecifierNonEmptyStringList(name, specifierAttribute.When, (UhtSpecifierNonEmptyStringListDelegate)Delegate.CreateDelegate(typeof(UhtSpecifierNonEmptyStringListDelegate), methodInfo)));
 					break;
 				case UhtSpecifierValueType.Legacy:
-					Table.Add(new UhtSpecifierLegacy(Name, (UhtSpecifierLegacyDelegate)Delegate.CreateDelegate(typeof(UhtSpecifierLegacyDelegate), MethodInfo)));
+					table.Add(new UhtSpecifierLegacy(name, (UhtSpecifierLegacyDelegate)Delegate.CreateDelegate(typeof(UhtSpecifierLegacyDelegate), methodInfo)));
 					break;
 			}
 		}
diff --git a/Engine/Source/Programs/Shared/EpicGames.UHT/Tables/UhtSpecifierValidatorTable.cs b/Engine/Source/Programs/Shared/EpicGames.UHT/Tables/UhtSpecifierValidatorTable.cs
--- a/Engine/Source/Programs/Shared/EpicGames.UHT/Tables/UhtSpecifierValidatorTable.cs
+++ b/Engine/Source/Programs/Shared/EpicGames.UHT/Tables/UhtSpecifierValidatorTable.cs
@@ -1,10 +1,10 @@
 // Copyright Epic Games, Inc. All Rights Reserved.
 
+using System;
+using System.Reflection;
 using EpicGames.Core;
 using EpicGames.UHT.Types;
 using EpicGames.UHT.Utils;
-using System;
-using System.Reflection;
 
 namespace EpicGames.UHT.Tables
 {
@@ -12,11 +12,11 @@ namespace EpicGames.UHT.Tables
 	/// <summary>
 	/// Delegate used to validate a specifier
 	/// </summary>
-	/// <param name="Type">Containing type</param>
-	/// <param name="MetaData">Containing meta data</param>
-	/// <param name="Key">Key of the meta data entry</param>
-	/// <param name="Value">Value of the meta data entry</param>
-	public delegate void UhtSpecifierValidatorDelegate(UhtType Type, UhtMetaData MetaData, UhtMetaDataKey Key, StringView Value);
+	/// <param name="type">Containing type</param>
+	/// <param name="metaData">Containing meta data</param>
+	/// <param name="key">Key of the meta data entry</param>
+	/// <param name="value">Value of the meta data entry</param>
+	public delegate void UhtSpecifierValidatorDelegate(UhtType type, UhtMetaData metaData, UhtMetaDataKey key, StringView value);
 
 	/// <summary>
 	/// Defines a specifier validated created from the attribute
@@ -37,12 +37,12 @@ namespace EpicGames.UHT.Tables
 		/// <summary>
 		/// Construct a new instance
 		/// </summary>
-		/// <param name="Name">Name of the validator</param>
-		/// <param name="Delegate">Delegate of the validator</param>
-		public UhtSpecifierValidator(string Name, UhtSpecifierValidatorDelegate Delegate)
+		/// <param name="name">Name of the validator</param>
+		/// <param name="specifierValidatorDelegate">Delegate of the validator</param>
+		public UhtSpecifierValidator(string name, UhtSpecifierValidatorDelegate specifierValidatorDelegate)
 		{
-			this.Name = Name;
-			this.Delegate = Delegate;
+			this.Name = name;
+			this.Delegate = specifierValidatorDelegate;
 		}
 	}
 
@@ -80,10 +80,10 @@ namespace EpicGames.UHT.Tables
 		/// <summary>
 		/// Add the given value to the lookup table.  It will throw an exception if it is a duplicate.
 		/// </summary>
-		/// <param name="Specifier">Validator to add</param>
-		public UhtSpecifierValidatorTable Add(UhtSpecifierValidator Specifier)
+		/// <param name="specifier">Validator to add</param>
+		public UhtSpecifierValidatorTable Add(UhtSpecifierValidator specifier)
 		{
-			base.Add(Specifier.Name, Specifier);
+			base.Add(specifier.Name, specifier);
 			return this;
 		}
 	}
@@ -104,20 +104,20 @@ namespace EpicGames.UHT.Tables
 		/// <summary>
 		/// Handle the attribute appearing on a method
 		/// </summary>
-		/// <param name="Type">Type containing the method</param>
-		/// <param name="MethodInfo">The method</param>
-		/// <param name="SpecifierValidatorAttribute">Attribute</param>
+		/// <param name="type">Type containing the method</param>
+		/// <param name="methodInfo">The method</param>
+		/// <param name="specifierValidatorAttribute">Attribute</param>
 		/// <exception cref="UhtIceException">Thrown if the validator isn't properly defined</exception>
-		public void OnSpecifierValidatorAttribute(Type Type, MethodInfo MethodInfo, UhtSpecifierValidatorAttribute SpecifierValidatorAttribute)
+		public void OnSpecifierValidatorAttribute(Type type, MethodInfo methodInfo, UhtSpecifierValidatorAttribute specifierValidatorAttribute)
 		{
-			string Name = UhtLookupTableBase.GetSuffixedName(Type, MethodInfo, SpecifierValidatorAttribute.Name, "SpecifierValidator");
+			string name = UhtLookupTableBase.GetSuffixedName(type, methodInfo, specifierValidatorAttribute.Name, "SpecifierValidator");
 
-			if (string.IsNullOrEmpty(SpecifierValidatorAttribute.Extends))
+			if (String.IsNullOrEmpty(specifierValidatorAttribute.Extends))
 			{
-				throw new UhtIceException($"The 'SpecifierValidator' attribute on the {Type.Name}.{MethodInfo.Name} method doesn't have a table specified.");
+				throw new UhtIceException($"The 'SpecifierValidator' attribute on the {type.Name}.{methodInfo.Name} method doesn't have a table specified.");
 			}
-			UhtSpecifierValidatorTable Table = Get(SpecifierValidatorAttribute.Extends);
-			Table.Add(new UhtSpecifierValidator(Name, (UhtSpecifierValidatorDelegate)Delegate.CreateDelegate(typeof(UhtSpecifierValidatorDelegate), MethodInfo)));
+			UhtSpecifierValidatorTable table = Get(specifierValidatorAttribute.Extends);
+			table.Add(new UhtSpecifierValidator(name, (UhtSpecifierValidatorDelegate)Delegate.CreateDelegate(typeof(UhtSpecifierValidatorDelegate), methodInfo)));
 		}
 	}
 }
diff --git a/Engine/Source/Programs/Shared/EpicGames.UHT/Tables/UhtStructDefaultValueTable.cs b/Engine/Source/Programs/Shared/EpicGames.UHT/Tables/UhtStructDefaultValueTable.cs
--- a/Engine/Source/Programs/Shared/EpicGames.UHT/Tables/UhtStructDefaultValueTable.cs
+++ b/Engine/Source/Programs/Shared/EpicGames.UHT/Tables/UhtStructDefaultValueTable.cs
@@ -1,13 +1,13 @@
 // Copyright Epic Games, Inc. All Rights Reserved.
 
-using EpicGames.Core;
-using EpicGames.UHT.Tokenizer;
 using System;
 using System.Collections.Generic;
 using System.Reflection;
+using System.Text;
+using EpicGames.Core;
+using EpicGames.UHT.Tokenizer;
 using EpicGames.UHT.Types;
 using EpicGames.UHT.Utils;
-using System.Text;
 
 namespace EpicGames.UHT.Tables
 {
@@ -15,11 +15,11 @@ namespace EpicGames.UHT.Tables
 	/// <summary>
 	/// Delegate for invoking structure default value sanitizer
 	/// </summary>
-	/// <param name="Property"></param>
-	/// <param name="DefaultValueReader"></param>
-	/// <param name="InnerDefaultValue"></param>
+	/// <param name="property"></param>
+	/// <param name="defaultValueReader"></param>
+	/// <param name="innerDefaultValue"></param>
 	/// <returns></returns>
-	public delegate bool UhtStructDefaultValueDelegate(UhtStructProperty Property, IUhtTokenReader DefaultValueReader, StringBuilder InnerDefaultValue);
+	public delegate bool UhtStructDefaultValueDelegate(UhtStructProperty property, IUhtTokenReader defaultValueReader, StringBuilder innerDefaultValue);
 
 	/// <summary>
 	/// Options for structure default value sanitizer
@@ -49,35 +49,35 @@ namespace EpicGames.UHT.Tables
 		/// <summary>
 		/// Test to see if any of the specified flags are set
 		/// </summary>
-		/// <param name="InFlags">Current flags</param>
-		/// <param name="TestFlags">Flags to test for</param>
+		/// <param name="inFlags">Current flags</param>
+		/// <param name="testFlags">Flags to test for</param>
 		/// <returns>True if any of the flags are set</returns>
-		public static bool HasAnyFlags(this UhtStructDefaultValueOptions InFlags, UhtStructDefaultValueOptions TestFlags)
+		public static bool HasAnyFlags(this UhtStructDefaultValueOptions inFlags, UhtStructDefaultValueOptions testFlags)
 		{
-			return (InFlags & TestFlags) != 0;
+			return (inFlags & testFlags) != 0;
 		}
 
 		/// <summary>
 		/// Test to see if all of the specified flags are set
 		/// </summary>
-		/// <param name="InFlags">Current flags</param>
-		/// <param name="TestFlags">Flags to test for</param>
+		/// <param name="inFlags">Current flags</param>
+		/// <param name="testFlags">Flags to test for</param>
 		/// <returns>True if all the flags are set</returns>
-		public static bool HasAllFlags(this UhtStructDefaultValueOptions InFlags, UhtStructDefaultValueOptions TestFlags)
+		public static bool HasAllFlags(this UhtStructDefaultValueOptions inFlags, UhtStructDefaultValueOptions testFlags)
 		{
-			return (InFlags & TestFlags) == TestFlags;
+			return (inFlags & testFlags) == testFlags;
 		}
 
 		/// <summary>
 		/// Test to see if a specific set of flags have a specific value.
 		/// </summary>
-		/// <param name="InFlags">Current flags</param>
-		/// <param name="TestFlags">Flags to test for</param>
-		/// <param name="MatchFlags">Expected value of the tested flags</param>
+		/// <param name="inFlags">Current flags</param>
+		/// <param name="testFlags">Flags to test for</param>
+		/// <param name="matchFlags">Expected value of the tested flags</param>
 		/// <returns>True if the given flags have a specific value.</returns>
-		public static bool HasExactFlags(this UhtStructDefaultValueOptions InFlags, UhtStructDefaultValueOptions TestFlags, UhtStructDefaultValueOptions MatchFlags)
+		public static bool HasExactFlags(this UhtStructDefaultValueOptions inFlags, UhtStructDefaultValueOptions testFlags, UhtStructDefaultValueOptions matchFlags)
 		{
-			return (InFlags & TestFlags) == MatchFlags;
+			return (inFlags & testFlags) == matchFlags;
 		}
 	}
 
@@ -116,8 +116,8 @@ namespace EpicGames.UHT.Tables
 	/// </summary>
 	public class UhtStructDefaultValueTable
 	{
-		private readonly Dictionary<StringView, UhtStructDefaultValue> StructDefaultValues = new Dictionary<StringView, UhtStructDefaultValue>();
-		private UhtStructDefaultValue? DefaultInternal = null;
+		private readonly Dictionary<StringView, UhtStructDefaultValue> _structDefaultValues = new Dictionary<StringView, UhtStructDefaultValue>();
+		private UhtStructDefaultValue? _defaultInternal = null;
 
 		/// <summary>
 		/// Fetch the default sanitizer
@@ -126,54 +126,54 @@ namespace EpicGames.UHT.Tables
 		{
 			get
 			{
-				if (this.DefaultInternal == null)
+				if (this._defaultInternal == null)
 				{
 					throw new UhtIceException("No struct default value has been marked as default");
 				}
-				return (UhtStructDefaultValue)this.DefaultInternal;
+				return (UhtStructDefaultValue)this._defaultInternal;
 			}
 		}
 
 		/// <summary>
 		/// Return the structure default value associated with the given name
 		/// </summary>
-		/// <param name="Name"></param>
-		/// <param name="StructDefaultValue">Structure default value handler</param>
+		/// <param name="name"></param>
+		/// <param name="structDefaultValue">Structure default value handler</param>
 		/// <returns></returns>
-		public bool TryGet(StringView Name, out UhtStructDefaultValue StructDefaultValue)
+		public bool TryGet(StringView name, out UhtStructDefaultValue structDefaultValue)
 		{
-			return this.StructDefaultValues.TryGetValue(Name, out StructDefaultValue);
+			return this._structDefaultValues.TryGetValue(name, out structDefaultValue);
 		}
 
 		/// <summary>
 		/// Handle a structure default value sanitizer attribute
 		/// </summary>
-		/// <param name="MethodInfo">Method information</param>
-		/// <param name="StructDefaultValueAttribute">Found attribute</param>
+		/// <param name="methodInfo">Method information</param>
+		/// <param name="structDefaultValueAttribute">Found attribute</param>
 		/// <exception cref="UhtIceException">Thrown if the attribute isn't property defined</exception>
-		public void OnStructDefaultValueAttribute(MethodInfo MethodInfo, UhtStructDefaultValueAttribute StructDefaultValueAttribute)
+		public void OnStructDefaultValueAttribute(MethodInfo methodInfo, UhtStructDefaultValueAttribute structDefaultValueAttribute)
 		{
-			if (string.IsNullOrEmpty(StructDefaultValueAttribute.Name) && !StructDefaultValueAttribute.Options.HasAnyFlags(UhtStructDefaultValueOptions.Default))
+			if (String.IsNullOrEmpty(structDefaultValueAttribute.Name) && !structDefaultValueAttribute.Options.HasAnyFlags(UhtStructDefaultValueOptions.Default))
 			{
 				throw new UhtIceException("A struct default value attribute must have a name or be marked as default");
 			}
 
-			UhtStructDefaultValue StructDefaultValue = new UhtStructDefaultValue
+			UhtStructDefaultValue structDefaultValue = new UhtStructDefaultValue
 			{
-				Delegate = (UhtStructDefaultValueDelegate)Delegate.CreateDelegate(typeof(UhtStructDefaultValueDelegate), MethodInfo)
+				Delegate = (UhtStructDefaultValueDelegate)Delegate.CreateDelegate(typeof(UhtStructDefaultValueDelegate), methodInfo)
 			};
 
-			if (StructDefaultValueAttribute.Options.HasAnyFlags(UhtStructDefaultValueOptions.Default))
+			if (structDefaultValueAttribute.Options.HasAnyFlags(UhtStructDefaultValueOptions.Default))
 			{
-				if (this.DefaultInternal != null)
+				if (this._defaultInternal != null)
 				{
 					throw new UhtIceException("Only one struct default value attribute dispatcher can be marked as default");
 				}
-				this.DefaultInternal = StructDefaultValue;
+				this._defaultInternal = structDefaultValue;
 			}
-			else if (!string.IsNullOrEmpty(StructDefaultValueAttribute.Name))
+			else if (!String.IsNullOrEmpty(structDefaultValueAttribute.Name))
 			{
-				this.StructDefaultValues.Add(new StringView(StructDefaultValueAttribute.Name), StructDefaultValue);
+				this._structDefaultValues.Add(new StringView(structDefaultValueAttribute.Name), structDefaultValue);
 			}
 		}
 	}
diff --git a/Engine/Source/Programs/Shared/EpicGames.UHT/Tokenizer/UhtToken.cs b/Engine/Source/Programs/Shared/EpicGames.UHT/Tokenizer/UhtToken.cs
--- a/Engine/Source/Programs/Shared/EpicGames.UHT/Tokenizer/UhtToken.cs
+++ b/Engine/Source/Programs/Shared/EpicGames.UHT/Tokenizer/UhtToken.cs
@@ -1,11 +1,11 @@
 // Copyright Epic Games, Inc. All Rights Reserved.
 
-using EpicGames.Core;
-using EpicGames.UHT.Utils;
 using System;
 using System.Collections.Generic;
 using System.Globalization;
 using System.Text;
+using EpicGames.Core;
+using EpicGames.UHT.Utils;
 
 namespace EpicGames.UHT.Tokenizer
 {
@@ -86,11 +86,11 @@ namespace EpicGames.UHT.Tokenizer
 		/// <summary>
 		/// Return true if the token type is an end type
 		/// </summary>
-		/// <param name="TokenType">Token type in question</param>
+		/// <param name="tokenType">Token type in question</param>
 		/// <returns>True if the token type is an end type</returns>
-		public static bool IsEndType(this UhtTokenType TokenType)
+		public static bool IsEndType(this UhtTokenType tokenType)
 		{
-			switch (TokenType)
+			switch (tokenType)
 			{
 				case UhtTokenType.EndOfFile:
 				case UhtTokenType.EndOfDefault:
@@ -142,7 +142,7 @@ namespace EpicGames.UHT.Tokenizer
 		/// <summary>
 		/// End position of the token value
 		/// </summary>
-		public int InputEndPos { get => this.InputStartPos + this.Value.Span.Length; }
+		public int InputEndPos => this.InputStartPos + this.Value.Span.Length;
 
 		/// <summary>
 		/// Line containing the token
@@ -157,10 +157,10 @@ namespace EpicGames.UHT.Tokenizer
 		/// <summary>
 		/// Construct a new token
 		/// </summary>
-		/// <param name="TokenType">Type of the token</param>
-		public UhtToken(UhtTokenType TokenType)
+		/// <param name="tokenType">Type of the token</param>
+		public UhtToken(UhtTokenType tokenType)
 		{
-			this.TokenType = TokenType;
+			this.TokenType = tokenType;
 			this.UngetPos = 0;
 			this.UngetLine = 0;
 			this.InputStartPos = 0;
@@ -171,29 +171,29 @@ namespace EpicGames.UHT.Tokenizer
 		/// <summary>
 		/// Construct a new token
 		/// </summary>
-		/// <param name="TokenType">Type of token</param>
-		/// <param name="UngetPos">Unget position</param>
-		/// <param name="UngetLine">Unget line</param>
-		/// <param name="InputStartPos">Start position of value</param>
-		/// <param name="InputLine">Line of value</param>
-		/// <param name="Value">Token value</param>
-		public UhtToken(UhtTokenType TokenType, int UngetPos, int UngetLine, int InputStartPos, int InputLine, StringView Value)
+		/// <param name="tokenType">Type of token</param>
+		/// <param name="ungetPos">Unget position</param>
+		/// <param name="ungetLine">Unget line</param>
+		/// <param name="inputStartPos">Start position of value</param>
+		/// <param name="inputLine">Line of value</param>
+		/// <param name="value">Token value</param>
+		public UhtToken(UhtTokenType tokenType, int ungetPos, int ungetLine, int inputStartPos, int inputLine, StringView value)
 		{
-			this.TokenType = TokenType;
-			this.UngetPos = UngetPos;
-			this.UngetLine = UngetLine;
-			this.InputStartPos = InputStartPos;
-			this.InputLine = InputLine;
-			this.Value = Value;
+			this.TokenType = tokenType;
+			this.UngetPos = ungetPos;
+			this.UngetLine = ungetLine;
+			this.InputStartPos = inputStartPos;
+			this.InputLine = inputLine;
+			this.Value = value;
 		}
 
 		/// <summary>
 		/// True if the token isn't an end token
 		/// </summary>
-		/// <param name="Token">Token in question</param>
-		public static implicit operator bool(UhtToken Token)
+		/// <param name="token">Token in question</param>
+		public static implicit operator bool(UhtToken token)
 		{
-			return !Token.IsEndType();
+			return !token.IsEndType();
 		}
 
 		/// <summary>
@@ -208,44 +208,44 @@ namespace EpicGames.UHT.Tokenizer
 		/// <summary>
 		/// Test to see if the value matches the given character
 		/// </summary>
-		/// <param name="Value">Value to test</param>
+		/// <param name="value">Value to test</param>
 		/// <returns>True if the token value matches the given value</returns>
-		public bool IsValue(char Value)
+		public bool IsValue(char value)
 		{
-			return this.Value.Span.Length == 1 && this.Value.Span[0] == Value;
+			return this.Value.Span.Length == 1 && this.Value.Span[0] == value;
 		}
 
 		/// <summary>
 		/// Test to see if the value matches the given string
 		/// </summary>
-		/// <param name="Value">Value to test</param>
-		/// <param name="bIgnoreCase">If true, ignore case</param>
+		/// <param name="value">Value to test</param>
+		/// <param name="ignoreCase">If true, ignore case</param>
 		/// <returns>True if the value matches</returns>
-		public bool IsValue(string Value, bool bIgnoreCase = false)
+		public bool IsValue(string value, bool ignoreCase = false)
 		{
-			return this.Value.Span.Equals(Value, bIgnoreCase ? StringComparison.OrdinalIgnoreCase : StringComparison.Ordinal);
+			return this.Value.Span.Equals(value, ignoreCase ? StringComparison.OrdinalIgnoreCase : StringComparison.Ordinal);
 		}
 
 		/// <summary>
 		/// Test to see if the value matches the given string
 		/// </summary>
-		/// <param name="Value">Value to test</param>
-		/// <param name="bIgnoreCase">If true, ignore case</param>
+		/// <param name="value">Value to test</param>
+		/// <param name="ignoreCase">If true, ignore case</param>
 		/// <returns>True if the value matches</returns>
-		public bool IsValue(StringView Value, bool bIgnoreCase = false)
+		public bool IsValue(StringView value, bool ignoreCase = false)
 		{
-			return this.Value.Span.Equals(Value.Span, bIgnoreCase ? StringComparison.OrdinalIgnoreCase : StringComparison.Ordinal);
+			return this.Value.Span.Equals(value.Span, ignoreCase ? StringComparison.OrdinalIgnoreCase : StringComparison.Ordinal);
 		}
 
 		/// <summary>
 		/// Test to see if the value starts with the given string
 		/// </summary>
-		/// <param name="Value">Value to test</param>
-		/// <param name="bIgnoreCase">If true, ignore case</param>
+		/// <param name="value">Value to test</param>
+		/// <param name="ignoreCase">If true, ignore case</param>
 		/// <returns>True is the value starts with the given string</returns>
-		public bool ValueStartsWith(string Value, bool bIgnoreCase = false)
+		public bool ValueStartsWith(string value, bool ignoreCase = false)
 		{
-			return this.Value.Span.StartsWith(Value, bIgnoreCase ? StringComparison.OrdinalIgnoreCase : StringComparison.Ordinal);
+			return this.Value.Span.StartsWith(value, ignoreCase ? StringComparison.OrdinalIgnoreCase : StringComparison.Ordinal);
 		}
 
 		/// <summary>
@@ -260,23 +260,23 @@ namespace EpicGames.UHT.Tokenizer
 		/// <summary>
 		/// Return true if the identifier matches
 		/// </summary>
-		/// <param name="Identifier">Identifier to test</param>
-		/// <param name="bIgnoreCase">If true, ignore case</param>
+		/// <param name="identifier">Identifier to test</param>
+		/// <param name="ignoreCase">If true, ignore case</param>
 		/// <returns>True if the identifier matches</returns>
-		public bool IsIdentifier(string Identifier, bool bIgnoreCase = false)
+		public bool IsIdentifier(string identifier, bool ignoreCase = false)
 		{
-			return IsIdentifier() && IsValue(Identifier, bIgnoreCase);
+			return IsIdentifier() && IsValue(identifier, ignoreCase);
 		}
 
 		/// <summary>
 		/// Return true if the identifier matches
 		/// </summary>
-		/// <param name="Identifier">Identifier to test</param>
-		/// <param name="bIgnoreCase">If true, ignore case</param>
+		/// <param name="identifier">Identifier to test</param>
+		/// <param name="ignoreCase">If true, ignore case</param>
 		/// <returns>True if the identifier matches</returns>
-		public bool IsIdentifier(StringView Identifier, bool bIgnoreCase = false)
+		public bool IsIdentifier(StringView identifier, bool ignoreCase = false)
 		{
-			return IsIdentifier() && IsValue(Identifier, bIgnoreCase);
+			return IsIdentifier() && IsValue(identifier, ignoreCase);
 		}
 
 		/// <summary>
@@ -291,31 +291,31 @@ namespace EpicGames.UHT.Tokenizer
 		/// <summary>
 		/// Return true if the symbol matches
 		/// </summary>
-		/// <param name="Symbol">Symbol to test</param>
+		/// <param name="symbol">Symbol to test</param>
 		/// <returns>True if the symbol matches</returns>
-		public bool IsSymbol(char Symbol)
+		public bool IsSymbol(char symbol)
 		{
-			return IsSymbol() && IsValue(Symbol);
+			return IsSymbol() && IsValue(symbol);
 		}
 
 		/// <summary>
 		/// Return true if the symbol matches
 		/// </summary>
-		/// <param name="Symbol">Symbol to test</param>
+		/// <param name="symbol">Symbol to test</param>
 		/// <returns>True if the symbol matches</returns>
-		public bool IsSymbol(string Symbol)
+		public bool IsSymbol(string symbol)
 		{
-			return IsSymbol() && IsValue(Symbol);
+			return IsSymbol() && IsValue(symbol);
 		}
 
 		/// <summary>
 		/// Return true if the symbol matches
 		/// </summary>
-		/// <param name="Symbol">Symbol to test</param>
+		/// <param name="symbol">Symbol to test</param>
 		/// <returns>True if the symbol matches</returns>
-		public bool IsSymbol(StringView Symbol)
+		public bool IsSymbol(StringView symbol)
 		{
-			return IsSymbol() && IsValue(Symbol);
+			return IsSymbol() && IsValue(symbol);
 		}
 
 		/// <summary>
@@ -339,26 +339,26 @@ namespace EpicGames.UHT.Tokenizer
 		/// <summary>
 		/// Get the integer value of the token.  Only supported for decimal, hexadecimal, and floating point values
 		/// </summary>
-		/// <param name="Value">Resulting value</param>
+		/// <param name="value">Resulting value</param>
 		/// <returns>True if the value was set</returns>
-		public bool GetConstInt(out int Value)
+		public bool GetConstInt(out int value)
 		{
 			switch (TokenType)
 			{
 				case UhtTokenType.DecimalConst:
-					Value = (int)GetDecimalValue();
+					value = (int)GetDecimalValue();
 					return true;
 				case UhtTokenType.HexConst:
-					Value = (int)GetHexValue();
+					value = (int)GetHexValue();
 					return true;
 				case UhtTokenType.FloatConst:
 					{
-						float Float = GetFloatValue();
-						Value = (int)Float;
-						return Float == Value;
+						float floatValue = GetFloatValue();
+						value = (int)floatValue;
+						return floatValue == value;
 					}
 				default:
-					Value = 0;
+					value = 0;
 					return false;
 			}
 		}
@@ -366,26 +366,26 @@ namespace EpicGames.UHT.Tokenizer
 		/// <summary>
 		/// Get the integer value of the token.  Only supported for decimal, hexadecimal, and floating point values
 		/// </summary>
-		/// <param name="Value">Resulting value</param>
+		/// <param name="value">Resulting value</param>
 		/// <returns>True if the value was set</returns>
-		public bool GetConstLong(out long Value)
+		public bool GetConstLong(out long value)
 		{
 			switch (TokenType)
 			{
 				case UhtTokenType.DecimalConst:
-					Value = GetDecimalValue();
+					value = GetDecimalValue();
 					return true;
 				case UhtTokenType.HexConst:
-					Value = GetHexValue();
+					value = GetHexValue();
 					return true;
 				case UhtTokenType.FloatConst:
 					{
-						float Float = GetFloatValue();
-						Value = (long)Float;
-						return Float == Value;
+						float floatValue = GetFloatValue();
+						value = (long)floatValue;
+						return floatValue == value;
 					}
 				default:
-					Value = 0;
+					value = 0;
 					return false;
 			}
 		}
@@ -393,23 +393,23 @@ namespace EpicGames.UHT.Tokenizer
 		/// <summary>
 		/// Get the float value of the token.  Only supported for decimal, hexadecimal, and floating point values
 		/// </summary>
-		/// <param name="Value">Resulting value</param>
+		/// <param name="value">Resulting value</param>
 		/// <returns>True if the value was set</returns>
-		public bool GetConstFloat(out float Value)
+		public bool GetConstFloat(out float value)
 		{
 			switch (TokenType)
 			{
 				case UhtTokenType.DecimalConst:
-					Value = (float)GetDecimalValue();
+					value = (float)GetDecimalValue();
 					return true;
 				case UhtTokenType.HexConst:
-					Value = (float)GetHexValue();
+					value = (float)GetHexValue();
 					return true;
 				case UhtTokenType.FloatConst:
-					Value = GetFloatValue();
+					value = GetFloatValue();
 					return true;
 				default:
-					Value = 0;
+					value = 0;
 					return false;
 			}
 		}
@@ -417,23 +417,23 @@ namespace EpicGames.UHT.Tokenizer
 		/// <summary>
 		/// Get the double value of the token.  Only supported for decimal, hexadecimal, and floating point values
 		/// </summary>
-		/// <param name="Value">Resulting value</param>
+		/// <param name="value">Resulting value</param>
 		/// <returns>True if the value was set</returns>
-		public bool GetConstDouble(out double Value)
+		public bool GetConstDouble(out double value)
 		{
 			switch (TokenType)
 			{
 				case UhtTokenType.DecimalConst:
-					Value = (float)GetDecimalValue();
+					value = (double)GetDecimalValue();
 					return true;
 				case UhtTokenType.HexConst:
-					Value = (float)GetHexValue();
+					value = (double)GetHexValue();
 					return true;
 				case UhtTokenType.FloatConst:
-					Value = GetDoubleValue();
+					value = GetDoubleValue();
 					return true;
 				default:
-					Value = 0;
+					value = 0;
 					return false;
 			}
 		}
@@ -452,35 +452,35 @@ namespace EpicGames.UHT.Tokenizer
 		/// <summary>
 		/// Return an un-escaped string.  The surrounding quotes will be removed and escaped characters will be converted to the actual values.
 		/// </summary>
-		/// <param name="MessageSite"></param>
+		/// <param name="messageSite"></param>
 		/// <returns>Resulting string</returns>
 		/// <exception cref="UhtException">Thrown if the token type is not a string or character constant</exception>
-		public StringView GetUnescapedString(IUhtMessageSite MessageSite)
+		public StringView GetUnescapedString(IUhtMessageSite messageSite)
 		{
 			switch (this.TokenType)
 			{
 				case UhtTokenType.StringConst:
-					ReadOnlySpan<char> Span = this.Value.Span.Slice(1, this.Value.Span.Length - 2);
-					int Index = Span.IndexOf('\\');
-					if (Index == -1)
+					ReadOnlySpan<char> span = this.Value.Span.Slice(1, this.Value.Span.Length - 2);
+					int index = span.IndexOf('\\');
+					if (index == -1)
 					{
-						return new StringView(this.Value, 1, Span.Length);
+						return new StringView(this.Value, 1, span.Length);
 					}
 					else
 					{
-						StringBuilder Builder = new StringBuilder();
-						while (Index >= 0)
+						StringBuilder builder = new StringBuilder();
+						while (index >= 0)
 						{
-							Builder.Append(Span.Slice(0, Index));
-							if (Span[Index + 1] == 'n')
+							builder.Append(span.Slice(0, index));
+							if (span[index + 1] == 'n')
 							{
-								Builder.Append('\n');
+								builder.Append('\n');
 							}
-							Span = Span.Slice(Index + 1);
-							Index = Span.IndexOf('\\');
+							span = span.Slice(index + 1);
+							index = span.IndexOf('\\');
 						}
-						Builder.Append(Span);
-						return new StringView(Builder.ToString());
+						builder.Append(span);
+						return new StringView(builder.ToString());
 					}
 
 				case UhtTokenType.CharConst:
@@ -504,15 +504,15 @@ namespace EpicGames.UHT.Tokenizer
 					}
 			}
 
-			throw new UhtException(MessageSite, this.InputLine, "Call to GetUnescapedString on token that isn't a string or char constant");
+			throw new UhtException(messageSite, this.InputLine, "Call to GetUnescapedString on token that isn't a string or char constant");
 		}
 
 		/// <summary>
 		/// Return a string representation of the token value.  This will convert numeric values and format them.
 		/// </summary>
-		/// <param name="bRespectQuotes">If true, embedded quotes will be respected</param>
+		/// <param name="respectQuotes">If true, embedded quotes will be respected</param>
 		/// <returns>Resulting string</returns>
-		public StringView GetConstantValue(bool bRespectQuotes = false)
+		public StringView GetConstantValue(bool respectQuotes = false)
 		{
 			switch (this.TokenType)
 			{
@@ -524,7 +524,7 @@ namespace EpicGames.UHT.Tokenizer
 					return new StringView(GetFloatValue().ToString("F6", NumberFormatInfo.InvariantInfo));
 				case UhtTokenType.CharConst:
 				case UhtTokenType.StringConst:
-					return GetTokenString(bRespectQuotes);
+					return GetTokenString(respectQuotes);
 				default:
 					return "NotConstant";
 			}
@@ -533,139 +533,139 @@ namespace EpicGames.UHT.Tokenizer
 		/// <summary>
 		/// Return an un-escaped string.  The surrounding quotes will be removed and escaped characters will be converted to the actual values.
 		/// </summary>
-		/// <param name="bRespectQuotes">If true, respect embedded quotes</param>
+		/// <param name="respectQuotes">If true, respect embedded quotes</param>
 		/// <returns>Resulting string</returns>
-		public StringView GetTokenString(bool bRespectQuotes = false)
+		public StringView GetTokenString(bool respectQuotes = false)
 		{
-			StringViewBuilder SVB = new StringViewBuilder();
+			StringViewBuilder builder = new StringViewBuilder();
 			switch (this.TokenType)
 			{
 				case UhtTokenType.StringConst:
-					StringView SubValue = new StringView(this.Value, 1, this.Value.Span.Length - 2);
-					while (SubValue.Span.Length > 0)
+					StringView subValue = new StringView(this.Value, 1, this.Value.Span.Length - 2);
+					while (subValue.Span.Length > 0)
 					{
-						int SlashIndex = SubValue.Span.IndexOf('\\');
-						if (SlashIndex == -1)
+						int slashIndex = subValue.Span.IndexOf('\\');
+						if (slashIndex == -1)
 						{
-							SVB.Append(SubValue);
+							builder.Append(subValue);
 							break;
 						}
-						if (SlashIndex > 0)
+						if (slashIndex > 0)
 						{
-							SVB.Append(new StringView(SubValue, 0, SlashIndex));
+							builder.Append(new StringView(subValue, 0, slashIndex));
 						}
-						if (SlashIndex + 1 == SubValue.Span.Length)
+						if (slashIndex + 1 == subValue.Span.Length)
 						{
 							break;
 						}
 
-						if (SlashIndex + 1 < SubValue.Span.Length)
+						if (slashIndex + 1 < subValue.Span.Length)
 						{
-							char c = SubValue.Span[SlashIndex + 1];
+							char c = subValue.Span[slashIndex + 1];
 							if (c == 'n')
 							{
 								c = '\n';
 							}
-							else if (bRespectQuotes && c == '"')
+							else if (respectQuotes && c == '"')
 							{
-								SVB.Append('\\');
+								builder.Append('\\');
 							}
-							SVB.Append(c);
-							SubValue = new StringView(SubValue, SlashIndex + 2);
+							builder.Append(c);
+							subValue = new StringView(subValue, slashIndex + 2);
 						}
 					}
 					break;
 
 				case UhtTokenType.CharConst:
-					char CharConst = this.Value.Span[1];
-					if (CharConst == '\\')
+					char charConst = this.Value.Span[1];
+					if (charConst == '\\')
 					{
-						CharConst = this.Value.Span[2];
-						switch (CharConst)
+						charConst = this.Value.Span[2];
+						switch (charConst)
 						{
 							case 't':
-								CharConst = '\t';
+								charConst = '\t';
 								break;
 							case 'n':
-								CharConst = '\n';
+								charConst = '\n';
 								break;
 							case 'r':
-								CharConst = '\r';
+								charConst = '\r';
 								break;
 						}
 					}
-					SVB.Append(CharConst);
+					builder.Append(charConst);
 					break;
 
 				default:
 					throw new UhtIceException("Call to GetTokenString on a token that isn't a string or char constant");
 			}
-			return SVB.ToStringView();
+			return builder.ToStringView();
 		}
 
 		/// <summary>
 		/// Join the given tokens into a string
 		/// </summary>
-		/// <param name="Tokens">Tokens to join</param>
+		/// <param name="tokens">Tokens to join</param>
 		/// <returns>Joined strings</returns>
-		public static string Join(IEnumerable<UhtToken> Tokens)
+		public static string Join(IEnumerable<UhtToken> tokens)
 		{
-			StringBuilder Builder = new StringBuilder();
-			foreach (var Token in Tokens)
+			StringBuilder builder = new StringBuilder();
+			foreach (UhtToken token in tokens)
 			{
-				Builder.Append(Token.Value.ToString());
+				builder.Append(token.Value.ToString());
 			}
-			return Builder.ToString();
+			return builder.ToString();
 		}
 
 		/// <summary>
 		/// Join the given tokens into a string
 		/// </summary>
-		/// <param name="Separator">Separator between tokens</param>
-		/// <param name="Tokens">Tokens to join</param>
+		/// <param name="separator">Separator between tokens</param>
+		/// <param name="tokens">Tokens to join</param>
 		/// <returns>Joined strings</returns>
-		public static string Join(char Separator, IEnumerable<UhtToken> Tokens)
+		public static string Join(char separator, IEnumerable<UhtToken> tokens)
 		{
-			StringBuilder Builder = new StringBuilder();
-			bool bIncludeSeperator = false;
-			foreach (var Token in Tokens)
+			StringBuilder builder = new StringBuilder();
+			bool includeSeperator = false;
+			foreach (UhtToken token in tokens)
 			{
-				if (!bIncludeSeperator)
+				if (!includeSeperator)
 				{
-					bIncludeSeperator = true;
+					includeSeperator = true;
 				}
 				else
 				{
-					Builder.Append(Separator);
+					builder.Append(separator);
 				}
-				Builder.Append(Token.Value.ToString());
+				builder.Append(token.Value.ToString());
 			}
-			return Builder.ToString();
+			return builder.ToString();
 		}
 
 		/// <summary>
 		/// Join the given tokens into a string
 		/// </summary>
-		/// <param name="Separator">Separator between tokens</param>
-		/// <param name="Tokens">Tokens to join</param>
+		/// <param name="separator">Separator between tokens</param>
+		/// <param name="tokens">Tokens to join</param>
 		/// <returns>Joined strings</returns>
-		public static string Join(string Separator, IEnumerable<UhtToken> Tokens)
+		public static string Join(string separator, IEnumerable<UhtToken> tokens)
 		{
-			StringBuilder Builder = new StringBuilder();
-			bool bIncludeSeperator = false;
-			foreach (var Token in Tokens)
+			StringBuilder builder = new StringBuilder();
+			bool includeSeperator = false;
+			foreach (UhtToken token in tokens)
 			{
-				if (!bIncludeSeperator)
+				if (!includeSeperator)
 				{
-					bIncludeSeperator = true;
+					includeSeperator = true;
 				}
 				else
 				{
-					Builder.Append(Separator);
+					builder.Append(separator);
 				}
-				Builder.Append(Token.Value.ToString());
+				builder.Append(token.Value.ToString());
 			}
-			return Builder.ToString();
+			return builder.ToString();
 		}
 
 		/// <summary>
@@ -688,48 +688,48 @@ namespace EpicGames.UHT.Tokenizer
 		{
 			if (Value.Span.Length > 0)
 			{
-				if (UhtFCString.IsFloatMarker(Value.Span[Value.Span.Length - 1]))
+				if (UhtFCString.IsFloatMarker(Value.Span[^1]))
 				{
-					return float.Parse(Value.Span.Slice(0, Value.Span.Length - 1));
+					return Single.Parse(Value.Span.Slice(0, Value.Span.Length - 1));
 				}
 			}
-			return float.Parse(Value.Span);
+			return Single.Parse(Value.Span);
 		}
 
 		double GetDoubleValue()
 		{
 			if (Value.Span.Length > 0)
 			{
-				if (UhtFCString.IsFloatMarker(Value.Span[Value.Span.Length - 1]))
+				if (UhtFCString.IsFloatMarker(Value.Span[^1]))
 				{
-					return double.Parse(Value.Span.Slice(0, Value.Span.Length - 1));
+					return Double.Parse(Value.Span.Slice(0, Value.Span.Length - 1));
 				}
 			}
-			return double.Parse(Value.Span);
+			return Double.Parse(Value.Span);
 		}
 
 		long GetDecimalValue()
 		{
-			ReadOnlySpan<char> Span = this.Value.Span;
-			bool bIsUnsigned = false;
-			while (Span.Length > 0)
+			ReadOnlySpan<char> span = this.Value.Span;
+			bool isUnsigned = false;
+			while (span.Length > 0)
 			{
-				char C = Span[Span.Length - 1];
-				if (UhtFCString.IsLongMarker(C))
+				char c = span[^1];
+				if (UhtFCString.IsLongMarker(c))
 				{
-					Span = Span.Slice(0, Span.Length - 1);
+					span = span.Slice(0, span.Length - 1);
 				}
-				else if (UhtFCString.IsUnsignedMarker(C))
+				else if (UhtFCString.IsUnsignedMarker(c))
 				{
-					bIsUnsigned = true;
-					Span = Span.Slice(0, Span.Length - 1);
+					isUnsigned = true;
+					span = span.Slice(0, span.Length - 1);
 				}
 				else
 				{
 					break;
 				}
 			}
-			return bIsUnsigned ? (long)Convert.ToUInt64(Span.ToString(), 10) : Convert.ToInt64(Span.ToString(), 10);
+			return isUnsigned ? (long)Convert.ToUInt64(span.ToString(), 10) : Convert.ToInt64(span.ToString(), 10);
 		}
 
 		long GetHexValue()
