More coding standard fixes

diff --git a/Engine/Source/Programs/Shared/EpicGames.UHT/Tokenizer/UhtTokenBufferReader.cs b/Engine/Source/Programs/Shared/EpicGames.UHT/Tokenizer/UhtTokenBufferReader.cs
--- a/Engine/Source/Programs/Shared/EpicGames.UHT/Tokenizer/UhtTokenBufferReader.cs
+++ b/Engine/Source/Programs/Shared/EpicGames.UHT/Tokenizer/UhtTokenBufferReader.cs
@@ -1,10 +1,9 @@
 // Copyright Epic Games, Inc. All Rights Reserved.
 
-using EpicGames.Core;
-using EpicGames.UHT.Utils;
 using System;
 using System.Collections.Generic;
-using System.Diagnostics;
+using EpicGames.Core;
+using EpicGames.UHT.Utils;
 
 namespace EpicGames.UHT.Tokenizer
 {
@@ -14,44 +13,44 @@ namespace EpicGames.UHT.Tokenizer
 	/// </summary>
 	public sealed class UhtTokenBufferReader : IUhtTokenReader, IUhtMessageLineNumber
 	{
-		private static readonly StringView[] EmptyComments = Array.Empty<StringView>();
-
-		private readonly IUhtMessageSite MessageSiteInternal;
-		private List<StringView>? CommentsInternal = null;
-		private readonly List<UhtToken> RecordedTokensInternal = new List<UhtToken>();
-		private IUhtTokenPreprocessor? TokenPreprocessorInternal = null;
-		private UhtToken CurrentToken = new UhtToken(); // PeekToken must have been invoked first
-		private readonly ReadOnlyMemory<char> Data;
-		private int PrevPos = 0;
-		private int PrevLine = 1;
-		private bool bHasToken = false;
-		private int PreCurrentTokenInputPos = 0;
-		private int PreCurrentTokenInputLine = 1;
-		private int InputPosInternal = 0;
-		private int InputLineInternal = 1;
-		private int CommentsDisableCount = 0;
-		private int CommittedComments = 0;
-		private List<StringView>? SavedComments = null;
-		private int SavedInputPos = 0;
-		private int SavedInputLine = 0;
-		private bool bHasSavedState = false;
-		private bool bRecordTokens = false;
-		private int PreprocessorPendingCommentsCount = 0;
+		private static readonly StringView[] s_emptyComments = Array.Empty<StringView>();
+
+		private readonly IUhtMessageSite _messageSiteInternal;
+		private List<StringView>? _commentsInternal = null;
+		private readonly List<UhtToken> _recordedTokensInternal = new List<UhtToken>();
+		private IUhtTokenPreprocessor? _tokenPreprocessorInternal = null;
+		private UhtToken _currentToken = new UhtToken(); // PeekToken must have been invoked first
+		private readonly ReadOnlyMemory<char> _data;
+		private int _prevPos = 0;
+		private int _prevLine = 1;
+		private bool _hasToken = false;
+		private int _preCurrentTokenInputPos = 0;
+		private int _preCurrentTokenInputLine = 1;
+		private int _inputPosInternal = 0;
+		private int _inputLineInternal = 1;
+		private int _commentsDisableCount = 0;
+		private int _committedComments = 0;
+		private List<StringView>? _savedComments = null;
+		private int _savedInputPos = 0;
+		private int _savedInputLine = 0;
+		private bool _hasSavedState = false;
+		private bool _recordTokens = false;
+		private int _preprocessorPendingCommentsCount = 0;
 
 		/// <summary>
 		/// Construct a new token reader
 		/// </summary>
-		/// <param name="MessageSite">Message site for messages</param>
-		/// <param name="Input">Input source</param>
-		public UhtTokenBufferReader(IUhtMessageSite MessageSite, ReadOnlyMemory<char> Input)
+		/// <param name="messageSite">Message site for messages</param>
+		/// <param name="input">Input source</param>
+		public UhtTokenBufferReader(IUhtMessageSite messageSite, ReadOnlyMemory<char> input)
 		{
-			this.MessageSiteInternal = MessageSite;
-			this.Data = Input;
+			this._messageSiteInternal = messageSite;
+			this._data = input;
 		}
 
 		#region IUHTMessageSite Implementation
-		IUhtMessageSession IUhtMessageSite.MessageSession => this.MessageSiteInternal.MessageSession;
-		IUhtMessageSource? IUhtMessageSite.MessageSource => this.MessageSiteInternal.MessageSource;
+		IUhtMessageSession IUhtMessageSite.MessageSession => this._messageSiteInternal.MessageSession;
+		IUhtMessageSource? IUhtMessageSite.MessageSource => this._messageSiteInternal.MessageSource;
 		IUhtMessageLineNumber? IUhtMessageSite.MessageLineNumber => this;
 		#endregion
 
@@ -61,90 +60,87 @@ namespace EpicGames.UHT.Tokenizer
 
 		#region ITokenReader Implementation
 		/// <inheritdoc/>
-		public bool bIsEOF
+		public bool IsEOF
 		{
 			get
 			{
-				if (this.bHasToken)
+				if (this._hasToken)
 				{
-					return this.CurrentToken.TokenType.IsEndType();
+					return this._currentToken.TokenType.IsEndType();
 				}
 				else
 				{
-					return this.InputPos == this.Data.Length;
+					return this.InputPos == this._data.Length;
 				}
 			}
 		}
 
 		/// <inheritdoc/>
-		public int InputPos
-		{
-			get => this.bHasToken ? this.PreCurrentTokenInputPos : this.InputPosInternal;
-		}
+		public int InputPos => this._hasToken ? this._preCurrentTokenInputPos : this._inputPosInternal;
 
 		/// <inheritdoc/>
 		public int InputLine
 		{
-			get => this.bHasToken ? this.PreCurrentTokenInputLine : this.InputLineInternal;
-			set 
-			{ 
-				ClearToken(); 
-				this.InputLineInternal = value; 
+			get => this._hasToken ? this._preCurrentTokenInputLine : this._inputLineInternal;
+			set
+			{
+				ClearToken();
+				this._inputLineInternal = value;
 			}
 		}
 
 		/// <inheritdoc/>
-		public ReadOnlySpan<StringView> Comments 
+		public ReadOnlySpan<StringView> Comments
 		{
 			get
 			{
-				if (this.CommentsInternal != null && this.CommittedComments != 0)
+				if (this._commentsInternal != null && this._committedComments != 0)
 				{
-					return new ReadOnlySpan<StringView>(this.CommentsInternal.ToArray(), 0, this.CommittedComments);
+					return new ReadOnlySpan<StringView>(this._commentsInternal.ToArray(), 0, this._committedComments);
 				}
 				else
 				{
-					return new ReadOnlySpan<StringView>(UhtTokenBufferReader.EmptyComments);
+					return new ReadOnlySpan<StringView>(UhtTokenBufferReader.s_emptyComments);
 				}
 			}
 		}
 
 		/// <inheritdoc/>
-		public IUhtTokenPreprocessor? TokenPreprocessor { get => this.TokenPreprocessorInternal; set => this.TokenPreprocessorInternal = value; }
+		public IUhtTokenPreprocessor? TokenPreprocessor { get => this._tokenPreprocessorInternal; set => this._tokenPreprocessorInternal = value; }
 
 		/// <inheritdoc/>
 		public ref UhtToken PeekToken()
 		{
-			if (!this.bHasToken)
+			if (!this._hasToken)
 			{
-				this.CurrentToken = GetTokenInternal(true);
-				this.bHasToken = true;
+				this._currentToken = GetTokenInternal(true);
+				this._hasToken = true;
 			}
-			return ref this.CurrentToken;
+			return ref this._currentToken;
 		}
 
 		/// <inheritdoc/>
 		public void SkipWhitespaceAndComments()
 		{
-			bool bGotInlineComment = false;
-			SkipWhitespaceAndCommentsInternal(ref bGotInlineComment, true);
+			bool gotInlineComment = false;
+			SkipWhitespaceAndCommentsInternal(ref gotInlineComment, true);
 		}
 
 		/// <inheritdoc/>
 		public void ConsumeToken()
 		{
-			if (this.bRecordTokens && !this.CurrentToken.IsEndType())
+			if (this._recordTokens && !this._currentToken.IsEndType())
 			{
-				this.RecordedTokensInternal.Add(this.CurrentToken);
+				this._recordedTokensInternal.Add(this._currentToken);
 			}
-			this.bHasToken = false;
+			this._hasToken = false;
 
 			// When comments are disabled, we are still collecting comments, but aren't committing them
-			if (this.CommentsDisableCount == 0)
+			if (this._commentsDisableCount == 0)
 			{
-				if (this.CommentsInternal != null)
+				if (this._commentsInternal != null)
 				{
-					this.CommittedComments = this.CommentsInternal.Count;
+					this._committedComments = this._commentsInternal.Count;
 				}
 			}
 			else
@@ -156,266 +152,262 @@ namespace EpicGames.UHT.Tokenizer
 		/// <inheritdoc/>
 		public UhtToken GetToken()
 		{
-			UhtToken Token = PeekToken();
+			UhtToken token = PeekToken();
 			ConsumeToken();
-			return Token;
+			return token;
 		}
 
 		/// <inheritdoc/>
-		public StringView GetRawString(char Terminator, UhtRawStringOptions Options)
+		public StringView GetRawString(char terminator, UhtRawStringOptions options)
 		{
-			ReadOnlySpan<char> Span = this.Data.Span;
+			ReadOnlySpan<char> span = this._data.Span;
 
 			ClearToken();
 			SkipWhitespaceAndComments();
 
-			int StartPos = this.InputPos;
-			bool bInQuotes = false;
+			int startPos = this.InputPos;
+			bool inQuotes = false;
 			while (true)
 			{
-				char C = InternalGetChar(Span);
+				char c = InternalGetChar(span);
 
 				// Check for end of file
-				if (C == 0)
+				if (c == 0)
 				{
 					break;
 				}
 
 				// Check for end of line
-				if (C == '\r' || C == '\n')
+				if (c == '\r' || c == '\n')
 				{
-					--this.InputPosInternal;
+					--this._inputPosInternal;
 					break;
 				}
 
 				// Check for terminator as long as we aren't in quotes
-				if (C == Terminator && !bInQuotes)
+				if (c == terminator && !inQuotes)
 				{
-					if (Options.HasAnyFlags(UhtRawStringOptions.DontConsumeTerminator))
+					if (options.HasAnyFlags(UhtRawStringOptions.DontConsumeTerminator))
 					{
-						--this.InputPosInternal;
+						--this._inputPosInternal;
 					}
 					break;
 				}
 
 				// Check for comment
-				if (!bInQuotes && C == '/')
+				if (!inQuotes && c == '/')
 				{
-					char P = InternalPeekChar(Span);
-					if (P == '*' || P == '/')
+					char p = InternalPeekChar(span);
+					if (p == '*' || p == '/')
 					{
-						--this.InputPosInternal;
+						--this._inputPosInternal;
 						break;
 					}
 				}
 
 				// Check for quotes
-				if (C == '"' && Options.HasAnyFlags(UhtRawStringOptions.RespectQuotes))
+				if (c == '"' && options.HasAnyFlags(UhtRawStringOptions.RespectQuotes))
 				{
-					bInQuotes = !bInQuotes;
+					inQuotes = !inQuotes;
 				}
 			}
 
 			// If EOF, then error
-			if (bInQuotes)
+			if (inQuotes)
 			{
 				throw new UhtException(this, "Unterminated quoted string");
 			}
 
 			// Remove trailing whitespace
-			int EndPos = this.InputPos;
-			for (; EndPos > StartPos; --EndPos)
+			int endPos = this.InputPos;
+			for (; endPos > startPos; --endPos)
 			{
-				char C = Span[EndPos - 1];
-				if (C != ' ' && C != '\t')
+				char c = span[endPos - 1];
+				if (c != ' ' && c != '\t')
 				{
 					break;
 				}
 			}
 
 			// Check for too long
-			if (EndPos - StartPos >= UhtToken.MaxStringLength)
+			if (endPos - startPos >= UhtToken.MaxStringLength)
 			{
 				throw new UhtException(this, $"String exceeds maximum of {UhtToken.MaxStringLength} characters");
 			}
-			return new StringView(this.Data, StartPos, EndPos - StartPos);
+			return new StringView(this._data, startPos, endPos - startPos);
 		}
 
 		/// <inheritdoc/>
 		public UhtToken GetLine()
 		{
-			ReadOnlySpan<char> Span = this.Data.Span;
+			ReadOnlySpan<char> span = this._data.Span;
 
 			ClearToken();
-			this.PrevPos = this.InputPosInternal;
-			this.PrevLine = this.InputLineInternal;
+			this._prevPos = this._inputPosInternal;
+			this._prevLine = this._inputLineInternal;
 
-			if (this.PrevPos == Span.Length)
+			if (this._prevPos == span.Length)
 			{
-				return new UhtToken(UhtTokenType.EndOfFile, this.PrevPos, this.PrevLine, this.PrevPos, this.PrevLine, new StringView(this.Data.Slice(this.PrevPos, 0)));
+				return new UhtToken(UhtTokenType.EndOfFile, this._prevPos, this._prevLine, this._prevPos, this._prevLine, new StringView(this._data.Slice(this._prevPos, 0)));
 			}
 
-			int LastPos = this.InputPosInternal;
+			int lastPos = this._inputPosInternal;
 			while (true)
 			{
-				char Char = InternalGetChar(Span);
-				if (Char == 0)
+				char c = InternalGetChar(span);
+				if (c == 0)
 				{
 					break;
 				}
-				else if (Char == '\r')
+				else if (c == '\r')
 				{
 				}
-				else if (Char == '\n')
+				else if (c == '\n')
 				{
-					++this.InputLineInternal;
+					++this._inputLineInternal;
 					break;
 				}
 				else
 				{
-					LastPos = this.InputPosInternal;
+					lastPos = this._inputPosInternal;
 				}
 			}
 
-			return new UhtToken(UhtTokenType.Line, this.PrevPos, this.PrevLine, this.PrevPos, this.PrevLine, new StringView(this.Data.Slice(this.PrevPos, LastPos - this.PrevPos)));
+			return new UhtToken(UhtTokenType.Line, this._prevPos, this._prevLine, this._prevPos, this._prevLine, new StringView(this._data.Slice(this._prevPos, lastPos - this._prevPos)));
 		}
 
 		/// <inheritdoc/>
-		public StringView GetStringView(int StartPos, int Count)
+		public StringView GetStringView(int startPos, int count)
 		{
-			return new StringView(this.Data, StartPos, Count);
+			return new StringView(this._data, startPos, count);
 		}
 
 		/// <inheritdoc/>
 		public void ClearComments()
 		{
-			if (this.PreprocessorPendingCommentsCount != 0)
-			{
-				Debugger.Break();
-			}
-			if (this.CommentsInternal != null)
+			if (this._commentsInternal != null)
 			{
 
 				// Clearing comments does not remove any uncommitted comments
-				this.CommentsInternal.RemoveRange(0, this.CommittedComments);
-				this.CommittedComments = 0;
+				this._commentsInternal.RemoveRange(0, this._committedComments);
+				this._committedComments = 0;
 			}
 		}
 
 		/// <inheritdoc/>
 		public void DisableComments()
 		{
-			++this.CommentsDisableCount;
+			++this._commentsDisableCount;
 		}
 
 		/// <inheritdoc/>
 		public void EnableComments()
 		{
-			--this.CommentsDisableCount;
+			--this._commentsDisableCount;
 		}
 
 		/// <inheritdoc/>
 		public void CommitPendingComments()
 		{
-			if (this.CommentsInternal != null)
+			if (this._commentsInternal != null)
 			{
-				this.CommittedComments = this.CommentsInternal.Count;
+				this._committedComments = this._commentsInternal.Count;
 			}
 		}
 
 		/// <inheritdoc/>
-		public bool IsFirstTokenInLine(ref UhtToken Token)
+		public bool IsFirstTokenInLine(ref UhtToken token)
 		{
-			return IsFirstTokenInLine(this.Data.Span, Token.InputStartPos);
+			return IsFirstTokenInLine(this._data.Span, token.InputStartPos);
 		}
 
 		/// <inheritdoc/>
 		public void SaveState()
 		{
-			if (this.bHasSavedState)
+			if (this._hasSavedState)
 			{
 				throw new UhtIceException("Can not save more than one state");
 			}
-			this.bHasSavedState = true;
-			this.SavedInputLine = this.InputLine;
-			this.SavedInputPos = this.InputPos;
-			if (this.CommentsInternal != null)
+			this._hasSavedState = true;
+			this._savedInputLine = this.InputLine;
+			this._savedInputPos = this.InputPos;
+			if (this._commentsInternal != null)
 			{
-				if (this.SavedComments == null)
+				if (this._savedComments == null)
 				{
-					this.SavedComments = new List<StringView>();
+					this._savedComments = new List<StringView>();
 				}
-				this.SavedComments.Clear();
-				this.SavedComments.AddRange(this.CommentsInternal);
+				this._savedComments.Clear();
+				this._savedComments.AddRange(this._commentsInternal);
 			}
-			if (this.TokenPreprocessorInternal != null)
+			if (this._tokenPreprocessorInternal != null)
 			{
-				this.TokenPreprocessorInternal.SaveState();
+				this._tokenPreprocessorInternal.SaveState();
 			}
 		}
 
 		/// <inheritdoc/>
 		public void RestoreState()
 		{
-			if (!this.bHasSavedState)
+			if (!this._hasSavedState)
 			{
 				throw new UhtIceException("Can not restore state when none have been saved");
 			}
-			this.bHasSavedState = false;
+			this._hasSavedState = false;
 			ClearToken();
-			this.InputPosInternal = this.SavedInputPos;
-			this.InputLineInternal = this.SavedInputLine;
-			if (this.SavedComments != null && this.CommentsInternal != null)
+			this._inputPosInternal = this._savedInputPos;
+			this._inputLineInternal = this._savedInputLine;
+			if (this._savedComments != null && this._commentsInternal != null)
 			{
-				this.CommentsInternal.Clear();
-				this.CommentsInternal.AddRange(this.SavedComments);
+				this._commentsInternal.Clear();
+				this._commentsInternal.AddRange(this._savedComments);
 			}
-			if (this.TokenPreprocessorInternal != null)
+			if (this._tokenPreprocessorInternal != null)
 			{
-				this.TokenPreprocessorInternal.RestoreState();
+				this._tokenPreprocessorInternal.RestoreState();
 			}
 		}
 
 		/// <inheritdoc/>
 		public void AbandonState()
 		{
-			if (!this.bHasSavedState)
+			if (!this._hasSavedState)
 			{
 				throw new UhtIceException("Can not abandon state when none have been saved");
 			}
-			this.bHasSavedState = false;
+			this._hasSavedState = false;
 			ClearToken();
 		}
 
 		/// <inheritdoc/>
 		public void EnableRecording()
 		{
-			if (this.bRecordTokens)
+			if (this._recordTokens)
 			{
 				throw new UhtIceException("Can not nest token recording");
 			}
-			this.bRecordTokens = true;
+			this._recordTokens = true;
 		}
 
 		/// <inheritdoc/>
-		public void RecordToken(ref UhtToken Token)
+		public void RecordToken(ref UhtToken token)
 		{
-			if (!this.bRecordTokens)
+			if (!this._recordTokens)
 			{
 				throw new UhtIceException("Attempt to disable token recording when it isn't enabled");
 			}
-			this.RecordedTokensInternal.Add(Token);
+			this._recordedTokensInternal.Add(token);
 		}
 
 		/// <inheritdoc/>
 		public void DisableRecording()
 		{
-			if (!this.bRecordTokens)
+			if (!this._recordTokens)
 			{
 				throw new UhtIceException("Attempt to disable token recording when it isn't enabled");
 			}
-			this.RecordedTokensInternal.Clear();
-			this.bRecordTokens = false;
+			this._recordedTokensInternal.Clear();
+			this._recordTokens = false;
 		}
 
 		/// <inheritdoc/>
@@ -423,11 +415,11 @@ namespace EpicGames.UHT.Tokenizer
 		{
 			get
 			{
-				if (!this.bRecordTokens)
+				if (!this._recordTokens)
 				{
 					throw new UhtIceException("Attempt to get recorded tokens when it isn't enabled");
 				}
-				return this.RecordedTokensInternal;
+				return this._recordedTokensInternal;
 			}
 		}
 		#endregion
@@ -435,31 +427,27 @@ namespace EpicGames.UHT.Tokenizer
 		#region Internals
 		private void ClearAllComments()
 		{
-			if (this.PreprocessorPendingCommentsCount != 0)
-			{
-				Debugger.Break();
-			}
-			if (this.CommentsInternal != null)
+			if (this._commentsInternal != null)
 			{
-				this.CommentsInternal.Clear();
-				this.CommittedComments = 0;
+				this._commentsInternal.Clear();
+				this._committedComments = 0;
 			}
 		}
 
 		private void ClearPendingComments()
 		{
-			if (this.CommentsInternal != null)
+			if (this._commentsInternal != null)
 			{
-				int StartingComment = this.CommittedComments + this.PreprocessorPendingCommentsCount;
-				this.CommentsInternal.RemoveRange(StartingComment, this.CommentsInternal.Count - StartingComment);
+				int startingComment = this._committedComments + this._preprocessorPendingCommentsCount;
+				this._commentsInternal.RemoveRange(startingComment, this._commentsInternal.Count - startingComment);
 			}
 		}
 
-		private static bool IsFirstTokenInLine(ReadOnlySpan<char> Span, int StartPos)
+		private static bool IsFirstTokenInLine(ReadOnlySpan<char> span, int startPos)
 		{
-			for (int Pos = StartPos; --Pos > 0;)
+			for (int pos = startPos; --pos > 0;)
 			{
-				switch (Span[Pos])
+				switch (span[pos])
 				{
 					case '\r':
 					case '\n':
@@ -478,74 +466,74 @@ namespace EpicGames.UHT.Tokenizer
 		/// Get the next token
 		/// </summary>
 		/// <returns>Return the next token from the stream.  If the end of stream is reached, a token type of None will be returned.</returns>
-		private UhtToken GetTokenInternal(bool bEnablePreprocessor)
+		private UhtToken GetTokenInternal(bool enablePreprocessor)
 		{
-			ReadOnlySpan<char> Span = this.Data.Span;
-			bool bGotInlineComment = false;
-		Restart:
-			this.PreCurrentTokenInputLine = this.InputLineInternal;
-			this.PreCurrentTokenInputPos = this.InputPosInternal;
-			SkipWhitespaceAndCommentsInternal(ref bGotInlineComment, bEnablePreprocessor);
-			this.PrevPos = this.InputPosInternal;
-			this.PrevLine = this.InputLineInternal;
-
-			UhtTokenType TokenType = UhtTokenType.EndOfFile;
-			int StartPos = this.InputPosInternal;
-			int StartLine = this.InputLineInternal;
-
-			char C = InternalGetChar(Span);
-			if (C == 0)
+			ReadOnlySpan<char> span = this._data.Span;
+			bool gotInlineComment = false;
+Restart:
+			this._preCurrentTokenInputLine = this._inputLineInternal;
+			this._preCurrentTokenInputPos = this._inputPosInternal;
+			SkipWhitespaceAndCommentsInternal(ref gotInlineComment, enablePreprocessor);
+			this._prevPos = this._inputPosInternal;
+			this._prevLine = this._inputLineInternal;
+
+			UhtTokenType tokenType = UhtTokenType.EndOfFile;
+			int startPos = this._inputPosInternal;
+			int startLine = this._inputLineInternal;
+
+			char c = InternalGetChar(span);
+			if (c == 0)
 			{
 			}
 
-			else if (UhtFCString.IsAlpha(C) || C == '_')
+			else if (UhtFCString.IsAlpha(c) || c == '_')
 			{
 
-				for (; this.InputPosInternal < Span.Length; ++this.InputPosInternal)
+				for (; this._inputPosInternal < span.Length; ++this._inputPosInternal)
 				{
-					C = Span[this.InputPosInternal];
-					if (!(UhtFCString.IsAlpha(C) || UhtFCString.IsDigit(C) || C == '_'))
+					c = span[this._inputPosInternal];
+					if (!(UhtFCString.IsAlpha(c) || UhtFCString.IsDigit(c) || c == '_'))
 					{
 						break;
 					}
 				}
 
-				if (this.InputPosInternal - StartPos >= UhtToken.MaxNameLength)
+				if (this._inputPosInternal - startPos >= UhtToken.MaxNameLength)
 				{
 					throw new UhtException(this, $"Identifier length exceeds maximum of {UhtToken.MaxNameLength}");
 				}
 
-				TokenType = UhtTokenType.Identifier;
+				tokenType = UhtTokenType.Identifier;
 			}
 
 			// Check for any numerics 
-			else if (IsNumeric(Span, C))
+			else if (IsNumeric(span, c))
 			{
 				// Integer or floating point constant.
 
-				bool bIsFloat = C == '.';
-				bool bIsHex = false;
+				bool isFloat = c == '.';
+				bool isHex = false;
 
 				// Ignore the starting sign for this part of the parsing
-				if (UhtFCString.IsSign(C))
+				if (UhtFCString.IsSign(c))
 				{
-					C = InternalGetChar(Span); // We know this won't fail since the code above checked the Peek
+					c = InternalGetChar(span); // We know this won't fail since the code above checked the Peek
 				}
 
 				// Check for a hex valid
-				if (C == '0')
+				if (c == '0')
 				{
-					bIsHex = UhtFCString.IsHexMarker(InternalPeekChar(Span));
-					if (bIsHex)
+					isHex = UhtFCString.IsHexMarker(InternalPeekChar(span));
+					if (isHex)
 					{
-						InternalGetChar(Span);
+						InternalGetChar(span);
 					}
 				}
 
 				// If we have a hex constant
-				if (bIsHex)
+				if (isHex)
 				{
-					for (; this.InputPosInternal < Span.Length && UhtFCString.IsHexDigit(Span[this.InputPosInternal]); ++this.InputPosInternal)
+					for (; this._inputPosInternal < span.Length && UhtFCString.IsHexDigit(span[this._inputPosInternal]); ++this._inputPosInternal)
 					{
 					}
 				}
@@ -555,115 +543,115 @@ namespace EpicGames.UHT.Tokenizer
 				{
 
 					// Skip all digits
-					for (; this.InputPosInternal < Span.Length && UhtFCString.IsDigit(Span[this.InputPosInternal]); ++this.InputPosInternal)
+					for (; this._inputPosInternal < span.Length && UhtFCString.IsDigit(span[this._inputPosInternal]); ++this._inputPosInternal)
 					{
 					}
 
 					// If we have a '.'
-					if (this.InputPosInternal < Span.Length && Span[this.InputPosInternal] == '.')
+					if (this._inputPosInternal < span.Length && span[this._inputPosInternal] == '.')
 					{
-						bIsFloat = true;
-						++this.InputPosInternal;
+						isFloat = true;
+						++this._inputPosInternal;
 
 						// Skip all digits
-						for (; this.InputPosInternal < Span.Length && UhtFCString.IsDigit(Span[this.InputPosInternal]); ++this.InputPosInternal)
+						for (; this._inputPosInternal < span.Length && UhtFCString.IsDigit(span[this._inputPosInternal]); ++this._inputPosInternal)
 						{
 						}
 					}
 
 					// If we have a 'e'
-					if (this.InputPosInternal < Span.Length && UhtFCString.IsExponentMarker(Span[this.InputPosInternal]))
+					if (this._inputPosInternal < span.Length && UhtFCString.IsExponentMarker(span[this._inputPosInternal]))
 					{
-						bIsFloat = true;
-						++this.InputPosInternal;
+						isFloat = true;
+						++this._inputPosInternal;
 
 						// Skip any signs
-						if (this.InputPosInternal < Span.Length && UhtFCString.IsSign(Span[this.InputPosInternal]))
+						if (this._inputPosInternal < span.Length && UhtFCString.IsSign(span[this._inputPosInternal]))
 						{
-							++this.InputPosInternal;
+							++this._inputPosInternal;
 						}
 
 						// Skip all digits
-						for (; this.InputPosInternal < Span.Length && UhtFCString.IsDigit(Span[this.InputPosInternal]); ++this.InputPosInternal)
+						for (; this._inputPosInternal < span.Length && UhtFCString.IsDigit(span[this._inputPosInternal]); ++this._inputPosInternal)
 						{
 						}
 					}
 
 					// If we have a 'f'
-					if (this.InputPosInternal < Span.Length && UhtFCString.IsFloatMarker(Span[this.InputPosInternal]))
+					if (this._inputPosInternal < span.Length && UhtFCString.IsFloatMarker(span[this._inputPosInternal]))
 					{
-						bIsFloat = true;
-						++this.InputPosInternal;
+						isFloat = true;
+						++this._inputPosInternal;
 					}
 
 					// Check for u/l markers
-					while (this.InputPosInternal < Span.Length && 
-						(UhtFCString.IsUnsignedMarker(Span[this.InputPosInternal]) || UhtFCString.IsLongMarker(Span[this.InputPosInternal])))
+					while (this._inputPosInternal < span.Length &&
+						(UhtFCString.IsUnsignedMarker(span[this._inputPosInternal]) || UhtFCString.IsLongMarker(span[this._inputPosInternal])))
 					{
-						++this.InputPosInternal;
+						++this._inputPosInternal;
 					}
 				}
 
-				if (this.InputPosInternal - StartPos >= UhtToken.MaxNameLength)
+				if (this._inputPosInternal - startPos >= UhtToken.MaxNameLength)
 				{
 					throw new UhtException(this, $"Number length exceeds maximum of {UhtToken.MaxNameLength}");
 				}
 
-				TokenType = bIsFloat ? UhtTokenType.FloatConst : (bIsHex ? UhtTokenType.HexConst : UhtTokenType.DecimalConst);
+				tokenType = isFloat ? UhtTokenType.FloatConst : (isHex ? UhtTokenType.HexConst : UhtTokenType.DecimalConst);
 			}
 
 			// Escaped character constant
-			else if (C == '\'')
+			else if (c == '\'')
 			{
 
 				// We try to skip the character constant value. But if it is backslash, we have to skip another character
-				if (InternalGetChar(Span) == '\\')
+				if (InternalGetChar(span) == '\\')
 				{
-					InternalGetChar(Span);
+					InternalGetChar(span);
 				}
 
-				char NextChar = InternalGetChar(Span);
-				if (NextChar != '\'')
+				char nextChar = InternalGetChar(span);
+				if (nextChar != '\'')
 				{
 					throw new UhtException(this, "Unterminated character constant");
 				}
 
-				TokenType = UhtTokenType.CharConst;
+				tokenType = UhtTokenType.CharConst;
 			}
 
 			// String constant
-			else if (C == '"')
+			else if (c == '"')
 			{
 				for (; ; )
 				{
-					char NextChar = InternalGetChar(Span);
-					if (NextChar == '\r' || NextChar == '\n')
+					char nextChar = InternalGetChar(span);
+					if (nextChar == '\r' || nextChar == '\n')
 					{
 						// throw
 					}
 
-					if (NextChar == '\\')
+					if (nextChar == '\\')
 					{
-						NextChar = InternalGetChar(Span);
-						if (NextChar == '\r' || NextChar == '\n')
+						nextChar = InternalGetChar(span);
+						if (nextChar == '\r' || nextChar == '\n')
 						{
 							// throw
 						}
-						NextChar = InternalGetChar(Span);
+						nextChar = InternalGetChar(span);
 					}
 
-					if (NextChar == '"')
+					if (nextChar == '"')
 					{
 						break;
 					}
 				}
 
-				if (this.InputPosInternal - StartPos >= UhtToken.MaxStringLength)
+				if (this._inputPosInternal - startPos >= UhtToken.MaxStringLength)
 				{
 					throw new UhtException(this, $"String constant exceeds maximum of {UhtToken.MaxStringLength} characters");
 				}
 
-				TokenType = UhtTokenType.StringConst;
+				tokenType = UhtTokenType.StringConst;
 			}
 
 			// Assume everything else is a symbol.
@@ -671,28 +659,28 @@ namespace EpicGames.UHT.Tokenizer
 			else
 			{
 				{
-					if (this.InputPosInternal < Span.Length)
+					if (this._inputPosInternal < span.Length)
 					{
-						char D = Span[this.InputPosInternal];
-						if ((C == '<' && D == '<')
-							|| (C == '!' && D == '=')
-							|| (C == '<' && D == '=')
-							|| (C == '>' && D == '=')
-							|| (C == '+' && D == '+')
-							|| (C == '-' && D == '-')
-							|| (C == '+' && D == '=')
-							|| (C == '-' && D == '=')
-							|| (C == '*' && D == '=')
-							|| (C == '/' && D == '=')
-							|| (C == '&' && D == '&')
-							|| (C == '|' && D == '|')
-							|| (C == '^' && D == '^')
-							|| (C == '=' && D == '=')
-							|| (C == '*' && D == '*')
-							|| (C == '~' && D == '=')
-							|| (C == ':' && D == ':'))
+						char d = span[this._inputPosInternal];
+						if ((c == '<' && d == '<')
+							|| (c == '!' && d == '=')
+							|| (c == '<' && d == '=')
+							|| (c == '>' && d == '=')
+							|| (c == '+' && d == '+')
+							|| (c == '-' && d == '-')
+							|| (c == '+' && d == '=')
+							|| (c == '-' && d == '=')
+							|| (c == '*' && d == '=')
+							|| (c == '/' && d == '=')
+							|| (c == '&' && d == '&')
+							|| (c == '|' && d == '|')
+							|| (c == '^' && d == '^')
+							|| (c == '=' && d == '=')
+							|| (c == '*' && d == '*')
+							|| (c == '~' && d == '=')
+							|| (c == ':' && d == ':'))
 						{
-							++this.InputPosInternal;
+							++this._inputPosInternal;
 						}
 					}
 
@@ -703,308 +691,307 @@ namespace EpicGames.UHT.Tokenizer
 					// 1) If the block is being skipped, we want to preserve any comments ahead of the token and eliminate all other comments found in the #if block
 					// 2) If the block isn't being skipped or in the case of things such as #pragma or #include, they were considered statements in the UHT
 					//	  and would result in the comments being dropped.  In the new UHT, we just eliminate all pending tokens.
-					if (C == '#' && bEnablePreprocessor && this.TokenPreprocessorInternal != null && IsFirstTokenInLine(Span, StartPos))
+					if (c == '#' && enablePreprocessor && this._tokenPreprocessorInternal != null && IsFirstTokenInLine(span, startPos))
 					{
-						this.PreprocessorPendingCommentsCount = this.CommentsInternal != null ? this.CommentsInternal.Count - this.CommittedComments : 0;
-						UhtToken Temp = new UhtToken(TokenType, StartPos, StartLine, StartPos, this.InputLineInternal, new StringView(this.Data.Slice(StartPos, this.InputPosInternal - StartPos)));
-						bool bInclude = this.TokenPreprocessorInternal.ParsePreprocessorDirective(ref Temp, true, out bool bClearComments, out bool bIllegalContentsCheck);
-						if (!bInclude)
+						this._preprocessorPendingCommentsCount = this._commentsInternal != null ? this._commentsInternal.Count - this._committedComments : 0;
+						UhtToken temp = new UhtToken(tokenType, startPos, startLine, startPos, this._inputLineInternal, new StringView(this._data.Slice(startPos, this._inputPosInternal - startPos)));
+						bool include = this._tokenPreprocessorInternal.ParsePreprocessorDirective(ref temp, true, out bool clearComments, out bool illegalContentsCheck);
+						if (!include)
 						{
-							++this.CommentsDisableCount;
-							int CheckStateMachine = 0;
+							++this._commentsDisableCount;
+							int checkStateMachine = 0;
 							while (true)
 							{
-								if (bIsEOF)
+								if (IsEOF)
 								{
 									break;
 								}
-								UhtToken LocalToken = GetTokenInternal(false);
-								ReadOnlySpan<char> LocalValueSpan = LocalToken.Value.Span;
-								if (LocalToken.IsSymbol('#') && LocalValueSpan.Length == 1 && LocalValueSpan[0] == '#' && IsFirstTokenInLine(Span, LocalToken.InputStartPos))
+								UhtToken localToken = GetTokenInternal(false);
+								ReadOnlySpan<char> localValueSpan = localToken.Value.Span;
+								if (localToken.IsSymbol('#') && localValueSpan.Length == 1 && localValueSpan[0] == '#' && IsFirstTokenInLine(span, localToken.InputStartPos))
 								{
-									if (this.TokenPreprocessorInternal.ParsePreprocessorDirective(ref Temp, false, out bool _, out bool bScratchIllegalContentsCheck))
+									if (this._tokenPreprocessorInternal.ParsePreprocessorDirective(ref temp, false, out bool _, out bool scratchIllegalContentsCheck))
 									{
 										break;
 									}
-									bIllegalContentsCheck = bScratchIllegalContentsCheck;
+									illegalContentsCheck = scratchIllegalContentsCheck;
 								}
-								else if (bIllegalContentsCheck)
+								else if (illegalContentsCheck)
 								{
-									switch (CheckStateMachine)
+									switch (checkStateMachine)
 									{
 										case 0:
-										ResetStateMachineCheck:
-											if (LocalValueSpan.CompareTo("UPROPERTY", StringComparison.Ordinal) == 0 ||
-												LocalValueSpan.CompareTo("UCLASS", StringComparison.Ordinal) == 0 ||
-												LocalValueSpan.CompareTo("USTRUCT", StringComparison.Ordinal) == 0 ||
-												LocalValueSpan.CompareTo("UENUM", StringComparison.Ordinal) == 0 ||
-												LocalValueSpan.CompareTo("UINTERFACE", StringComparison.Ordinal) == 0 ||
-												LocalValueSpan.CompareTo("UDELEGATE", StringComparison.Ordinal) == 0 ||
-												LocalValueSpan.CompareTo("UFUNCTION", StringComparison.Ordinal) == 0)
+ResetStateMachineCheck:
+											if (localValueSpan.CompareTo("UPROPERTY", StringComparison.Ordinal) == 0 ||
+												localValueSpan.CompareTo("UCLASS", StringComparison.Ordinal) == 0 ||
+												localValueSpan.CompareTo("USTRUCT", StringComparison.Ordinal) == 0 ||
+												localValueSpan.CompareTo("UENUM", StringComparison.Ordinal) == 0 ||
+												localValueSpan.CompareTo("UINTERFACE", StringComparison.Ordinal) == 0 ||
+												localValueSpan.CompareTo("UDELEGATE", StringComparison.Ordinal) == 0 ||
+												localValueSpan.CompareTo("UFUNCTION", StringComparison.Ordinal) == 0)
 											{
-												this.LogError($"'{LocalValueSpan.ToString()}' must not be inside preprocessor blocks, except for WITH_EDITORONLY_DATA");
+												this.LogError($"'{localValueSpan.ToString()}' must not be inside preprocessor blocks, except for WITH_EDITORONLY_DATA");
 											}
-											else if (LocalValueSpan.CompareTo("void", StringComparison.Ordinal) == 0)
+											else if (localValueSpan.CompareTo("void", StringComparison.Ordinal) == 0)
 											{
-												CheckStateMachine = 1;
+												checkStateMachine = 1;
 											}
 											break;
 
 										case 1:
-											if (LocalValueSpan.CompareTo("Serialize", StringComparison.Ordinal) == 0)
+											if (localValueSpan.CompareTo("Serialize", StringComparison.Ordinal) == 0)
 											{
-												CheckStateMachine = 2;
+												checkStateMachine = 2;
 											}
 											else
 											{
-												CheckStateMachine = 0;
+												checkStateMachine = 0;
 												goto ResetStateMachineCheck;
 											}
 											break;
 
 										case 2:
-											if (LocalValueSpan.CompareTo("(", StringComparison.Ordinal) == 0)
+											if (localValueSpan.CompareTo("(", StringComparison.Ordinal) == 0)
 											{
-												CheckStateMachine = 3;
+												checkStateMachine = 3;
 											}
 											else
 											{
-												CheckStateMachine = 0;
+												checkStateMachine = 0;
 												goto ResetStateMachineCheck;
 											}
 											break;
 
 										case 3:
-											if (LocalValueSpan.CompareTo("FArchive", StringComparison.Ordinal) == 0 || 
-												LocalValueSpan.CompareTo("FStructuredArchiveRecord", StringComparison.Ordinal) == 0)
+											if (localValueSpan.CompareTo("FArchive", StringComparison.Ordinal) == 0 ||
+												localValueSpan.CompareTo("FStructuredArchiveRecord", StringComparison.Ordinal) == 0)
 											{
 												this.LogError($"Engine serialization functions must not be inside preprocessor blocks, except for WITH_EDITORONLY_DATA");
-												CheckStateMachine = 0;
+												checkStateMachine = 0;
 											}
-											else if (LocalValueSpan.CompareTo("FStructuredArchive", StringComparison.Ordinal) == 0)
+											else if (localValueSpan.CompareTo("FStructuredArchive", StringComparison.Ordinal) == 0)
 											{
-												CheckStateMachine = 4;
+												checkStateMachine = 4;
 											}
 											else
 											{
-												CheckStateMachine = 0;
+												checkStateMachine = 0;
 												goto ResetStateMachineCheck;
 											}
 											break;
 
 										case 4:
-											if (LocalValueSpan.CompareTo("::", StringComparison.Ordinal) == 0)
+											if (localValueSpan.CompareTo("::", StringComparison.Ordinal) == 0)
 											{
-												CheckStateMachine = 5;
+												checkStateMachine = 5;
 											}
 											else
 											{
-												CheckStateMachine = 0;
+												checkStateMachine = 0;
 												goto ResetStateMachineCheck;
 											}
 											break;
 
 										case 5:
-											if (LocalValueSpan.CompareTo("FRecord", StringComparison.Ordinal) == 0)
+											if (localValueSpan.CompareTo("FRecord", StringComparison.Ordinal) == 0)
 											{
 												this.LogError($"Engine serialization functions must not be inside preprocessor blocks, except for WITH_EDITORONLY_DATA");
-												CheckStateMachine = 0;
+												checkStateMachine = 0;
 											}
 											else
 											{
-												CheckStateMachine = 0;
+												checkStateMachine = 0;
 												goto ResetStateMachineCheck;
 											}
 											break;
 									}
 								}
 							}
-							--this.CommentsDisableCount;
+							--this._commentsDisableCount;
 
 							// Clear any extra pending comments at this time
 							ClearPendingComments();
 						}
-						this.PreprocessorPendingCommentsCount = 0;
+						this._preprocessorPendingCommentsCount = 0;
 
 						// Depending on the type of directive and/or the #if expression, we clear comments.
-						if (bClearComments)
+						if (clearComments)
 						{
 							ClearPendingComments();
-							bGotInlineComment = false;
+							gotInlineComment = false;
 						}
 						goto Restart;
 					}
 				}
-				TokenType = UhtTokenType.Symbol;
+				tokenType = UhtTokenType.Symbol;
 			}
 
-			UhtToken Out = new UhtToken(TokenType, StartPos, StartLine, StartPos, this.InputLineInternal, new StringView(this.Data.Slice(StartPos, this.InputPosInternal - StartPos)));
-			return Out;
+			return new UhtToken(tokenType, startPos, startLine, startPos, this._inputLineInternal, new StringView(this._data.Slice(startPos, this._inputPosInternal - startPos)));
 		}
 
 		/// <summary>
 		/// Skip all leading whitespace and collect any comments
 		/// </summary>
-		private void SkipWhitespaceAndCommentsInternal(ref bool bPersistGotInlineComment, bool bEnablePreprocessor)
+		private void SkipWhitespaceAndCommentsInternal(ref bool persistGotInlineComment, bool enablePreprocessor)
 		{
-			ReadOnlySpan<char> Span = this.Data.Span;
+			ReadOnlySpan<char> span = this._data.Span;
 
-			bool bGotNewlineBetweenComments = false;
-			bool bGotInlineComment = bPersistGotInlineComment;
+			bool gotNewlineBetweenComments = false;
+			bool gotInlineComment = persistGotInlineComment;
 			for (; ; )
 			{
-				uint Char = InternalGetChar(Span);
-				if (Char == 0)
+				uint c = InternalGetChar(span);
+				if (c == 0)
 				{
 					break;
 				}
-				else if (Char == '\n')
+				else if (c == '\n')
 				{
-					bGotNewlineBetweenComments |= bGotInlineComment;
-					++this.InputLineInternal;
+					gotNewlineBetweenComments |= gotInlineComment;
+					++this._inputLineInternal;
 				}
-				else if (Char == '\r' || Char == '\t' || Char == ' ')
+				else if (c == '\r' || c == '\t' || c == ' ')
 				{
 				}
-				else if (Char == '/')
+				else if (c == '/')
 				{
-					uint NextChar = InternalPeekChar(Span);
-					if (NextChar == '*')
+					uint nextChar = InternalPeekChar(span);
+					if (nextChar == '*')
 					{
-						if (bEnablePreprocessor)
+						if (enablePreprocessor)
 						{
 							ClearAllComments();
 						}
-						int CommentStart = this.InputPosInternal - 1;
-						++this.InputPosInternal;
+						int commentStart = this._inputPosInternal - 1;
+						++this._inputPosInternal;
 						for (; ; )
 						{
-							char CommentChar = InternalGetChar(Span);
-							if (CommentChar == 0)
+							char commentChar = InternalGetChar(span);
+							if (commentChar == 0)
 							{
-								if (bEnablePreprocessor)
+								if (enablePreprocessor)
 								{
 									ClearAllComments();
 								}
 								throw new UhtException(this, "End of header encountered inside comment");
 							}
-							else if (CommentChar == '\n')
+							else if (commentChar == '\n')
 							{
-								++this.InputLineInternal;
+								++this._inputLineInternal;
 							}
-							else if (CommentChar == '*' && InternalPeekChar(Span) == '/')
+							else if (commentChar == '*' && InternalPeekChar(span) == '/')
 							{
-								++this.InputPosInternal;
+								++this._inputPosInternal;
 								break;
 							}
 						}
-						if (bEnablePreprocessor)
+						if (enablePreprocessor)
 						{
-							AddComment(new StringView(this.Data.Slice(CommentStart, this.InputPosInternal - CommentStart)));
+							AddComment(new StringView(this._data.Slice(commentStart, this._inputPosInternal - commentStart)));
 						}
 					}
-					else if (NextChar == '/')
+					else if (nextChar == '/')
 					{
-						if (bGotNewlineBetweenComments)
+						if (gotNewlineBetweenComments)
 						{
-							bGotNewlineBetweenComments = false;
-							if (bEnablePreprocessor)
+							gotNewlineBetweenComments = false;
+							if (enablePreprocessor)
 							{
 								ClearAllComments();
 							}
 						}
-						bGotInlineComment = true;
-						int CommentStart = this.InputPosInternal - 1;
-						++this.InputPosInternal;
+						gotInlineComment = true;
+						int commentStart = this._inputPosInternal - 1;
+						++this._inputPosInternal;
 
 						// Scan to the end of the line
 						for (; ; )
 						{
-							char CommentChar = InternalGetChar(Span);
-							if (CommentChar == 0)
+							char commentChar = InternalGetChar(span);
+							if (commentChar == 0)
 							{
 								//--Pos;
 								break;
 							}
-							if (CommentChar == '\r')
+							if (commentChar == '\r')
 							{
 							}
-							else if (CommentChar == '\n')
+							else if (commentChar == '\n')
 							{
-								++InputLineInternal;
+								++_inputLineInternal;
 								break;
 							}
 						}
-						if (bEnablePreprocessor)
+						if (enablePreprocessor)
 						{
-							AddComment(new StringView(this.Data.Slice(CommentStart, this.InputPosInternal - CommentStart)));
+							AddComment(new StringView(this._data.Slice(commentStart, this._inputPosInternal - commentStart)));
 						}
 					}
 					else
 					{
-						--this.InputPosInternal;
+						--this._inputPosInternal;
 						break;
 					}
 				}
 				else
 				{
-					--this.InputPosInternal;
+					--this._inputPosInternal;
 					break;
 				}
 			}
-			bPersistGotInlineComment = bGotInlineComment;
+			persistGotInlineComment = gotInlineComment;
 			return;
 		}
 
-		private bool IsNumeric(ReadOnlySpan<char> Span, char C)
+		private bool IsNumeric(ReadOnlySpan<char> span, char c)
 		{
 			// Check for [0..9]
-			if (UhtFCString.IsDigit(C))
+			if (UhtFCString.IsDigit(c))
 			{
 				return true;
 			}
 
 			// Check for [+-]...
-			if (UhtFCString.IsSign(C))
+			if (UhtFCString.IsSign(c))
 			{
-				if (this.InputPosInternal == Span.Length)
+				if (this._inputPosInternal == span.Length)
 				{
 					return false;
 				}
 
 				// Check for [+-][0..9]...
-				if (UhtFCString.IsDigit(Span[this.InputPosInternal]))
+				if (UhtFCString.IsDigit(span[this._inputPosInternal]))
 				{
 					return true;
 				}
 
 				// Check for [+-][.][0..9]...
-				if (Span[this.InputPosInternal] != '.')
+				if (span[this._inputPosInternal] != '.')
 				{
 					return false;
 				}
 
-				if (this.InputPosInternal + 1 == Span.Length)
+				if (this._inputPosInternal + 1 == span.Length)
 				{
 					return false;
 				}
 
-				if (UhtFCString.IsDigit(Span[this.InputPosInternal + 1]))
+				if (UhtFCString.IsDigit(span[this._inputPosInternal + 1]))
 				{
 					return true;
 				}
 				return false;
 			}
 
-			if (C == '.')
+			if (c == '.')
 			{
-				if (this.InputPosInternal == Span.Length)
+				if (this._inputPosInternal == span.Length)
 				{
 					return false;
 				}
 
 				// Check for [.][0..9]...
-				if (UhtFCString.IsDigit(Span[this.InputPosInternal]))
+				if (UhtFCString.IsDigit(span[this._inputPosInternal]))
 				{
 					return true;
 				}
@@ -1016,22 +1003,22 @@ namespace EpicGames.UHT.Tokenizer
 		/// Fetch the next character in the input stream or zero if we have reached the end.
 		/// The current offset in the buffer is not advanced.  The method does not support UTF-8
 		/// </summary>
-		/// <param name="Span">The span containing the data</param>
+		/// <param name="span">The span containing the data</param>
 		/// <returns>Next character in the stream or zero</returns>
-		private char InternalPeekChar(ReadOnlySpan<char> Span)
+		private char InternalPeekChar(ReadOnlySpan<char> span)
 		{
-			return this.InputPosInternal < Span.Length ? Span[this.InputPosInternal] : '\0';
+			return this._inputPosInternal < span.Length ? span[this._inputPosInternal] : '\0';
 		}
 
 		/// <summary>
 		/// Fetch the next character in the input stream or zero if we have reached the end.
 		/// The current offset in the buffer is advanced.  The method does not support UTF-8.
 		/// </summary>
-		/// <param name="Span">The span containing the data</param>
+		/// <param name="span">The span containing the data</param>
 		/// <returns>Next character in the stream or zero</returns>
-		private char InternalGetChar(ReadOnlySpan<char> Span)
+		private char InternalGetChar(ReadOnlySpan<char> span)
 		{
-			return this.InputPosInternal < Span.Length ? Span[this.InputPosInternal++] : '\0';
+			return this._inputPosInternal < span.Length ? span[this._inputPosInternal++] : '\0';
 		}
 
 		/// <summary>
@@ -1039,25 +1026,25 @@ namespace EpicGames.UHT.Tokenizer
 		/// </summary>
 		private void ClearToken()
 		{
-			if (this.bHasToken)
+			if (this._hasToken)
 			{
-				this.bHasToken = false;
-				if (this.CommentsInternal != null && this.CommentsInternal.Count > this.CommittedComments)
+				this._hasToken = false;
+				if (this._commentsInternal != null && this._commentsInternal.Count > this._committedComments)
 				{
-					this.CommentsInternal.RemoveRange(this.CommittedComments, this.CommentsInternal.Count - this.CommittedComments);
+					this._commentsInternal.RemoveRange(this._committedComments, this._commentsInternal.Count - this._committedComments);
 				}
-				this.InputPosInternal = this.CurrentToken.UngetPos;
-				this.InputLineInternal = this.CurrentToken.UngetLine;
+				this._inputPosInternal = this._currentToken.UngetPos;
+				this._inputLineInternal = this._currentToken.UngetLine;
 			}
 		}
 
-		private void AddComment(StringView Comment)
+		private void AddComment(StringView comment)
 		{
-			if (this.CommentsInternal == null)
+			if (this._commentsInternal == null)
 			{
-				this.CommentsInternal = new List<StringView>(4);
+				this._commentsInternal = new List<StringView>(4);
 			}
-			this.CommentsInternal.Add(Comment);
+			this._commentsInternal.Add(comment);
 		}
 		#endregion
 	}
diff --git a/Engine/Source/Programs/Shared/EpicGames.UHT/Tokenizer/UhtTokenReader.cs b/Engine/Source/Programs/Shared/EpicGames.UHT/Tokenizer/UhtTokenReader.cs
--- a/Engine/Source/Programs/Shared/EpicGames.UHT/Tokenizer/UhtTokenReader.cs
+++ b/Engine/Source/Programs/Shared/EpicGames.UHT/Tokenizer/UhtTokenReader.cs
@@ -1,11 +1,11 @@
 // Copyright Epic Games, Inc. All Rights Reserved.
 
-using EpicGames.Core;
-using EpicGames.UHT.Utils;
 using System;
 using System.Collections.Generic;
 using System.Text;
 using System.Threading;
+using EpicGames.Core;
+using EpicGames.UHT.Utils;
 
 namespace EpicGames.UHT.Tokenizer
 {
@@ -43,35 +43,35 @@ namespace EpicGames.UHT.Tokenizer
 		/// <summary>
 		/// Test to see if any of the specified flags are set
 		/// </summary>
-		/// <param name="InFlags">Current flags</param>
-		/// <param name="TestFlags">Flags to test for</param>
+		/// <param name="inFlags">Current flags</param>
+		/// <param name="testFlags">Flags to test for</param>
 		/// <returns>True if any of the flags are set</returns>
-		public static bool HasAnyFlags(this UhtRawStringOptions InFlags, UhtRawStringOptions TestFlags)
+		public static bool HasAnyFlags(this UhtRawStringOptions inFlags, UhtRawStringOptions testFlags)
 		{
-			return (InFlags & TestFlags) != 0;
+			return (inFlags & testFlags) != 0;
 		}
 
 		/// <summary>
 		/// Test to see if all of the specified flags are set
 		/// </summary>
-		/// <param name="InFlags">Current flags</param>
-		/// <param name="TestFlags">Flags to test for</param>
+		/// <param name="inFlags">Current flags</param>
+		/// <param name="testFlags">Flags to test for</param>
 		/// <returns>True if all the flags are set</returns>
-		public static bool HasAllFlags(this UhtRawStringOptions InFlags, UhtRawStringOptions TestFlags)
+		public static bool HasAllFlags(this UhtRawStringOptions inFlags, UhtRawStringOptions testFlags)
 		{
-			return (InFlags & TestFlags) == TestFlags;
+			return (inFlags & testFlags) == testFlags;
 		}
 
 		/// <summary>
 		/// Test to see if a specific set of flags have a specific value.
 		/// </summary>
-		/// <param name="InFlags">Current flags</param>
-		/// <param name="TestFlags">Flags to test for</param>
-		/// <param name="MatchFlags">Expected value of the tested flags</param>
+		/// <param name="inFlags">Current flags</param>
+		/// <param name="testFlags">Flags to test for</param>
+		/// <param name="matchFlags">Expected value of the tested flags</param>
 		/// <returns>True if the given flags have a specific value.</returns>
-		public static bool HasExactFlags(this UhtRawStringOptions InFlags, UhtRawStringOptions TestFlags, UhtRawStringOptions MatchFlags)
+		public static bool HasExactFlags(this UhtRawStringOptions inFlags, UhtRawStringOptions testFlags, UhtRawStringOptions matchFlags)
 		{
-			return (InFlags & TestFlags) == MatchFlags;
+			return (inFlags & testFlags) == matchFlags;
 		}
 	}
 
@@ -84,12 +84,12 @@ namespace EpicGames.UHT.Tokenizer
 		/// <summary>
 		/// Parse a preprocessor directive
 		/// </summary>
-		/// <param name="Token">Token starting the directive.  Will be only the '#'</param>
-		/// <param name="bIsBeingIncluded">If true, the directive the source is being included.  Otherwise it is being skipped as part of an #if block</param>
-		/// <param name="bClearComments">If true, comments should be cleared</param>
-		/// <param name="bIllegalContentsCheck">If true, excluded contents should be checked for unparsed UE macros (i.e. UCLASS) </param>
+		/// <param name="token">Token starting the directive.  Will be only the '#'</param>
+		/// <param name="isBeingIncluded">If true, the directive the source is being included.  Otherwise it is being skipped as part of an #if block</param>
+		/// <param name="clearComments">If true, comments should be cleared</param>
+		/// <param name="illegalContentsCheck">If true, excluded contents should be checked for unparsed UE macros (i.e. UCLASS) </param>
 		/// <returns>True if the source should continue to be included</returns>
-		public bool ParsePreprocessorDirective(ref UhtToken Token, bool bIsBeingIncluded, out bool bClearComments, out bool bIllegalContentsCheck);
+		public bool ParsePreprocessorDirective(ref UhtToken token, bool isBeingIncluded, out bool clearComments, out bool illegalContentsCheck);
 
 		/// <summary>
 		/// Save the current preprocessor state
@@ -110,7 +110,7 @@ namespace EpicGames.UHT.Tokenizer
 		/// <summary>
 		/// True if the reader is at the end of the stream
 		/// </summary>
-		public bool bIsEOF { get; }
+		public bool IsEOF { get; }
 
 		/// <summary>
 		/// Current input position in the stream by characters
@@ -148,9 +148,9 @@ namespace EpicGames.UHT.Tokenizer
 		/// <summary>
 		/// Tests to see if the given token is the first token of a line
 		/// </summary>
-		/// <param name="Token">The token to test</param>
+		/// <param name="token">The token to test</param>
 		/// <returns>True if the token is the first token on the line</returns>
-		public bool IsFirstTokenInLine(ref UhtToken Token);
+		public bool IsFirstTokenInLine(ref UhtToken token);
 
 		/// <summary>
 		/// Skip any whitespace and comments at the current buffer position
@@ -166,18 +166,18 @@ namespace EpicGames.UHT.Tokenizer
 		/// <summary>
 		/// Get a view of the buffer being read
 		/// </summary>
-		/// <param name="StartPos">Starting character offset in the buffer.</param>
-		/// <param name="Count">Length of the span</param>
+		/// <param name="startPos">Starting character offset in the buffer.</param>
+		/// <param name="count">Length of the span</param>
 		/// <returns>The string view into the buffer</returns>
-		public StringView GetStringView(int StartPos, int Count);
+		public StringView GetStringView(int startPos, int count);
 
 		/// <summary>
 		/// Return a string terminated by the given character.
 		/// </summary>
-		/// <param name="Terminator">The character to stop at.</param>
-		/// <param name="Options">Options</param>
+		/// <param name="terminator">The character to stop at.</param>
+		/// <param name="options">Options</param>
 		/// <returns>The parsed string</returns>
-		public StringView GetRawString(char Terminator, UhtRawStringOptions Options);
+		public StringView GetRawString(char terminator, UhtRawStringOptions options);
 
 		/// <summary>
 		/// The current collection of parsed comments.  This does not include any comments parsed as part of a 
@@ -237,8 +237,8 @@ namespace EpicGames.UHT.Tokenizer
 		/// <summary>
 		/// Record the given token to the list of recorded tokens
 		/// </summary>
-		/// <param name="Token">Token to record</param>
-		public void RecordToken(ref UhtToken Token);
+		/// <param name="token">Token to record</param>
+		public void RecordToken(ref UhtToken token);
 
 		/// <summary>
 		/// Get the current collection of recorded tokens
@@ -265,35 +265,35 @@ namespace EpicGames.UHT.Tokenizer
 		/// <summary>
 		/// Join the tokens in the list
 		/// </summary>
-		/// <param name="Builder">Destination builder</param>
-		/// <param name="Seperator">Separator between the tokens</param>
+		/// <param name="builder">Destination builder</param>
+		/// <param name="seperator">Separator between the tokens</param>
 		/// <returns></returns>
-		public void Join(StringBuilder Builder, string Seperator = "")
+		public void Join(StringBuilder builder, string seperator = "")
 		{
-			Builder.Append(this.Token.Value.ToString());
-			UhtTokenList List = this;
-			while (List.Next != null)
+			builder.Append(this.Token.Value.ToString());
+			UhtTokenList list = this;
+			while (list.Next != null)
 			{
-				List = List.Next;
-				Builder.Append(Seperator);
-				Builder.Append(List.Token.Value.ToString());
+				list = list.Next;
+				builder.Append(seperator);
+				builder.Append(list.Token.Value.ToString());
 			}
 		}
 
 		/// <summary>
 		/// Join the tokens in the list
 		/// </summary>
-		/// <param name="Seperator">Separator between the tokens</param>
+		/// <param name="seperator">Separator between the tokens</param>
 		/// <returns></returns>
-		public string Join(string Seperator = "")
+		public string Join(string seperator = "")
 		{
 			if (this.Next == null)
 			{
 				return this.Token.Value.ToString();
 			}
-			StringBuilder Builder = new StringBuilder();
-			Join(Builder, Seperator);
-			return Builder.ToString();
+			StringBuilder builder = new StringBuilder();
+			Join(builder, seperator);
+			return builder.ToString();
 		}
 
 		/// <summary>
@@ -302,20 +302,20 @@ namespace EpicGames.UHT.Tokenizer
 		/// <returns></returns>
 		public UhtToken[] ToArray()
 		{
-			int Count = 1;
-			for (UhtTokenList Temp = this; Temp.Next != null; Temp = Temp.Next)
+			int count = 1;
+			for (UhtTokenList temp = this; temp.Next != null; temp = temp.Next)
 			{
-				++Count;
+				++count;
 			}
-			UhtToken[] Out = new UhtToken[Count];
-			Out[0] = this.Token;
-			Count = 1;
-			for (UhtTokenList Temp = this; Temp.Next != null; Temp = Temp.Next)
+			UhtToken[] outTokens = new UhtToken[count];
+			outTokens[0] = this.Token;
+			count = 1;
+			for (UhtTokenList temp = this; temp.Next != null; temp = temp.Next)
 			{
-				Out[Count] = Temp.Next.Token;
-				++Count;
+				outTokens[count] = temp.Next.Token;
+				++count;
 			}
-			return Out;
+			return outTokens;
 		}
 	}
 
@@ -324,101 +324,101 @@ namespace EpicGames.UHT.Tokenizer
 	/// </summary>
 	public static class UhtTokenListCache
 	{
-		private static readonly ThreadLocal<UhtTokenList?> Tls = new ThreadLocal<UhtTokenList?>(() => null);
+		private static readonly ThreadLocal<UhtTokenList?> s_tls = new ThreadLocal<UhtTokenList?>(() => null);
 
 		/// <summary>
 		/// Borrow a token list
 		/// </summary>
-		/// <param name="Token">Starting token</param>
+		/// <param name="token">Starting token</param>
 		/// <returns>Token list</returns>
-		public static UhtTokenList Borrow(UhtToken Token)
+		public static UhtTokenList Borrow(UhtToken token)
 		{
-			UhtTokenList? Identifier = Tls.Value;
-			if (Identifier != null)
+			UhtTokenList? identifier = s_tls.Value;
+			if (identifier != null)
 			{
-				Tls.Value = Identifier.Next;
+				s_tls.Value = identifier.Next;
 			}
 			else
 			{
-				Identifier = new UhtTokenList();
+				identifier = new UhtTokenList();
 			}
-			Identifier.Token = Token;
-			Identifier.Next = null;
-			return Identifier;
+			identifier.Token = token;
+			identifier.Next = null;
+			return identifier;
 		}
 
 		/// <summary>
 		/// Return a token list to the cache
 		/// </summary>
-		/// <param name="Identifier"></param>
-		public static void Return(UhtTokenList Identifier)
+		/// <param name="identifier"></param>
+		public static void Return(UhtTokenList identifier)
 		{
-			UhtTokenList? Tail = Tls.Value;
-			if (Tail != null)
+			UhtTokenList? tail = s_tls.Value;
+			if (tail != null)
 			{
-				Tail.Next = Identifier;
+				tail.Next = identifier;
 			}
 
-			for (; Identifier.Next != null; Identifier = Identifier.Next)
+			for (; identifier.Next != null; identifier = identifier.Next)
 			{
 			}
 
-			Tls.Value = Identifier;
+			s_tls.Value = identifier;
 		}
 	}
 
 	/// <summary>
 	/// Delegate for when a token is parsed
 	/// </summary>
-	/// <param name="Token">The token in question</param>
-	public delegate void UhtTokenDelegate(ref UhtToken Token);
+	/// <param name="token">The token in question</param>
+	public delegate void UhtTokenDelegate(ref UhtToken token);
 
 	/// <summary>
 	/// Delegate for when a token is parsed in an until block
 	/// </summary>
-	/// <param name="Token">The token in question</param>
+	/// <param name="token">The token in question</param>
 	/// <returns>True if parsing should continue</returns>
-	public delegate bool UhtTokensUntilDelegate(ref UhtToken Token);
+	public delegate bool UhtTokensUntilDelegate(ref UhtToken token);
 
 	/// <summary>
 	/// Delegate for an enumeration of tokens
 	/// </summary>
-	/// <param name="Tokens">Parsed tokens</param>
-	public delegate void UhtTokensDelegate(IEnumerable<UhtToken> Tokens);
+	/// <param name="tokens">Parsed tokens</param>
+	public delegate void UhtTokensDelegate(IEnumerable<UhtToken> tokens);
 
 	/// <summary>
 	/// Delegate for cached token list
 	/// </summary>
-	/// <param name="TokenList">Token list that can be cached</param>
-	public delegate void UhtTokenListDelegate(UhtTokenList TokenList);
+	/// <param name="tokenList">Token list that can be cached</param>
+	public delegate void UhtTokenListDelegate(UhtTokenList tokenList);
 
 	/// <summary>
 	/// Delegate for a constant float
 	/// </summary>
-	/// <param name="Value">Value in question</param>
-	public delegate void UhtTokenConstFloatDelegate(float Value);
+	/// <param name="value">Value in question</param>
+	public delegate void UhtTokenConstFloatDelegate(float value);
 
 	/// <summary>
 	/// Delegate for a constant double
 	/// </summary>
-	/// <param name="Value">Value in question</param>
-	public delegate void UhtTokenConstDoubleDelegate(double Value);
+	/// <param name="value">Value in question</param>
+	public delegate void UhtTokenConstDoubleDelegate(double value);
 
 	/// <summary>
 	/// Helper struct to disable comment parsing.  Should be used in a using block
 	/// </summary>
 	public struct UhtTokenDisableComments : IDisposable
 	{
-		private readonly IUhtTokenReader TokenReader;
+		private readonly IUhtTokenReader _tokenReader;
 
 		/// <summary>
 		/// Construct instance
 		/// </summary>
-		/// <param name="TokenReader">Token reader to disable</param>
-		public UhtTokenDisableComments(IUhtTokenReader TokenReader)
+		/// <param name="tokenReader">Token reader to disable</param>
+		public UhtTokenDisableComments(IUhtTokenReader tokenReader)
 		{
-			this.TokenReader = TokenReader;
-			this.TokenReader.DisableComments();
+			this._tokenReader = tokenReader;
+			this._tokenReader.DisableComments();
 		}
 
 		/// <summary>
@@ -426,7 +426,7 @@ namespace EpicGames.UHT.Tokenizer
 		/// </summary>
 		public void Dispose()
 		{
-			this.TokenReader.EnableComments();
+			this._tokenReader.EnableComments();
 		}
 	}
 
@@ -435,18 +435,18 @@ namespace EpicGames.UHT.Tokenizer
 	/// </summary>
 	public struct UhtTokenSaveState : IDisposable
 	{
-		private readonly IUhtTokenReader TokenReader;
-		private bool bHandled;
+		private readonly IUhtTokenReader _tokenReader;
+		private bool _handled;
 
 		/// <summary>
 		/// Construct instance
 		/// </summary>
-		/// <param name="TokenReader">Token reader</param>
-		public UhtTokenSaveState(IUhtTokenReader TokenReader)
+		/// <param name="tokenReader">Token reader</param>
+		public UhtTokenSaveState(IUhtTokenReader tokenReader)
 		{
-			this.TokenReader = TokenReader;
-			this.bHandled = false;
-			this.TokenReader.SaveState();
+			this._tokenReader = tokenReader;
+			this._handled = false;
+			this._tokenReader.SaveState();
 		}
 
 		/// <summary>
@@ -454,7 +454,7 @@ namespace EpicGames.UHT.Tokenizer
 		/// </summary>
 		public void Dispose()
 		{
-			if (!this.bHandled)
+			if (!this._handled)
 			{
 				RestoreState();
 			}
@@ -466,12 +466,12 @@ namespace EpicGames.UHT.Tokenizer
 		/// <exception cref="UhtIceException">Thrown if state has already been restored or aborted</exception>
 		public void RestoreState()
 		{
-			if (this.bHandled)
+			if (this._handled)
 			{
 				throw new UhtIceException("Can not call RestoreState/AbandonState more than once");
 			}
-			this.TokenReader.RestoreState();
-			this.bHandled = true;
+			this._tokenReader.RestoreState();
+			this._handled = true;
 		}
 
 		/// <summary>
@@ -480,12 +480,12 @@ namespace EpicGames.UHT.Tokenizer
 		/// <exception cref="UhtIceException">Thrown if state has already been restored or aborted</exception>
 		public void AbandonState()
 		{
-			if (this.bHandled)
+			if (this._handled)
 			{
 				throw new UhtIceException("Can not call RestoreState/AbandonState more than once");
 			}
-			this.TokenReader.AbandonState();
-			this.bHandled = true;
+			this._tokenReader.AbandonState();
+			this._handled = true;
 		}
 	}
 }
diff --git a/Engine/Source/Programs/Shared/EpicGames.UHT/Tokenizer/UhtTokenReaderFloatExtensions.cs b/Engine/Source/Programs/Shared/EpicGames.UHT/Tokenizer/UhtTokenReaderFloatExtensions.cs
--- a/Engine/Source/Programs/Shared/EpicGames.UHT/Tokenizer/UhtTokenReaderFloatExtensions.cs
+++ b/Engine/Source/Programs/Shared/EpicGames.UHT/Tokenizer/UhtTokenReaderFloatExtensions.cs
@@ -13,171 +13,171 @@ namespace EpicGames.UHT.Tokenizer
 		/// <summary>
 		/// Get the next token as a float.  If the next token is not a float, no token is consumed.
 		/// </summary>
-		/// <param name="TokenReader">Token reader</param>
-		/// <param name="Value">The float value of the token</param>
+		/// <param name="tokenReader">Token reader</param>
+		/// <param name="value">The float value of the token</param>
 		/// <returns>True if the next token was an float, false if not.</returns>
-		public static bool TryOptionalConstFloat(this IUhtTokenReader TokenReader, out float Value)
+		public static bool TryOptionalConstFloat(this IUhtTokenReader tokenReader, out float value)
 		{
-			ref UhtToken Token = ref TokenReader.PeekToken();
-			if (Token.IsConstFloat() && Token.GetConstFloat(out Value)) // NOTE: This is restricted to only float values
+			ref UhtToken token = ref tokenReader.PeekToken();
+			if (token.IsConstFloat() && token.GetConstFloat(out value)) // NOTE: This is restricted to only float values
 			{
-				TokenReader.ConsumeToken();
+				tokenReader.ConsumeToken();
 				return true;
 			}
-			Value = 0;
+			value = 0;
 			return false;
 		}
 
 		/// <summary>
 		/// Get the next token as a float.  If the next token is not a float, no token is consumed.
 		/// </summary>
-		/// <param name="TokenReader"></param>
+		/// <param name="tokenReader"></param>
 		/// <returns>The token reader</returns>
-		public static IUhtTokenReader OptionalConstFloat(this IUhtTokenReader TokenReader)
+		public static IUhtTokenReader OptionalConstFloat(this IUhtTokenReader tokenReader)
 		{
-			TokenReader.TryOptionalConstFloat(out float _);
-			return TokenReader;
+			tokenReader.TryOptionalConstFloat(out float _);
+			return tokenReader;
 		}
 
 		/// <summary>
 		/// Get the next token as a float.  If the next token is not a float, no token is consumed.
 		/// </summary>
-		/// <param name="TokenReader">Token reader</param>
-		/// <param name="Delegate">Delegate to invoke with the float value</param>
+		/// <param name="tokenReader">Token reader</param>
+		/// <param name="floatDelegate">Delegate to invoke with the float value</param>
 		/// <returns>The token reader</returns>
-		public static IUhtTokenReader OptionalConstFloat(this IUhtTokenReader TokenReader, UhtTokenConstFloatDelegate Delegate)
+		public static IUhtTokenReader OptionalConstFloat(this IUhtTokenReader tokenReader, UhtTokenConstFloatDelegate floatDelegate)
 		{
-			if (TokenReader.TryOptionalConstFloat(out float Value))
+			if (tokenReader.TryOptionalConstFloat(out float value))
 			{
-				Delegate(Value);
+				floatDelegate(value);
 			}
-			return TokenReader;
+			return tokenReader;
 		}
 
 		/// <summary>
 		/// Get the next token as a float.  If the next token is not a float, an exception is thrown
 		/// </summary>
-		/// <param name="TokenReader">Token reader</param>
-		/// <param name="ExceptionContext">If not null, an exception will be thrown with the given text as part of the message.</param>
+		/// <param name="tokenReader">Token reader</param>
+		/// <param name="exceptionContext">If not null, an exception will be thrown with the given text as part of the message.</param>
 		/// <returns>True if the next token was an float, false if not.</returns>
-		public static IUhtTokenReader RequireConstFloat(this IUhtTokenReader TokenReader, object? ExceptionContext = null)
+		public static IUhtTokenReader RequireConstFloat(this IUhtTokenReader tokenReader, object? exceptionContext = null)
 		{
-			if (!TokenReader.TryOptionalConstFloat(out float _))
+			if (!tokenReader.TryOptionalConstFloat(out float _))
 			{
-				throw new UhtTokenException(TokenReader, TokenReader.PeekToken(), "constant float", ExceptionContext);
+				throw new UhtTokenException(tokenReader, tokenReader.PeekToken(), "constant float", exceptionContext);
 			}
-			return TokenReader;
+			return tokenReader;
 		}
 
 		/// <summary>
 		/// Get the next token as a float.  If the next token is not a float, an exception is thrown
 		/// </summary>
-		/// <param name="TokenReader">Token reader</param>
-		/// <param name="ExceptionContext">If not null, an exception will be thrown with the given text as part of the message.</param>
+		/// <param name="tokenReader">Token reader</param>
+		/// <param name="exceptionContext">If not null, an exception will be thrown with the given text as part of the message.</param>
 		/// <returns>The floating point value of the token</returns>
-		public static float GetConstFloat(this IUhtTokenReader TokenReader, object? ExceptionContext = null)
+		public static float GetConstFloat(this IUhtTokenReader tokenReader, object? exceptionContext = null)
 		{
-			float Value;
-			if (!TokenReader.TryOptionalConstFloat(out Value))
+			float value;
+			if (!tokenReader.TryOptionalConstFloat(out value))
 			{
-				throw new UhtTokenException(TokenReader, TokenReader.PeekToken(), "constant float", ExceptionContext);
+				throw new UhtTokenException(tokenReader, tokenReader.PeekToken(), "constant float", exceptionContext);
 			}
-			return Value;
+			return value;
 		}
 
 		/// <summary>
 		/// Get the next float.  It also handles [+/-] token followed by an float.
 		/// </summary>
-		/// <param name="TokenReader">Source tokens</param>
-		/// <param name="Value">The float value of the token</param>
+		/// <param name="tokenReader">Source tokens</param>
+		/// <param name="value">The float value of the token</param>
 		/// <returns>True if the next token was an float, false if not.</returns>
-		public static bool TryOptionalLeadingSignConstFloat(this IUhtTokenReader TokenReader, out float Value)
+		public static bool TryOptionalLeadingSignConstFloat(this IUhtTokenReader tokenReader, out float value)
 		{
-			float LocalValue = 0;
-			bool Results = TokenReader.TryOptionalLeadingSignConstNumeric((ref UhtToken Token) =>
+			float localValue = 0;
+			bool results = tokenReader.TryOptionalLeadingSignConstNumeric((ref UhtToken token) =>
 			{
-				return (Token.IsConstInt() || Token.IsConstFloat()) && Token.GetConstFloat(out LocalValue);
+				return (token.IsConstInt() || token.IsConstFloat()) && token.GetConstFloat(out localValue);
 			});
-			Value = LocalValue;
-			return Results;
+			value = localValue;
+			return results;
 		}
 
 		/// <summary>
 		/// Get the next float.  It also handles [+/-] token followed by an float.
 		/// </summary>
-		/// <param name="TokenReader">Source tokens</param>
-		/// <param name="Value">The float value of the token</param>
+		/// <param name="tokenReader">Source tokens</param>
+		/// <param name="value">The float value of the token</param>
 		/// <returns>True if the next token was an float, false if not.</returns>
-		public static bool TryOptionalConstFloatExpression(this IUhtTokenReader TokenReader, out float Value)
+		public static bool TryOptionalConstFloatExpression(this IUhtTokenReader tokenReader, out float value)
 		{
-			float LocalValue = 0;
-			bool Results = TokenReader.TryOptionalLeadingSignConstNumeric((ref UhtToken Token) =>
+			float localValue = 0;
+			bool results = tokenReader.TryOptionalLeadingSignConstNumeric((ref UhtToken token) =>
 			{
-				return (Token.IsConstInt() || Token.IsConstFloat()) && Token.GetConstFloat(out LocalValue);
+				return (token.IsConstInt() || token.IsConstFloat()) && token.GetConstFloat(out localValue);
 			});
-			Value = LocalValue;
-			return Results;
+			value = localValue;
+			return results;
 		}
 
 		/// <summary>
 		/// Get the next float.  It also handles [+/-] token followed by an float.
 		/// </summary>
-		/// <param name="TokenReader">Source tokens</param>
+		/// <param name="tokenReader">Source tokens</param>
 		/// <returns>The double value</returns>
-		public static float GetConstFloatExpression(this IUhtTokenReader TokenReader)
+		public static float GetConstFloatExpression(this IUhtTokenReader tokenReader)
 		{
-			if (!TokenReader.TryOptionalConstFloatExpression(out float LocalValue))
+			if (!tokenReader.TryOptionalConstFloatExpression(out float localValue))
 			{
-				throw new UhtTokenException(TokenReader, TokenReader.PeekToken(), "constant float", null);
+				throw new UhtTokenException(tokenReader, tokenReader.PeekToken(), "constant float", null);
 			}
-			return LocalValue;
+			return localValue;
 		}
 
 		/// <summary>
 		/// Get the next double.  It also handles [+/-] token followed by an double.
 		/// </summary>
-		/// <param name="TokenReader">Source tokens</param>
-		/// <param name="Value">The double value of the token</param>
+		/// <param name="tokenReader">Source tokens</param>
+		/// <param name="value">The double value of the token</param>
 		/// <returns>True if the next token was an double, false if not.</returns>
-		public static bool TryOptionalConstDoubleExpression(this IUhtTokenReader TokenReader, out double Value)
+		public static bool TryOptionalConstDoubleExpression(this IUhtTokenReader tokenReader, out double value)
 		{
-			double LocalValue = 0;
-			bool Results = TokenReader.TryOptionalLeadingSignConstNumeric((ref UhtToken Token) =>
+			double localValue = 0;
+			bool results = tokenReader.TryOptionalLeadingSignConstNumeric((ref UhtToken token) =>
 			{
-				return (Token.IsConstInt() || Token.IsConstFloat()) && Token.GetConstDouble(out LocalValue);
+				return (token.IsConstInt() || token.IsConstFloat()) && token.GetConstDouble(out localValue);
 			});
-			Value = LocalValue;
-			return Results;
+			value = localValue;
+			return results;
 		}
 
 		/// <summary>
 		/// Get the next double.  It also handles [+/-] token followed by an double.
 		/// </summary>
-		/// <param name="TokenReader">Source tokens</param>
-		/// <param name="Delegate">Delegate to invoke if the double is parsed</param>
+		/// <param name="tokenReader">Source tokens</param>
+		/// <param name="doubleDelegate">Delegate to invoke if the double is parsed</param>
 		/// <returns>The supplied token reader</returns>
-		public static IUhtTokenReader RequireConstDoubleExpression(this IUhtTokenReader TokenReader, UhtTokenConstDoubleDelegate Delegate)
+		public static IUhtTokenReader RequireConstDoubleExpression(this IUhtTokenReader tokenReader, UhtTokenConstDoubleDelegate doubleDelegate)
 		{
-			if (!TokenReader.TryOptionalConstDoubleExpression(out double LocalValue))
+			if (!tokenReader.TryOptionalConstDoubleExpression(out double localValue))
 			{
-				throw new UhtTokenException(TokenReader, TokenReader.PeekToken(), "constant double", null);
+				throw new UhtTokenException(tokenReader, tokenReader.PeekToken(), "constant double", null);
 			}
-			Delegate(LocalValue);
-			return TokenReader;
+			doubleDelegate(localValue);
+			return tokenReader;
 		}
 
 		/// <summary>
 		/// Get the next double.  It also handles [+/-] token followed by an double.
 		/// </summary>
-		/// <param name="TokenReader">Source tokens</param>
+		/// <param name="tokenReader">Source tokens</param>
 		/// <returns>The double value</returns>
-		public static double GetConstDoubleExpression(this IUhtTokenReader TokenReader)
+		public static double GetConstDoubleExpression(this IUhtTokenReader tokenReader)
 		{
-			if (!TokenReader.TryOptionalConstDoubleExpression(out double LocalValue))
+			if (!tokenReader.TryOptionalConstDoubleExpression(out double localValue))
 			{
-				throw new UhtTokenException(TokenReader, TokenReader.PeekToken(), "constant double", null);
+				throw new UhtTokenException(tokenReader, tokenReader.PeekToken(), "constant double", null);
 			}
-			return LocalValue;
+			return localValue;
 		}
 	}
 }
