More coding standard fixes

diff --git a/Engine/Source/Programs/Shared/EpicGames.UHT/Tokenizer/UhtTokenReaderFlowControlExtensions.cs b/Engine/Source/Programs/Shared/EpicGames.UHT/Tokenizer/UhtTokenReaderFlowControlExtensions.cs
--- a/Engine/Source/Programs/Shared/EpicGames.UHT/Tokenizer/UhtTokenReaderFlowControlExtensions.cs
+++ b/Engine/Source/Programs/Shared/EpicGames.UHT/Tokenizer/UhtTokenReaderFlowControlExtensions.cs
@@ -1,8 +1,8 @@
 // Copyright Epic Games, Inc. All Rights Reserved.
 
-using EpicGames.UHT.Utils;
 using System;
 using System.Collections.Generic;
+using EpicGames.UHT.Utils;
 
 namespace EpicGames.UHT.Tokenizer
 {
@@ -16,18 +16,18 @@ namespace EpicGames.UHT.Tokenizer
 		/// <summary>
 		/// Parse an optional list
 		/// </summary>
-		/// <param name="TokenReader">Token reader</param>
-		/// <param name="Initiator">Initiating symbol</param>
-		/// <param name="Terminator">Terminating symbol</param>
-		/// <param name="Separator">Separator symbol</param>
-		/// <param name="bAllowTrailingSeparator">If true, allow trailing separators</param>
-		/// <param name="Action">Action to be invoked for each list element.</param>
+		/// <param name="tokenReader">Token reader</param>
+		/// <param name="initiator">Initiating symbol</param>
+		/// <param name="terminator">Terminating symbol</param>
+		/// <param name="separator">Separator symbol</param>
+		/// <param name="allowTrailingSeparator">If true, allow trailing separators</param>
+		/// <param name="action">Action to be invoked for each list element.</param>
 		/// <returns>True if a list was read</returns>
-		public static bool TryOptionalList(this IUhtTokenReader TokenReader, char Initiator, char Terminator, char Separator, bool bAllowTrailingSeparator, Action Action)
+		public static bool TryOptionalList(this IUhtTokenReader tokenReader, char initiator, char terminator, char separator, bool allowTrailingSeparator, Action action)
 		{
-			if (TokenReader.TryOptional(Initiator))
+			if (tokenReader.TryOptional(initiator))
 			{
-				TokenReader.RequireList(Terminator, Separator, bAllowTrailingSeparator, Action);
+				tokenReader.RequireList(terminator, separator, allowTrailingSeparator, action);
 				return true;
 			}
 			return false;
@@ -36,65 +36,65 @@ namespace EpicGames.UHT.Tokenizer
 		/// <summary>
 		/// Parse a required list
 		/// </summary>
-		/// <param name="TokenReader">Token reader</param>
-		/// <param name="Initiator">Initiating symbol</param>
-		/// <param name="Terminator">Terminating symbol</param>
-		/// <param name="Separator">Separator symbol</param>
-		/// <param name="bAllowTrailingSeparator">If true, allow trailing separators</param>
-		/// <param name="Action">Action to be invoked for each list element.</param>
+		/// <param name="tokenReader">Token reader</param>
+		/// <param name="initiator">Initiating symbol</param>
+		/// <param name="terminator">Terminating symbol</param>
+		/// <param name="separator">Separator symbol</param>
+		/// <param name="allowTrailingSeparator">If true, allow trailing separators</param>
+		/// <param name="action">Action to be invoked for each list element.</param>
 		/// <returns>Token reader</returns>
-		public static IUhtTokenReader RequireList(this IUhtTokenReader TokenReader, char Initiator, char Terminator, char Separator, bool bAllowTrailingSeparator, Action Action)
+		public static IUhtTokenReader RequireList(this IUhtTokenReader tokenReader, char initiator, char terminator, char separator, bool allowTrailingSeparator, Action action)
 		{
-			return TokenReader.RequireList(Initiator, Terminator, Separator, bAllowTrailingSeparator, null, Action);
+			return tokenReader.RequireList(initiator, terminator, separator, allowTrailingSeparator, null, action);
 		}
 
 		/// <summary>
 		/// Parse a required list
 		/// </summary>
-		/// <param name="TokenReader">Token reader</param>
-		/// <param name="Initiator">Initiating symbol</param>
-		/// <param name="Terminator">Terminating symbol</param>
-		/// <param name="Separator">Separator symbol</param>
-		/// <param name="bAllowTrailingSeparator">If true, allow trailing separators</param>
-		/// <param name="ExceptionContext">Extra context for error messages</param>
-		/// <param name="Action">Action to be invoked for each list element.</param>
+		/// <param name="tokenReader">Token reader</param>
+		/// <param name="initiator">Initiating symbol</param>
+		/// <param name="terminator">Terminating symbol</param>
+		/// <param name="separator">Separator symbol</param>
+		/// <param name="allowTrailingSeparator">If true, allow trailing separators</param>
+		/// <param name="exceptionContext">Extra context for error messages</param>
+		/// <param name="action">Action to be invoked for each list element.</param>
 		/// <returns>Token reader</returns>
-		public static IUhtTokenReader RequireList(this IUhtTokenReader TokenReader, char Initiator, char Terminator, char Separator, bool bAllowTrailingSeparator, object? ExceptionContext, Action Action)
+		public static IUhtTokenReader RequireList(this IUhtTokenReader tokenReader, char initiator, char terminator, char separator, bool allowTrailingSeparator, object? exceptionContext, Action action)
 		{
-			TokenReader.Require(Initiator);
-			return TokenReader.RequireList(Terminator, Separator, bAllowTrailingSeparator, ExceptionContext, Action);
+			tokenReader.Require(initiator);
+			return tokenReader.RequireList(terminator, separator, allowTrailingSeparator, exceptionContext, action);
 		}
 
 		/// <summary>
 		/// Parse a required list.  Initiating token must have already been parsed
 		/// </summary>
-		/// <param name="TokenReader">Token reader</param>
-		/// <param name="Terminator">Terminating symbol</param>
-		/// <param name="Separator">Separator symbol</param>
-		/// <param name="bAllowTrailingSeparator">If true, allow trailing separators</param>
-		/// <param name="Action">Action to be invoked for each list element.</param>
+		/// <param name="tokenReader">Token reader</param>
+		/// <param name="terminator">Terminating symbol</param>
+		/// <param name="separator">Separator symbol</param>
+		/// <param name="allowTrailingSeparator">If true, allow trailing separators</param>
+		/// <param name="action">Action to be invoked for each list element.</param>
 		/// <returns>Token reader</returns>
-		public static IUhtTokenReader RequireList(this IUhtTokenReader TokenReader, char Terminator, char Separator, bool bAllowTrailingSeparator, Action Action)
+		public static IUhtTokenReader RequireList(this IUhtTokenReader tokenReader, char terminator, char separator, bool allowTrailingSeparator, Action action)
 		{
-			return TokenReader.RequireList(Terminator, Separator, bAllowTrailingSeparator, null, Action);
+			return tokenReader.RequireList(terminator, separator, allowTrailingSeparator, null, action);
 		}
 
 		/// <summary>
 		/// Parse a required list.  Initiating token must have already been parsed
 		/// </summary>
-		/// <param name="TokenReader">Token reader</param>
-		/// <param name="Terminator">Terminating symbol</param>
-		/// <param name="Separator">Separator symbol</param>
-		/// <param name="bAllowTrailingSeparator">If true, allow trailing separators</param>
-		/// <param name="ExceptionContext">Extra context for error messages</param>
-		/// <param name="Action">Action to be invoked for each list element.</param>
+		/// <param name="tokenReader">Token reader</param>
+		/// <param name="terminator">Terminating symbol</param>
+		/// <param name="separator">Separator symbol</param>
+		/// <param name="allowTrailingSeparator">If true, allow trailing separators</param>
+		/// <param name="exceptionContext">Extra context for error messages</param>
+		/// <param name="action">Action to be invoked for each list element.</param>
 		/// <returns>Token reader</returns>
-		public static IUhtTokenReader RequireList(this IUhtTokenReader TokenReader, char Terminator, char Separator, bool bAllowTrailingSeparator, object? ExceptionContext, Action Action)
+		public static IUhtTokenReader RequireList(this IUhtTokenReader tokenReader, char terminator, char separator, bool allowTrailingSeparator, object? exceptionContext, Action action)
 		{
 			// Check for an empty list
-			if (TokenReader.TryOptional(Terminator))
+			if (tokenReader.TryOptional(terminator))
 			{
-				return TokenReader;
+				return tokenReader;
 			}
 
 			// Process the body
@@ -102,48 +102,48 @@ namespace EpicGames.UHT.Tokenizer
 			{
 
 				// Read the element via the lambda
-				Action();
+				action();
 
 				// Make sure we haven't reached the EOF
-				if (TokenReader.bIsEOF)
+				if (tokenReader.IsEOF)
 				{
-					throw new UhtTokenException(TokenReader, TokenReader.GetToken(), $"'{Separator}' or '{Terminator}'", ExceptionContext);
+					throw new UhtTokenException(tokenReader, tokenReader.GetToken(), $"'{separator}' or '{terminator}'", exceptionContext);
 				}
 
 				// If we have a separator, then it might be a trailing separator 
-				if (TokenReader.TryOptional(Separator))
+				if (tokenReader.TryOptional(separator))
 				{
-					if (bAllowTrailingSeparator && TokenReader.TryOptional(Terminator))
+					if (allowTrailingSeparator && tokenReader.TryOptional(terminator))
 					{
-						return TokenReader;
+						return tokenReader;
 					}
 					continue;
 				}
 
 				// Otherwise, we must have an terminator
-				if (!TokenReader.TryOptional(Terminator))
+				if (!tokenReader.TryOptional(terminator))
 				{
-					throw new UhtTokenException(TokenReader, TokenReader.GetToken(), $"'{Separator}' or '{Terminator}'", ExceptionContext);
+					throw new UhtTokenException(tokenReader, tokenReader.GetToken(), $"'{separator}' or '{terminator}'", exceptionContext);
 				}
-				return TokenReader;
+				return tokenReader;
 			}
 		}
 
 		/// <summary>
 		/// Parse an optional list
 		/// </summary>
-		/// <param name="TokenReader">Token reader</param>
-		/// <param name="Initiator">Initiating symbol</param>
-		/// <param name="Terminator">Terminating symbol</param>
-		/// <param name="Separator">Separator symbol</param>
-		/// <param name="bAllowTrailingSeparator">If true, allow trailing separators</param>
-		/// <param name="Delegate">Delegate to invoke with the parsed token list</param>
+		/// <param name="tokenReader">Token reader</param>
+		/// <param name="initiator">Initiating symbol</param>
+		/// <param name="terminator">Terminating symbol</param>
+		/// <param name="separator">Separator symbol</param>
+		/// <param name="allowTrailingSeparator">If true, allow trailing separators</param>
+		/// <param name="tokensDelegate">Delegate to invoke with the parsed token list</param>
 		/// <returns>True if a list was read</returns>
-		public static bool TryOptionalList(this IUhtTokenReader TokenReader, char Initiator, char Terminator, char Separator, bool bAllowTrailingSeparator, UhtTokensDelegate Delegate)
+		public static bool TryOptionalList(this IUhtTokenReader tokenReader, char initiator, char terminator, char separator, bool allowTrailingSeparator, UhtTokensDelegate tokensDelegate)
 		{
-			if (TokenReader.TryOptional(Initiator))
+			if (tokenReader.TryOptional(initiator))
 			{
-				TokenReader.RequireList(Terminator, Separator, bAllowTrailingSeparator, Delegate);
+				tokenReader.RequireList(terminator, separator, allowTrailingSeparator, tokensDelegate);
 				return true;
 			}
 			return false;
@@ -152,104 +152,104 @@ namespace EpicGames.UHT.Tokenizer
 		/// <summary>
 		/// Parse a required list
 		/// </summary>
-		/// <param name="TokenReader">Token reader</param>
-		/// <param name="Initiator">Initiating symbol</param>
-		/// <param name="Terminator">Terminating symbol</param>
-		/// <param name="Separator">Separator symbol</param>
-		/// <param name="bAllowTrailingSeparator">If true, allow trailing separators</param>
-		/// <param name="Delegate">Delegate to invoke with the parsed token list</param>
+		/// <param name="tokenReader">Token reader</param>
+		/// <param name="initiator">Initiating symbol</param>
+		/// <param name="terminator">Terminating symbol</param>
+		/// <param name="separator">Separator symbol</param>
+		/// <param name="allowTrailingSeparator">If true, allow trailing separators</param>
+		/// <param name="tokensDelegate">Delegate to invoke with the parsed token list</param>
 		/// <returns>Token reader</returns>
-		public static IUhtTokenReader RequireList(this IUhtTokenReader TokenReader, char Initiator, char Terminator, char Separator, bool bAllowTrailingSeparator, UhtTokensDelegate Delegate)
+		public static IUhtTokenReader RequireList(this IUhtTokenReader tokenReader, char initiator, char terminator, char separator, bool allowTrailingSeparator, UhtTokensDelegate tokensDelegate)
 		{
-			return TokenReader.RequireList(Initiator, Terminator, Separator, bAllowTrailingSeparator, null, Delegate);
+			return tokenReader.RequireList(initiator, terminator, separator, allowTrailingSeparator, null, tokensDelegate);
 		}
 
 		/// <summary>
 		/// Parse a required list
 		/// </summary>
-		/// <param name="TokenReader">Token reader</param>
-		/// <param name="Initiator">Initiating symbol</param>
-		/// <param name="Terminator">Terminating symbol</param>
-		/// <param name="Separator">Separator symbol</param>
-		/// <param name="bAllowTrailingSeparator">If true, allow trailing separators</param>
-		/// <param name="ExceptionContext">Extra context for error messages</param>
-		/// <param name="Delegate">Delegate to invoke with the parsed token list</param>
+		/// <param name="tokenReader">Token reader</param>
+		/// <param name="initiator">Initiating symbol</param>
+		/// <param name="terminator">Terminating symbol</param>
+		/// <param name="separator">Separator symbol</param>
+		/// <param name="allowTrailingSeparator">If true, allow trailing separators</param>
+		/// <param name="exceptionContext">Extra context for error messages</param>
+		/// <param name="tokensDelegate">Delegate to invoke with the parsed token list</param>
 		/// <returns>Token reader</returns>
-		public static IUhtTokenReader RequireList(this IUhtTokenReader TokenReader, char Initiator, char Terminator, char Separator, bool bAllowTrailingSeparator, object? ExceptionContext, UhtTokensDelegate Delegate)
+		public static IUhtTokenReader RequireList(this IUhtTokenReader tokenReader, char initiator, char terminator, char separator, bool allowTrailingSeparator, object? exceptionContext, UhtTokensDelegate tokensDelegate)
 		{
-			TokenReader.Require(Initiator);
-			return TokenReader.RequireList(Terminator, Separator, bAllowTrailingSeparator, ExceptionContext, Delegate);
+			tokenReader.Require(initiator);
+			return tokenReader.RequireList(terminator, separator, allowTrailingSeparator, exceptionContext, tokensDelegate);
 		}
 
 		/// <summary>
 		/// Parse a required list.  Initiating token must have already been parsed
 		/// </summary>
-		/// <param name="TokenReader">Token reader</param>
-		/// <param name="Terminator">Terminating symbol</param>
-		/// <param name="Separator">Separator symbol</param>
-		/// <param name="bAllowTrailingSeparator">If true, allow trailing separators</param>
-		/// <param name="Delegate">Delegate to invoke with the parsed token list</param>
+		/// <param name="tokenReader">Token reader</param>
+		/// <param name="terminator">Terminating symbol</param>
+		/// <param name="separator">Separator symbol</param>
+		/// <param name="allowTrailingSeparator">If true, allow trailing separators</param>
+		/// <param name="tokensDelegate">Delegate to invoke with the parsed token list</param>
 		/// <returns>Token reader</returns>
-		public static IUhtTokenReader RequireList(this IUhtTokenReader TokenReader, char Terminator, char Separator, bool bAllowTrailingSeparator, UhtTokensDelegate Delegate)
+		public static IUhtTokenReader RequireList(this IUhtTokenReader tokenReader, char terminator, char separator, bool allowTrailingSeparator, UhtTokensDelegate tokensDelegate)
 		{
-			return TokenReader.RequireList(Terminator, Separator, bAllowTrailingSeparator, null, Delegate);
+			return tokenReader.RequireList(terminator, separator, allowTrailingSeparator, null, tokensDelegate);
 		}
 
 		/// <summary>
 		/// Parse a required list.  Initiating token must have already been parsed
 		/// </summary>
-		/// <param name="TokenReader">Token reader</param>
-		/// <param name="Terminator">Terminating symbol</param>
-		/// <param name="Separator">Separator symbol</param>
-		/// <param name="bAllowTrailingSeparator">If true, allow trailing separators</param>
-		/// <param name="ExceptionContext">Extra context for error messages</param>
-		/// <param name="Delegate">Delegate to invoke with the parsed token list</param>
+		/// <param name="tokenReader">Token reader</param>
+		/// <param name="terminator">Terminating symbol</param>
+		/// <param name="separator">Separator symbol</param>
+		/// <param name="allowTrailingSeparator">If true, allow trailing separators</param>
+		/// <param name="exceptionContext">Extra context for error messages</param>
+		/// <param name="tokensDelegate">Delegate to invoke with the parsed token list</param>
 		/// <returns>Token reader</returns>
-		public static IUhtTokenReader RequireList(this IUhtTokenReader TokenReader, char Terminator, char Separator, bool bAllowTrailingSeparator, object? ExceptionContext, UhtTokensDelegate Delegate)
+		public static IUhtTokenReader RequireList(this IUhtTokenReader tokenReader, char terminator, char separator, bool allowTrailingSeparator, object? exceptionContext, UhtTokensDelegate tokensDelegate)
 		{
 			// Check for an empty list
-			if (TokenReader.TryOptional(Terminator))
+			if (tokenReader.TryOptional(terminator))
 			{
-				return TokenReader;
+				return tokenReader;
 			}
 
 			// Process the body
 			while (true)
 			{
-				List<UhtToken> Tokens = new List<UhtToken>();
+				List<UhtToken> tokens = new List<UhtToken>();
 
 				// Read the tokens until we hit the end
 				while (true)
 				{
 					// Make sure we haven't reached the EOF
-					if (TokenReader.bIsEOF)
+					if (tokenReader.IsEOF)
 					{
-						throw new UhtTokenException(TokenReader, TokenReader.GetToken(), $"'{Separator}' or '{Terminator}'", ExceptionContext);
+						throw new UhtTokenException(tokenReader, tokenReader.GetToken(), $"'{separator}' or '{terminator}'", exceptionContext);
 					}
 
 					// If we have a separator, then it might be a trailing separator 
-					if (TokenReader.TryOptional(Separator))
+					if (tokenReader.TryOptional(separator))
 					{
-						Delegate(Tokens);
-						if (TokenReader.TryOptional(Terminator))
+						tokensDelegate(tokens);
+						if (tokenReader.TryOptional(terminator))
 						{
-							if (bAllowTrailingSeparator)
+							if (allowTrailingSeparator)
 							{
-								return TokenReader;
+								return tokenReader;
 							}
-							throw new UhtException(TokenReader, $"A separator '{Separator}' followed immediately by the terminator '{Terminator}' is invalid");
+							throw new UhtException(tokenReader, $"A separator '{separator}' followed immediately by the terminator '{terminator}' is invalid");
 						}
 						break;
 					}
 
 					// If this is the terminator, then we are done
-					if (TokenReader.TryOptional(Terminator))
+					if (tokenReader.TryOptional(terminator))
 					{
-						Delegate(Tokens);
-						return TokenReader;
+						tokensDelegate(tokens);
+						return tokenReader;
 					}
 
-					Tokens.Add(TokenReader.GetToken());
+					tokens.Add(tokenReader.GetToken());
 				}
 			}
 		}
@@ -257,91 +257,91 @@ namespace EpicGames.UHT.Tokenizer
 		/// <summary>
 		/// Consume a block of tokens bounded by the two given symbols.
 		/// </summary>
-		/// <param name="TokenReader">Token reader</param>
-		/// <param name="Initiator">The next token must be the given symbol.</param>
-		/// <param name="Terminator">The tokens are read until the given symbol is found.  The terminating symbol will be consumed.</param>
-		/// <param name="ExceptionContext">Extra context for any error messages</param>
+		/// <param name="tokenReader">Token reader</param>
+		/// <param name="initiator">The next token must be the given symbol.</param>
+		/// <param name="terminator">The tokens are read until the given symbol is found.  The terminating symbol will be consumed.</param>
+		/// <param name="exceptionContext">Extra context for any error messages</param>
 		/// <returns>The input token reader</returns>
-		public static IUhtTokenReader RequireList(this IUhtTokenReader TokenReader, char Initiator, char Terminator, object? ExceptionContext = null)
+		public static IUhtTokenReader RequireList(this IUhtTokenReader tokenReader, char initiator, char terminator, object? exceptionContext = null)
 		{
-			TokenReader.Require(Initiator);
+			tokenReader.Require(initiator);
 
 			// Process the body
 			while (true)
 			{
 
 				// Make sure we haven't reached the EOF
-				if (TokenReader.bIsEOF)
+				if (tokenReader.IsEOF)
 				{
-					throw new UhtTokenException(TokenReader, TokenReader.GetToken(), $"'{Terminator}'", ExceptionContext);
+					throw new UhtTokenException(tokenReader, tokenReader.GetToken(), $"'{terminator}'", exceptionContext);
 				}
 
 				// Look for the terminator
-				if (TokenReader.TryOptional(Terminator))
+				if (tokenReader.TryOptional(terminator))
 				{
 					break;
 				}
 
 				// Consume the current token
-				TokenReader.ConsumeToken();
+				tokenReader.ConsumeToken();
 			}
-			return TokenReader;
+			return tokenReader;
 		}
 
 		/// <summary>
 		/// Invoke action while the next token is the given string
 		/// </summary>
-		/// <param name="TokenReader">Token reader</param>
-		/// <param name="Text">Text to match to continue invoking Action</param>
-		/// <param name="Action">Action to invoke if and only if the prior text was parsed.</param>
+		/// <param name="tokenReader">Token reader</param>
+		/// <param name="text">Text to match to continue invoking Action</param>
+		/// <param name="action">Action to invoke if and only if the prior text was parsed.</param>
 		/// <returns>Token reader</returns>
-		public static IUhtTokenReader While(this IUhtTokenReader TokenReader, string Text, Action Action)
+		public static IUhtTokenReader While(this IUhtTokenReader tokenReader, string text, Action action)
 		{
-			while (TokenReader.TryOptional(Text))
+			while (tokenReader.TryOptional(text))
 			{
-				Action();
+				action();
 			}
-			return TokenReader;
+			return tokenReader;
 		}
 
 		/// <summary>
 		/// Invoke action while the next token is the given string
 		/// </summary>
-		/// <param name="TokenReader">Token reader</param>
-		/// <param name="Text">Text to match to continue invoking Action</param>
-		/// <param name="Action">Action to invoke if and only if the prior text was parsed.</param>
+		/// <param name="tokenReader">Token reader</param>
+		/// <param name="text">Text to match to continue invoking Action</param>
+		/// <param name="action">Action to invoke if and only if the prior text was parsed.</param>
 		/// <returns>Token reader</returns>
-		public static IUhtTokenReader While(this IUhtTokenReader TokenReader, char Text, Action Action)
+		public static IUhtTokenReader While(this IUhtTokenReader tokenReader, char text, Action action)
 		{
-			while (TokenReader.TryOptional(Text))
+			while (tokenReader.TryOptional(text))
 			{
-				Action();
+				action();
 			}
-			return TokenReader;
+			return tokenReader;
 		}
 
 		/// <summary>
 		/// Read tokens until the delegate return false.  The terminating token is not consumed.
 		/// </summary>
-		/// <param name="TokenReader">Token reader</param>
-		/// <param name="Delegate">Invoked with each read token.  Return true to continue tokenizing or false to terminate.</param>
+		/// <param name="tokenReader">Token reader</param>
+		/// <param name="untilDelegate">Invoked with each read token.  Return true to continue tokenizing or false to terminate.</param>
 		/// <returns>Token reader</returns>
-		public static IUhtTokenReader While(this IUhtTokenReader TokenReader, UhtTokensUntilDelegate Delegate)
+		public static IUhtTokenReader While(this IUhtTokenReader tokenReader, UhtTokensUntilDelegate untilDelegate)
 		{
 			while (true)
 			{
-				if (TokenReader.bIsEOF)
+				if (tokenReader.IsEOF)
 				{
-					throw new UhtTokenException(TokenReader, TokenReader.GetToken(), null, null);
+					throw new UhtTokenException(tokenReader, tokenReader.GetToken(), null, null);
 				}
-				ref UhtToken Token = ref TokenReader.PeekToken();
-				if (Delegate(ref Token))
+				ref UhtToken token = ref tokenReader.PeekToken();
+				if (untilDelegate(ref token))
 				{
-					TokenReader.ConsumeToken();
+					tokenReader.ConsumeToken();
 				}
 				else
 				{
-					return TokenReader;
+					return tokenReader;
 				}
 			}
 		}
@@ -349,45 +349,45 @@ namespace EpicGames.UHT.Tokenizer
 		/// <summary>
 		/// Consume tokens until one of the strings are found.  Terminating token will not be consumed.
 		/// </summary>
-		/// <param name="TokenReader">Token reader</param>
-		/// <param name="Terminators">Strings that will terminate processing.</param>
+		/// <param name="tokenReader">Token reader</param>
+		/// <param name="terminators">Strings that will terminate processing.</param>
 		/// <returns>Number of tokens consumed</returns>
-		public static int ConsumeUntil(this IUhtTokenReader TokenReader, string[] Terminators)
+		public static int ConsumeUntil(this IUhtTokenReader tokenReader, string[] terminators)
 		{
-			int ConsumedTokens = 0;
-			while (!TokenReader.bIsEOF)
+			int consumedTokens = 0;
+			while (!tokenReader.IsEOF)
 			{
-				ref UhtToken Token = ref TokenReader.PeekToken();
-				foreach (string Candidate in Terminators)
+				ref UhtToken token = ref tokenReader.PeekToken();
+				foreach (string candidate in terminators)
 				{
-					if ((Token.IsIdentifier() || Token.IsSymbol()) && Token.IsValue(Candidate))
+					if ((token.IsIdentifier() || token.IsSymbol()) && token.IsValue(candidate))
 					{
-						return ConsumedTokens;
+						return consumedTokens;
 					}
 				}
-				TokenReader.ConsumeToken();
-				++ConsumedTokens;
+				tokenReader.ConsumeToken();
+				++consumedTokens;
 			}
-			return ConsumedTokens;
+			return consumedTokens;
 		}
 
 		/// <summary>
 		/// Consume until the given terminator is found.  Terminating token will be consumed
 		/// </summary>
-		/// <param name="TokenReader">Token reader</param>
-		/// <param name="Terminator">Terminating symbol</param>
+		/// <param name="tokenReader">Token reader</param>
+		/// <param name="terminator">Terminating symbol</param>
 		/// <returns>Token reader</returns>
-		public static IUhtTokenReader ConsumeUntil(this IUhtTokenReader TokenReader, char Terminator)
+		public static IUhtTokenReader ConsumeUntil(this IUhtTokenReader tokenReader, char terminator)
 		{
-			while (!TokenReader.bIsEOF)
+			while (!tokenReader.IsEOF)
 			{
-				if (TokenReader.TryOptional(Terminator))
+				if (tokenReader.TryOptional(terminator))
 				{
 					break;
 				}
-				TokenReader.ConsumeToken();
+				tokenReader.ConsumeToken();
 			}
-			return TokenReader;
+			return tokenReader;
 		}
 	}
 }
diff --git a/Engine/Source/Programs/Shared/EpicGames.UHT/Tokenizer/UhtTokenReaderGeneralExtensions.cs b/Engine/Source/Programs/Shared/EpicGames.UHT/Tokenizer/UhtTokenReaderGeneralExtensions.cs
--- a/Engine/Source/Programs/Shared/EpicGames.UHT/Tokenizer/UhtTokenReaderGeneralExtensions.cs
+++ b/Engine/Source/Programs/Shared/EpicGames.UHT/Tokenizer/UhtTokenReaderGeneralExtensions.cs
@@ -1,12 +1,12 @@
 // Copyright Epic Games, Inc. All Rights Reserved.
 
+using System;
 using EpicGames.Core;
 using EpicGames.UHT.Utils;
-using System;
 
 namespace EpicGames.UHT.Tokenizer
 {
-	
+
 	/// <summary>
 	/// Collection of general token reader extensions
 	/// </summary>
@@ -16,15 +16,15 @@ namespace EpicGames.UHT.Tokenizer
 		/// <summary>
 		/// Try to parse the given text
 		/// </summary>
-		/// <param name="TokenReader">Token reader</param>
-		/// <param name="Text">Text to match</param>
+		/// <param name="tokenReader">Token reader</param>
+		/// <param name="text">Text to match</param>
 		/// <returns>True if the text matched</returns>
-		public static bool TryOptional(this IUhtTokenReader TokenReader, string Text)
+		public static bool TryOptional(this IUhtTokenReader tokenReader, string text)
 		{
-			ref UhtToken Token = ref TokenReader.PeekToken();
-			if (Token.IsIdentifier(Text) || Token.IsSymbol(Text))
+			ref UhtToken token = ref tokenReader.PeekToken();
+			if (token.IsIdentifier(text) || token.IsSymbol(text))
 			{
-				TokenReader.ConsumeToken();
+				tokenReader.ConsumeToken();
 				return true;
 			}
 			return false;
@@ -33,15 +33,15 @@ namespace EpicGames.UHT.Tokenizer
 		/// <summary>
 		/// Try to parse the given text
 		/// </summary>
-		/// <param name="TokenReader">Token reader</param>
-		/// <param name="Text">Text to match</param>
+		/// <param name="tokenReader">Token reader</param>
+		/// <param name="text">Text to match</param>
 		/// <returns>True if the text matched</returns>
-		public static bool TryOptional(this IUhtTokenReader TokenReader, StringView Text)
+		public static bool TryOptional(this IUhtTokenReader tokenReader, StringView text)
 		{
-			ref UhtToken Token = ref TokenReader.PeekToken();
-			if (Token.IsIdentifier(Text) || Token.IsSymbol(Text))
+			ref UhtToken token = ref tokenReader.PeekToken();
+			if (token.IsIdentifier(text) || token.IsSymbol(text))
 			{
-				TokenReader.ConsumeToken();
+				tokenReader.ConsumeToken();
 				return true;
 			}
 			return false;
@@ -50,30 +50,30 @@ namespace EpicGames.UHT.Tokenizer
 		/// <summary>
 		/// Try to parse the given text.  However, the matching token will not be consumed.
 		/// </summary>
-		/// <param name="TokenReader">Token reader</param>
-		/// <param name="Text">Text to match</param>
+		/// <param name="tokenReader">Token reader</param>
+		/// <param name="text">Text to match</param>
 		/// <returns>True if the text matched</returns>
-		public static bool TryPeekOptional(this IUhtTokenReader TokenReader, string Text)
+		public static bool TryPeekOptional(this IUhtTokenReader tokenReader, string text)
 		{
-			ref UhtToken Token = ref TokenReader.PeekToken();
-			return Token.IsIdentifier(Text) || Token.IsSymbol(Text);
+			ref UhtToken token = ref tokenReader.PeekToken();
+			return token.IsIdentifier(text) || token.IsSymbol(text);
 		}
 
 		/// <summary>
 		/// Test to see if the next token is one of the given strings.
 		/// </summary>
-		/// <param name="TokenReader">Token reader</param>
-		/// <param name="Text">List of keywords to test</param>
+		/// <param name="tokenReader">Token reader</param>
+		/// <param name="text">List of keywords to test</param>
 		/// <returns>Index of matched string or -1 if nothing matched</returns>
-		public static int TryOptional(this IUhtTokenReader TokenReader, string[] Text)
+		public static int TryOptional(this IUhtTokenReader tokenReader, string[] text)
 		{
-			ref UhtToken Token = ref TokenReader.PeekToken();
-			for (int Index = 0, EndIndex = Text.Length; Index < EndIndex; ++Index)
+			ref UhtToken token = ref tokenReader.PeekToken();
+			for (int index = 0, endIndex = text.Length; index < endIndex; ++index)
 			{
-				if (Token.IsIdentifier(Text[Index]) || Token.IsSymbol(Text[Index]))
+				if (token.IsIdentifier(text[index]) || token.IsSymbol(text[index]))
 				{
-					TokenReader.ConsumeToken();
-					return Index;
+					tokenReader.ConsumeToken();
+					return index;
 				}
 			}
 			return -1;
@@ -82,15 +82,15 @@ namespace EpicGames.UHT.Tokenizer
 		/// <summary>
 		/// Try to parse the given text
 		/// </summary>
-		/// <param name="TokenReader">Token reader</param>
-		/// <param name="Text">Text to match</param>
+		/// <param name="tokenReader">Token reader</param>
+		/// <param name="text">Text to match</param>
 		/// <returns>True if the text matched</returns>
-		public static bool TryOptional(this IUhtTokenReader TokenReader, char Text)
+		public static bool TryOptional(this IUhtTokenReader tokenReader, char text)
 		{
-			ref UhtToken Token = ref TokenReader.PeekToken();
-			if (Token.IsSymbol(Text))
+			ref UhtToken token = ref tokenReader.PeekToken();
+			if (token.IsSymbol(text))
 			{
-				TokenReader.ConsumeToken();
+				tokenReader.ConsumeToken();
 				return true;
 			}
 			return false;
@@ -99,256 +99,256 @@ namespace EpicGames.UHT.Tokenizer
 		/// <summary>
 		/// Try to parse the given text
 		/// </summary>
-		/// <param name="TokenReader">Token reader</param>
-		/// <param name="Text">Text to match</param>
-		/// <param name="OutToken">Open that was matched</param>
+		/// <param name="tokenReader">Token reader</param>
+		/// <param name="text">Text to match</param>
+		/// <param name="outToken">Open that was matched</param>
 		/// <returns>True if the text matched</returns>
-		public static bool TryOptional(this IUhtTokenReader TokenReader, char Text, out UhtToken OutToken)
+		public static bool TryOptional(this IUhtTokenReader tokenReader, char text, out UhtToken outToken)
 		{
-			ref UhtToken Token = ref TokenReader.PeekToken();
-			if (Token.IsSymbol(Text))
+			ref UhtToken token = ref tokenReader.PeekToken();
+			if (token.IsSymbol(text))
 			{
-				OutToken = Token;
-				TokenReader.ConsumeToken();
+				outToken = token;
+				tokenReader.ConsumeToken();
 				return true;
 			}
-			OutToken = new UhtToken();
+			outToken = new UhtToken();
 			return false;
 		}
 
 		/// <summary>
 		/// Try to parse the given text.  However, the matching token will not be consumed.
 		/// </summary>
-		/// <param name="TokenReader">Token reader</param>
-		/// <param name="Text">Text to match</param>
+		/// <param name="tokenReader">Token reader</param>
+		/// <param name="text">Text to match</param>
 		/// <returns>True if the text matched</returns>
-		public static bool TryPeekOptional(this IUhtTokenReader TokenReader, char Text)
+		public static bool TryPeekOptional(this IUhtTokenReader tokenReader, char text)
 		{
-			ref UhtToken Token = ref TokenReader.PeekToken();
-			return Token.IsSymbol(Text);
+			ref UhtToken token = ref tokenReader.PeekToken();
+			return token.IsSymbol(text);
 		}
 
 		/// <summary>
 		/// Parse optional text
 		/// </summary>
-		/// <param name="TokenReader">Token reader</param>
-		/// <param name="Text">Text to match</param>
+		/// <param name="tokenReader">Token reader</param>
+		/// <param name="text">Text to match</param>
 		/// <returns>Token reader</returns>
-		public static IUhtTokenReader Optional(this IUhtTokenReader TokenReader, string Text)
+		public static IUhtTokenReader Optional(this IUhtTokenReader tokenReader, string text)
 		{
-			ref UhtToken Token = ref TokenReader.PeekToken();
-			if (Token.IsIdentifier(Text) || Token.IsSymbol(Text))
+			ref UhtToken token = ref tokenReader.PeekToken();
+			if (token.IsIdentifier(text) || token.IsSymbol(text))
 			{
-				TokenReader.ConsumeToken();
+				tokenReader.ConsumeToken();
 			}
-			return TokenReader;
+			return tokenReader;
 		}
 
 		/// <summary>
 		/// Parse optional text
 		/// </summary>
-		/// <param name="TokenReader">Token reader</param>
-		/// <param name="Text">Text to match</param>
-		/// <param name="Action">Action to invoke if the text was found</param>
+		/// <param name="tokenReader">Token reader</param>
+		/// <param name="text">Text to match</param>
+		/// <param name="action">Action to invoke if the text was found</param>
 		/// <returns>Token reader</returns>
-		public static IUhtTokenReader Optional(this IUhtTokenReader TokenReader, string Text, Action Action)
+		public static IUhtTokenReader Optional(this IUhtTokenReader tokenReader, string text, Action action)
 		{
-			ref UhtToken Token = ref TokenReader.PeekToken();
-			if (Token.IsIdentifier(Text) || Token.IsSymbol(Text))
+			ref UhtToken token = ref tokenReader.PeekToken();
+			if (token.IsIdentifier(text) || token.IsSymbol(text))
 			{
-				TokenReader.ConsumeToken();
-				Action();
+				tokenReader.ConsumeToken();
+				action();
 			}
-			return TokenReader;
+			return tokenReader;
 		}
 
 		/// <summary>
 		/// Parse optional text
 		/// </summary>
-		/// <param name="TokenReader">Token reader</param>
-		/// <param name="Text">Text to match</param>
+		/// <param name="tokenReader">Token reader</param>
+		/// <param name="text">Text to match</param>
 		/// <returns>Token reader</returns>
-		public static IUhtTokenReader Optional(this IUhtTokenReader TokenReader, char Text)
+		public static IUhtTokenReader Optional(this IUhtTokenReader tokenReader, char text)
 		{
-			ref UhtToken Token = ref TokenReader.PeekToken();
-			if (Token.IsSymbol(Text))
+			ref UhtToken token = ref tokenReader.PeekToken();
+			if (token.IsSymbol(text))
 			{
-				TokenReader.ConsumeToken();
+				tokenReader.ConsumeToken();
 			}
-			return TokenReader;
+			return tokenReader;
 		}
 
 		/// <summary>
 		/// Parse optional text
 		/// </summary>
-		/// <param name="TokenReader">Token reader</param>
-		/// <param name="Text">Text to match</param>
-		/// <param name="Action">Action to invoke if the text was found</param>
+		/// <param name="tokenReader">Token reader</param>
+		/// <param name="text">Text to match</param>
+		/// <param name="action">Action to invoke if the text was found</param>
 		/// <returns>Token reader</returns>
-		public static IUhtTokenReader Optional(this IUhtTokenReader TokenReader, char Text, Action Action)
+		public static IUhtTokenReader Optional(this IUhtTokenReader tokenReader, char text, Action action)
 		{
-			ref UhtToken Token = ref TokenReader.PeekToken();
-			if (Token.IsSymbol(Text))
+			ref UhtToken token = ref tokenReader.PeekToken();
+			if (token.IsSymbol(text))
 			{
-				TokenReader.ConsumeToken();
-				Action();
+				tokenReader.ConsumeToken();
+				action();
 			}
-			return TokenReader;
+			return tokenReader;
 		}
 
 		/// <summary>
 		/// Parse optional token that starts with the given text 
 		/// </summary>
-		/// <param name="TokenReader">Token reader</param>
-		/// <param name="Text">Text to match</param>
+		/// <param name="tokenReader">Token reader</param>
+		/// <param name="text">Text to match</param>
 		/// <returns>Token reader</returns>
-		public static IUhtTokenReader OptionalStartsWith(this IUhtTokenReader TokenReader, string Text)
+		public static IUhtTokenReader OptionalStartsWith(this IUhtTokenReader tokenReader, string text)
 		{
-			ref UhtToken Token = ref TokenReader.PeekToken();
-			if (Token.IsIdentifier() && Token.ValueStartsWith(Text))
+			ref UhtToken token = ref tokenReader.PeekToken();
+			if (token.IsIdentifier() && token.ValueStartsWith(text))
 			{
-				TokenReader.ConsumeToken();
+				tokenReader.ConsumeToken();
 			}
-			return TokenReader;
+			return tokenReader;
 		}
 
 		/// <summary>
 		/// Parse optional token that starts with the given text 
 		/// </summary>
-		/// <param name="TokenReader">Token reader</param>
-		/// <param name="Text">Text to match</param>
-		/// <param name="TokenDelegate">Delegate to invoke on a match</param>
+		/// <param name="tokenReader">Token reader</param>
+		/// <param name="text">Text to match</param>
+		/// <param name="tokenDelegate">Delegate to invoke on a match</param>
 		/// <returns>Token reader</returns>
-		public static IUhtTokenReader OptionalStartsWith(this IUhtTokenReader TokenReader, string Text, UhtTokenDelegate TokenDelegate)
+		public static IUhtTokenReader OptionalStartsWith(this IUhtTokenReader tokenReader, string text, UhtTokenDelegate tokenDelegate)
 		{
-			ref UhtToken Token = ref TokenReader.PeekToken();
-			if (Token.IsIdentifier() && Token.ValueStartsWith(Text))
+			ref UhtToken token = ref tokenReader.PeekToken();
+			if (token.IsIdentifier() && token.ValueStartsWith(text))
 			{
-				UhtToken CurrentToken = Token;
-				TokenReader.ConsumeToken();
-				TokenDelegate(ref CurrentToken);
+				UhtToken currentToken = token;
+				tokenReader.ConsumeToken();
+				tokenDelegate(ref currentToken);
 			}
-			return TokenReader;
+			return tokenReader;
 		}
 
 		/// <summary>
 		/// Require the given text
 		/// </summary>
-		/// <param name="TokenReader">Token reader</param>
-		/// <param name="Text">Required text</param>
-		/// <param name="ExceptionContext">Extra exception context</param>
+		/// <param name="tokenReader">Token reader</param>
+		/// <param name="text">Required text</param>
+		/// <param name="exceptionContext">Extra exception context</param>
 		/// <returns>Token reader</returns>
 		/// <exception cref="UhtTokenException">Thrown if text is not found</exception>
-		public static IUhtTokenReader Require(this IUhtTokenReader TokenReader, string Text, object? ExceptionContext = null)
+		public static IUhtTokenReader Require(this IUhtTokenReader tokenReader, string text, object? exceptionContext = null)
 		{
-			ref UhtToken Token = ref TokenReader.PeekToken();
-			if (Token.IsIdentifier(Text) || Token.IsSymbol(Text))
+			ref UhtToken token = ref tokenReader.PeekToken();
+			if (token.IsIdentifier(text) || token.IsSymbol(text))
 			{
-				TokenReader.ConsumeToken();
+				tokenReader.ConsumeToken();
 			}
 			else
 			{
-				throw new UhtTokenException(TokenReader, Token, Text, ExceptionContext);
+				throw new UhtTokenException(tokenReader, token, text, exceptionContext);
 			}
-			return TokenReader;
+			return tokenReader;
 		}
 
 		/// <summary>
 		/// Require the given text
 		/// </summary>
-		/// <param name="TokenReader">Token reader</param>
-		/// <param name="Text">Required text</param>
-		/// <param name="TokenDelegate">Delegate to invoke on a match</param>
+		/// <param name="tokenReader">Token reader</param>
+		/// <param name="text">Required text</param>
+		/// <param name="tokenDelegate">Delegate to invoke on a match</param>
 		/// <returns>Token reader</returns>
 		/// <exception cref="UhtTokenException">Thrown if text is not found</exception>
-		public static IUhtTokenReader Require(this IUhtTokenReader TokenReader, string Text, UhtTokenDelegate TokenDelegate)
+		public static IUhtTokenReader Require(this IUhtTokenReader tokenReader, string text, UhtTokenDelegate tokenDelegate)
 		{
-			return TokenReader.Require(Text, null, TokenDelegate);
+			return tokenReader.Require(text, null, tokenDelegate);
 		}
 
 		/// <summary>
 		/// Require the given text
 		/// </summary>
-		/// <param name="TokenReader">Token reader</param>
-		/// <param name="Text">Required text</param>
-		/// <param name="ExceptionContext">Extra exception context</param>
-		/// <param name="TokenDelegate">Delegate to invoke on a match</param>
+		/// <param name="tokenReader">Token reader</param>
+		/// <param name="text">Required text</param>
+		/// <param name="exceptionContext">Extra exception context</param>
+		/// <param name="tokenDelegate">Delegate to invoke on a match</param>
 		/// <returns>Token reader</returns>
 		/// <exception cref="UhtTokenException">Thrown if text is not found</exception>
-		public static IUhtTokenReader Require(this IUhtTokenReader TokenReader, string Text, object? ExceptionContext, UhtTokenDelegate TokenDelegate)
+		public static IUhtTokenReader Require(this IUhtTokenReader tokenReader, string text, object? exceptionContext, UhtTokenDelegate tokenDelegate)
 		{
-			ref UhtToken Token = ref TokenReader.PeekToken();
-			if (Token.IsIdentifier(Text) || Token.IsSymbol(Text))
+			ref UhtToken token = ref tokenReader.PeekToken();
+			if (token.IsIdentifier(text) || token.IsSymbol(text))
 			{
-				UhtToken CurrentToken = Token;
-				TokenReader.ConsumeToken();
-				TokenDelegate(ref CurrentToken);
+				UhtToken currentToken = token;
+				tokenReader.ConsumeToken();
+				tokenDelegate(ref currentToken);
 			}
 			else
 			{
-				throw new UhtTokenException(TokenReader, Token, Text, ExceptionContext);
+				throw new UhtTokenException(tokenReader, token, text, exceptionContext);
 			}
-			return TokenReader;
+			return tokenReader;
 		}
 
 		/// <summary>
 		/// Require the given text
 		/// </summary>
-		/// <param name="TokenReader">Token reader</param>
-		/// <param name="Text">Required text</param>
-		/// <param name="ExceptionContext">Extra exception context</param>
+		/// <param name="tokenReader">Token reader</param>
+		/// <param name="text">Required text</param>
+		/// <param name="exceptionContext">Extra exception context</param>
 		/// <returns>Token reader</returns>
 		/// <exception cref="UhtTokenException">Thrown if text is not found</exception>
-		public static IUhtTokenReader Require(this IUhtTokenReader TokenReader, char Text, object? ExceptionContext = null)
+		public static IUhtTokenReader Require(this IUhtTokenReader tokenReader, char text, object? exceptionContext = null)
 		{
-			ref UhtToken Token = ref TokenReader.PeekToken();
-			if (Token.IsSymbol(Text))
+			ref UhtToken token = ref tokenReader.PeekToken();
+			if (token.IsSymbol(text))
 			{
-				TokenReader.ConsumeToken();
+				tokenReader.ConsumeToken();
 			}
 			else
 			{
-				throw new UhtTokenException(TokenReader, Token, Text, ExceptionContext);
+				throw new UhtTokenException(tokenReader, token, text, exceptionContext);
 			}
-			return TokenReader;
+			return tokenReader;
 		}
 
 		/// <summary>
 		/// Require the given text
 		/// </summary>
-		/// <param name="TokenReader">Token reader</param>
-		/// <param name="Text">Required text</param>
-		/// <param name="TokenDelegate">Delegate to invoke on a match</param>
+		/// <param name="tokenReader">Token reader</param>
+		/// <param name="text">Required text</param>
+		/// <param name="tokenDelegate">Delegate to invoke on a match</param>
 		/// <returns>Token reader</returns>
 		/// <exception cref="UhtTokenException">Thrown if text is not found</exception>
-		public static IUhtTokenReader Require(this IUhtTokenReader TokenReader, char Text, UhtTokenDelegate TokenDelegate)
+		public static IUhtTokenReader Require(this IUhtTokenReader tokenReader, char text, UhtTokenDelegate tokenDelegate)
 		{
-			return TokenReader.Require(Text, null, TokenDelegate);
+			return tokenReader.Require(text, null, tokenDelegate);
 		}
 
 		/// <summary>
 		/// Require the given text
 		/// </summary>
-		/// <param name="TokenReader">Token reader</param>
-		/// <param name="Text">Required text</param>
-		/// <param name="ExceptionContext">Extra exception context</param>
-		/// <param name="TokenDelegate">Delegate to invoke on a match</param>
+		/// <param name="tokenReader">Token reader</param>
+		/// <param name="text">Required text</param>
+		/// <param name="exceptionContext">Extra exception context</param>
+		/// <param name="tokenDelegate">Delegate to invoke on a match</param>
 		/// <returns>Token reader</returns>
 		/// <exception cref="UhtTokenException">Thrown if text is not found</exception>
-		public static IUhtTokenReader Require(this IUhtTokenReader TokenReader, char Text, object? ExceptionContext, UhtTokenDelegate TokenDelegate)
+		public static IUhtTokenReader Require(this IUhtTokenReader tokenReader, char text, object? exceptionContext, UhtTokenDelegate tokenDelegate)
 		{
-			ref UhtToken Token = ref TokenReader.PeekToken();
-			if (Token.IsSymbol(Text))
+			ref UhtToken token = ref tokenReader.PeekToken();
+			if (token.IsSymbol(text))
 			{
-				UhtToken CurrentToken = Token;
-				TokenReader.ConsumeToken();
-				TokenDelegate(ref CurrentToken);
+				UhtToken currentToken = token;
+				tokenReader.ConsumeToken();
+				tokenDelegate(ref currentToken);
 			}
 			else
 			{
-				throw new UhtTokenException(TokenReader, Token, Text, ExceptionContext);
+				throw new UhtTokenException(tokenReader, token, text, exceptionContext);
 			}
-			return TokenReader;
+			return tokenReader;
 		}
 	}
 }
diff --git a/Engine/Source/Programs/Shared/EpicGames.UHT/Tokenizer/UhtTokenReaderIdentifierExtensions.cs b/Engine/Source/Programs/Shared/EpicGames.UHT/Tokenizer/UhtTokenReaderIdentifierExtensions.cs
--- a/Engine/Source/Programs/Shared/EpicGames.UHT/Tokenizer/UhtTokenReaderIdentifierExtensions.cs
+++ b/Engine/Source/Programs/Shared/EpicGames.UHT/Tokenizer/UhtTokenReaderIdentifierExtensions.cs
@@ -1,9 +1,9 @@
 // Copyright Epic Games, Inc. All Rights Reserved.
 
-using EpicGames.UHT.Utils;
 using System;
 using System.Collections.Generic;
 using System.Diagnostics.CodeAnalysis;
+using EpicGames.UHT.Utils;
 
 namespace EpicGames.UHT.Tokenizer
 {
@@ -36,35 +36,35 @@ namespace EpicGames.UHT.Tokenizer
 		/// <summary>
 		/// Test to see if any of the specified flags are set
 		/// </summary>
-		/// <param name="InFlags">Current flags</param>
-		/// <param name="TestFlags">Flags to test for</param>
+		/// <param name="inFlags">Current flags</param>
+		/// <param name="testFlags">Flags to test for</param>
 		/// <returns>True if any of the flags are set</returns>
-		public static bool HasAnyFlags(this UhtCppIdentifierOptions InFlags, UhtCppIdentifierOptions TestFlags)
+		public static bool HasAnyFlags(this UhtCppIdentifierOptions inFlags, UhtCppIdentifierOptions testFlags)
 		{
-			return (InFlags & TestFlags) != 0;
+			return (inFlags & testFlags) != 0;
 		}
 
 		/// <summary>
 		/// Test to see if all of the specified flags are set
 		/// </summary>
-		/// <param name="InFlags">Current flags</param>
-		/// <param name="TestFlags">Flags to test for</param>
+		/// <param name="inFlags">Current flags</param>
+		/// <param name="testFlags">Flags to test for</param>
 		/// <returns>True if all the flags are set</returns>
-		public static bool HasAllFlags(this UhtCppIdentifierOptions InFlags, UhtCppIdentifierOptions TestFlags)
+		public static bool HasAllFlags(this UhtCppIdentifierOptions inFlags, UhtCppIdentifierOptions testFlags)
 		{
-			return (InFlags & TestFlags) == TestFlags;
+			return (inFlags & testFlags) == testFlags;
 		}
 
 		/// <summary>
 		/// Test to see if a specific set of flags have a specific value.
 		/// </summary>
-		/// <param name="InFlags">Current flags</param>
-		/// <param name="TestFlags">Flags to test for</param>
-		/// <param name="MatchFlags">Expected value of the tested flags</param>
+		/// <param name="inFlags">Current flags</param>
+		/// <param name="testFlags">Flags to test for</param>
+		/// <param name="matchFlags">Expected value of the tested flags</param>
 		/// <returns>True if the given flags have a specific value.</returns>
-		public static bool HasExactFlags(this UhtCppIdentifierOptions InFlags, UhtCppIdentifierOptions TestFlags, UhtCppIdentifierOptions MatchFlags)
+		public static bool HasExactFlags(this UhtCppIdentifierOptions inFlags, UhtCppIdentifierOptions testFlags, UhtCppIdentifierOptions matchFlags)
 		{
-			return (InFlags & TestFlags) == MatchFlags;
+			return (inFlags & testFlags) == matchFlags;
 		}
 	}
 
@@ -78,12 +78,12 @@ namespace EpicGames.UHT.Tokenizer
 		/// Get the next token and verify that it is an identifier
 		/// </summary>
 		/// <returns>True if it is an identifier, false if not.</returns>
-		public static bool TryOptionalIdentifier(this IUhtTokenReader TokenReader)
+		public static bool TryOptionalIdentifier(this IUhtTokenReader tokenReader)
 		{
-			ref UhtToken Token = ref TokenReader.PeekToken();
-			if (Token.IsIdentifier())
+			ref UhtToken token = ref tokenReader.PeekToken();
+			if (token.IsIdentifier())
 			{
-				TokenReader.ConsumeToken();
+				tokenReader.ConsumeToken();
 				return true;
 			}
 			return false;
@@ -92,169 +92,169 @@ namespace EpicGames.UHT.Tokenizer
 		/// <summary>
 		/// Get the next token and verify that it is an identifier
 		/// </summary>
-		/// <param name="TokenReader">Token reader</param>
-		/// <param name="Identifier">The fetched value of the identifier</param>
+		/// <param name="tokenReader">Token reader</param>
+		/// <param name="identifier">The fetched value of the identifier</param>
 		/// <returns>True if it is an identifier, false if not.</returns>
-		public static bool TryOptionalIdentifier(this IUhtTokenReader TokenReader, out UhtToken Identifier)
+		public static bool TryOptionalIdentifier(this IUhtTokenReader tokenReader, out UhtToken identifier)
 		{
-			ref UhtToken Token = ref TokenReader.PeekToken();
-			if (Token.IsIdentifier())
+			ref UhtToken token = ref tokenReader.PeekToken();
+			if (token.IsIdentifier())
 			{
-				Identifier = Token;
-				TokenReader.ConsumeToken();
+				identifier = token;
+				tokenReader.ConsumeToken();
 				return true;
 			}
-			Identifier = new UhtToken();
+			identifier = new UhtToken();
 			return false;
 		}
 
 		/// <summary>
 		/// Parse an optional identifier
 		/// </summary>
-		/// <param name="TokenReader">Token reader</param>
-		/// <param name="TokenDelegate">Invoked of an identifier is parsed</param>
+		/// <param name="tokenReader">Token reader</param>
+		/// <param name="tokenDelegate">Invoked of an identifier is parsed</param>
 		/// <returns>Token reader</returns>
-		public static IUhtTokenReader OptionalIdentifier(this IUhtTokenReader TokenReader, UhtTokenDelegate TokenDelegate)
+		public static IUhtTokenReader OptionalIdentifier(this IUhtTokenReader tokenReader, UhtTokenDelegate tokenDelegate)
 		{
-			ref UhtToken Token = ref TokenReader.PeekToken();
-			if (Token.IsIdentifier())
+			ref UhtToken token = ref tokenReader.PeekToken();
+			if (token.IsIdentifier())
 			{
-				UhtToken TokenCopy = Token;
-				TokenReader.ConsumeToken();
-				TokenDelegate(ref TokenCopy);
+				UhtToken tokenCopy = token;
+				tokenReader.ConsumeToken();
+				tokenDelegate(ref tokenCopy);
 			}
-			return TokenReader;
+			return tokenReader;
 		}
 
 		/// <summary>
 		/// Parse a required identifier
 		/// </summary>
-		/// <param name="TokenReader">Token reader</param>
-		/// <param name="ExceptionContext">Extra exception context</param>
+		/// <param name="tokenReader">Token reader</param>
+		/// <param name="exceptionContext">Extra exception context</param>
 		/// <returns>Token reader</returns>
 		/// <exception cref="UhtTokenException">Thrown if an identifier isn't found</exception>
-		public static IUhtTokenReader RequireIdentifier(this IUhtTokenReader TokenReader, object? ExceptionContext = null)
+		public static IUhtTokenReader RequireIdentifier(this IUhtTokenReader tokenReader, object? exceptionContext = null)
 		{
-			ref UhtToken Token = ref TokenReader.PeekToken();
-			if (Token.IsIdentifier())
+			ref UhtToken token = ref tokenReader.PeekToken();
+			if (token.IsIdentifier())
 			{
-				TokenReader.ConsumeToken();
-				return TokenReader;
+				tokenReader.ConsumeToken();
+				return tokenReader;
 			}
-			throw new UhtTokenException(TokenReader, Token, "an identifier", ExceptionContext);
+			throw new UhtTokenException(tokenReader, token, "an identifier", exceptionContext);
 		}
 
 		/// <summary>
 		/// Parse a required identifier
 		/// </summary>
-		/// <param name="TokenReader">Token reader</param>
-		/// <param name="TokenDelegate">Invoked if an identifier is parsed</param>
+		/// <param name="tokenReader">Token reader</param>
+		/// <param name="tokenDelegate">Invoked if an identifier is parsed</param>
 		/// <returns>Token reader</returns>
-		public static IUhtTokenReader RequireIdentifier(this IUhtTokenReader TokenReader, UhtTokenDelegate TokenDelegate)
+		public static IUhtTokenReader RequireIdentifier(this IUhtTokenReader tokenReader, UhtTokenDelegate tokenDelegate)
 		{
-			return TokenReader.RequireIdentifier(null, TokenDelegate);
+			return tokenReader.RequireIdentifier(null, tokenDelegate);
 		}
 
 		/// <summary>
 		/// Parse a required identifier
 		/// </summary>
-		/// <param name="TokenReader">Token reader</param>
-		/// <param name="ExceptionContext">Extra exception context</param>
-		/// <param name="TokenDelegate">Invoked if an identifier is parsed</param>
+		/// <param name="tokenReader">Token reader</param>
+		/// <param name="exceptionContext">Extra exception context</param>
+		/// <param name="tokenDelegate">Invoked if an identifier is parsed</param>
 		/// <returns>Token reader</returns>
 		/// <exception cref="UhtTokenException">Thrown if an identifier isn't found</exception>
-		public static IUhtTokenReader RequireIdentifier(this IUhtTokenReader TokenReader, object? ExceptionContext, UhtTokenDelegate TokenDelegate)
+		public static IUhtTokenReader RequireIdentifier(this IUhtTokenReader tokenReader, object? exceptionContext, UhtTokenDelegate tokenDelegate)
 		{
-			ref UhtToken Token = ref TokenReader.PeekToken();
-			if (Token.IsIdentifier())
+			ref UhtToken token = ref tokenReader.PeekToken();
+			if (token.IsIdentifier())
 			{
-				UhtToken CurrentToken = Token;
-				TokenReader.ConsumeToken();
-				TokenDelegate(ref CurrentToken);
-				return TokenReader;
+				UhtToken currentToken = token;
+				tokenReader.ConsumeToken();
+				tokenDelegate(ref currentToken);
+				return tokenReader;
 			}
-			throw new UhtTokenException(TokenReader, Token, "an identifier", ExceptionContext);
+			throw new UhtTokenException(tokenReader, token, "an identifier", exceptionContext);
 		}
 
 		/// <summary>
 		/// Get a required identifier
 		/// </summary>
-		/// <param name="TokenReader">Token reader</param>
-		/// <param name="ExceptionContext">Extra exception context</param>
+		/// <param name="tokenReader">Token reader</param>
+		/// <param name="exceptionContext">Extra exception context</param>
 		/// <returns>Identifier token</returns>
 		/// <exception cref="UhtTokenException">Thrown if an identifier isn't found</exception>
-		public static UhtToken GetIdentifier(this IUhtTokenReader TokenReader, object? ExceptionContext = null)
+		public static UhtToken GetIdentifier(this IUhtTokenReader tokenReader, object? exceptionContext = null)
 		{
-			ref UhtToken Token = ref TokenReader.PeekToken();
-			if (Token.IsIdentifier())
+			ref UhtToken token = ref tokenReader.PeekToken();
+			if (token.IsIdentifier())
 			{
-				UhtToken CurrentToken = Token;
-				TokenReader.ConsumeToken();
-				return CurrentToken;
+				UhtToken currentToken = token;
+				tokenReader.ConsumeToken();
+				return currentToken;
 			}
-			throw new UhtTokenException(TokenReader, Token, ExceptionContext != null ? ExceptionContext : "an identifier");
+			throw new UhtTokenException(tokenReader, token, exceptionContext ?? "an identifier");
 		}
 
 		/// <summary>
 		/// Parse an optional cpp identifier
 		/// </summary>
-		/// <param name="TokenReader">Token reader</param>
-		/// <param name="Identifier">Enumeration of the identifier tokens</param>
-		/// <param name="Options">Identifier options</param>
+		/// <param name="tokenReader">Token reader</param>
+		/// <param name="identifier">Enumeration of the identifier tokens</param>
+		/// <param name="options">Identifier options</param>
 		/// <returns>True if an identifier is parsed</returns>
-		public static bool TryOptionalCppIdentifier(this IUhtTokenReader TokenReader, [NotNullWhen(true)] out IEnumerable<UhtToken>? Identifier, UhtCppIdentifierOptions Options = UhtCppIdentifierOptions.None)
+		public static bool TryOptionalCppIdentifier(this IUhtTokenReader tokenReader, [NotNullWhen(true)] out IEnumerable<UhtToken>? identifier, UhtCppIdentifierOptions options = UhtCppIdentifierOptions.None)
 		{
-			Identifier = null;
-			List<UhtToken> LocalIdentifier = new List<UhtToken>();
+			identifier = null;
+			List<UhtToken> localIdentifier = new List<UhtToken>();
 
-			using (var SavedState = new UhtTokenSaveState(TokenReader))
+			using (UhtTokenSaveState savedState = new UhtTokenSaveState(tokenReader))
 			{
-				UhtToken Token;
+				UhtToken token;
 
-				if (TokenReader.TryOptionalIdentifier(out Token))
+				if (tokenReader.TryOptionalIdentifier(out token))
 				{
-					LocalIdentifier.Add(Token);
+					localIdentifier.Add(token);
 				}
 				else
 				{
 					return false;
 				}
 
-				if (Options.HasAnyFlags(UhtCppIdentifierOptions.AllowTemplates))
+				if (options.HasAnyFlags(UhtCppIdentifierOptions.AllowTemplates))
 				{
 					while (true)
 					{
-						if (TokenReader.TryPeekOptional('<'))
+						if (tokenReader.TryPeekOptional('<'))
 						{
-							LocalIdentifier.Add(TokenReader.GetToken());
-							int NestedScopes = 1;
-							while (NestedScopes > 0)
+							localIdentifier.Add(tokenReader.GetToken());
+							int nestedScopes = 1;
+							while (nestedScopes > 0)
 							{
-								Token = TokenReader.GetToken();
-								if (Token.TokenType.IsEndType())
+								token = tokenReader.GetToken();
+								if (token.TokenType.IsEndType())
 								{
 									return false;
 								}
-								LocalIdentifier.Add(Token);
-								if (Token.IsSymbol('<'))
+								localIdentifier.Add(token);
+								if (token.IsSymbol('<'))
 								{
-									++NestedScopes;
+									++nestedScopes;
 								}
-								else if (Token.IsSymbol('>'))
+								else if (token.IsSymbol('>'))
 								{
-									--NestedScopes;
+									--nestedScopes;
 								}
 							}
 						}
 
-						if (!TokenReader.TryPeekOptional("::"))
+						if (!tokenReader.TryPeekOptional("::"))
 						{
 							break;
 						}
-						LocalIdentifier.Add(TokenReader.GetToken());
-						if (!TokenReader.TryOptionalIdentifier(out Token))
+						localIdentifier.Add(tokenReader.GetToken());
+						if (!tokenReader.TryOptionalIdentifier(out token))
 						{
-							LocalIdentifier.Add(Token);
+							localIdentifier.Add(token);
 						}
 						else
 						{
@@ -264,11 +264,11 @@ namespace EpicGames.UHT.Tokenizer
 				}
 				else
 				{
-					while (TokenReader.TryOptional("::"))
+					while (tokenReader.TryOptional("::"))
 					{
-						if (TokenReader.TryOptionalIdentifier(out Token))
+						if (tokenReader.TryOptionalIdentifier(out token))
 						{
-							LocalIdentifier.Add(Token);
+							localIdentifier.Add(token);
 						}
 						else
 						{
@@ -276,8 +276,8 @@ namespace EpicGames.UHT.Tokenizer
 						}
 					}
 				}
-				Identifier = LocalIdentifier;
-				SavedState.AbandonState();
+				identifier = localIdentifier;
+				savedState.AbandonState();
 				return true;
 			}
 		}
@@ -285,156 +285,156 @@ namespace EpicGames.UHT.Tokenizer
 		/// <summary>
 		/// Parse a required cpp identifier
 		/// </summary>
-		/// <param name="TokenReader">Token reader</param>
-		/// <param name="Options">Parsing options</param>
+		/// <param name="tokenReader">Token reader</param>
+		/// <param name="options">Parsing options</param>
 		/// <returns>Token reader</returns>
-		public static IUhtTokenReader RequireCppIdentifier(this IUhtTokenReader TokenReader, UhtCppIdentifierOptions Options = UhtCppIdentifierOptions.None)
+		public static IUhtTokenReader RequireCppIdentifier(this IUhtTokenReader tokenReader, UhtCppIdentifierOptions options = UhtCppIdentifierOptions.None)
 		{
-			TokenReader.GetCppIdentifier(Options);
-			return TokenReader;
+			tokenReader.GetCppIdentifier(options);
+			return tokenReader;
 		}
 
 		/// <summary>
 		/// Parse a required cpp identifier
 		/// </summary>
-		/// <param name="TokenReader">Token reader</param>
-		/// <param name="InitialIdentifier">Initial token of the identifier</param>
-		/// <param name="Options">Parsing options</param>
+		/// <param name="tokenReader">Token reader</param>
+		/// <param name="initialIdentifier">Initial token of the identifier</param>
+		/// <param name="options">Parsing options</param>
 		/// <returns>Token reader</returns>
-		public static IUhtTokenReader RequireCppIdentifier(this IUhtTokenReader TokenReader, ref UhtToken InitialIdentifier, UhtCppIdentifierOptions Options = UhtCppIdentifierOptions.None)
+		public static IUhtTokenReader RequireCppIdentifier(this IUhtTokenReader tokenReader, ref UhtToken initialIdentifier, UhtCppIdentifierOptions options = UhtCppIdentifierOptions.None)
 		{
-			TokenReader.GetCppIdentifier(ref InitialIdentifier, Options);
-			return TokenReader;
+			tokenReader.GetCppIdentifier(ref initialIdentifier, options);
+			return tokenReader;
 		}
 
 		/// <summary>
 		/// Parse a required cpp identifier
 		/// </summary>
-		/// <param name="TokenReader">Token reader</param>
-		/// <param name="TokenListDelegate">Invoked when identifier is parsed</param>
+		/// <param name="tokenReader">Token reader</param>
+		/// <param name="tokenListDelegate">Invoked when identifier is parsed</param>
 		/// <returns>Token reader</returns>
-		public static IUhtTokenReader RequireCppIdentifier(this IUhtTokenReader TokenReader, UhtTokenListDelegate TokenListDelegate)
+		public static IUhtTokenReader RequireCppIdentifier(this IUhtTokenReader tokenReader, UhtTokenListDelegate tokenListDelegate)
 		{
-			TokenReader.RequireCppIdentifier(UhtCppIdentifierOptions.None, TokenListDelegate);
-			return TokenReader;
+			tokenReader.RequireCppIdentifier(UhtCppIdentifierOptions.None, tokenListDelegate);
+			return tokenReader;
 		}
 
 		/// <summary>
 		/// Parse a required cpp identifier
 		/// </summary>
-		/// <param name="TokenReader">Token reader</param>
-		/// <param name="Options">Parsing options</param>
-		/// <param name="TokenListDelegate">Invoked when identifier is parsed</param>
+		/// <param name="tokenReader">Token reader</param>
+		/// <param name="options">Parsing options</param>
+		/// <param name="tokenListDelegate">Invoked when identifier is parsed</param>
 		/// <returns>Token reader</returns>
-		public static IUhtTokenReader RequireCppIdentifier(this IUhtTokenReader TokenReader, UhtCppIdentifierOptions Options, UhtTokenListDelegate TokenListDelegate)
+		public static IUhtTokenReader RequireCppIdentifier(this IUhtTokenReader tokenReader, UhtCppIdentifierOptions options, UhtTokenListDelegate tokenListDelegate)
 		{
-			UhtTokenList TokenList = TokenReader.GetCppIdentifier(Options);
-			TokenListDelegate(TokenList);
-			UhtTokenListCache.Return(TokenList);
-			return TokenReader;
+			UhtTokenList tokenList = tokenReader.GetCppIdentifier(options);
+			tokenListDelegate(tokenList);
+			UhtTokenListCache.Return(tokenList);
+			return tokenReader;
 		}
 
 		/// <summary>
 		/// Parse a required cpp identifier
 		/// </summary>
-		/// <param name="TokenReader">Token reader</param>
-		/// <param name="InitialIdentifier">Initial token of the identifier</param>
-		/// <param name="TokenListDelegate">Invoked when identifier is parsed</param>
+		/// <param name="tokenReader">Token reader</param>
+		/// <param name="initialIdentifier">Initial token of the identifier</param>
+		/// <param name="tokenListDelegate">Invoked when identifier is parsed</param>
 		/// <returns>Token reader</returns>
-		public static IUhtTokenReader RequireCppIdentifier(this IUhtTokenReader TokenReader, ref UhtToken InitialIdentifier, UhtTokenListDelegate TokenListDelegate)
+		public static IUhtTokenReader RequireCppIdentifier(this IUhtTokenReader tokenReader, ref UhtToken initialIdentifier, UhtTokenListDelegate tokenListDelegate)
 		{
-			TokenReader.RequireCppIdentifier(ref InitialIdentifier, UhtCppIdentifierOptions.None, TokenListDelegate);
-			return TokenReader;
+			tokenReader.RequireCppIdentifier(ref initialIdentifier, UhtCppIdentifierOptions.None, tokenListDelegate);
+			return tokenReader;
 		}
 
 		/// <summary>
 		/// Parse a required cpp identifier
 		/// </summary>
-		/// <param name="TokenReader">Token reader</param>
-		/// <param name="InitialIdentifier">Initial token of the identifier</param>
-		/// <param name="Options">Parsing options</param>
-		/// <param name="TokenListDelegate">Invoked when identifier is parsed</param>
+		/// <param name="tokenReader">Token reader</param>
+		/// <param name="initialIdentifier">Initial token of the identifier</param>
+		/// <param name="options">Parsing options</param>
+		/// <param name="tokenListDelegate">Invoked when identifier is parsed</param>
 		/// <returns>Token reader</returns>
-		public static IUhtTokenReader RequireCppIdentifier(this IUhtTokenReader TokenReader, ref UhtToken InitialIdentifier, UhtCppIdentifierOptions Options, UhtTokenListDelegate TokenListDelegate)
+		public static IUhtTokenReader RequireCppIdentifier(this IUhtTokenReader tokenReader, ref UhtToken initialIdentifier, UhtCppIdentifierOptions options, UhtTokenListDelegate tokenListDelegate)
 		{
-			UhtTokenList TokenList = TokenReader.GetCppIdentifier(ref InitialIdentifier, Options);
-			TokenListDelegate(TokenList);
-			UhtTokenListCache.Return(TokenList);
-			return TokenReader;
+			UhtTokenList tokenList = tokenReader.GetCppIdentifier(ref initialIdentifier, options);
+			tokenListDelegate(tokenList);
+			UhtTokenListCache.Return(tokenList);
+			return tokenReader;
 		}
 
 		/// <summary>
 		/// Get a required cpp identifier
 		/// </summary>
-		/// <param name="TokenReader">Token reader</param>
-		/// <param name="Options">Parsing options</param>
+		/// <param name="tokenReader">Token reader</param>
+		/// <param name="options">Parsing options</param>
 		/// <returns>Token list</returns>
-		public static UhtTokenList GetCppIdentifier(this IUhtTokenReader TokenReader, UhtCppIdentifierOptions Options = UhtCppIdentifierOptions.None)
+		public static UhtTokenList GetCppIdentifier(this IUhtTokenReader tokenReader, UhtCppIdentifierOptions options = UhtCppIdentifierOptions.None)
 		{
-			UhtToken Token = TokenReader.GetIdentifier();
-			return TokenReader.GetCppIdentifier(ref Token, Options);
+			UhtToken token = tokenReader.GetIdentifier();
+			return tokenReader.GetCppIdentifier(ref token, options);
 		}
 
 		/// <summary>
 		/// Get a required cpp identifier
 		/// </summary>
-		/// <param name="TokenReader">Token reader</param>
-		/// <param name="InitialIdentifier">Initial token of the identifier</param>
-		/// <param name="Options">Parsing options</param>
+		/// <param name="tokenReader">Token reader</param>
+		/// <param name="initialIdentifier">Initial token of the identifier</param>
+		/// <param name="options">Parsing options</param>
 		/// <returns>Token list</returns>
-		public static UhtTokenList GetCppIdentifier(this IUhtTokenReader TokenReader, ref UhtToken InitialIdentifier, UhtCppIdentifierOptions Options = UhtCppIdentifierOptions.None)
+		public static UhtTokenList GetCppIdentifier(this IUhtTokenReader tokenReader, ref UhtToken initialIdentifier, UhtCppIdentifierOptions options = UhtCppIdentifierOptions.None)
 		{
-			UhtTokenList ListHead = UhtTokenListCache.Borrow(InitialIdentifier);
-			UhtTokenList ListTail = ListHead;
+			UhtTokenList listHead = UhtTokenListCache.Borrow(initialIdentifier);
+			UhtTokenList listTail = listHead;
 
-			if (Options.HasAnyFlags(UhtCppIdentifierOptions.AllowTemplates))
+			if (options.HasAnyFlags(UhtCppIdentifierOptions.AllowTemplates))
 			{
 				while (true)
 				{
-					if (TokenReader.TryPeekOptional('<'))
+					if (tokenReader.TryPeekOptional('<'))
 					{
-						ListTail.Next = UhtTokenListCache.Borrow(TokenReader.GetToken());
-						ListTail = ListTail.Next;
-						int NestedScopes = 1;
-						while (NestedScopes > 0)
+						listTail.Next = UhtTokenListCache.Borrow(tokenReader.GetToken());
+						listTail = listTail.Next;
+						int nestedScopes = 1;
+						while (nestedScopes > 0)
 						{
-							UhtToken Token = TokenReader.GetToken();
-							if (Token.TokenType.IsEndType())
+							UhtToken token = tokenReader.GetToken();
+							if (token.TokenType.IsEndType())
 							{
-								throw new UhtTokenException(TokenReader, Token, new string[] { "<", ">" }, "template");
+								throw new UhtTokenException(tokenReader, token, new string[] { "<", ">" }, "template");
 							}
-							ListTail.Next = UhtTokenListCache.Borrow(Token);
-							ListTail = ListTail.Next;
-							if (Token.IsSymbol('<'))
+							listTail.Next = UhtTokenListCache.Borrow(token);
+							listTail = listTail.Next;
+							if (token.IsSymbol('<'))
 							{
-								++NestedScopes;
+								++nestedScopes;
 							}
-							else if (Token.IsSymbol('>'))
+							else if (token.IsSymbol('>'))
 							{
-								--NestedScopes;
+								--nestedScopes;
 							}
 						}
 					}
 
-					if (!TokenReader.TryPeekOptional("::"))
+					if (!tokenReader.TryPeekOptional("::"))
 					{
 						break;
 					}
-					ListTail.Next = UhtTokenListCache.Borrow(TokenReader.GetToken());
-					ListTail = ListTail.Next;
-					ListTail.Next = UhtTokenListCache.Borrow(TokenReader.GetIdentifier());
-					ListTail = ListTail.Next;
+					listTail.Next = UhtTokenListCache.Borrow(tokenReader.GetToken());
+					listTail = listTail.Next;
+					listTail.Next = UhtTokenListCache.Borrow(tokenReader.GetIdentifier());
+					listTail = listTail.Next;
 				}
 			}
-			else if (TokenReader.PeekToken().IsSymbol("::"))
+			else if (tokenReader.PeekToken().IsSymbol("::"))
 			{
-				TokenReader.While("::", () =>
+				tokenReader.While("::", () =>
 				{
-					ListTail.Next = UhtTokenListCache.Borrow(TokenReader.GetIdentifier());
-					ListTail = ListTail.Next;
+					listTail.Next = UhtTokenListCache.Borrow(tokenReader.GetIdentifier());
+					listTail = listTail.Next;
 				});
 			}
-			return ListHead;
+			return listHead;
 		}
 	}
 }
diff --git a/Engine/Source/Programs/Shared/EpicGames.UHT/Tokenizer/UhtTokenReaderIntegerExtensions.cs b/Engine/Source/Programs/Shared/EpicGames.UHT/Tokenizer/UhtTokenReaderIntegerExtensions.cs
--- a/Engine/Source/Programs/Shared/EpicGames.UHT/Tokenizer/UhtTokenReaderIntegerExtensions.cs
+++ b/Engine/Source/Programs/Shared/EpicGames.UHT/Tokenizer/UhtTokenReaderIntegerExtensions.cs
@@ -9,9 +9,9 @@ namespace EpicGames.UHT.Tokenizer
 	/// <summary>
 	/// Notification of signed token being parsed
 	/// </summary>
-	/// <param name="Token">Token in question</param>
+	/// <param name="token">Token in question</param>
 	/// <returns>True if the token value is acceptable</returns>
-	public delegate bool UhtParseMergedSignToken(ref UhtToken Token);
+	public delegate bool UhtParseMergedSignToken(ref UhtToken token);
 
 	/// <summary>
 	/// Collection of helper methods to parse integers
@@ -21,156 +21,156 @@ namespace EpicGames.UHT.Tokenizer
 		/// <summary>
 		/// Get the next token as an integer.  If the next token is not an integer, no token is consumed.
 		/// </summary>
-		/// <param name="TokenReader">Token reader</param>
-		/// <param name="Value">The integer value of the token</param>
+		/// <param name="tokenReader">Token reader</param>
+		/// <param name="value">The integer value of the token</param>
 		/// <returns>True if the next token was an integer, false if not.</returns>
-		public static bool TryOptionalConstInt(this IUhtTokenReader TokenReader, out int Value)
+		public static bool TryOptionalConstInt(this IUhtTokenReader tokenReader, out int value)
 		{
-			ref UhtToken Token = ref TokenReader.PeekToken();
-			if (Token.GetConstInt(out Value))
+			ref UhtToken token = ref tokenReader.PeekToken();
+			if (token.GetConstInt(out value))
 			{
-				TokenReader.ConsumeToken();
+				tokenReader.ConsumeToken();
 				return true;
 			}
-			Value = 0;
+			value = 0;
 			return false;
 		}
 
 		/// <summary>
 		/// Get the next token as an integer.  If the next token is not an integer, an exception is thrown
 		/// </summary>
-		/// <param name="TokenReader">Token reader</param>
+		/// <param name="tokenReader">Token reader</param>
 		/// <returns>The value of the constant</returns>
-		public static IUhtTokenReader OptionalConstInt(this IUhtTokenReader TokenReader)
+		public static IUhtTokenReader OptionalConstInt(this IUhtTokenReader tokenReader)
 		{
-			TokenReader.TryOptionalConstInt(out int _);
-			return TokenReader;
+			tokenReader.TryOptionalConstInt(out int _);
+			return tokenReader;
 		}
 
 		/// <summary>
 		/// Get the next token as an integer.  If the next token is not an integer, an exception is thrown
 		/// </summary>
-		/// <param name="TokenReader">Token reader</param>
-		/// <param name="ExceptionContext">If not null, an exception will be thrown with the given text as part of the message.</param>
+		/// <param name="tokenReader">Token reader</param>
+		/// <param name="exceptionContext">If not null, an exception will be thrown with the given text as part of the message.</param>
 		/// <returns>The value of the constant</returns>
-		public static IUhtTokenReader RequireConstInt(this IUhtTokenReader TokenReader, object? ExceptionContext = null)
+		public static IUhtTokenReader RequireConstInt(this IUhtTokenReader tokenReader, object? exceptionContext = null)
 		{
-			if (!TokenReader.TryOptionalConstInt(out int _))
+			if (!tokenReader.TryOptionalConstInt(out int _))
 			{
-				throw new UhtTokenException(TokenReader, TokenReader.PeekToken(), "constant integer", ExceptionContext);
+				throw new UhtTokenException(tokenReader, tokenReader.PeekToken(), "constant integer", exceptionContext);
 			}
-			return TokenReader;
+			return tokenReader;
 		}
 
 		/// <summary>
 		/// Get the next token as an integer.  If the next token is not an integer, an exception is thrown
 		/// </summary>
-		/// <param name="TokenReader">Token reader</param>
-		/// <param name="ExceptionContext">If not null, an exception will be thrown with the given text as part of the message.</param>
+		/// <param name="tokenReader">Token reader</param>
+		/// <param name="exceptionContext">If not null, an exception will be thrown with the given text as part of the message.</param>
 		/// <returns>The value of the constant</returns>
-		public static int GetConstInt(this IUhtTokenReader TokenReader, object? ExceptionContext = null)
+		public static int GetConstInt(this IUhtTokenReader tokenReader, object? exceptionContext = null)
 		{
-			int Value;
-			if (!TokenReader.TryOptionalConstInt(out Value))
+			int value;
+			if (!tokenReader.TryOptionalConstInt(out value))
 			{
-				throw new UhtTokenException(TokenReader, TokenReader.PeekToken(), "constant integer", ExceptionContext);
+				throw new UhtTokenException(tokenReader, tokenReader.PeekToken(), "constant integer", exceptionContext);
 			}
-			return Value;
+			return value;
 		}
 
 		/// <summary>
 		/// Get the next token as an integer.  If the next token is not an integer, no token is consumed.
 		/// </summary>
-		/// <param name="TokenReader">Token reader</param>
-		/// <param name="Value">The integer value of the token</param>
+		/// <param name="tokenReader">Token reader</param>
+		/// <param name="value">The integer value of the token</param>
 		/// <returns>True if the next token was an integer, false if not.</returns>
-		public static bool TryOptionalConstLong(this IUhtTokenReader TokenReader, out long Value)
+		public static bool TryOptionalConstLong(this IUhtTokenReader tokenReader, out long value)
 		{
-			ref UhtToken Token = ref TokenReader.PeekToken();
-			if (Token.GetConstLong(out Value))
+			ref UhtToken token = ref tokenReader.PeekToken();
+			if (token.GetConstLong(out value))
 			{
-				TokenReader.ConsumeToken();
+				tokenReader.ConsumeToken();
 				return true;
 			}
-			Value = 0;
+			value = 0;
 			return false;
 		}
 
 		/// <summary>
 		/// Get the next token as an integer.  If the next token is not an integer, an exception is thrown
 		/// </summary>
-		/// <param name="TokenReader">Token reader</param>
+		/// <param name="tokenReader">Token reader</param>
 		/// <returns>The value of the constant</returns>
-		public static IUhtTokenReader OptionalConstLong(this IUhtTokenReader TokenReader)
+		public static IUhtTokenReader OptionalConstLong(this IUhtTokenReader tokenReader)
 		{
-			ref UhtToken Token = ref TokenReader.PeekToken();
+			ref UhtToken token = ref tokenReader.PeekToken();
 
-			if (Token.GetConstLong(out long _))
+			if (token.GetConstLong(out long _))
 			{
-				TokenReader.ConsumeToken();
+				tokenReader.ConsumeToken();
 			}
-			return TokenReader;
+			return tokenReader;
 		}
 
 		/// <summary>
 		/// Get the next token as an integer.  If the next token is not an integer, an exception is thrown
 		/// </summary>
-		/// <param name="TokenReader">Token reader</param>
-		/// <param name="ExceptionContext">If not null, an exception will be thrown with the given text as part of the message.</param>
+		/// <param name="tokenReader">Token reader</param>
+		/// <param name="exceptionContext">If not null, an exception will be thrown with the given text as part of the message.</param>
 		/// <returns>The value of the constant</returns>
-		public static IUhtTokenReader RequireConstLong(this IUhtTokenReader TokenReader, object? ExceptionContext = null)
+		public static IUhtTokenReader RequireConstLong(this IUhtTokenReader tokenReader, object? exceptionContext = null)
 		{
-			if (!TokenReader.TryOptionalConstLong(out long _))
+			if (!tokenReader.TryOptionalConstLong(out long _))
 			{
-				throw new UhtTokenException(TokenReader, TokenReader.PeekToken(), "constant long integer", ExceptionContext);
+				throw new UhtTokenException(tokenReader, tokenReader.PeekToken(), "constant long integer", exceptionContext);
 			}
-			return TokenReader;
+			return tokenReader;
 		}
 
 		/// <summary>
 		/// Get the next token as an integer.  If the next token is not an integer, an exception is thrown
 		/// </summary>
-		/// <param name="TokenReader">Token reader</param>
-		/// <param name="ExceptionContext">If not null, an exception will be thrown with the given text as part of the message.</param>
+		/// <param name="tokenReader">Token reader</param>
+		/// <param name="exceptionContext">If not null, an exception will be thrown with the given text as part of the message.</param>
 		/// <returns>The value of the constant</returns>
-		public static long GetConstLong(this IUhtTokenReader TokenReader, object? ExceptionContext = null)
+		public static long GetConstLong(this IUhtTokenReader tokenReader, object? exceptionContext = null)
 		{
-			long Value;
-			if (!TokenReader.TryOptionalConstLong(out Value))
+			long value;
+			if (!tokenReader.TryOptionalConstLong(out value))
 			{
-				throw new UhtTokenException(TokenReader, TokenReader.PeekToken(), "constant long integer", ExceptionContext);
+				throw new UhtTokenException(tokenReader, tokenReader.PeekToken(), "constant long integer", exceptionContext);
 			}
-			return Value;
+			return value;
 		}
 
 		/// <summary>
 		/// Helper method to combine any leading sign with the next numeric token
 		/// </summary>
-		/// <param name="TokenReader">Source tokens</param>
-		/// <param name="Delegate">Delegate to invoke with the merged value</param>
+		/// <param name="tokenReader">Source tokens</param>
+		/// <param name="tokenDelegate">Delegate to invoke with the merged value</param>
 		/// <returns>True if the next token was an parsed, false if not.</returns>
-		public static bool TryOptionalLeadingSignConstNumeric(this IUhtTokenReader TokenReader, UhtParseMergedSignToken Delegate)
+		public static bool TryOptionalLeadingSignConstNumeric(this IUhtTokenReader tokenReader, UhtParseMergedSignToken tokenDelegate)
 		{
-			using (var SavedState = new UhtTokenSaveState(TokenReader))
+			using (UhtTokenSaveState savedState = new UhtTokenSaveState(tokenReader))
 			{
 				// Check for a leading sign token
-				char Sign = ' ';
-				UhtToken Token = TokenReader.PeekToken();
-				if (Token.IsSymbol() && Token.Value.Length == 1 && UhtFCString.IsSign(Token.Value.Span[0]))
+				char sign = ' ';
+				UhtToken token = tokenReader.PeekToken();
+				if (token.IsSymbol() && token.Value.Length == 1 && UhtFCString.IsSign(token.Value.Span[0]))
 				{
-					Sign = Token.Value.Span[0];
-					TokenReader.ConsumeToken();
-					Token = TokenReader.PeekToken();
-					if (UhtFCString.IsSign(Token.Value.Span[0]))
+					sign = token.Value.Span[0];
+					tokenReader.ConsumeToken();
+					token = tokenReader.PeekToken();
+					if (UhtFCString.IsSign(token.Value.Span[0]))
 					{
 						return false;
 					}
-					Token.Value = new StringView($"{Sign}{Token.Value}");
+					token.Value = new StringView($"{sign}{token.Value}");
 				}
-				if (Delegate(ref Token))
+				if (tokenDelegate(ref token))
 				{
-					TokenReader.ConsumeToken();
-					SavedState.AbandonState();
+					tokenReader.ConsumeToken();
+					savedState.AbandonState();
 					return true;
 				}
 				return false;
@@ -180,65 +180,65 @@ namespace EpicGames.UHT.Tokenizer
 		/// <summary>
 		/// Get the next integer.  It also handled [+/-] token followed by an integer.
 		/// </summary>
-		/// <param name="TokenReader">Source tokens</param>
-		/// <param name="Value">The integer value of the token</param>
+		/// <param name="tokenReader">Source tokens</param>
+		/// <param name="value">The integer value of the token</param>
 		/// <returns>True if the next token was an integer, false if not.</returns>
-		public static bool TryOptionalConstIntExpression(this IUhtTokenReader TokenReader, out int Value)
+		public static bool TryOptionalConstIntExpression(this IUhtTokenReader tokenReader, out int value)
 		{
-			int LocalValue = 0;
-			bool Results = TokenReader.TryOptionalLeadingSignConstNumeric((ref UhtToken Token) =>
+			int localValue = 0;
+			bool results = tokenReader.TryOptionalLeadingSignConstNumeric((ref UhtToken token) =>
 			{
-				return Token.IsConstInt() && Token.GetConstInt(out LocalValue);
+				return token.IsConstInt() && token.GetConstInt(out localValue);
 			});
-			Value = LocalValue;
-			return Results;
+			value = localValue;
+			return results;
 		}
 
 		/// <summary>
 		/// Get the next integer.  It also handled [+/-] token followed by an integer.
 		/// </summary>
-		/// <param name="TokenReader">Source tokens</param>
+		/// <param name="tokenReader">Source tokens</param>
 		/// <returns>The integer value</returns>
-		public static int GetConstIntExpression(this IUhtTokenReader TokenReader)
+		public static int GetConstIntExpression(this IUhtTokenReader tokenReader)
 		{
-			int LocalValue = 0;
-			bool Results = TokenReader.TryOptionalLeadingSignConstNumeric((ref UhtToken Token) =>
+			int localValue = 0;
+			bool results = tokenReader.TryOptionalLeadingSignConstNumeric((ref UhtToken token) =>
 			{
-				return Token.IsConstInt() && Token.GetConstInt(out LocalValue);
+				return token.IsConstInt() && token.GetConstInt(out localValue);
 			});
-			return LocalValue;
+			return localValue;
 		}
 
 		/// <summary>
 		/// Get the next integer.  It also handled [+/-] token followed by an integer.
 		/// </summary>
-		/// <param name="TokenReader">Source tokens</param>
-		/// <param name="Value">The integer value of the token</param>
+		/// <param name="tokenReader">Source tokens</param>
+		/// <param name="value">The integer value of the token</param>
 		/// <returns>True if the next token was an integer, false if not.</returns>
-		public static bool TryOptionalConstLongExpression(this IUhtTokenReader TokenReader, out long Value)
+		public static bool TryOptionalConstLongExpression(this IUhtTokenReader tokenReader, out long value)
 		{
-			long LocalValue = 0;
-			bool Results = TokenReader.TryOptionalLeadingSignConstNumeric((ref UhtToken Token) =>
+			long localValue = 0;
+			bool results = tokenReader.TryOptionalLeadingSignConstNumeric((ref UhtToken token) =>
 			{
-				return Token.IsConstInt() && Token.GetConstLong(out LocalValue);
+				return token.IsConstInt() && token.GetConstLong(out localValue);
 			});
-			Value = LocalValue;
-			return Results;
+			value = localValue;
+			return results;
 		}
 
 		/// <summary>
 		/// Get the next long.  It also handled [+/-] token followed by an long.
 		/// </summary>
-		/// <param name="TokenReader">Source tokens</param>
+		/// <param name="tokenReader">Source tokens</param>
 		/// <returns>The long value</returns>
-		public static int GetConstLongExpression(this IUhtTokenReader TokenReader)
+		public static int GetConstLongExpression(this IUhtTokenReader tokenReader)
 		{
-			int LocalValue = 0;
-			bool Results = TokenReader.TryOptionalLeadingSignConstNumeric((ref UhtToken Token) =>
+			int localValue = 0;
+			bool results = tokenReader.TryOptionalLeadingSignConstNumeric((ref UhtToken token) =>
 			{
-				return Token.IsConstInt() && Token.GetConstInt(out LocalValue);
+				return token.IsConstInt() && token.GetConstInt(out localValue);
 			});
-			return LocalValue;
+			return localValue;
 		}
 	}
 }
diff --git a/Engine/Source/Programs/Shared/EpicGames.UHT/Tokenizer/UhtTokenReaderSkipExtensions.cs b/Engine/Source/Programs/Shared/EpicGames.UHT/Tokenizer/UhtTokenReaderSkipExtensions.cs
--- a/Engine/Source/Programs/Shared/EpicGames.UHT/Tokenizer/UhtTokenReaderSkipExtensions.cs
+++ b/Engine/Source/Programs/Shared/EpicGames.UHT/Tokenizer/UhtTokenReaderSkipExtensions.cs
@@ -13,62 +13,62 @@ namespace EpicGames.UHT.Tokenizer
 		/// <summary>
 		/// Skip a token regardless of the type.
 		/// </summary>
-		/// <param name="TokenReader">Token reader</param>
+		/// <param name="tokenReader">Token reader</param>
 		/// <returns>Token reader</returns>
-		public static IUhtTokenReader SkipOne(this IUhtTokenReader TokenReader)
+		public static IUhtTokenReader SkipOne(this IUhtTokenReader tokenReader)
 		{
-			TokenReader.PeekToken();
-			TokenReader.ConsumeToken();
-			return TokenReader;
+			tokenReader.PeekToken();
+			tokenReader.ConsumeToken();
+			return tokenReader;
 		}
 
 		/// <summary>
 		/// Skip an alignas expression
 		/// </summary>
-		/// <param name="TokenReader">Token reader</param>
+		/// <param name="tokenReader">Token reader</param>
 		/// <returns>Token reader</returns>
-		public static IUhtTokenReader SkipAlignasIfNecessary(this IUhtTokenReader TokenReader)
+		public static IUhtTokenReader SkipAlignasIfNecessary(this IUhtTokenReader tokenReader)
 		{
 			const string Identifier = "alignas";
-			if (TokenReader.TryOptional(Identifier))
+			if (tokenReader.TryOptional(Identifier))
 			{
-				TokenReader
+				tokenReader
 					.Require('(', Identifier)
 					.RequireConstInt(Identifier)
 					.Require(')', Identifier);
 			}
-			return TokenReader;
+			return tokenReader;
 		}
 
 		/// <summary>
 		/// Skip deprecation macro
 		/// </summary>
-		/// <param name="TokenReader">Token reader</param>
+		/// <param name="tokenReader">Token reader</param>
 		/// <returns>Token reader</returns>
-		public static IUhtTokenReader SkipDeprecatedMacroIfNecessary(this IUhtTokenReader TokenReader)
+		public static IUhtTokenReader SkipDeprecatedMacroIfNecessary(this IUhtTokenReader tokenReader)
 		{
-			if (TokenReader.TryOptional(new string[] { "DEPRECATED", "UE_DEPRECATED" }) >= 0)
+			if (tokenReader.TryOptional(new string[] { "DEPRECATED", "UE_DEPRECATED" }) >= 0)
 			{
-				TokenReader
+				tokenReader
 					.Require('(', "deprecation macro")
 					.RequireConstFloat("version in deprecation macro")
 					.Require(',', "deprecation macro")
 					.RequireConstString("message in deprecation macro")
 					.Require(')', "deprecation macro");
 			}
-			return TokenReader;
+			return tokenReader;
 		}
 
 		/// <summary>
 		/// Skip alignas and/or deprecation macros
 		/// </summary>
-		/// <param name="TokenReader">Token reader</param>
+		/// <param name="tokenReader">Token reader</param>
 		/// <returns>Token reader</returns>
-		public static IUhtTokenReader SkipAlignasAndDeprecatedMacroIfNecessary(this IUhtTokenReader TokenReader)
+		public static IUhtTokenReader SkipAlignasAndDeprecatedMacroIfNecessary(this IUhtTokenReader tokenReader)
 		{
 			// alignas() can come before or after the deprecation macro.
 			// We can't have both, but the compiler will catch that anyway.
-			return TokenReader
+			return tokenReader
 				.SkipAlignasIfNecessary()
 				.SkipDeprecatedMacroIfNecessary()
 				.SkipAlignasIfNecessary();
@@ -77,64 +77,64 @@ namespace EpicGames.UHT.Tokenizer
 		/// <summary>
 		/// Skip any block of tokens wrapped by the given token symbols
 		/// </summary>
-		/// <param name="TokenReader">Token reader</param>
-		/// <param name="Initiator">Initiating token (i.e. &quot;(&quot;)</param>
-		/// <param name="Terminator">Terminating token (i.e. &quot;)&quot;)</param>
-		/// <param name="InitialNesting">If true, start with an initial nesting count of one (assume we already parsed an initiator)</param>
-		/// <param name="ExceptionContext">Extra context for any errors</param>
+		/// <param name="tokenReader">Token reader</param>
+		/// <param name="initiator">Initiating token (i.e. &quot;(&quot;)</param>
+		/// <param name="terminator">Terminating token (i.e. &quot;)&quot;)</param>
+		/// <param name="initialNesting">If true, start with an initial nesting count of one (assume we already parsed an initiator)</param>
+		/// <param name="exceptionContext">Extra context for any errors</param>
 		/// <returns>Token reader</returns>
 		/// <exception cref="UhtTokenException">Throw if end of file is reached</exception>
-		public static IUhtTokenReader SkipBrackets(this IUhtTokenReader TokenReader, char Initiator, char Terminator, int InitialNesting, object? ExceptionContext = null)
+		public static IUhtTokenReader SkipBrackets(this IUhtTokenReader tokenReader, char initiator, char terminator, int initialNesting, object? exceptionContext = null)
 		{
-			int Nesting = InitialNesting;
-			if (Nesting == 0)
+			int nesting = initialNesting;
+			if (nesting == 0)
 			{
-				TokenReader.Require(Initiator, ExceptionContext);
-				++Nesting;
+				tokenReader.Require(initiator, exceptionContext);
+				++nesting;
 			}
 
 			do
 			{
-				UhtToken SkipToken = TokenReader.GetToken();
-				if (SkipToken.TokenType.IsEndType())
+				UhtToken skipToken = tokenReader.GetToken();
+				if (skipToken.TokenType.IsEndType())
 				{
-					throw new UhtTokenException(TokenReader, SkipToken, Terminator, ExceptionContext);
+					throw new UhtTokenException(tokenReader, skipToken, terminator, exceptionContext);
 				}
-				else if (SkipToken.IsSymbol(Initiator))
+				else if (skipToken.IsSymbol(initiator))
 				{
-					++Nesting;
+					++nesting;
 				}
-				else if (SkipToken.IsSymbol(Terminator))
+				else if (skipToken.IsSymbol(terminator))
 				{
-					--Nesting;
+					--nesting;
 				}
-			} while (Nesting != 0);
-			return TokenReader;
+			} while (nesting != 0);
+			return tokenReader;
 		}
 
 		/// <summary>
 		/// Skip tokens until the given terminator is found.  The terminator will not be consumed.
 		/// </summary>
-		/// <param name="TokenReader">Token reader</param>
-		/// <param name="Terminator">Terminator to skip until</param>
-		/// <param name="ExceptionContext">Extra context for any exceptions</param>
+		/// <param name="tokenReader">Token reader</param>
+		/// <param name="terminator">Terminator to skip until</param>
+		/// <param name="exceptionContext">Extra context for any exceptions</param>
 		/// <returns>Token reader</returns>
-		public static IUhtTokenReader SkipUntil(this IUhtTokenReader TokenReader, char Terminator, object? ExceptionContext = null)
+		public static IUhtTokenReader SkipUntil(this IUhtTokenReader tokenReader, char terminator, object? exceptionContext = null)
 		{
 			while (true)
 			{
-				ref UhtToken SkipToken = ref TokenReader.PeekToken();
-				if (SkipToken.TokenType.IsEndType())
+				ref UhtToken skipToken = ref tokenReader.PeekToken();
+				if (skipToken.TokenType.IsEndType())
 				{
-					throw new UhtTokenException(TokenReader, SkipToken, Terminator, ExceptionContext);
+					throw new UhtTokenException(tokenReader, skipToken, terminator, exceptionContext);
 				}
-				else if (SkipToken.IsSymbol(Terminator))
+				else if (skipToken.IsSymbol(terminator))
 				{
 					break;
 				}
-				TokenReader.ConsumeToken();
+				tokenReader.ConsumeToken();
 			}
-			return TokenReader;
+			return tokenReader;
 		}
 	}
 }
