More coding standard fixes

diff --git a/Engine/Source/Programs/Shared/EpicGames.UHT/Tokenizer/UhtTokenReaderStringExtensions.cs b/Engine/Source/Programs/Shared/EpicGames.UHT/Tokenizer/UhtTokenReaderStringExtensions.cs
--- a/Engine/Source/Programs/Shared/EpicGames.UHT/Tokenizer/UhtTokenReaderStringExtensions.cs
+++ b/Engine/Source/Programs/Shared/EpicGames.UHT/Tokenizer/UhtTokenReaderStringExtensions.cs
@@ -14,111 +14,111 @@ namespace EpicGames.UHT.Tokenizer
 		/// <summary>
 		/// Get the next token as a string.  If the next token is not a string, no token is consumed.
 		/// </summary>
-		/// <param name="TokenReader">Token reader</param>
-		/// <param name="Value">The string value of the token</param>
+		/// <param name="tokenReader">Token reader</param>
+		/// <param name="value">The string value of the token</param>
 		/// <returns>True if the next token was an string, false if not.</returns>
-		public static bool TryOptionalConstString(this IUhtTokenReader TokenReader, out StringView Value)
+		public static bool TryOptionalConstString(this IUhtTokenReader tokenReader, out StringView value)
 		{
-			ref UhtToken Token = ref TokenReader.PeekToken();
-			if (Token.IsConstString())
+			ref UhtToken token = ref tokenReader.PeekToken();
+			if (token.IsConstString())
 			{
-				Value = Token.GetTokenString();
-				TokenReader.ConsumeToken();
+				value = token.GetTokenString();
+				tokenReader.ConsumeToken();
 				return true;
 			}
-			Value = "";
+			value = "";
 			return false;
 		}
 
 		/// <summary>
 		/// Get the next token as a string.  If the next token is not a string, no token is consumed.
 		/// </summary>
-		/// <param name="TokenReader">Token reader</param>
+		/// <param name="tokenReader">Token reader</param>
 		/// <returns>True if the next token was an string, false if not.</returns>
-		public static IUhtTokenReader OptionalConstString(this IUhtTokenReader TokenReader)
+		public static IUhtTokenReader OptionalConstString(this IUhtTokenReader tokenReader)
 		{
-			ref UhtToken Token = ref TokenReader.PeekToken();
-			if (Token.IsConstString())
+			ref UhtToken token = ref tokenReader.PeekToken();
+			if (token.IsConstString())
 			{
-				TokenReader.ConsumeToken();
+				tokenReader.ConsumeToken();
 			}
-			return TokenReader;
+			return tokenReader;
 		}
 
 		/// <summary>
 		/// Verify that the next token is a string.
 		/// </summary>
-		/// <param name="TokenReader">Token reader</param>
-		/// <param name="ExceptionContext">If not null, an exception will be thrown with the given text as part of the message.</param>
+		/// <param name="tokenReader">Token reader</param>
+		/// <param name="exceptionContext">If not null, an exception will be thrown with the given text as part of the message.</param>
 		/// <returns>True if the next token was a string, false if not.</returns>
-		public static IUhtTokenReader RequireConstString(this IUhtTokenReader TokenReader, object? ExceptionContext)
+		public static IUhtTokenReader RequireConstString(this IUhtTokenReader tokenReader, object? exceptionContext)
 		{
-			ref UhtToken Token = ref TokenReader.PeekToken();
-			if (Token.IsConstString())
+			ref UhtToken token = ref tokenReader.PeekToken();
+			if (token.IsConstString())
 			{
-				TokenReader.ConsumeToken();
-				return TokenReader;
+				tokenReader.ConsumeToken();
+				return tokenReader;
 			}
-			throw new UhtTokenException(TokenReader, Token, "constant string", ExceptionContext);
+			throw new UhtTokenException(tokenReader, token, "constant string", exceptionContext);
 		}
 
 		/// <summary>
 		/// Get the next token as a string.  If the next token is not a string, an exception is thrown
 		/// </summary>
-		/// <param name="TokenReader">Token reader</param>
-		/// <param name="ExceptionContext">If not null, an exception will be thrown with the given text as part of the message.</param>
+		/// <param name="tokenReader">Token reader</param>
+		/// <param name="exceptionContext">If not null, an exception will be thrown with the given text as part of the message.</param>
 		/// <returns>The value of the string.</returns>
-		public static StringView GetConstString(this IUhtTokenReader TokenReader, object? ExceptionContext = null)
+		public static StringView GetConstString(this IUhtTokenReader tokenReader, object? exceptionContext = null)
 		{
-			ref UhtToken Token = ref TokenReader.PeekToken();
-			if (Token.IsConstString())
+			ref UhtToken token = ref tokenReader.PeekToken();
+			if (token.IsConstString())
 			{
-				StringView Output = Token.GetTokenString();
-				TokenReader.ConsumeToken();
-				return Output;
+				StringView output = token.GetTokenString();
+				tokenReader.ConsumeToken();
+				return output;
 			}
-			throw new UhtTokenException(TokenReader, Token, "constant string", ExceptionContext);
+			throw new UhtTokenException(tokenReader, token, "constant string", exceptionContext);
 		}
 
 		/// <summary>
 		/// Get the next token as a quoted string.  If the next token is not a string, an exception is thrown.
 		/// Character constants are not considered strings by this routine.
 		/// </summary>
-		/// <param name="TokenReader">Token reader</param>
-		/// <param name="ExceptionContext">If not null, an exception will be thrown with the given text as part of the message.</param>
+		/// <param name="tokenReader">Token reader</param>
+		/// <param name="exceptionContext">If not null, an exception will be thrown with the given text as part of the message.</param>
 		/// <returns>The value of the string.</returns>
-		public static StringView GetConstQuotedString(this IUhtTokenReader TokenReader, object? ExceptionContext = null)
+		public static StringView GetConstQuotedString(this IUhtTokenReader tokenReader, object? exceptionContext = null)
 		{
-			ref UhtToken Token = ref TokenReader.PeekToken();
-			if (Token.TokenType == UhtTokenType.StringConst)
+			ref UhtToken token = ref tokenReader.PeekToken();
+			if (token.TokenType == UhtTokenType.StringConst)
 			{
-				StringView Output = Token.Value;
-				TokenReader.ConsumeToken();
-				return Output;
+				StringView output = token.Value;
+				tokenReader.ConsumeToken();
+				return output;
 			}
-			throw new UhtTokenException(TokenReader, Token, "constant quoted string", ExceptionContext);
+			throw new UhtTokenException(tokenReader, token, "constant quoted string", exceptionContext);
 		}
 
 		/// <summary>
 		/// Get a const string that can optionally be wrapped with a TEXT() macro
 		/// </summary>
-		/// <param name="TokenReader">Token reader</param>
-		/// <param name="ExceptionContext">If not null, an exception will be thrown with the given text as part of the message.</param>
+		/// <param name="tokenReader">Token reader</param>
+		/// <param name="exceptionContext">If not null, an exception will be thrown with the given text as part of the message.</param>
 		/// <returns>The value of the string</returns>
-		public static StringView GetWrappedConstString(this IUhtTokenReader TokenReader, object? ExceptionContext = null)
+		public static StringView GetWrappedConstString(this IUhtTokenReader tokenReader, object? exceptionContext = null)
 		{
-			ref UhtToken Token = ref TokenReader.PeekToken();
-			if (Token.IsIdentifier("TEXT"))
+			ref UhtToken token = ref tokenReader.PeekToken();
+			if (token.IsIdentifier("TEXT"))
 			{
-				TokenReader.ConsumeToken();
-				TokenReader.Require('(');
-				StringView Out = TokenReader.GetConstString(ExceptionContext);
-				TokenReader.Require(')');
-				return Out;
+				tokenReader.ConsumeToken();
+				tokenReader.Require('(');
+				StringView output = tokenReader.GetConstString(exceptionContext);
+				tokenReader.Require(')');
+				return output;
 			}
 			else
 			{
-				return TokenReader.GetConstString(ExceptionContext);
+				return tokenReader.GetConstString(exceptionContext);
 			}
 		}
 	}
diff --git a/Engine/Source/Programs/Shared/EpicGames.UHT/Tokenizer/UhtTokenReaderUtilityExtensions.cs b/Engine/Source/Programs/Shared/EpicGames.UHT/Tokenizer/UhtTokenReaderUtilityExtensions.cs
--- a/Engine/Source/Programs/Shared/EpicGames.UHT/Tokenizer/UhtTokenReaderUtilityExtensions.cs
+++ b/Engine/Source/Programs/Shared/EpicGames.UHT/Tokenizer/UhtTokenReaderUtilityExtensions.cs
@@ -1,9 +1,9 @@
 // Copyright Epic Games, Inc. All Rights Reserved.
 
-using EpicGames.Core;
-using EpicGames.UHT.Utils;
 using System;
 using System.Collections.Generic;
+using EpicGames.Core;
+using EpicGames.UHT.Utils;
 
 namespace EpicGames.UHT.Tokenizer
 {
@@ -13,7 +13,7 @@ namespace EpicGames.UHT.Tokenizer
 	/// </summary>
 	public static class UhtTokenReaderUtilityExtensions
 	{
-		private static readonly HashSet<StringView> SkipDeclarationWarningStrings = new HashSet<StringView>
+		private static readonly HashSet<StringView> s_skipDeclarationWarningStrings = new HashSet<StringView>
 		{
 			"GENERATED_BODY",
 			"GENERATED_IINTERFACE_BODY",
@@ -34,23 +34,23 @@ namespace EpicGames.UHT.Tokenizer
 		/// <summary>
 		/// When processing type, make sure that the next token is the expected token
 		/// </summary>
-		/// <param name="TokenReader">Token reader</param>
-		/// <param name="ExpectedIdentifier">Expected identifier</param>
-		/// <param name="bIsMember">If true, log an error if the type begins with const</param>
+		/// <param name="tokenReader">Token reader</param>
+		/// <param name="expectedIdentifier">Expected identifier</param>
+		/// <param name="isMember">If true, log an error if the type begins with const</param>
 		/// <returns>true if there could be more header to process, false if the end was reached.</returns>
-		public static bool SkipExpectedType(this IUhtTokenReader TokenReader, StringView ExpectedIdentifier, bool bIsMember)
+		public static bool SkipExpectedType(this IUhtTokenReader tokenReader, StringView expectedIdentifier, bool isMember)
 		{
-			if (TokenReader.TryOptional(ExpectedIdentifier))
+			if (tokenReader.TryOptional(expectedIdentifier))
 			{
 				return true;
 			}
-			if (bIsMember && TokenReader.TryOptional("const"))
+			if (isMember && tokenReader.TryOptional("const"))
 			{
-				TokenReader.LogError("Const properties are not supported.");
+				tokenReader.LogError("Const properties are not supported.");
 			}
 			else
 			{
-				TokenReader.LogError($"Inappropriate keyword '{TokenReader.PeekToken().Value}' on variable of type '{ExpectedIdentifier}'");
+				tokenReader.LogError($"Inappropriate keyword '{tokenReader.PeekToken().Value}' on variable of type '{expectedIdentifier}'");
 			}
 			return false;
 		}
@@ -58,232 +58,232 @@ namespace EpicGames.UHT.Tokenizer
 		/// <summary>
 		/// Try to parse an optional _API macro
 		/// </summary>
-		/// <param name="TokenReader">Token reader</param>
-		/// <param name="APIMacroToken">_API macro parsed</param>
+		/// <param name="tokenReader">Token reader</param>
+		/// <param name="apiMacroToken">_API macro parsed</param>
 		/// <returns>True if an _API macro was parsed</returns>
-		public static bool TryOptionalAPIMacro(this IUhtTokenReader TokenReader, out UhtToken APIMacroToken)
+		public static bool TryOptionalAPIMacro(this IUhtTokenReader tokenReader, out UhtToken apiMacroToken)
 		{
-			ref UhtToken Token = ref TokenReader.PeekToken();
-			if (Token.IsIdentifier() && Token.Value.Span.EndsWith("_API"))
+			ref UhtToken token = ref tokenReader.PeekToken();
+			if (token.IsIdentifier() && token.Value.Span.EndsWith("_API"))
 			{
-				APIMacroToken = Token;
-				TokenReader.ConsumeToken();
+				apiMacroToken = token;
+				tokenReader.ConsumeToken();
 				return true;
 			}
-			APIMacroToken = new UhtToken();
+			apiMacroToken = new UhtToken();
 			return false;
 		}
 
 		/// <summary>
 		/// Parse an optional single inheritance
 		/// </summary>
-		/// <param name="TokenReader">Token reader</param>
-		/// <param name="SuperClassDelegate">Invoked with the inherited type name</param>
+		/// <param name="tokenReader">Token reader</param>
+		/// <param name="superClassDelegate">Invoked with the inherited type name</param>
 		/// <returns>Token reader</returns>
-		public static IUhtTokenReader OptionalInheritance(this IUhtTokenReader TokenReader, UhtTokenDelegate SuperClassDelegate)
+		public static IUhtTokenReader OptionalInheritance(this IUhtTokenReader tokenReader, UhtTokenDelegate superClassDelegate)
 		{
-			TokenReader.Optional(':', () =>
+			tokenReader.Optional(':', () =>
 			{
-				TokenReader
+				tokenReader
 					.Require("public", "public access modifier")
-					.RequireIdentifier(SuperClassDelegate);
+					.RequireIdentifier(superClassDelegate);
 			});
-			return TokenReader;
+			return tokenReader;
 		}
 
 		/// <summary>
 		/// Parse an optional inheritance
 		/// </summary>
-		/// <param name="TokenReader">Token reader</param>
-		/// <param name="SuperClassDelegate">Invoked with the inherited type name</param>
-		/// <param name="BaseClassDelegate">Invoked when other base classes are parsed</param>
+		/// <param name="tokenReader">Token reader</param>
+		/// <param name="superClassDelegate">Invoked with the inherited type name</param>
+		/// <param name="baseClassDelegate">Invoked when other base classes are parsed</param>
 		/// <returns>Token reader</returns>
-		public static IUhtTokenReader OptionalInheritance(this IUhtTokenReader TokenReader, UhtTokenDelegate SuperClassDelegate, UhtTokenListDelegate BaseClassDelegate)
+		public static IUhtTokenReader OptionalInheritance(this IUhtTokenReader tokenReader, UhtTokenDelegate superClassDelegate, UhtTokenListDelegate baseClassDelegate)
 		{
-			TokenReader.Optional(':', () =>
+			tokenReader.Optional(':', () =>
 			{
-				TokenReader
+				tokenReader
 					.Require("public", "public access modifier")
-					.RequireIdentifier(SuperClassDelegate)
+					.RequireIdentifier(superClassDelegate)
 					.While(',', () =>
 					{
-						TokenReader
+						tokenReader
 							.Require("public", "public interface access specifier")
-							.RequireCppIdentifier(UhtCppIdentifierOptions.AllowTemplates, BaseClassDelegate);
+							.RequireCppIdentifier(UhtCppIdentifierOptions.AllowTemplates, baseClassDelegate);
 					});
 			});
-			return TokenReader;
+			return tokenReader;
 		}
 
 		/// <summary>
 		/// Given a declaration/statement that starts with the given token, skip the declaration in the header.
 		/// </summary>
-		/// <param name="TokenReader">Token reader</param>
-		/// <param name="Token">Token that started the process</param>
+		/// <param name="tokenReader">Token reader</param>
+		/// <param name="token">Token that started the process</param>
 		/// <returns>true if there could be more header to process, false if the end was reached.</returns>
-		public static bool SkipDeclaration(this IUhtTokenReader TokenReader, ref UhtToken Token)
+		public static bool SkipDeclaration(this IUhtTokenReader tokenReader, ref UhtToken token)
 		{
 			// Consume all tokens until the end of declaration/definition has been found.
-			int NestedScopes = 0;
-			bool bEndOfDeclarationFound = false;
+			int nestedScopes = 0;
+			bool endOfDeclarationFound = false;
 
 			// Store the current value of PrevComment so it can be restored after we parsed everything.
-			using (UhtTokenDisableComments DisableComments = new UhtTokenDisableComments(TokenReader))
+			using (UhtTokenDisableComments disableComments = new UhtTokenDisableComments(tokenReader))
 			{
 
 				// Check if this is a class/struct declaration in which case it can be followed by member variable declaration.	
-				bool bPossiblyClassDeclaration = Token.IsIdentifier() && (Token.IsValue("class") || Token.IsValue("struct"));
+				bool possiblyClassDeclaration = token.IsIdentifier() && (token.IsValue("class") || token.IsValue("struct"));
 
 				// (known) macros can end without ; or } so use () to find the end of the declaration.
 				// However, we don't want to use it with DECLARE_FUNCTION, because we need it to be treated like a function.
-				bool bMacroDeclaration = IsProbablyAMacro(Token.Value) && !Token.IsIdentifier("DECLARE_FUNCTION");
+				bool macroDeclaration = IsProbablyAMacro(token.Value) && !token.IsIdentifier("DECLARE_FUNCTION");
 
-				bool bDefinitionFound = false;
-				char OpeningBracket = bMacroDeclaration ? '(' : '{';
-				char ClosingBracket = bMacroDeclaration ? ')' : '}';
+				bool definitionFound = false;
+				char openingBracket = macroDeclaration ? '(' : '{';
+				char closingBracket = macroDeclaration ? ')' : '}';
 
-				bool bRetestCurrentToken = false;
+				bool retestCurrentToken = false;
 				while (true)
 				{
-					if (!bRetestCurrentToken)
+					if (!retestCurrentToken)
 					{
-						Token = TokenReader.GetToken();
-						if (Token.TokenType.IsEndType())
+						token = tokenReader.GetToken();
+						if (token.TokenType.IsEndType())
 						{
 							break;
 						}
 					}
 					else
 					{
-						bRetestCurrentToken = false;
+						retestCurrentToken = false;
 					}
-					ReadOnlySpan<char> Span = Token.Value.Span;
+					ReadOnlySpan<char> span = token.Value.Span;
 
 					// If this is a macro, consume it
 					// If we find parentheses at top-level and we think it's a class declaration then it's more likely
 					// to be something like: class UThing* GetThing();
-					if (bPossiblyClassDeclaration && NestedScopes == 0 && Token.IsSymbol() && Span.Length == 1 && Span[0] == '(')
+					if (possiblyClassDeclaration && nestedScopes == 0 && token.IsSymbol() && span.Length == 1 && span[0] == '(')
 					{
-						bPossiblyClassDeclaration = false;
+						possiblyClassDeclaration = false;
 					}
 
-					if (Token.IsSymbol() && Span.Length == 1 && Span[0] == ';' && NestedScopes == 0)
+					if (token.IsSymbol() && span.Length == 1 && span[0] == ';' && nestedScopes == 0)
 					{
-						bEndOfDeclarationFound = true;
+						endOfDeclarationFound = true;
 						break;
 					}
 
-					if (Token.IsIdentifier())
+					if (token.IsIdentifier())
 					{
 						// Use a trivial pre-filter to avoid doing the search on things that aren't UE keywords we care about
-						if (Span[0] == 'G' || Span[0] == 'R' || Span[0] == 'U')
+						if (span[0] == 'G' || span[0] == 'R' || span[0] == 'U')
 						{
-							if (SkipDeclarationWarningStrings.Contains(Token.Value))
+							if (s_skipDeclarationWarningStrings.Contains(token.Value))
 							{
-								TokenReader.LogWarning($"The identifier \'{Token.Value}\' was detected in a block being skipped. Was this intentional?");
+								tokenReader.LogWarning($"The identifier \'{token.Value}\' was detected in a block being skipped. Was this intentional?");
 							}
 						}
 					}
 
-					if (!bMacroDeclaration && Token.IsIdentifier() && Span.Equals("PURE_VIRTUAL", StringComparison.Ordinal) && NestedScopes == 0)
+					if (!macroDeclaration && token.IsIdentifier() && span.Equals("PURE_VIRTUAL", StringComparison.Ordinal) && nestedScopes == 0)
 					{
-						OpeningBracket = '(';
-						ClosingBracket = ')';
+						openingBracket = '(';
+						closingBracket = ')';
 					}
 
-					if (Token.IsSymbol() && Span.Length == 1 && Span[0] == OpeningBracket)
+					if (token.IsSymbol() && span.Length == 1 && span[0] == openingBracket)
 					{
 						// This is a function definition or class declaration.
-						bDefinitionFound = true;
-						NestedScopes++;
+						definitionFound = true;
+						nestedScopes++;
 					}
-					else if (Token.IsSymbol() && Span.Length == 1 && Span[0] == ClosingBracket)
+					else if (token.IsSymbol() && span.Length == 1 && span[0] == closingBracket)
 					{
-						NestedScopes--;
-						if (NestedScopes == 0)
+						nestedScopes--;
+						if (nestedScopes == 0)
 						{
 							// Could be a class declaration in all capitals, and not a macro
-							bool bReallyEndDeclaration = true;
-							if (bMacroDeclaration)
+							bool reallyEndDeclaration = true;
+							if (macroDeclaration)
 							{
-								bReallyEndDeclaration = !TokenReader.TryPeekOptional('{');
+								reallyEndDeclaration = !tokenReader.TryPeekOptional('{');
 							}
 
-							if (bReallyEndDeclaration)
+							if (reallyEndDeclaration)
 							{
-								bEndOfDeclarationFound = true;
+								endOfDeclarationFound = true;
 								break;
 							}
 						}
 
-						if (NestedScopes < 0)
+						if (nestedScopes < 0)
 						{
-							throw new UhtException(TokenReader, Token.InputLine, $"Unexpected '{ClosingBracket}'. Did you miss a semi-colon?");
+							throw new UhtException(tokenReader, token.InputLine, $"Unexpected '{closingBracket}'. Did you miss a semi-colon?");
 						}
 					}
-					else if (bMacroDeclaration && NestedScopes == 0)
+					else if (macroDeclaration && nestedScopes == 0)
 					{
-						bMacroDeclaration = false;
-						OpeningBracket = '{';
-						ClosingBracket = '}';
-						bRetestCurrentToken = true;
+						macroDeclaration = false;
+						openingBracket = '{';
+						closingBracket = '}';
+						retestCurrentToken = true;
 					}
 				}
-				if (bEndOfDeclarationFound)
+				if (endOfDeclarationFound)
 				{
 					// Member variable declaration after class declaration (see bPossiblyClassDeclaration).
-					if (bPossiblyClassDeclaration && bDefinitionFound)
+					if (possiblyClassDeclaration && definitionFound)
 					{
 						// Should syntax errors be also handled when someone declares a variable after function definition?
 						// Consume the variable name.
-						if (TokenReader.bIsEOF)
+						if (tokenReader.IsEOF)
 						{
 							return true;
 						}
-						if (TokenReader.TryOptionalIdentifier())
+						if (tokenReader.TryOptionalIdentifier())
 						{
-							TokenReader.Require(';');
+							tokenReader.Require(';');
 						}
 					}
 
 					// C++ allows any number of ';' after member declaration/definition.
-					while (TokenReader.TryOptional(';'))
-					{ 
+					while (tokenReader.TryOptional(';'))
+					{
 					}
 				}
 			}
 
 			// Successfully consumed C++ declaration unless mismatched pair of brackets has been found.
-			return NestedScopes == 0 && bEndOfDeclarationFound;
+			return nestedScopes == 0 && endOfDeclarationFound;
 		}
 
-		private static bool IsProbablyAMacro(StringView Identifier)
+		private static bool IsProbablyAMacro(StringView identifier)
 		{
-			ReadOnlySpan<char> Span = Identifier.Span;
-			if (Span.Length == 0)
+			ReadOnlySpan<char> span = identifier.Span;
+			if (span.Length == 0)
 			{
 				return false;
 			}
 
 			// Macros must start with a capitalized alphanumeric character or underscore
-			char FirstChar = Span[0];
-			if (FirstChar != '_' && (FirstChar < 'A' || FirstChar > 'Z'))
+			char firstChar = span[0];
+			if (firstChar != '_' && (firstChar < 'A' || firstChar > 'Z'))
 			{
 				return false;
 			}
 
 			// Test for known delegate and event macros.
-			if (Span.StartsWith("DECLARE_MULTICAST_DELEGATE") ||
-				Span.StartsWith("DECLARE_DELEGATE") ||
-				Span.StartsWith("DECLARE_EVENT"))
+			if (span.StartsWith("DECLARE_MULTICAST_DELEGATE") ||
+				span.StartsWith("DECLARE_DELEGATE") ||
+				span.StartsWith("DECLARE_EVENT"))
 			{
 				return true;
 			}
 
 			// Failing that, we'll guess about it being a macro based on it being a fully-capitalized identifier.
-			foreach (char Ch in Span.Slice(1))
+			foreach (char ch in span.Slice(1))
 			{
-				if (Ch != '_' && (Ch < 'A' || Ch > 'Z') && (Ch < '0' || Ch > '9'))
+				if (ch != '_' && (ch < 'A' || ch > 'Z') && (ch < '0' || ch > '9'))
 				{
 					return false;
 				}
diff --git a/Engine/Source/Programs/Shared/EpicGames.UHT/Tokenizer/UhtTokenReplayReader.cs b/Engine/Source/Programs/Shared/EpicGames.UHT/Tokenizer/UhtTokenReplayReader.cs
--- a/Engine/Source/Programs/Shared/EpicGames.UHT/Tokenizer/UhtTokenReplayReader.cs
+++ b/Engine/Source/Programs/Shared/EpicGames.UHT/Tokenizer/UhtTokenReplayReader.cs
@@ -1,10 +1,10 @@
 // Copyright Epic Games, Inc. All Rights Reserved.
 
-using EpicGames.Core;
-using EpicGames.UHT.Utils;
 using System;
 using System.Collections.Generic;
 using System.Threading;
+using EpicGames.Core;
+using EpicGames.UHT.Utils;
 
 namespace EpicGames.UHT.Tokenizer
 {
@@ -14,39 +14,39 @@ namespace EpicGames.UHT.Tokenizer
 	/// </summary>
 	public sealed class UhtTokenReplayReader : IUhtTokenReader, IUhtMessageLineNumber
 	{
-		private static readonly ThreadLocal<UhtTokenReplayReader> Tls = new ThreadLocal<UhtTokenReplayReader>(() => new UhtTokenReplayReader());
+		private static readonly ThreadLocal<UhtTokenReplayReader> s_tls = new ThreadLocal<UhtTokenReplayReader>(() => new UhtTokenReplayReader());
 
 		const int MaxSavedStates = 2;
 		private struct SavedState
 		{
-			public int TokenIndex;
-			public bool bHasToken;
+			public int _tokenIndex;
+			public bool _hasToken;
 		}
 
-		private IUhtMessageSite MessageSiteInternal;
-		private ReadOnlyMemory<UhtToken> Tokens;
-		private ReadOnlyMemory<char> Data;
-		private int CurrentTokenIndex = -1;
-		private bool bHasToken = false;
-		private UhtToken CurrentToken = new UhtToken();
-		private readonly SavedState[] SavedStates = new SavedState[MaxSavedStates];
-		private int SavedStateCount = 0;
-		private UhtTokenType EndTokenType = UhtTokenType.EndOfFile;
+		private IUhtMessageSite _messageSiteInternal;
+		private ReadOnlyMemory<UhtToken> _tokens;
+		private ReadOnlyMemory<char> _data;
+		private int _currentTokenIndex = -1;
+		private bool _hasToken = false;
+		private UhtToken _currentToken = new UhtToken();
+		private readonly SavedState[] _savedStates = new SavedState[MaxSavedStates];
+		private int _savedStateCount = 0;
+		private UhtTokenType _endTokenType = UhtTokenType.EndOfFile;
 
 		/// <summary>
 		/// Construct new token reader
 		/// </summary>
-		/// <param name="MessageSite">Message site for generating errors</param>
-		/// <param name="Data">Complete data for the token (i.e. original source)</param>
-		/// <param name="Tokens">Tokens to replay</param>
-		/// <param name="EndTokenType">Token type to return when end of tokens reached</param>
-		public UhtTokenReplayReader(IUhtMessageSite MessageSite, ReadOnlyMemory<char> Data, ReadOnlyMemory<UhtToken> Tokens, UhtTokenType EndTokenType)
+		/// <param name="messageSite">Message site for generating errors</param>
+		/// <param name="data">Complete data for the token (i.e. original source)</param>
+		/// <param name="tokens">Tokens to replay</param>
+		/// <param name="endTokenType">Token type to return when end of tokens reached</param>
+		public UhtTokenReplayReader(IUhtMessageSite messageSite, ReadOnlyMemory<char> data, ReadOnlyMemory<UhtToken> tokens, UhtTokenType endTokenType)
 		{
-			this.MessageSiteInternal = MessageSite;
-			this.Tokens = Tokens;
-			this.Data = Data;
-			this.EndTokenType = EndTokenType;
-			this.CurrentToken = new UhtToken(EndTokenType);
+			this._messageSiteInternal = messageSite;
+			this._tokens = tokens;
+			this._data = data;
+			this._endTokenType = endTokenType;
+			this._currentToken = new UhtToken(endTokenType);
 		}
 
 		/// <summary>
@@ -54,97 +54,97 @@ namespace EpicGames.UHT.Tokenizer
 		/// </summary>
 		public UhtTokenReplayReader()
 		{
-			this.MessageSiteInternal = new UhtEmptyMessageSite();
-			this.Tokens = Array.Empty<UhtToken>().AsMemory();
-			this.Data = Array.Empty<char>().AsMemory();
+			this._messageSiteInternal = new UhtEmptyMessageSite();
+			this._tokens = Array.Empty<UhtToken>().AsMemory();
+			this._data = Array.Empty<char>().AsMemory();
 		}
 
 		/// <summary>
 		/// Reset a cached replay reader for replaying a new stream of tokens
 		/// </summary>
-		/// <param name="MessageSite">Message site for generating errors</param>
-		/// <param name="Data">Complete data for the token (i.e. original source)</param>
-		/// <param name="Tokens">Tokens to replay</param>
-		/// <param name="EndTokenType">Token type to return when end of tokens reached</param>
+		/// <param name="messageSite">Message site for generating errors</param>
+		/// <param name="data">Complete data for the token (i.e. original source)</param>
+		/// <param name="tokens">Tokens to replay</param>
+		/// <param name="endTokenType">Token type to return when end of tokens reached</param>
 		/// <returns>The replay reader</returns>
-		public UhtTokenReplayReader Reset(IUhtMessageSite MessageSite, ReadOnlyMemory<char> Data, ReadOnlyMemory<UhtToken> Tokens, UhtTokenType EndTokenType)
+		public UhtTokenReplayReader Reset(IUhtMessageSite messageSite, ReadOnlyMemory<char> data, ReadOnlyMemory<UhtToken> tokens, UhtTokenType endTokenType)
 		{
-			this.MessageSiteInternal = MessageSite;
-			this.Tokens = Tokens;
-			this.Data = Data;
-			this.CurrentTokenIndex = -1;
-			this.bHasToken = false;
-			this.CurrentToken = new UhtToken(EndTokenType);
-			this.SavedStateCount = 0;
-			this.EndTokenType = EndTokenType;
+			this._messageSiteInternal = messageSite;
+			this._tokens = tokens;
+			this._data = data;
+			this._currentTokenIndex = -1;
+			this._hasToken = false;
+			this._currentToken = new UhtToken(endTokenType);
+			this._savedStateCount = 0;
+			this._endTokenType = endTokenType;
 			return this;
 		}
 
 		/// <summary>
 		/// Return the replay reader associated with the current thread.  Only one replay reader is cached per thread.
 		/// </summary>
-		/// <param name="MessageSite">The message site used to log errors</param>
-		/// <param name="Data">Source data where tokens were originally parsed</param>
-		/// <param name="Tokens">Collection of tokens to replay</param>
-		/// <param name="EndTokenType">Type of end token marker to return when the end of the token list is reached.  This is used to produce errors in the context of the replay</param>
+		/// <param name="messageSite">The message site used to log errors</param>
+		/// <param name="data">Source data where tokens were originally parsed</param>
+		/// <param name="tokens">Collection of tokens to replay</param>
+		/// <param name="endTokenType">Type of end token marker to return when the end of the token list is reached.  This is used to produce errors in the context of the replay</param>
 		/// <returns>The threaded instance of the replay reader</returns>
 		/// <exception cref="UhtIceException">Thrown if the TLS value can not be retrieved.</exception>
-		public static UhtTokenReplayReader GetThreadInstance(IUhtMessageSite MessageSite, ReadOnlyMemory<char> Data, ReadOnlyMemory<UhtToken> Tokens, UhtTokenType EndTokenType)
+		public static UhtTokenReplayReader GetThreadInstance(IUhtMessageSite messageSite, ReadOnlyMemory<char> data, ReadOnlyMemory<UhtToken> tokens, UhtTokenType endTokenType)
 		{
-			UhtTokenReplayReader? Reader = Tls.Value;
-			if (Reader == null)
+			UhtTokenReplayReader? reader = s_tls.Value;
+			if (reader == null)
 			{
 				throw new UhtIceException("Unable to acquire a UhtTokenReplayReader");
 			}
-			Reader.Reset(MessageSite, Data, Tokens, EndTokenType);
-			return Reader;
+			reader.Reset(messageSite, data, tokens, endTokenType);
+			return reader;
 		}
 
 		#region IUHTMessageSite Implementation
-		IUhtMessageSession IUhtMessageSite.MessageSession => this.MessageSiteInternal.MessageSession;
-		IUhtMessageSource? IUhtMessageSite.MessageSource => this.MessageSiteInternal.MessageSource;
+		IUhtMessageSession IUhtMessageSite.MessageSession => this._messageSiteInternal.MessageSession;
+		IUhtMessageSource? IUhtMessageSite.MessageSource => this._messageSiteInternal.MessageSource;
 		IUhtMessageLineNumber? IUhtMessageSite.MessageLineNumber => this;
 		#endregion
 
 		#region ITokenReader Implementation
 		/// <inheritdoc/>
-		public bool bIsEOF
+		public bool IsEOF
 		{
 			get
 			{
-				if (!this.bHasToken)
+				if (!this._hasToken)
 				{
 					GetTokenInternal();
 				}
-				return this.CurrentTokenIndex == this.Tokens.Span.Length;
+				return this._currentTokenIndex == this._tokens.Span.Length;
 			}
 		}
 
 		/// <inheritdoc/>
-		public int InputPos 
+		public int InputPos
 		{
 			get
 			{
-				if (!this.bHasToken)
+				if (!this._hasToken)
 				{
 					GetTokenInternal();
 				}
-				return CurrentToken.InputStartPos;
+				return _currentToken.InputStartPos;
 			}
 		}
 
 		/// <inheritdoc/>
-		public int InputLine 
+		public int InputLine
 		{
 			get
 			{
-				if (!this.bHasToken)
+				if (!this._hasToken)
 				{
 					GetTokenInternal();
 				}
-				return CurrentToken.InputLine;
+				return _currentToken.InputLine;
 			}
-			set => throw new NotImplementedException(); 
+			set => throw new NotImplementedException();
 		}
 
 		/// <inheritdoc/>
@@ -165,7 +165,7 @@ namespace EpicGames.UHT.Tokenizer
 		/// <inheritdoc/>
 		public void ConsumeToken()
 		{
-			this.bHasToken = false;
+			this._hasToken = false;
 		}
 
 		/// <inheritdoc/>
@@ -193,27 +193,27 @@ namespace EpicGames.UHT.Tokenizer
 		}
 
 		/// <inheritdoc/>
-		public StringView GetRawString(char Terminator, UhtRawStringOptions Options)
+		public StringView GetRawString(char terminator, UhtRawStringOptions options)
 		{
 			throw new NotImplementedException();
 		}
 
 		/// <inheritdoc/>
-		public StringView GetStringView(int StartPos, int Count)
+		public StringView GetStringView(int startPos, int count)
 		{
-			return new StringView(this.Data, StartPos, Count);
+			return new StringView(this._data, startPos, count);
 		}
 
 		/// <inheritdoc/>
 		public UhtToken GetToken()
 		{
-			UhtToken Token = PeekToken();
+			UhtToken token = PeekToken();
 			ConsumeToken();
-			return Token;
+			return token;
 		}
 
 		/// <inheritdoc/>
-		public bool IsFirstTokenInLine(ref UhtToken Token)
+		public bool IsFirstTokenInLine(ref UhtToken token)
 		{
 			throw new NotImplementedException();
 		}
@@ -221,11 +221,11 @@ namespace EpicGames.UHT.Tokenizer
 		/// <inheritdoc/>
 		public ref UhtToken PeekToken()
 		{
-			if (!this.bHasToken)
+			if (!this._hasToken)
 			{
 				GetTokenInternal();
 			}
-			return ref this.CurrentToken;
+			return ref this._currentToken;
 		}
 
 		/// <inheritdoc/>
@@ -237,36 +237,36 @@ namespace EpicGames.UHT.Tokenizer
 		/// <inheritdoc/>
 		public void SaveState()
 		{
-			if (this.SavedStateCount == MaxSavedStates)
+			if (this._savedStateCount == MaxSavedStates)
 			{
 				throw new UhtIceException("Token reader saved states full");
 			}
-			this.SavedStates[this.SavedStateCount] = new SavedState { TokenIndex = this.CurrentTokenIndex, bHasToken = this.bHasToken };
-			++this.SavedStateCount;
+			this._savedStates[this._savedStateCount] = new SavedState { _tokenIndex = this._currentTokenIndex, _hasToken = this._hasToken };
+			++this._savedStateCount;
 		}
 
 		/// <inheritdoc/>
 		public void RestoreState()
 		{
-			if (this.SavedStateCount == 0)
+			if (this._savedStateCount == 0)
 			{
 				throw new UhtIceException("Attempt to restore a state when none have been saved");
 			}
 
-			--this.SavedStateCount;
-			this.CurrentTokenIndex = this.SavedStates[this.SavedStateCount].TokenIndex;
-			this.bHasToken = this.SavedStates[this.SavedStateCount].bHasToken;
+			--this._savedStateCount;
+			this._currentTokenIndex = this._savedStates[this._savedStateCount]._tokenIndex;
+			this._hasToken = this._savedStates[this._savedStateCount]._hasToken;
 		}
 
 		/// <inheritdoc/>
 		public void AbandonState()
 		{
-			if (this.SavedStateCount == 0)
+			if (this._savedStateCount == 0)
 			{
 				throw new UhtIceException("Attempt to abandon a state when none have been saved");
 			}
 
-			--this.SavedStateCount;
+			--this._savedStateCount;
 		}
 
 		/// <inheritdoc/>
@@ -282,19 +282,13 @@ namespace EpicGames.UHT.Tokenizer
 		}
 
 		/// <inheritdoc/>
-		public void RecordToken(ref UhtToken Token)
+		public void RecordToken(ref UhtToken token)
 		{
 			throw new NotImplementedException();
 		}
 
 		/// <inheritdoc/>
-		public List<UhtToken> RecordedTokens
-		{
-			get
-			{
-				throw new NotImplementedException();
-			}
-		}
+		public List<UhtToken> RecordedTokens => throw new NotImplementedException();
 		#endregion
 
 		#region IUHTMessageLineNumber implementation
@@ -302,45 +296,45 @@ namespace EpicGames.UHT.Tokenizer
 		{
 			get
 			{
-				if (this.Tokens.Length == 0)
+				if (this._tokens.Length == 0)
 				{
 					return -1;
 				}
-				if (this.CurrentTokenIndex < 0)
+				if (this._currentTokenIndex < 0)
 				{
-					return this.Tokens.Span[0].InputLine;
+					return this._tokens.Span[0].InputLine;
 				}
-				if (this.CurrentTokenIndex < this.Tokens.Length)
+				if (this._currentTokenIndex < this._tokens.Length)
 				{
-					return this.Tokens.Span[this.CurrentTokenIndex].InputLine;
+					return this._tokens.Span[this._currentTokenIndex].InputLine;
 				}
-				return this.Tokens.Span[this.Tokens.Length - 1].InputLine;
+				return this._tokens.Span[this._tokens.Length - 1].InputLine;
 			}
 		}
 		#endregion
 
 		private ref UhtToken GetTokenInternal()
 		{
-			if (this.CurrentTokenIndex < this.Tokens.Length)
+			if (this._currentTokenIndex < this._tokens.Length)
 			{
-				++this.CurrentTokenIndex;
+				++this._currentTokenIndex;
 			}
-			if (this.CurrentTokenIndex < this.Tokens.Length)
+			if (this._currentTokenIndex < this._tokens.Length)
 			{
-				this.CurrentToken = this.Tokens.Span[this.CurrentTokenIndex];
+				this._currentToken = this._tokens.Span[this._currentTokenIndex];
 			}
-			else if (this.Tokens.Length == 0)
+			else if (this._tokens.Length == 0)
 			{
-				this.CurrentToken = new UhtToken();
+				this._currentToken = new UhtToken();
 			}
 			else
 			{
-				UhtToken LastToken = this.Tokens.Span[this.Tokens.Length - 1];
-				int EndPos = LastToken.InputEndPos;
-				this.CurrentToken = new UhtToken(this.EndTokenType, EndPos, LastToken.InputLine, EndPos, LastToken.InputLine, new StringView());
+				UhtToken lastToken = this._tokens.Span[this._tokens.Length - 1];
+				int endPos = lastToken.InputEndPos;
+				this._currentToken = new UhtToken(this._endTokenType, endPos, lastToken.InputLine, endPos, lastToken.InputLine, new StringView());
 			}
-			this.bHasToken = true;
-			return ref this.CurrentToken;
+			this._hasToken = true;
+			return ref this._currentToken;
 		}
 	}
 }
diff --git a/Engine/Source/Programs/Shared/EpicGames.UHT/Types/Properties/UhtArrayProperty.cs b/Engine/Source/Programs/Shared/EpicGames.UHT/Types/Properties/UhtArrayProperty.cs
--- a/Engine/Source/Programs/Shared/EpicGames.UHT/Types/Properties/UhtArrayProperty.cs
+++ b/Engine/Source/Programs/Shared/EpicGames.UHT/Types/Properties/UhtArrayProperty.cs
@@ -1,13 +1,13 @@
 // Copyright Epic Games, Inc. All Rights Reserved.
 
+using System.Collections.Generic;
+using System.Diagnostics.CodeAnalysis;
+using System.Text;
 using EpicGames.Core;
 using EpicGames.UHT.Parsers;
 using EpicGames.UHT.Tables;
 using EpicGames.UHT.Tokenizer;
 using EpicGames.UHT.Utils;
-using System.Collections.Generic;
-using System.Diagnostics.CodeAnalysis;
-using System.Text;
 
 namespace EpicGames.UHT.Types
 {
@@ -20,31 +20,31 @@ namespace EpicGames.UHT.Types
 	public class UhtArrayProperty : UhtContainerBaseProperty
 	{
 		/// <inheritdoc/>
-		public override string EngineClassName { get => "ArrayProperty"; }
+		public override string EngineClassName => "ArrayProperty";
 
 		/// <inheritdoc/>
-		protected override string CppTypeText { get => "TArray"; }
+		protected override string CppTypeText => "TArray";
 
 		/// <inheritdoc/>
-		protected override string PGetMacroText { get => "TARRAY"; }
+		protected override string PGetMacroText => "TARRAY";
 
 		/// <inheritdoc/>
-		protected override UhtPGetArgumentType PGetTypeArgument { get => UhtPGetArgumentType.TypeText; }
+		protected override UhtPGetArgumentType PGetTypeArgument => UhtPGetArgumentType.TypeText;
 
 		/// <summary>
 		/// Construct a new array property
 		/// </summary>
-		/// <param name="PropertySettings">Property settings</param>
-		/// <param name="Value">Inner property value</param>
-		public UhtArrayProperty(UhtPropertySettings PropertySettings, UhtProperty Value) : base(PropertySettings, Value)
+		/// <param name="propertySettings">Property settings</param>
+		/// <param name="value">Inner property value</param>
+		public UhtArrayProperty(UhtPropertySettings propertySettings, UhtProperty value) : base(propertySettings, value)
 		{
 			// If the creation of the value property set more flags, then copy those flags to ourselves
-			this.PropertyFlags |= Value.PropertyFlags & EPropertyFlags.UObjectWrapper;
+			this.PropertyFlags |= value.PropertyFlags & EPropertyFlags.UObjectWrapper;
 
-			if (Value.MetaData.ContainsKey(UhtNames.NativeConst))
+			if (value.MetaData.ContainsKey(UhtNames.NativeConst))
 			{
 				this.MetaData.Add(UhtNames.NativeConstTemplateArg, "");
-				Value.MetaData.Remove(UhtNames.NativeConst);
+				value.MetaData.Remove(UhtNames.NativeConst);
 			}
 
 			this.PropertyCaps |= UhtPropertyCaps.PassCppArgsByRef | UhtPropertyCaps.SupportsRigVM | UhtPropertyCaps.IsRigVMArray;
@@ -59,26 +59,26 @@ namespace EpicGames.UHT.Types
 		}
 
 		/// <inheritdoc/>
-		protected override bool ResolveSelf(UhtResolvePhase Phase)
+		protected override bool ResolveSelf(UhtResolvePhase phase)
 		{
-			bool bResults = base.ResolveSelf(Phase);
-			switch (Phase)
+			bool results = base.ResolveSelf(phase);
+			switch (phase)
 			{
 				case UhtResolvePhase.Final:
-					this.PropertyFlags |= ResolveAndReturnNewFlags(this.ValueProperty, Phase);
+					this.PropertyFlags |= ResolveAndReturnNewFlags(this.ValueProperty, phase);
 					this.MetaData.Add(this.ValueProperty.MetaData);
 					this.ValueProperty.PropertyFlags = this.PropertyFlags & EPropertyFlags.PropagateToArrayInner;
 					this.ValueProperty.MetaData.Clear();
 					PropagateFlagsFromInnerAndHandlePersistentInstanceMetadata(this, this.MetaData, this.ValueProperty);
 					break;
 			}
-			return bResults;
+			return results;
 		}
 
 		/// <inheritdoc/>
-		public override void CollectReferencesInternal(IUhtReferenceCollector Collector, bool bTemplateProperty)
+		public override void CollectReferencesInternal(IUhtReferenceCollector collector, bool templateProperty)
 		{
-			this.ValueProperty.CollectReferencesInternal(Collector, true);
+			this.ValueProperty.CollectReferencesInternal(collector, true);
 		}
 
 		/// <inheritdoc/>
@@ -90,92 +90,92 @@ namespace EpicGames.UHT.Types
 		/// <inheritdoc/>
 		public override IEnumerable<UhtType> EnumerateReferencedTypes()
 		{
-			foreach (UhtType Type in this.ValueProperty.EnumerateReferencedTypes())
+			foreach (UhtType type in this.ValueProperty.EnumerateReferencedTypes())
 			{
-				yield return Type;
+				yield return type;
 			}
 		}
 
 		/// <inheritdoc/>
-		public override StringBuilder AppendText(StringBuilder Builder, UhtPropertyTextType TextType, bool bIsTemplateArgument)
+		public override StringBuilder AppendText(StringBuilder builder, UhtPropertyTextType textType, bool isTemplateArgument)
 		{
-			switch (TextType)
+			switch (textType)
 			{
 				case UhtPropertyTextType.SparseShort:
-					Builder.Append("TArray");
+					builder.Append("TArray");
 					break;
 
 				case UhtPropertyTextType.FunctionThunkParameterArgType:
-					Builder.AppendFunctionThunkParameterArrayType(this.ValueProperty);
+					builder.AppendFunctionThunkParameterArrayType(this.ValueProperty);
 					break;
 
 				default:
-					Builder.Append("TArray<").AppendPropertyText(this.ValueProperty, TextType, true);
-					if (Builder[Builder.Length - 1] == '>')
+					builder.Append("TArray<").AppendPropertyText(this.ValueProperty, textType, true);
+					if (builder[^1] == '>')
 					{
 						// if our internal property type is a template class, add a space between the closing brackets b/c VS.NET cannot parse this correctly
-						Builder.Append(' ');
+						builder.Append(' ');
 					}
-					Builder.Append('>');
+					builder.Append('>');
 					break;
 			}
-			return Builder;
+			return builder;
 		}
 
 		/// <inheritdoc/>
-		public override StringBuilder AppendMemberDecl(StringBuilder Builder, IUhtPropertyMemberContext Context, string Name, string NameSuffix, int Tabs)
+		public override StringBuilder AppendMemberDecl(StringBuilder builder, IUhtPropertyMemberContext context, string name, string nameSuffix, int tabs)
 		{
-			Builder.AppendMemberDecl(this.ValueProperty, Context, Name, GetNameSuffix(NameSuffix, "_Inner"), Tabs);
-			return AppendMemberDecl(Builder, Context, Name, NameSuffix, Tabs, "FArrayPropertyParams");
+			builder.AppendMemberDecl(this.ValueProperty, context, name, GetNameSuffix(nameSuffix, "_Inner"), tabs);
+			return AppendMemberDecl(builder, context, name, nameSuffix, tabs, "FArrayPropertyParams");
 		}
 
 		/// <inheritdoc/>
-		public override StringBuilder AppendMemberDef(StringBuilder Builder, IUhtPropertyMemberContext Context, string Name, string NameSuffix, string? Offset, int Tabs)
+		public override StringBuilder AppendMemberDef(StringBuilder builder, IUhtPropertyMemberContext context, string name, string nameSuffix, string? offset, int tabs)
 		{
-			Builder.AppendMemberDef(this.ValueProperty, Context, Name, GetNameSuffix(NameSuffix, "_Inner"), "0", Tabs);
-			AppendMemberDefStart(Builder, Context, Name, NameSuffix, Offset, Tabs, "FArrayPropertyParams", "UECodeGen_Private::EPropertyGenFlags::Array");
-			Builder.Append(this.Allocator == UhtPropertyAllocator.MemoryImage ? "EArrayPropertyFlags::UsesMemoryImageAllocator" : "EArrayPropertyFlags::None").Append(", ");
-			AppendMemberDefEnd(Builder, Context, Name, NameSuffix);
-			return Builder;
+			builder.AppendMemberDef(this.ValueProperty, context, name, GetNameSuffix(nameSuffix, "_Inner"), "0", tabs);
+			AppendMemberDefStart(builder, context, name, nameSuffix, offset, tabs, "FArrayPropertyParams", "UECodeGen_Private::EPropertyGenFlags::Array");
+			builder.Append(this.Allocator == UhtPropertyAllocator.MemoryImage ? "EArrayPropertyFlags::UsesMemoryImageAllocator" : "EArrayPropertyFlags::None").Append(", ");
+			AppendMemberDefEnd(builder, context, name, nameSuffix);
+			return builder;
 		}
 
 		/// <inheritdoc/>
-		public override StringBuilder AppendMemberPtr(StringBuilder Builder, IUhtPropertyMemberContext Context, string Name, string NameSuffix, int Tabs)
+		public override StringBuilder AppendMemberPtr(StringBuilder builder, IUhtPropertyMemberContext context, string name, string nameSuffix, int tabs)
 		{
-			Builder.AppendMemberPtr(this.ValueProperty, Context, Name, GetNameSuffix(NameSuffix, "_Inner"), Tabs);
-			base.AppendMemberPtr(Builder, Context, Name, NameSuffix, Tabs);
-			return Builder;
+			builder.AppendMemberPtr(this.ValueProperty, context, name, GetNameSuffix(nameSuffix, "_Inner"), tabs);
+			base.AppendMemberPtr(builder, context, name, nameSuffix, tabs);
+			return builder;
 		}
 
 		/// <inheritdoc/>
-		public override void AppendObjectHashes(StringBuilder Builder, int StartingLength, IUhtPropertyMemberContext Context)
+		public override void AppendObjectHashes(StringBuilder builder, int startingLength, IUhtPropertyMemberContext context)
 		{
-			this.ValueProperty.AppendObjectHashes(Builder, StartingLength, Context);
+			this.ValueProperty.AppendObjectHashes(builder, startingLength, context);
 		}
 
 		/// <inheritdoc/>
-		public override StringBuilder AppendNullConstructorArg(StringBuilder Builder, bool bIsInitializer)
+		public override StringBuilder AppendNullConstructorArg(StringBuilder builder, bool isInitializer)
 		{
-			Builder.AppendPropertyText(this, UhtPropertyTextType.Construction).Append("()");
-			return Builder;
+			builder.AppendPropertyText(this, UhtPropertyTextType.Construction).Append("()");
+			return builder;
 		}
 
 		/// <inheritdoc/>
-		public override bool SanitizeDefaultValue(IUhtTokenReader DefaultValueReader, StringBuilder InnerDefaultValue)
+		public override bool SanitizeDefaultValue(IUhtTokenReader defaultValueReader, StringBuilder innerDefaultValue)
 		{
 			return false;
 		}
 
 		/// <inheritdoc/>
-		protected override void ValidateFunctionArgument(UhtFunction Function, UhtValidationOptions Options)
+		protected override void ValidateFunctionArgument(UhtFunction function, UhtValidationOptions options)
 		{
-			base.ValidateFunctionArgument(Function, Options);
+			base.ValidateFunctionArgument(function, options);
 
-			if (Function.FunctionFlags.HasAnyFlags(EFunctionFlags.Net))
+			if (function.FunctionFlags.HasAnyFlags(EFunctionFlags.Net))
 			{
-				if (!Function.FunctionFlags.HasAnyFlags(EFunctionFlags.NetRequest))
+				if (!function.FunctionFlags.HasAnyFlags(EFunctionFlags.NetRequest))
 				{
-					if (this.RefQualifier != UhtPropertyRefQualifier.ConstRef && !this.bIsStaticArray)
+					if (this.RefQualifier != UhtPropertyRefQualifier.ConstRef && !this.IsStaticArray)
 					{
 						this.LogError("Replicated TArray parameters must be passed by const reference");
 					}
@@ -184,9 +184,9 @@ namespace EpicGames.UHT.Types
 		}
 
 		/// <inheritdoc/>
-		public override void Validate(UhtStruct OuterStruct, UhtProperty OutermostProperty, UhtValidationOptions Options)
+		public override void Validate(UhtStruct outerStruct, UhtProperty outermostProperty, UhtValidationOptions options)
 		{
-			base.Validate(OuterStruct, OutermostProperty, Options);
+			base.Validate(outerStruct, outermostProperty, options);
 
 			if (this.Allocator != UhtPropertyAllocator.Default)
 			{
@@ -198,65 +198,65 @@ namespace EpicGames.UHT.Types
 		}
 
 		/// <inheritdoc/>
-		public override bool IsSameType(UhtProperty Other)
+		public override bool IsSameType(UhtProperty other)
 		{
-			if (Other is UhtArrayProperty OtherArray)
+			if (other is UhtArrayProperty otherArray)
 			{
-				return this.ValueProperty.IsSameType(OtherArray.ValueProperty);
+				return this.ValueProperty.IsSameType(otherArray.ValueProperty);
 			}
 			return false;
 		}
 
-#region Keyword
+		#region Keyword
 		[UhtPropertyType(Keyword = "TArray")]
 		[SuppressMessage("CodeQuality", "IDE0051:Remove unused private members", Justification = "Attribute accessed method")]
-		private static UhtProperty? ArrayProperty(UhtPropertyResolvePhase ResolvePhase, UhtPropertySettings PropertySettings, IUhtTokenReader TokenReader, UhtToken MatchedToken)
+		private static UhtProperty? ArrayProperty(UhtPropertyResolvePhase resolvePhase, UhtPropertySettings propertySettings, IUhtTokenReader tokenReader, UhtToken matchedToken)
 		{
-			using (var TokenContext = new UhtMessageContext("TArray"))
+			using (UhtMessageContext tokenContext = new UhtMessageContext("TArray"))
 			{
-				if (!TokenReader.SkipExpectedType(MatchedToken.Value, PropertySettings.PropertyCategory == UhtPropertyCategory.Member))
+				if (!tokenReader.SkipExpectedType(matchedToken.Value, propertySettings.PropertyCategory == UhtPropertyCategory.Member))
 				{
 					return null;
 				}
-				TokenReader.Require('<');
+				tokenReader.Require('<');
 
 				// Parse the value type
-				UhtProperty? Value = UhtPropertyParser.ParseTemplateParam(ResolvePhase, PropertySettings, PropertySettings.SourceName, TokenReader);
-				if (Value == null)
+				UhtProperty? value = UhtPropertyParser.ParseTemplateParam(resolvePhase, propertySettings, propertySettings.SourceName, tokenReader);
+				if (value == null)
 				{
 					return null;
 				}
 
-				if (!Value.PropertyCaps.HasAnyFlags(UhtPropertyCaps.CanBeContainerValue))
+				if (!value.PropertyCaps.HasAnyFlags(UhtPropertyCaps.CanBeContainerValue))
 				{
-					TokenReader.LogError($"The type \'{Value.GetUserFacingDecl()}\' can not be used as a value in a TArray");
+					tokenReader.LogError($"The type \'{value.GetUserFacingDecl()}\' can not be used as a value in a TArray");
 				}
 
-				if (TokenReader.TryOptional(','))
+				if (tokenReader.TryOptional(','))
 				{
 					// If we found a comma, read the next thing, assume it's an allocator, and report that
-					UhtToken AllocatorToken = TokenReader.GetIdentifier();
-					if (AllocatorToken.IsIdentifier("FMemoryImageAllocator"))
+					UhtToken allocatorToken = tokenReader.GetIdentifier();
+					if (allocatorToken.IsIdentifier("FMemoryImageAllocator"))
 					{
-						PropertySettings.Allocator = UhtPropertyAllocator.MemoryImage;
+						propertySettings.Allocator = UhtPropertyAllocator.MemoryImage;
 					}
-					else if (AllocatorToken.IsIdentifier("TMemoryImageAllocator"))
+					else if (allocatorToken.IsIdentifier("TMemoryImageAllocator"))
 					{
-						TokenReader.RequireList('<', '>', "TMemoryImageAllocator template arguments");
-						PropertySettings.Allocator = UhtPropertyAllocator.MemoryImage;
+						tokenReader.RequireList('<', '>', "TMemoryImageAllocator template arguments");
+						propertySettings.Allocator = UhtPropertyAllocator.MemoryImage;
 					}
 					else
 					{
-						throw new UhtException(TokenReader, $"Found '{AllocatorToken.Value.ToString()}' - explicit allocators are not supported in TArray properties.");
+						throw new UhtException(tokenReader, $"Found '{allocatorToken.Value}' - explicit allocators are not supported in TArray properties.");
 					}
 				}
-				TokenReader.Require('>');
+				tokenReader.Require('>');
 
 				//@TODO: Prevent sparse delegate types from being used in a container
 
-				return new UhtArrayProperty(PropertySettings, Value);
+				return new UhtArrayProperty(propertySettings, value);
 			}
 		}
-#endregion
+		#endregion
 	}
 }
diff --git a/Engine/Source/Programs/Shared/EpicGames.UHT/Types/Properties/UhtBoolProperty.cs b/Engine/Source/Programs/Shared/EpicGames.UHT/Types/Properties/UhtBoolProperty.cs
--- a/Engine/Source/Programs/Shared/EpicGames.UHT/Types/Properties/UhtBoolProperty.cs
+++ b/Engine/Source/Programs/Shared/EpicGames.UHT/Types/Properties/UhtBoolProperty.cs
@@ -1,11 +1,11 @@
 // Copyright Epic Games, Inc. All Rights Reserved.
 
+using System.Diagnostics.CodeAnalysis;
+using System.Text;
 using EpicGames.Core;
 using EpicGames.UHT.Tables;
 using EpicGames.UHT.Tokenizer;
 using EpicGames.UHT.Utils;
-using System.Diagnostics.CodeAnalysis;
-using System.Text;
 
 namespace EpicGames.UHT.Types
 {
@@ -50,7 +50,7 @@ namespace EpicGames.UHT.Types
 	public class UhtBoolProperty : UhtProperty
 	{
 		/// <inheritdoc/>
-		public override string EngineClassName { get => "BoolProperty"; }
+		public override string EngineClassName => "BoolProperty";
 
 		/// <inheritdoc/>
 		protected override string CppTypeText
@@ -101,7 +101,7 @@ namespace EpicGames.UHT.Types
 		/// <summary>
 		/// If true, the boolean is a native bool and not a UBOOL
 		/// </summary>
-		protected bool bIsNativeBool { get => this.BoolType == UhtBoolType.Native; }
+		protected bool IsNativeBool => this.BoolType == UhtBoolType.Native;
 
 		/// <summary>
 		/// Type of the boolean
@@ -111,23 +111,23 @@ namespace EpicGames.UHT.Types
 		/// <summary>
 		/// Construct a new boolean property
 		/// </summary>
-		/// <param name="PropertySettings">Property settings</param>
-		/// <param name="BoolType">Type of the boolean</param>
-		public UhtBoolProperty(UhtPropertySettings PropertySettings, UhtBoolType BoolType) : base(PropertySettings)
+		/// <param name="propertySettings">Property settings</param>
+		/// <param name="boolType">Type of the boolean</param>
+		public UhtBoolProperty(UhtPropertySettings propertySettings, UhtBoolType boolType) : base(propertySettings)
 		{
-			this.PropertyCaps |= UhtPropertyCaps.RequiresNullConstructorArg | UhtPropertyCaps.IsParameterSupportedByBlueprint | 
+			this.PropertyCaps |= UhtPropertyCaps.RequiresNullConstructorArg | UhtPropertyCaps.IsParameterSupportedByBlueprint |
 				UhtPropertyCaps.IsMemberSupportedByBlueprint | UhtPropertyCaps.SupportsRigVM;
-			if (BoolType == UhtBoolType.Native || BoolType == UhtBoolType.UInt8)
+			if (boolType == UhtBoolType.Native || boolType == UhtBoolType.UInt8)
 			{
 				this.PropertyCaps |= UhtPropertyCaps.CanExposeOnSpawn;
 			}
-			this.BoolType = BoolType;
+			this.BoolType = boolType;
 		}
 
 		/// <inheritdoc/>
-		public override StringBuilder AppendText(StringBuilder Builder, UhtPropertyTextType TextType, bool bIsTemplateArgument = false)
+		public override StringBuilder AppendText(StringBuilder builder, UhtPropertyTextType textType, bool isTemplateArgument = false)
 		{
-			switch (TextType)
+			switch (textType)
 			{
 				case UhtPropertyTextType.Generic:
 				case UhtPropertyTextType.Construction:
@@ -136,7 +136,7 @@ namespace EpicGames.UHT.Types
 				case UhtPropertyTextType.RigVMType:
 				case UhtPropertyTextType.ExportMember:
 				case UhtPropertyTextType.GenericFunctionArgOrRetVal:
-					Builder.Append(this.CppTypeText);
+					builder.Append(this.CppTypeText);
 					break;
 
 				case UhtPropertyTextType.Sparse:
@@ -148,126 +148,126 @@ namespace EpicGames.UHT.Types
 				case UhtPropertyTextType.EventParameterMember:
 				case UhtPropertyTextType.EventParameterFunctionMember:
 				case UhtPropertyTextType.GetterSetterArg:
-					Builder.Append("bool");
+					builder.Append("bool");
 					break;
 			}
-			return Builder;
+			return builder;
 		}
 
 		/// <inheritdoc/>
-		public override StringBuilder AppendMemberDecl(StringBuilder Builder, IUhtPropertyMemberContext Context, string Name, string NameSuffix, int Tabs)
+		public override StringBuilder AppendMemberDecl(StringBuilder builder, IUhtPropertyMemberContext context, string name, string nameSuffix, int tabs)
 		{
-			Builder.AppendMetaDataDecl(this, Context, Name, NameSuffix, Tabs);
-			if (!(this.Outer is UhtProperty))
+			builder.AppendMetaDataDecl(this, context, name, nameSuffix, tabs);
+			if (this.Outer is not UhtProperty)
 			{
-				Builder.AppendTabs(Tabs).Append("static void ").AppendNameDecl(Context, Name, NameSuffix).Append("_SetBit(void* Obj);\r\n");
+				builder.AppendTabs(tabs).Append("static void ").AppendNameDecl(context, name, nameSuffix).Append("_SetBit(void* Obj);\r\n");
 			}
-			Builder.AppendTabs(Tabs).Append("static const UECodeGen_Private::").Append("FBoolPropertyParams").Append(' ').AppendNameDecl(Context, Name, NameSuffix).Append(";\r\n");
-			return Builder;
+			builder.AppendTabs(tabs).Append("static const UECodeGen_Private::").Append("FBoolPropertyParams").Append(' ').AppendNameDecl(context, name, nameSuffix).Append(";\r\n");
+			return builder;
 		}
 
 		/// <inheritdoc/>
-		public override StringBuilder AppendMemberDef(StringBuilder Builder, IUhtPropertyMemberContext Context, string Name, string NameSuffix, string? Offset, int Tabs)
+		public override StringBuilder AppendMemberDef(StringBuilder builder, IUhtPropertyMemberContext context, string name, string nameSuffix, string? offset, int tabs)
 		{
-			Builder.AppendMetaDataDef(this, Context, Name, NameSuffix, Tabs);
-			if (this.Outer == Context.OuterStruct)
+			builder.AppendMetaDataDef(this, context, name, nameSuffix, tabs);
+			if (this.Outer == context.OuterStruct)
 			{
-				Builder.AppendTabs(Tabs).Append("void ").AppendNameDef(Context, Name, NameSuffix).Append("_SetBit(void* Obj)\r\n");
-				Builder.AppendTabs(Tabs).Append("{\r\n");
-				Builder.AppendTabs(Tabs + 1).Append("((").Append(Context.OuterStructSourceName).Append("*)Obj)->").Append(this.SourceName).Append(" = 1;\r\n");
-				Builder.AppendTabs(Tabs).Append("}\r\n");
+				builder.AppendTabs(tabs).Append("void ").AppendNameDef(context, name, nameSuffix).Append("_SetBit(void* Obj)\r\n");
+				builder.AppendTabs(tabs).Append("{\r\n");
+				builder.AppendTabs(tabs + 1).Append("((").Append(context.OuterStructSourceName).Append("*)Obj)->").Append(this.SourceName).Append(" = 1;\r\n");
+				builder.AppendTabs(tabs).Append("}\r\n");
 			}
 
-			AppendMemberDefStart(Builder, Context, Name, NameSuffix, Offset, Tabs, "FBoolPropertyParams",
-				bIsNativeBool ?
+			AppendMemberDefStart(builder, context, name, nameSuffix, offset, tabs, "FBoolPropertyParams",
+				IsNativeBool ?
 				"UECodeGen_Private::EPropertyGenFlags::Bool | UECodeGen_Private::EPropertyGenFlags::NativeBool" :
-				"UECodeGen_Private::EPropertyGenFlags::Bool ", 
+				"UECodeGen_Private::EPropertyGenFlags::Bool ",
 				false, false);
 
-			Builder.Append("sizeof(").Append(this.CppTypeText).Append("), ");
+			builder.Append("sizeof(").Append(this.CppTypeText).Append("), ");
 
-			if (this.Outer == Context.OuterStruct)
+			if (this.Outer == context.OuterStruct)
 			{
-				Builder
-					.Append("sizeof(").Append(Context.OuterStructSourceName).Append("), ")
-					.Append('&').AppendNameDef(Context, Name, NameSuffix).Append("_SetBit, ");
+				builder
+					.Append("sizeof(").Append(context.OuterStructSourceName).Append("), ")
+					.Append('&').AppendNameDef(context, name, nameSuffix).Append("_SetBit, ");
 			}
 			else
 			{
-				Builder.Append("0, nullptr, ");
+				builder.Append("0, nullptr, ");
 			}
 
-			AppendMemberDefEnd(Builder, Context, Name, NameSuffix);
-			return Builder;
+			AppendMemberDefEnd(builder, context, name, nameSuffix);
+			return builder;
 		}
 
 		/// <inheritdoc/>
-		public override StringBuilder AppendFullDecl(StringBuilder Builder, UhtPropertyTextType TextType, bool bSkipParameterName = false)
+		public override StringBuilder AppendFullDecl(StringBuilder builder, UhtPropertyTextType textType, bool skipParameterName = false)
 		{
-			AppendText(Builder, TextType);
+			AppendText(builder, textType);
 
 			//@todo we currently can't have out bools.. so this isn't really necessary, but eventually out bools may be supported, so leave here for now
-			if (TextType.IsParameter() && this.PropertyFlags.HasAnyFlags(EPropertyFlags.OutParm))
+			if (textType.IsParameter() && this.PropertyFlags.HasAnyFlags(EPropertyFlags.OutParm))
 			{
-				Builder.Append('&');
+				builder.Append('&');
 			}
 
-			Builder.Append(' ');
+			builder.Append(' ');
 
-			if (!bSkipParameterName)
+			if (!skipParameterName)
 			{
-				Builder.Append(this.SourceName);
+				builder.Append(this.SourceName);
 			}
 
 			if (this.ArrayDimensions != null)
 			{
-				Builder.Append('[').Append(this.ArrayDimensions).Append(']');
+				builder.Append('[').Append(this.ArrayDimensions).Append(']');
 			}
-			else if (TextType == UhtPropertyTextType.ExportMember && !this.bIsNativeBool)
+			else if (textType == UhtPropertyTextType.ExportMember && !this.IsNativeBool)
 			{
-				Builder.Append(":1");
+				builder.Append(":1");
 			}
-			return Builder;
+			return builder;
 		}
 
 		/// <inheritdoc/>
-		public override StringBuilder AppendNullConstructorArg(StringBuilder Builder, bool bIsInitializer)
+		public override StringBuilder AppendNullConstructorArg(StringBuilder builder, bool isInitializer)
 		{
-			Builder.Append("false");
-			return Builder;
+			builder.Append("false");
+			return builder;
 		}
 
 		/// <inheritdoc/>
-		public override bool SanitizeDefaultValue(IUhtTokenReader DefaultValueReader, StringBuilder InnerDefaultValue)
+		public override bool SanitizeDefaultValue(IUhtTokenReader defaultValueReader, StringBuilder innerDefaultValue)
 		{
-			UhtToken Identifier = DefaultValueReader.GetIdentifier();
-			if (Identifier.IsValue("true") || Identifier.IsValue("false"))
+			UhtToken identifier = defaultValueReader.GetIdentifier();
+			if (identifier.IsValue("true") || identifier.IsValue("false"))
 			{
-				InnerDefaultValue.Append(Identifier.Value.ToString());
+				innerDefaultValue.Append(identifier.Value.ToString());
 				return true;
 			}
 			return false;
 		}
 
 		/// <inheritdoc/>
-		public override bool IsSameType(UhtProperty Other)
+		public override bool IsSameType(UhtProperty other)
 		{
 			// We don't test BoolType.
-			return Other is UhtBoolProperty;
+			return other is UhtBoolProperty;
 		}
 
 		#region Keyword
 		[UhtPropertyType(Keyword = "bool", Options = UhtPropertyTypeOptions.Simple | UhtPropertyTypeOptions.Immediate)]
 		[SuppressMessage("CodeQuality", "IDE0051:Remove unused private members", Justification = "Attribute accessed method")]
 		[SuppressMessage("Style", "IDE0060:Remove unused parameter", Justification = "Attribute accessed method")]
-		private static UhtProperty? BoolProperty(UhtPropertyResolvePhase ResolvePhase, UhtPropertySettings PropertySettings, IUhtTokenReader TokenReader, UhtToken MatchedToken)
+		private static UhtProperty? BoolProperty(UhtPropertyResolvePhase resolvePhase, UhtPropertySettings propertySettings, IUhtTokenReader tokenReader, UhtToken matchedToken)
 		{
-			if (PropertySettings.bIsBitfield)
+			if (propertySettings.IsBitfield)
 			{
-				TokenReader.LogError("bool bitfields are not supported.");
+				tokenReader.LogError("bool bitfields are not supported.");
 				return null;
 			}
-			return new UhtBoolProperty(PropertySettings, UhtBoolType.Native);
+			return new UhtBoolProperty(propertySettings, UhtBoolType.Native);
 		}
 		#endregion
 	}
diff --git a/Engine/Source/Programs/Shared/EpicGames.UHT/Types/Properties/UhtByteProperty.cs b/Engine/Source/Programs/Shared/EpicGames.UHT/Types/Properties/UhtByteProperty.cs
--- a/Engine/Source/Programs/Shared/EpicGames.UHT/Types/Properties/UhtByteProperty.cs
+++ b/Engine/Source/Programs/Shared/EpicGames.UHT/Types/Properties/UhtByteProperty.cs
@@ -1,13 +1,14 @@
 // Copyright Epic Games, Inc. All Rights Reserved.
 
-using EpicGames.Core;
-using EpicGames.UHT.Tables;
-using EpicGames.UHT.Tokenizer;
-using EpicGames.UHT.Utils;
+using System;
 using System.Collections.Generic;
 using System.Diagnostics.CodeAnalysis;
 using System.Text;
 using System.Text.Json.Serialization;
+using EpicGames.Core;
+using EpicGames.UHT.Tables;
+using EpicGames.UHT.Tokenizer;
+using EpicGames.UHT.Utils;
 
 namespace EpicGames.UHT.Types
 {
@@ -25,27 +26,27 @@ namespace EpicGames.UHT.Types
 		public UhtEnum? Enum { get; set; }
 
 		/// <inheritdoc/>
-		public override string EngineClassName { get => "ByteProperty"; }
+		public override string EngineClassName => "ByteProperty";
 
 		/// <inheritdoc/>
-		protected override string CppTypeText { get => "uint8"; }
+		protected override string CppTypeText => "uint8";
 
 		/// <inheritdoc/>
-		protected override string PGetMacroText { get => this.Enum == null || this.Enum.CppForm != UhtEnumCppForm.EnumClass ? "PROPERTY" : "ENUM"; }
+		protected override string PGetMacroText => this.Enum == null || this.Enum.CppForm != UhtEnumCppForm.EnumClass ? "PROPERTY" : "ENUM";
 
 		/// <inheritdoc/>
-		protected override UhtPGetArgumentType PGetTypeArgument { get => this.Enum == null || this.Enum.CppForm != UhtEnumCppForm.EnumClass ? UhtPGetArgumentType.EngineClass : UhtPGetArgumentType.TypeText; }
+		protected override UhtPGetArgumentType PGetTypeArgument => this.Enum == null || this.Enum.CppForm != UhtEnumCppForm.EnumClass ? UhtPGetArgumentType.EngineClass : UhtPGetArgumentType.TypeText;
 
 		/// <summary>
 		/// Construct a new property
 		/// </summary>
-		/// <param name="PropertySettings">Property settings</param>
-		/// <param name="IntType">Integer type</param>
-		/// <param name="Enum">Optional referenced enum</param>
-		public UhtByteProperty(UhtPropertySettings PropertySettings, UhtPropertyIntType IntType, UhtEnum? Enum = null) : base(PropertySettings, IntType)
+		/// <param name="propertySettings">Property settings</param>
+		/// <param name="intType">Integer type</param>
+		/// <param name="enumObj">Optional referenced enum</param>
+		public UhtByteProperty(UhtPropertySettings propertySettings, UhtPropertyIntType intType, UhtEnum? enumObj = null) : base(propertySettings, intType)
 		{
-			this.Enum = Enum;
-			this.PropertyCaps |= UhtPropertyCaps.CanExposeOnSpawn | UhtPropertyCaps.IsParameterSupportedByBlueprint | 
+			this.Enum = enumObj;
+			this.PropertyCaps |= UhtPropertyCaps.CanExposeOnSpawn | UhtPropertyCaps.IsParameterSupportedByBlueprint |
 				UhtPropertyCaps.IsMemberSupportedByBlueprint | UhtPropertyCaps.SupportsRigVM;
 			if (this.Enum != null)
 			{
@@ -63,91 +64,91 @@ namespace EpicGames.UHT.Types
 		}
 
 		/// <inheritdoc/>
-		public override void CollectReferencesInternal(IUhtReferenceCollector Collector, bool bTemplateProperty)
+		public override void CollectReferencesInternal(IUhtReferenceCollector collector, bool templateProperty)
 		{
-			Collector.AddCrossModuleReference(this.Enum, true);
+			collector.AddCrossModuleReference(this.Enum, true);
 		}
 
 		/// <inheritdoc/>
-		public override StringBuilder AppendText(StringBuilder Builder, UhtPropertyTextType TextType, bool bIsTemplateArgument)
+		public override StringBuilder AppendText(StringBuilder builder, UhtPropertyTextType textType, bool isTemplateArgument)
 		{
 			if (this.Enum != null)
 			{
-				return UhtEnumProperty.AppendEnumText(Builder, this, this.Enum, TextType, bIsTemplateArgument);
+				return UhtEnumProperty.AppendEnumText(builder, this, this.Enum, textType, isTemplateArgument);
 			}
 			else
 			{
-				return base.AppendText(Builder, TextType, bIsTemplateArgument);
+				return base.AppendText(builder, textType, isTemplateArgument);
 			}
 		}
 
 		/// <inheritdoc/>
-		public override StringBuilder AppendMemberDecl(StringBuilder Builder, IUhtPropertyMemberContext Context, string Name, string NameSuffix, int Tabs)
+		public override StringBuilder AppendMemberDecl(StringBuilder builder, IUhtPropertyMemberContext context, string name, string nameSuffix, int tabs)
 		{
-			return AppendMemberDecl(Builder, Context, Name, NameSuffix, Tabs, "FBytePropertyParams");
+			return AppendMemberDecl(builder, context, name, nameSuffix, tabs, "FBytePropertyParams");
 		}
 
 		/// <inheritdoc/>
-		public override StringBuilder AppendMemberDef(StringBuilder Builder, IUhtPropertyMemberContext Context, string Name, string NameSuffix, string? Offset, int Tabs)
+		public override StringBuilder AppendMemberDef(StringBuilder builder, IUhtPropertyMemberContext context, string name, string nameSuffix, string? offset, int tabs)
 		{
-			AppendMemberDefStart(Builder, Context, Name, NameSuffix, Offset, Tabs, "FBytePropertyParams", "UECodeGen_Private::EPropertyGenFlags::Byte");
-			AppendMemberDefRef(Builder, Context, this.Enum, true, true);
-			AppendMemberDefEnd(Builder, Context, Name, NameSuffix);
-			return Builder;
+			AppendMemberDefStart(builder, context, name, nameSuffix, offset, tabs, "FBytePropertyParams", "UECodeGen_Private::EPropertyGenFlags::Byte");
+			AppendMemberDefRef(builder, context, this.Enum, true, true);
+			AppendMemberDefEnd(builder, context, name, nameSuffix);
+			return builder;
 		}
 
 		/// <inheritdoc/>
-		public override void AppendObjectHashes(StringBuilder Builder, int StartingLength, IUhtPropertyMemberContext Context)
+		public override void AppendObjectHashes(StringBuilder builder, int startingLength, IUhtPropertyMemberContext context)
 		{
-			Builder.AppendObjectHash(StartingLength, this, Context, this.Enum);
+			builder.AppendObjectHash(startingLength, this, context, this.Enum);
 		}
 
 		/// <inheritdoc/>
-		public override StringBuilder AppendFunctionThunkParameterArg(StringBuilder Builder)
+		public override StringBuilder AppendFunctionThunkParameterArg(StringBuilder builder)
 		{
 			if (this.Enum != null)
 			{
-				return UhtEnumProperty.AppendEnumFunctionThunkParameterArg(Builder, this, this.Enum);
+				return UhtEnumProperty.AppendEnumFunctionThunkParameterArg(builder, this, this.Enum);
 			}
 			else
 			{
-				return base.AppendFunctionThunkParameterArg(Builder);
+				return base.AppendFunctionThunkParameterArg(builder);
 			}
 		}
 
 		/// <inheritdoc/>
-		public override bool SanitizeDefaultValue(IUhtTokenReader DefaultValueReader, StringBuilder InnerDefaultValue)
+		public override bool SanitizeDefaultValue(IUhtTokenReader defaultValueReader, StringBuilder innerDefaultValue)
 		{
 			if (this.Enum != null)
 			{
-				return UhtEnumProperty.SanitizeEnumDefaultValue(this, this.Enum, DefaultValueReader, InnerDefaultValue);
+				return UhtEnumProperty.SanitizeEnumDefaultValue(this, this.Enum, defaultValueReader, innerDefaultValue);
 			}
 
-			int Value = DefaultValueReader.GetConstIntExpression();
-			InnerDefaultValue.Append(Value);
-			return Value >= byte.MinValue && Value <= byte.MaxValue;
+			int value = defaultValueReader.GetConstIntExpression();
+			innerDefaultValue.Append(value);
+			return value >= Byte.MinValue && value <= Byte.MaxValue;
 		}
 
 		/// <inheritdoc/>
-		public override bool IsSameType(UhtProperty Other)
+		public override bool IsSameType(UhtProperty other)
 		{
-			if (Other is UhtByteProperty OtherByte)
+			if (other is UhtByteProperty otherByte)
 			{
-				return this.Enum == OtherByte.Enum;
+				return this.Enum == otherByte.Enum;
 			}
-			else if (Other is UhtEnumProperty OtherEnum)
+			else if (other is UhtEnumProperty otherEnum)
 			{
-				return this.Enum == OtherEnum.Enum;
+				return this.Enum == otherEnum.Enum;
 			}
 			return false;
 		}
 
 		/// <inheritdoc/>
-		protected override void ValidateFunctionArgument(UhtFunction Function, UhtValidationOptions Options)
+		protected override void ValidateFunctionArgument(UhtFunction function, UhtValidationOptions options)
 		{
-			base.ValidateFunctionArgument(Function, Options);
+			base.ValidateFunctionArgument(function, options);
 
-			if (Function.FunctionFlags.HasAnyFlags(EFunctionFlags.BlueprintEvent | EFunctionFlags.BlueprintCallable))
+			if (function.FunctionFlags.HasAnyFlags(EFunctionFlags.BlueprintEvent | EFunctionFlags.BlueprintCallable))
 			{
 				if (this.Enum != null && this.Enum.UnderlyingType != UhtEnumUnderlyingType.uint8 && this.Enum.UnderlyingType != UhtEnumUnderlyingType.Unspecified)
 				{
@@ -160,15 +161,15 @@ namespace EpicGames.UHT.Types
 		[UhtPropertyType(Keyword = "uint8", Options = UhtPropertyTypeOptions.Simple | UhtPropertyTypeOptions.Immediate)]
 		[SuppressMessage("CodeQuality", "IDE0051:Remove unused private members", Justification = "Attribute accessed method")]
 		[SuppressMessage("Style", "IDE0060:Remove unused parameter", Justification = "Attribute accessed method")]
-		private static UhtProperty? UInt8Property(UhtPropertyResolvePhase ResolvePhase, UhtPropertySettings PropertySettings, IUhtTokenReader TokenReader, UhtToken MatchedToken)
+		private static UhtProperty? UInt8Property(UhtPropertyResolvePhase resolvePhase, UhtPropertySettings propertySettings, IUhtTokenReader tokenReader, UhtToken matchedToken)
 		{
-			if (PropertySettings.bIsBitfield)
+			if (propertySettings.IsBitfield)
 			{
-				return new UhtBoolProperty(PropertySettings, UhtBoolType.UInt8);
+				return new UhtBoolProperty(propertySettings, UhtBoolType.UInt8);
 			}
 			else
 			{
-				return new UhtByteProperty(PropertySettings, UhtPropertyIntType.Sized);
+				return new UhtByteProperty(propertySettings, UhtPropertyIntType.Sized);
 			}
 		}
 		#endregion
diff --git a/Engine/Source/Programs/Shared/EpicGames.UHT/Types/Properties/UhtClassProperty.cs b/Engine/Source/Programs/Shared/EpicGames.UHT/Types/Properties/UhtClassProperty.cs
--- a/Engine/Source/Programs/Shared/EpicGames.UHT/Types/Properties/UhtClassProperty.cs
+++ b/Engine/Source/Programs/Shared/EpicGames.UHT/Types/Properties/UhtClassProperty.cs
@@ -1,11 +1,11 @@
 // Copyright Epic Games, Inc. All Rights Reserved.
 
+using System.Diagnostics.CodeAnalysis;
+using System.Text;
 using EpicGames.Core;
 using EpicGames.UHT.Tables;
 using EpicGames.UHT.Tokenizer;
 using EpicGames.UHT.Utils;
-using System.Diagnostics.CodeAnalysis;
-using System.Text;
 
 namespace EpicGames.UHT.Types
 {
@@ -17,25 +17,25 @@ namespace EpicGames.UHT.Types
 	public class UhtClassProperty : UhtObjectProperty
 	{
 		/// <inheritdoc/>
-		public override string EngineClassName { get => "ClassProperty"; }
+		public override string EngineClassName => "ClassProperty";
 
 		/// <inheritdoc/>
-		protected override bool bPGetPassAsNoPtr { get => this.PropertyFlags.HasAnyFlags(EPropertyFlags.UObjectWrapper); }
+		protected override bool PGetPassAsNoPtr => this.PropertyFlags.HasAnyFlags(EPropertyFlags.UObjectWrapper);
 
 		/// <inheritdoc/>
-		protected override UhtPGetArgumentType PGetTypeArgument { get => UhtPGetArgumentType.TypeText; }
+		protected override UhtPGetArgumentType PGetTypeArgument => UhtPGetArgumentType.TypeText;
 
 		/// <summary>
 		/// Construct a new property
 		/// </summary>
-		/// <param name="PropertySettings">Property settings</param>
-		/// <param name="Class">Referenced class</param>
-		/// <param name="MetaClass">Reference meta class</param>
-		/// <param name="ExtraFlags">Extra flags to apply to the property.</param>
-		public UhtClassProperty(UhtPropertySettings PropertySettings, UhtClass Class, UhtClass MetaClass, EPropertyFlags ExtraFlags = EPropertyFlags.None)
-			: base(PropertySettings, Class, MetaClass)
+		/// <param name="propertySettings">Property settings</param>
+		/// <param name="classObj">Referenced class</param>
+		/// <param name="metaClass">Reference meta class</param>
+		/// <param name="extraFlags">Extra flags to apply to the property.</param>
+		public UhtClassProperty(UhtPropertySettings propertySettings, UhtClass classObj, UhtClass metaClass, EPropertyFlags extraFlags = EPropertyFlags.None)
+			: base(propertySettings, classObj, metaClass)
 		{
-			this.PropertyFlags |= ExtraFlags;
+			this.PropertyFlags |= extraFlags;
 			this.PropertyCaps |= UhtPropertyCaps.CanHaveConfig;
 			this.PropertyCaps &= ~(UhtPropertyCaps.CanBeInstanced);
 		}
@@ -47,30 +47,30 @@ namespace EpicGames.UHT.Types
 		}
 
 		/// <inheritdoc/>
-		private StringBuilder AppendSubClassText(StringBuilder Builder)
+		private StringBuilder AppendSubClassText(StringBuilder builder)
 		{
-			Builder.Append("TSubclassOf<").Append(this.MetaClass?.SourceName).Append("> "); //COMPATIBILITY-TODO - Extra space in old UHT
-			return Builder;
+			builder.Append("TSubclassOf<").Append(this.MetaClass?.SourceName).Append("> "); //COMPATIBILITY-TODO - Extra space in old UHT
+			return builder;
 		}
 
 		/// <inheritdoc/>
-		private StringBuilder AppendText(StringBuilder Builder)
+		private StringBuilder AppendText(StringBuilder builder)
 		{
 			if (this.PropertyFlags.HasAnyFlags(EPropertyFlags.UObjectWrapper))
 			{
-				AppendSubClassText(Builder);
+				AppendSubClassText(builder);
 			}
 			else
 			{
-				Builder.Append(this.Class.SourceName).Append('*');
+				builder.Append(this.Class.SourceName).Append('*');
 			}
-			return Builder;
+			return builder;
 		}
 
 		/// <inheritdoc/>
-		public override StringBuilder AppendText(StringBuilder Builder, UhtPropertyTextType TextType, bool bIsTemplateArgument)
+		public override StringBuilder AppendText(StringBuilder builder, UhtPropertyTextType textType, bool isTemplateArgument)
 		{
-			switch (TextType)
+			switch (textType)
 			{
 				case UhtPropertyTextType.Generic:
 				case UhtPropertyTextType.Sparse:
@@ -85,21 +85,21 @@ namespace EpicGames.UHT.Types
 				case UhtPropertyTextType.FunctionThunkParameterArrayType:
 				case UhtPropertyTextType.RigVMType:
 				case UhtPropertyTextType.GetterSetterArg:
-					AppendText(Builder);
+					AppendText(builder);
 					break;
 
 				case UhtPropertyTextType.FunctionThunkRetVal:
 					if (this.PropertyFlags.HasAnyFlags(EPropertyFlags.ConstParm))
 					{
-						Builder.Append("const ");
+						builder.Append("const ");
 					}
 					if (this.PropertyFlags.HasAnyFlags(EPropertyFlags.UObjectWrapper))
 					{
-						AppendText(Builder);
+						AppendText(builder);
 					}
 					else
 					{
-						Builder.Append("UClass*");
+						builder.Append("UClass*");
 					}
 					break;
 
@@ -107,58 +107,58 @@ namespace EpicGames.UHT.Types
 				case UhtPropertyTextType.EventParameterFunctionMember:
 					if (this.PropertyFlags.HasAnyFlags(EPropertyFlags.ReturnParm) && !this.PropertyFlags.HasAnyFlags(EPropertyFlags.UObjectWrapper))
 					{
-						Builder.Append("UClass*");
+						builder.Append("UClass*");
 					}
 					else
 					{
-						AppendText(Builder);
+						AppendText(builder);
 					}
 					break;
 
 				case UhtPropertyTextType.FunctionThunkParameterArgType:
 					if (this.PropertyFlags.HasAllFlags(EPropertyFlags.OutParm | EPropertyFlags.UObjectWrapper))
 					{
-						AppendSubClassText(Builder);
+						AppendSubClassText(builder);
 					}
 					else
 					{
-						Builder.Append(this.Class.SourceName);
+						builder.Append(this.Class.SourceName);
 					}
 					break;
 			}
-			return Builder;
+			return builder;
 		}
 
 		/// <inheritdoc/>
-		public override StringBuilder AppendMemberDecl(StringBuilder Builder, IUhtPropertyMemberContext Context, string Name, string NameSuffix, int Tabs)
+		public override StringBuilder AppendMemberDecl(StringBuilder builder, IUhtPropertyMemberContext context, string name, string nameSuffix, int tabs)
 		{
-			return base.AppendMemberDecl(Builder, Context, Name, NameSuffix, Tabs, "FClassPropertyParams");
+			return base.AppendMemberDecl(builder, context, name, nameSuffix, tabs, "FClassPropertyParams");
 		}
 
 		/// <inheritdoc/>
-		public override StringBuilder AppendMemberDef(StringBuilder Builder, IUhtPropertyMemberContext Context, string Name, string NameSuffix, string? Offset, int Tabs)
+		public override StringBuilder AppendMemberDef(StringBuilder builder, IUhtPropertyMemberContext context, string name, string nameSuffix, string? offset, int tabs)
 		{
-			AppendMemberDefStart(Builder, Context, Name, NameSuffix, Offset, Tabs, "FClassPropertyParams", "UECodeGen_Private::EPropertyGenFlags::Class");
-			AppendMemberDefRef(Builder, Context, this.Class, false);
-			AppendMemberDefRef(Builder, Context, this.MetaClass, false);
-			AppendMemberDefEnd(Builder, Context, Name, NameSuffix);
-			return Builder;
+			AppendMemberDefStart(builder, context, name, nameSuffix, offset, tabs, "FClassPropertyParams", "UECodeGen_Private::EPropertyGenFlags::Class");
+			AppendMemberDefRef(builder, context, this.Class, false);
+			AppendMemberDefRef(builder, context, this.MetaClass, false);
+			AppendMemberDefEnd(builder, context, name, nameSuffix);
+			return builder;
 		}
 
 		#region Keyword
 		[UhtPropertyType(Keyword = "TSubclassOf")]
 		[SuppressMessage("CodeQuality", "IDE0051:Remove unused private members", Justification = "Attribute accessed method")]
 		[SuppressMessage("Style", "IDE0060:Remove unused parameter", Justification = "Attribute accessed method")]
-		private static UhtProperty? SubclassOfProperty(UhtPropertyResolvePhase ResolvePhase, UhtPropertySettings PropertySettings, IUhtTokenReader TokenReader, UhtToken MatchedToken)
+		private static UhtProperty? SubclassOfProperty(UhtPropertyResolvePhase resolvePhase, UhtPropertySettings propertySettings, IUhtTokenReader tokenReader, UhtToken matchedToken)
 		{
-			UhtClass? MetaClass = UhtObjectPropertyBase.ParseTemplateClass(PropertySettings, TokenReader, MatchedToken);
-			if (MetaClass == null)
+			UhtClass? metaClass = UhtObjectPropertyBase.ParseTemplateClass(propertySettings, tokenReader, matchedToken);
+			if (metaClass == null)
 			{
 				return null;
 			}
 
 			// With TSubclassOf, MetaClass is used as a class limiter.  
-			return new UhtClassProperty(PropertySettings, MetaClass.Session.UClass, MetaClass, EPropertyFlags.UObjectWrapper);
+			return new UhtClassProperty(propertySettings, metaClass.Session.UClass, metaClass, EPropertyFlags.UObjectWrapper);
 		}
 		#endregion
 	}
diff --git a/Engine/Source/Programs/Shared/EpicGames.UHT/Types/Properties/UhtClassPtrProperties.cs b/Engine/Source/Programs/Shared/EpicGames.UHT/Types/Properties/UhtClassPtrProperties.cs
--- a/Engine/Source/Programs/Shared/EpicGames.UHT/Types/Properties/UhtClassPtrProperties.cs
+++ b/Engine/Source/Programs/Shared/EpicGames.UHT/Types/Properties/UhtClassPtrProperties.cs
@@ -1,9 +1,9 @@
 // Copyright Epic Games, Inc. All Rights Reserved.
 
+using System.Text;
 using EpicGames.Core;
 using EpicGames.UHT.Tables;
 using EpicGames.UHT.Utils;
-using System.Text;
 
 namespace EpicGames.UHT.Types
 {
@@ -15,69 +15,69 @@ namespace EpicGames.UHT.Types
 	public class UhtClassPtrProperty : UhtClassProperty
 	{
 		/// <inheritdoc/>
-		public override string EngineClassName { get => "ClassPtrProperty"; }
+		public override string EngineClassName => "ClassPtrProperty";
 
 		/// <inheritdoc/>
-		protected override string CppTypeText { get => "ObjectPtr"; }
+		protected override string CppTypeText => "ObjectPtr";
 
 		/// <inheritdoc/>
-		protected override string PGetMacroText { get => "OBJECTPTR"; }
+		protected override string PGetMacroText => "OBJECTPTR";
 
 		/// <inheritdoc/>
-		protected override UhtPGetArgumentType PGetTypeArgument { get => UhtPGetArgumentType.TypeText; }
+		protected override UhtPGetArgumentType PGetTypeArgument => UhtPGetArgumentType.TypeText;
 
 		/// <summary>
 		/// Construct a new property
 		/// </summary>
-		/// <param name="PropertySettings"></param>
-		/// <param name="Class">Referenced class</param>
-		/// <param name="MetaClass">Meta data class</param>
-		/// <param name="ExtraFlags">Extra property flags to apply to the property</param>
-		public UhtClassPtrProperty(UhtPropertySettings PropertySettings, UhtClass Class, UhtClass MetaClass, EPropertyFlags ExtraFlags = EPropertyFlags.None)
-			: base(PropertySettings, Class, MetaClass)
+		/// <param name="propertySettings"></param>
+		/// <param name="classObj">Referenced class</param>
+		/// <param name="metaClass">Meta data class</param>
+		/// <param name="extraFlags">Extra property flags to apply to the property</param>
+		public UhtClassPtrProperty(UhtPropertySettings propertySettings, UhtClass classObj, UhtClass metaClass, EPropertyFlags extraFlags = EPropertyFlags.None)
+			: base(propertySettings, classObj, metaClass)
 		{
-			this.PropertyFlags |= ExtraFlags | EPropertyFlags.UObjectWrapper;
+			this.PropertyFlags |= extraFlags | EPropertyFlags.UObjectWrapper;
 		}
 
 		/// <inheritdoc/>
-		public override StringBuilder AppendText(StringBuilder Builder, UhtPropertyTextType TextType, bool bIsTemplateArgument)
+		public override StringBuilder AppendText(StringBuilder builder, UhtPropertyTextType textType, bool isTemplateArgument)
 		{
-			switch (TextType)
+			switch (textType)
 			{
 				default:
-					Builder.Append("TObjectPtr<").Append(this.Class.SourceName).Append('>');
+					builder.Append("TObjectPtr<").Append(this.Class.SourceName).Append('>');
 					break;
 			}
-			return Builder;
+			return builder;
 		}
 
 		/// <inheritdoc/>
-		public override StringBuilder AppendMemberDecl(StringBuilder Builder, IUhtPropertyMemberContext Context, string Name, string NameSuffix, int Tabs)
+		public override StringBuilder AppendMemberDecl(StringBuilder builder, IUhtPropertyMemberContext context, string name, string nameSuffix, int tabs)
 		{
-			return AppendMemberDecl(Builder, Context, Name, NameSuffix, Tabs, "FClassPtrPropertyParams");
+			return AppendMemberDecl(builder, context, name, nameSuffix, tabs, "FClassPtrPropertyParams");
 		}
 
 		/// <inheritdoc/>
-		public override StringBuilder AppendMemberDef(StringBuilder Builder, IUhtPropertyMemberContext Context, string Name, string NameSuffix, string? Offset, int Tabs)
+		public override StringBuilder AppendMemberDef(StringBuilder builder, IUhtPropertyMemberContext context, string name, string nameSuffix, string? offset, int tabs)
 		{
-			AppendMemberDefStart(Builder, Context, Name, NameSuffix, Offset, Tabs, "FClassPtrPropertyParams", 
+			AppendMemberDefStart(builder, context, name, nameSuffix, offset, tabs, "FClassPtrPropertyParams",
 				"UECodeGen_Private::EPropertyGenFlags::Class | UECodeGen_Private::EPropertyGenFlags::ObjectPtr");
-			AppendMemberDefRef(Builder, Context, this.Class, false);
-			AppendMemberDefRef(Builder, Context, this.MetaClass, false);
-			AppendMemberDefEnd(Builder, Context, Name, NameSuffix);
-			return Builder;
+			AppendMemberDefRef(builder, context, this.Class, false);
+			AppendMemberDefRef(builder, context, this.MetaClass, false);
+			AppendMemberDefEnd(builder, context, name, nameSuffix);
+			return builder;
 		}
 
 		/// <inheritdoc/>
-		public override void Validate(UhtStruct OuterStruct, UhtProperty OutermostProperty, UhtValidationOptions Options)
+		public override void Validate(UhtStruct outerStruct, UhtProperty outermostProperty, UhtValidationOptions options)
 		{
-			base.Validate(OuterStruct, OutermostProperty, Options);
+			base.Validate(outerStruct, outermostProperty, options);
 
 			// UFunctions with a smart pointer as input parameter wont compile anyway, because of missing P_GET_... macro.
 			// UFunctions with a smart pointer as return type will crash when called via blueprint, because they are not supported in VM.
 			if (this.PropertyCategory == UhtPropertyCategory.RegularParameter || this.PropertyCategory == UhtPropertyCategory.ReplicatedParameter)
 			{
-				OuterStruct.LogError("UFunctions cannot take a TObjectPtr as a parameter.");
+				outerStruct.LogError("UFunctions cannot take a TObjectPtr as a parameter.");
 			}
 		}
 	}
diff --git a/Engine/Source/Programs/Shared/EpicGames.UHT/Types/Properties/UhtContainerProperty.cs b/Engine/Source/Programs/Shared/EpicGames.UHT/Types/Properties/UhtContainerProperty.cs
--- a/Engine/Source/Programs/Shared/EpicGames.UHT/Types/Properties/UhtContainerProperty.cs
+++ b/Engine/Source/Programs/Shared/EpicGames.UHT/Types/Properties/UhtContainerProperty.cs
@@ -17,39 +17,39 @@ namespace EpicGames.UHT.Types
 		/// <summary>
 		/// Construct a new property
 		/// </summary>
-		/// <param name="PropertySettings">Property settings</param>
-		/// <param name="Value">Value property</param>
-		protected UhtContainerBaseProperty(UhtPropertySettings PropertySettings, UhtProperty Value) : base(PropertySettings)
+		/// <param name="propertySettings">Property settings</param>
+		/// <param name="value">Value property</param>
+		protected UhtContainerBaseProperty(UhtPropertySettings propertySettings, UhtProperty value) : base(propertySettings)
 		{
-			this.ValueProperty = Value;
-			this.PropertyCaps = (this.PropertyCaps & ~(UhtPropertyCaps.CanBeInstanced | UhtPropertyCaps.CanHaveConfig)) | 
+			this.ValueProperty = value;
+			this.PropertyCaps = (this.PropertyCaps & ~(UhtPropertyCaps.CanBeInstanced | UhtPropertyCaps.CanHaveConfig)) |
 				(this.ValueProperty.PropertyCaps & (UhtPropertyCaps.CanBeInstanced | UhtPropertyCaps.CanHaveConfig));
 		}
 
 		/// <inheritdoc/>
-		public override void Validate(UhtStruct OuterStruct, UhtProperty OutermostProperty, UhtValidationOptions Options)
+		public override void Validate(UhtStruct outerStruct, UhtProperty outermostProperty, UhtValidationOptions options)
 		{
-			base.Validate(OuterStruct, OutermostProperty, Options);
-			this.ValueProperty.Validate(OuterStruct, OutermostProperty, Options | UhtValidationOptions.IsValue);
+			base.Validate(outerStruct, outermostProperty, options);
+			this.ValueProperty.Validate(outerStruct, outermostProperty, options | UhtValidationOptions.IsValue);
 		}
 
 		/// <summary>
 		/// Propagate flags and meta data to/from child properties
 		/// </summary>
-		/// <param name="Container">Container property</param>
-		/// <param name="MetaData">Meta data</param>
-		/// <param name="Inner">Inner property</param>
-		protected static void PropagateFlagsFromInnerAndHandlePersistentInstanceMetadata(UhtProperty Container, UhtMetaData? MetaData, UhtProperty Inner)
+		/// <param name="container">Container property</param>
+		/// <param name="metaData">Meta data</param>
+		/// <param name="inner">Inner property</param>
+		protected static void PropagateFlagsFromInnerAndHandlePersistentInstanceMetadata(UhtProperty container, UhtMetaData? metaData, UhtProperty inner)
 		{
 			// Copy some of the property flags to the container property.
-			if (Inner.PropertyFlags.HasAnyFlags(EPropertyFlags.ContainsInstancedReference | EPropertyFlags.InstancedReference))
+			if (inner.PropertyFlags.HasAnyFlags(EPropertyFlags.ContainsInstancedReference | EPropertyFlags.InstancedReference))
 			{
-				Container.PropertyFlags |= EPropertyFlags.ContainsInstancedReference;
-				Container.PropertyFlags &= ~(EPropertyFlags.InstancedReference | EPropertyFlags.PersistentInstance); //this was propagated to the inner
+				container.PropertyFlags |= EPropertyFlags.ContainsInstancedReference;
+				container.PropertyFlags &= ~(EPropertyFlags.InstancedReference | EPropertyFlags.PersistentInstance); //this was propagated to the inner
 
-				if (MetaData != null && Inner.PropertyFlags.HasAnyFlags(EPropertyFlags.PersistentInstance))
+				if (metaData != null && inner.PropertyFlags.HasAnyFlags(EPropertyFlags.PersistentInstance))
 				{
-					Inner.MetaData.Add(MetaData);
+					inner.MetaData.Add(metaData);
 				}
 			}
 		}
@@ -57,21 +57,21 @@ namespace EpicGames.UHT.Types
 		/// <summary>
 		/// Resolve the child and return any new flags
 		/// </summary>
-		/// <param name="Child">Child to resolve</param>
-		/// <param name="Phase">Resolve phase</param>
+		/// <param name="child">Child to resolve</param>
+		/// <param name="phase">Resolve phase</param>
 		/// <returns>And new flags</returns>
-		protected static EPropertyFlags ResolveAndReturnNewFlags(UhtProperty Child, UhtResolvePhase Phase)
+		protected static EPropertyFlags ResolveAndReturnNewFlags(UhtProperty child, UhtResolvePhase phase)
 		{
-			EPropertyFlags OldFlags = Child.PropertyFlags;
-			Child.Resolve(Phase);
-			EPropertyFlags NewFlags = Child.PropertyFlags;
-			return NewFlags & ~OldFlags;
+			EPropertyFlags oldFlags = child.PropertyFlags;
+			child.Resolve(phase);
+			EPropertyFlags newFlags = child.PropertyFlags;
+			return newFlags & ~oldFlags;
 		}
 
 		/// <inheritdoc/>
-		public override bool ScanForInstancedReferenced(bool bDeepScan)
+		public override bool ScanForInstancedReferenced(bool deepScan)
 		{
-			return this.ValueProperty.ScanForInstancedReferenced(bDeepScan);
+			return this.ValueProperty.ScanForInstancedReferenced(deepScan);
 		}
 	}
 }
