More coding standard fixes

diff --git a/Engine/Source/Programs/Shared/EpicGames.UHT/Types/Properties/UhtDelegateProperty.cs b/Engine/Source/Programs/Shared/EpicGames.UHT/Types/Properties/UhtDelegateProperty.cs
--- a/Engine/Source/Programs/Shared/EpicGames.UHT/Types/Properties/UhtDelegateProperty.cs
+++ b/Engine/Source/Programs/Shared/EpicGames.UHT/Types/Properties/UhtDelegateProperty.cs
@@ -1,11 +1,11 @@
 // Copyright Epic Games, Inc. All Rights Reserved.
 
+using System.Text;
+using System.Text.Json.Serialization;
 using EpicGames.Core;
 using EpicGames.UHT.Tables;
 using EpicGames.UHT.Tokenizer;
 using EpicGames.UHT.Utils;
-using System.Text;
-using System.Text.Json.Serialization;
 
 namespace EpicGames.UHT.Types
 {
@@ -16,16 +16,16 @@ namespace EpicGames.UHT.Types
 	public class UhtDelegateProperty : UhtProperty
 	{
 		/// <inheritdoc/>
-		public override string EngineClassName { get => "DelegateProperty"; }
+		public override string EngineClassName => "DelegateProperty";
 
 		/// <inheritdoc/>
-		protected override string CppTypeText { get => "FScriptDelegate"; }
+		protected override string CppTypeText => "FScriptDelegate";
 
 		/// <inheritdoc/>
-		protected override string PGetMacroText { get => "PROPERTY"; }
+		protected override string PGetMacroText => "PROPERTY";
 
 		/// <inheritdoc/>
-		protected override UhtPGetArgumentType PGetTypeArgument { get => UhtPGetArgumentType.EngineClass; }
+		protected override UhtPGetArgumentType PGetTypeArgument => UhtPGetArgumentType.EngineClass;
 
 		/// <summary>
 		/// Referenced function
@@ -36,117 +36,117 @@ namespace EpicGames.UHT.Types
 		/// <summary>
 		/// Construct new property
 		/// </summary>
-		/// <param name="PropertySettings">Property settings</param>
-		/// <param name="Function">Referenced function</param>
-		public UhtDelegateProperty(UhtPropertySettings PropertySettings, UhtFunction Function) : base(PropertySettings)
+		/// <param name="propertySettings">Property settings</param>
+		/// <param name="function">Referenced function</param>
+		public UhtDelegateProperty(UhtPropertySettings propertySettings, UhtFunction function) : base(propertySettings)
 		{
-			this.Function = Function;
-			this.HeaderFile.AddReferencedHeader(Function);
-			this.PropertyCaps |= UhtPropertyCaps.PassCppArgsByRef | UhtPropertyCaps.IsParameterSupportedByBlueprint | 
+			this.Function = function;
+			this.HeaderFile.AddReferencedHeader(function);
+			this.PropertyCaps |= UhtPropertyCaps.PassCppArgsByRef | UhtPropertyCaps.IsParameterSupportedByBlueprint |
 				UhtPropertyCaps.IsMemberSupportedByBlueprint | UhtPropertyCaps.SupportsRigVM;
 		}
 
 		/// <inheritdoc/>
-		protected override bool ResolveSelf(UhtResolvePhase Phase)
+		protected override bool ResolveSelf(UhtResolvePhase phase)
 		{
-			bool bResults = base.ResolveSelf(Phase);
-			switch (Phase)
+			bool results = base.ResolveSelf(phase);
+			switch (phase)
 			{
 				case UhtResolvePhase.Final:
 					this.PropertyFlags |= EPropertyFlags.InstancedReference & ~this.DisallowPropertyFlags;
 					break;
 			}
-			return bResults;
+			return results;
 		}
 
 		/// <inheritdoc/>
-		public override bool ScanForInstancedReferenced(bool bDeepScan)
+		public override bool ScanForInstancedReferenced(bool deepScan)
 		{
 			return true;
 		}
 
 		/// <inheritdoc/>
-		public override void CollectReferencesInternal(IUhtReferenceCollector Collector, bool bTemplateProperty)
+		public override void CollectReferencesInternal(IUhtReferenceCollector collector, bool templateProperty)
 		{
-			Collector.AddCrossModuleReference(this.Function, true);
+			collector.AddCrossModuleReference(this.Function, true);
 		}
 
 		/// <inheritdoc/>
-		public override StringBuilder AppendText(StringBuilder Builder, UhtPropertyTextType TextType, bool bIsTemplateArgument)
+		public override StringBuilder AppendText(StringBuilder builder, UhtPropertyTextType textType, bool isTemplateArgument)
 		{
-			switch (TextType)
+			switch (textType)
 			{
 				case UhtPropertyTextType.ExportMember:
 				case UhtPropertyTextType.RigVMType:
 				case UhtPropertyTextType.EventParameterFunctionMember:
-					Builder.Append(this.CppTypeText);
+					builder.Append(this.CppTypeText);
 					break;
 
 				default:
-					Builder.Append(this.Function.SourceName);
+					builder.Append(this.Function.SourceName);
 					break;
 			}
-			return Builder;
+			return builder;
 		}
 
 		/// <inheritdoc/>
-		public override StringBuilder AppendMemberDecl(StringBuilder Builder, IUhtPropertyMemberContext Context, string Name, string NameSuffix, int Tabs)
+		public override StringBuilder AppendMemberDecl(StringBuilder builder, IUhtPropertyMemberContext context, string name, string nameSuffix, int tabs)
 		{
-			return AppendMemberDecl(Builder, Context, Name, NameSuffix, Tabs, "FDelegatePropertyParams");
+			return AppendMemberDecl(builder, context, name, nameSuffix, tabs, "FDelegatePropertyParams");
 		}
 
 		/// <inheritdoc/>
-		public override StringBuilder AppendMemberDef(StringBuilder Builder, IUhtPropertyMemberContext Context, string Name, string NameSuffix, string? Offset, int Tabs)
+		public override StringBuilder AppendMemberDef(StringBuilder builder, IUhtPropertyMemberContext context, string name, string nameSuffix, string? offset, int tabs)
 		{
-			AppendMemberDefStart(Builder, Context, Name, NameSuffix, Offset, Tabs, "FDelegatePropertyParams", "UECodeGen_Private::EPropertyGenFlags::Delegate");
-			AppendMemberDefRef(Builder, Context, this.Function, true);
-			AppendMemberDefEnd(Builder, Context, Name, NameSuffix);
-			return Builder;
+			AppendMemberDefStart(builder, context, name, nameSuffix, offset, tabs, "FDelegatePropertyParams", "UECodeGen_Private::EPropertyGenFlags::Delegate");
+			AppendMemberDefRef(builder, context, this.Function, true);
+			AppendMemberDefEnd(builder, context, name, nameSuffix);
+			return builder;
 		}
 
 		/// <inheritdoc/>
-		public override StringBuilder AppendFunctionThunkParameterArg(StringBuilder Builder)
+		public override StringBuilder AppendFunctionThunkParameterArg(StringBuilder builder)
 		{
-			return Builder.Append(this.Function.SourceName).Append('(').AppendFunctionThunkParameterName(this).Append(')');
+			return builder.Append(this.Function.SourceName).Append('(').AppendFunctionThunkParameterName(this).Append(')');
 		}
 
 		/// <inheritdoc/>
-		public override void AppendObjectHashes(StringBuilder Builder, int StartingLength, IUhtPropertyMemberContext Context)
+		public override void AppendObjectHashes(StringBuilder builder, int startingLength, IUhtPropertyMemberContext context)
 		{
-			Builder.AppendObjectHash(StartingLength, this, Context, this.Function);
+			builder.AppendObjectHash(startingLength, this, context, this.Function);
 		}
 
 		/// <inheritdoc/>
-		public override StringBuilder AppendNullConstructorArg(StringBuilder Builder, bool bIsInitializer)
+		public override StringBuilder AppendNullConstructorArg(StringBuilder builder, bool isInitializer)
 		{
-			Builder.AppendPropertyText(this, UhtPropertyTextType.Construction).Append("()");
-			return Builder;
+			builder.AppendPropertyText(this, UhtPropertyTextType.Construction).Append("()");
+			return builder;
 		}
 
 		/// <inheritdoc/>
-		public override bool SanitizeDefaultValue(IUhtTokenReader DefaultValueReader, StringBuilder InnerDefaultValue)
+		public override bool SanitizeDefaultValue(IUhtTokenReader defaultValueReader, StringBuilder innerDefaultValue)
 		{
 			return false;
 		}
 
 		/// <inheritdoc/>
-		public override bool IsSameType(UhtProperty Other)
+		public override bool IsSameType(UhtProperty other)
 		{
-			if (Other is UhtDelegateProperty OtherObject)
+			if (other is UhtDelegateProperty otherObject)
 			{
-				return this.Function == OtherObject.Function;
+				return this.Function == otherObject.Function;
 			}
 			return false;
 		}
 
 		/// <inheritdoc/>
-		protected override void ValidateFunctionArgument(UhtFunction Function, UhtValidationOptions Options)
+		protected override void ValidateFunctionArgument(UhtFunction function, UhtValidationOptions options)
 		{
-			base.ValidateFunctionArgument(Function, Options);
+			base.ValidateFunctionArgument(function, options);
 
-			if (Function.FunctionFlags.HasAnyFlags(EFunctionFlags.Net))
+			if (function.FunctionFlags.HasAnyFlags(EFunctionFlags.Net))
 			{
-				if (Function.FunctionFlags.HasAnyFlags(EFunctionFlags.NetRequest))
+				if (function.FunctionFlags.HasAnyFlags(EFunctionFlags.NetRequest))
 				{
 					if (!this.PropertyFlags.HasAnyFlags(EPropertyFlags.RepSkip))
 					{
diff --git a/Engine/Source/Programs/Shared/EpicGames.UHT/Types/Properties/UhtDoubleProperty.cs b/Engine/Source/Programs/Shared/EpicGames.UHT/Types/Properties/UhtDoubleProperty.cs
--- a/Engine/Source/Programs/Shared/EpicGames.UHT/Types/Properties/UhtDoubleProperty.cs
+++ b/Engine/Source/Programs/Shared/EpicGames.UHT/Types/Properties/UhtDoubleProperty.cs
@@ -1,10 +1,10 @@
 // Copyright Epic Games, Inc. All Rights Reserved.
 
+using System.Diagnostics.CodeAnalysis;
+using System.Text;
 using EpicGames.UHT.Tables;
 using EpicGames.UHT.Tokenizer;
 using EpicGames.UHT.Utils;
-using System.Diagnostics.CodeAnalysis;
-using System.Text;
 
 namespace EpicGames.UHT.Types
 {
@@ -17,54 +17,54 @@ namespace EpicGames.UHT.Types
 	public class UhtDoubleProperty : UhtNumericProperty
 	{
 		/// <inheritdoc/>
-		public override string EngineClassName { get => "DoubleProperty"; }
+		public override string EngineClassName => "DoubleProperty";
 
 		/// <inheritdoc/>
-		protected override string CppTypeText { get => "double"; }
+		protected override string CppTypeText => "double";
 
 		/// <summary>
 		/// Create new property
 		/// </summary>
-		/// <param name="PropertySettings">Property settings</param>
-		public UhtDoubleProperty(UhtPropertySettings PropertySettings) : base(PropertySettings, UhtPropertyIntType.None)
+		/// <param name="propertySettings">Property settings</param>
+		public UhtDoubleProperty(UhtPropertySettings propertySettings) : base(propertySettings, UhtPropertyIntType.None)
 		{
 			this.PropertyCaps |= UhtPropertyCaps.IsParameterSupportedByBlueprint | UhtPropertyCaps.IsMemberSupportedByBlueprint;
 		}
 
 		/// <inheritdoc/>
-		public override StringBuilder AppendMemberDecl(StringBuilder Builder, IUhtPropertyMemberContext Context, string Name, string NameSuffix, int Tabs)
+		public override StringBuilder AppendMemberDecl(StringBuilder builder, IUhtPropertyMemberContext context, string name, string nameSuffix, int tabs)
 		{
-			return AppendMemberDecl(Builder, Context, Name, NameSuffix, Tabs, "FDoublePropertyParams");
+			return AppendMemberDecl(builder, context, name, nameSuffix, tabs, "FDoublePropertyParams");
 		}
 
 		/// <inheritdoc/>
-		public override StringBuilder AppendMemberDef(StringBuilder Builder, IUhtPropertyMemberContext Context, string Name, string NameSuffix, string? Offset, int Tabs)
+		public override StringBuilder AppendMemberDef(StringBuilder builder, IUhtPropertyMemberContext context, string name, string nameSuffix, string? offset, int tabs)
 		{
-			AppendMemberDefStart(Builder, Context, Name, NameSuffix, Offset, Tabs, "FDoublePropertyParams", "UECodeGen_Private::EPropertyGenFlags::Double");
-			AppendMemberDefEnd(Builder, Context, Name, NameSuffix);
-			return Builder;
+			AppendMemberDefStart(builder, context, name, nameSuffix, offset, tabs, "FDoublePropertyParams", "UECodeGen_Private::EPropertyGenFlags::Double");
+			AppendMemberDefEnd(builder, context, name, nameSuffix);
+			return builder;
 		}
 
 		/// <inheritdoc/>
-		public override bool SanitizeDefaultValue(IUhtTokenReader DefaultValueReader, StringBuilder InnerDefaultValue)
+		public override bool SanitizeDefaultValue(IUhtTokenReader defaultValueReader, StringBuilder innerDefaultValue)
 		{
-			InnerDefaultValue.AppendFormat("{0:F6}", DefaultValueReader.GetConstFloatExpression());
+			innerDefaultValue.AppendFormat("{0:F6}", defaultValueReader.GetConstFloatExpression());
 			return true;
 		}
 
 		/// <inheritdoc/>
-		public override bool IsSameType(UhtProperty Other)
+		public override bool IsSameType(UhtProperty other)
 		{
-			return Other is UhtDoubleProperty;
+			return other is UhtDoubleProperty;
 		}
 
 		#region Keyword
 		[UhtPropertyType(Keyword = "double", Options = UhtPropertyTypeOptions.Simple | UhtPropertyTypeOptions.Immediate)]
 		[SuppressMessage("CodeQuality", "IDE0051:Remove unused private members", Justification = "Attribute accessed method")]
 		[SuppressMessage("Style", "IDE0060:Remove unused parameter", Justification = "Attribute accessed method")]
-		private static UhtProperty? DoubleProperty(UhtPropertyResolvePhase ResolvePhase, UhtPropertySettings PropertySettings, IUhtTokenReader TokenReader, UhtToken MatchedToken)
+		private static UhtProperty? DoubleProperty(UhtPropertyResolvePhase resolvePhase, UhtPropertySettings propertySettings, IUhtTokenReader tokenReader, UhtToken matchedToken)
 		{
-			return new UhtDoubleProperty(PropertySettings);
+			return new UhtDoubleProperty(propertySettings);
 		}
 		#endregion
 	}
diff --git a/Engine/Source/Programs/Shared/EpicGames.UHT/Types/Properties/UhtEnumProperty.cs b/Engine/Source/Programs/Shared/EpicGames.UHT/Types/Properties/UhtEnumProperty.cs
--- a/Engine/Source/Programs/Shared/EpicGames.UHT/Types/Properties/UhtEnumProperty.cs
+++ b/Engine/Source/Programs/Shared/EpicGames.UHT/Types/Properties/UhtEnumProperty.cs
@@ -1,13 +1,13 @@
 // Copyright Epic Games, Inc. All Rights Reserved.
 
-using EpicGames.Core;
-using EpicGames.UHT.Tables;
-using EpicGames.UHT.Tokenizer;
-using EpicGames.UHT.Utils;
 using System.Collections.Generic;
 using System.Diagnostics.CodeAnalysis;
 using System.Text;
 using System.Text.Json.Serialization;
+using EpicGames.Core;
+using EpicGames.UHT.Tables;
+using EpicGames.UHT.Tokenizer;
+using EpicGames.UHT.Utils;
 
 namespace EpicGames.UHT.Types
 {
@@ -20,16 +20,16 @@ namespace EpicGames.UHT.Types
 	public class UhtEnumProperty : UhtProperty
 	{
 		/// <inheritdoc/>
-		public override string EngineClassName { get => this.Enum.CppForm != UhtEnumCppForm.EnumClass ? "ByteProperty" : "EnumProperty"; }
+		public override string EngineClassName => this.Enum.CppForm != UhtEnumCppForm.EnumClass ? "ByteProperty" : "EnumProperty";
 
 		/// <inheritdoc/>
-		protected override string CppTypeText { get => "invalid"; }
+		protected override string CppTypeText => "invalid";
 
 		/// <inheritdoc/>
-		protected override string PGetMacroText { get => this.Enum.CppForm != UhtEnumCppForm.EnumClass ? "PROPERTY" : "ENUM"; }
+		protected override string PGetMacroText => this.Enum.CppForm != UhtEnumCppForm.EnumClass ? "PROPERTY" : "ENUM";
 
 		/// <inheritdoc/>
-		protected override UhtPGetArgumentType PGetTypeArgument { get => this.Enum.CppForm != UhtEnumCppForm.EnumClass ? UhtPGetArgumentType.EngineClass : UhtPGetArgumentType.TypeText; }
+		protected override UhtPGetArgumentType PGetTypeArgument => this.Enum.CppForm != UhtEnumCppForm.EnumClass ? UhtPGetArgumentType.EngineClass : UhtPGetArgumentType.TypeText;
 
 		/// <summary>
 		/// Referenced enum
@@ -55,14 +55,14 @@ namespace EpicGames.UHT.Types
 		/// <summary>
 		/// Construct property
 		/// </summary>
-		/// <param name="PropertySettings">Property settings</param>
-		/// <param name="Enum">Referenced enum</param>
-		public UhtEnumProperty(UhtPropertySettings PropertySettings, UhtEnum Enum) : base(PropertySettings)
+		/// <param name="propertySettings">Property settings</param>
+		/// <param name="enumObj">Referenced enum</param>
+		public UhtEnumProperty(UhtPropertySettings propertySettings, UhtEnum enumObj) : base(propertySettings)
 		{
 			this.PropertyCaps |= UhtPropertyCaps.RequiresNullConstructorArg | UhtPropertyCaps.CanExposeOnSpawn | UhtPropertyCaps.IsParameterSupportedByBlueprint |
 				UhtPropertyCaps.IsMemberSupportedByBlueprint | UhtPropertyCaps.SupportsRigVM | UhtPropertyCaps.IsRigVMEnum;
-			this.Enum = Enum;
-			this.HeaderFile.AddReferencedHeader(Enum);
+			this.Enum = enumObj;
+			this.HeaderFile.AddReferencedHeader(enumObj);
 			if (this.Enum.CppForm == UhtEnumCppForm.EnumClass)
 			{
 				this.UnderlyingProperty = CreateUnderlyingProperty();
@@ -74,9 +74,9 @@ namespace EpicGames.UHT.Types
 		}
 
 		/// <inheritdoc/>
-		public override void CollectReferencesInternal(IUhtReferenceCollector Collector, bool bTemplateProperty)
+		public override void CollectReferencesInternal(IUhtReferenceCollector collector, bool templateProperty)
 		{
-			Collector.AddCrossModuleReference(this.Enum, true);
+			collector.AddCrossModuleReference(this.Enum, true);
 		}
 
 		/// <inheritdoc/>
@@ -98,234 +98,236 @@ namespace EpicGames.UHT.Types
 		/// <summary>
 		/// Append enum text
 		/// </summary>
-		/// <param name="Builder">Destination builder</param>
-		/// <param name="Property">Property in question</param>
-		/// <param name="Enum">Referenced enum</param>
-		/// <param name="TextType">Type of text to append</param>
-		/// <param name="bIsTemplateArgument">If true, this property is a template argument</param>
+		/// <param name="builder">Destination builder</param>
+		/// <param name="property">Property in question</param>
+		/// <param name="enumObj">Referenced enum</param>
+		/// <param name="textType">Type of text to append</param>
+		/// <param name="isTemplateArgument">If true, this property is a template argument</param>
 		/// <returns></returns>
-		public static StringBuilder AppendEnumText(StringBuilder Builder, UhtProperty Property, UhtEnum Enum, UhtPropertyTextType TextType, bool bIsTemplateArgument)
+		public static StringBuilder AppendEnumText(StringBuilder builder, UhtProperty property, UhtEnum enumObj, UhtPropertyTextType textType, bool isTemplateArgument)
 		{
-			switch (TextType)
+			switch (textType)
 			{
 				case UhtPropertyTextType.GenericFunctionArgOrRetVal:
 				case UhtPropertyTextType.GenericFunctionArgOrRetValImpl:
 				case UhtPropertyTextType.ClassFunctionArgOrRetVal:
 				case UhtPropertyTextType.InterfaceFunctionArgOrRetVal:
 				case UhtPropertyTextType.EventFunctionArgOrRetVal:
-					if (Enum.CppForm == UhtEnumCppForm.EnumClass || (!bIsTemplateArgument && (Property.PropertyFlags.HasAnyFlags(EPropertyFlags.ReturnParm) || !Property.PropertyFlags.HasAnyFlags(EPropertyFlags.OutParm))))
+					if (enumObj.CppForm == UhtEnumCppForm.EnumClass || 
+						(!isTemplateArgument && (property.PropertyFlags.HasAnyFlags(EPropertyFlags.ReturnParm) || 
+						!property.PropertyFlags.HasAnyFlags(EPropertyFlags.OutParm))))
 					{
-						Builder.Append(Enum.CppType);
+						builder.Append(enumObj.CppType);
 					}
 					else
 					{
-						Builder.Append("TEnumAsByte<").Append(Enum.CppType).Append('>');
+						builder.Append("TEnumAsByte<").Append(enumObj.CppType).Append('>');
 					}
 					break;
 
 				case UhtPropertyTextType.FunctionThunkParameterArgType:
-					Builder.Append(Enum.SourceName);
+					builder.Append(enumObj.SourceName);
 					break;
 
 				case UhtPropertyTextType.RigVMType:
-					if (Enum.CppForm == UhtEnumCppForm.EnumClass || !bIsTemplateArgument)
+					if (enumObj.CppForm == UhtEnumCppForm.EnumClass || !isTemplateArgument)
 					{
-						Builder.Append(Enum.CppType);
+						builder.Append(enumObj.CppType);
 					}
 					else
 					{
-						Builder.Append("TEnumAsByte<").Append(Enum.CppType).Append('>');
+						builder.Append("TEnumAsByte<").Append(enumObj.CppType).Append('>');
 					}
 					break;
 
 				case UhtPropertyTextType.GetterSetterArg:
-					Builder.Append(Enum.CppType);
+					builder.Append(enumObj.CppType);
 					break;
 
 				default:
-					if (Enum.CppForm == UhtEnumCppForm.EnumClass)
+					if (enumObj.CppForm == UhtEnumCppForm.EnumClass)
 					{
-						Builder.Append(Enum.CppType);
+						builder.Append(enumObj.CppType);
 					}
 					else
 					{
-						Builder.Append("TEnumAsByte<").Append(Enum.CppType).Append('>');
+						builder.Append("TEnumAsByte<").Append(enumObj.CppType).Append('>');
 					}
 					break;
 			}
-			return Builder;
+			return builder;
 		}
 
 		/// <summary>
 		/// Append the text for a function thunk call argument
 		/// </summary>
-		/// <param name="Builder">Output builder</param>
-		/// <param name="Property">Property in question</param>
-		/// <param name="Enum">Referenced enum</param>
+		/// <param name="builder">Output builder</param>
+		/// <param name="property">Property in question</param>
+		/// <param name="enumObj">Referenced enum</param>
 		/// <returns>Output builder</returns>
-		public static StringBuilder AppendEnumFunctionThunkParameterArg(StringBuilder Builder, UhtProperty Property, UhtEnum Enum)
+		public static StringBuilder AppendEnumFunctionThunkParameterArg(StringBuilder builder, UhtProperty property, UhtEnum enumObj)
 		{
-			if (!Property.PropertyFlags.HasAnyFlags(EPropertyFlags.OutParm))
+			if (!property.PropertyFlags.HasAnyFlags(EPropertyFlags.OutParm))
 			{
-				Builder.Append(Enum.CppType).Append('(').AppendFunctionThunkParameterName(Property).Append(')');
+				builder.Append(enumObj.CppType).Append('(').AppendFunctionThunkParameterName(property).Append(')');
 			}
-			else if (Enum.CppForm == UhtEnumCppForm.EnumClass)
+			else if (enumObj.CppForm == UhtEnumCppForm.EnumClass)
 			{
-				Builder.Append('(').Append(Enum.CppType).Append("&)(").AppendFunctionThunkParameterName(Property).Append(')');
+				builder.Append('(').Append(enumObj.CppType).Append("&)(").AppendFunctionThunkParameterName(property).Append(')');
 			}
 			else
 			{
-				Builder.Append("(TEnumAsByte<").Append(Enum.CppType).Append(">&)(").AppendFunctionThunkParameterName(Property).Append(')');
+				builder.Append("(TEnumAsByte<").Append(enumObj.CppType).Append(">&)(").AppendFunctionThunkParameterName(property).Append(')');
 			}
-			return Builder;
+			return builder;
 		}
 
 		/// <summary>
 		/// Sanitize the default value for an enumeration
 		/// </summary>
-		/// <param name="Property">Property in question</param>
-		/// <param name="Enum">Referenced enumeration</param>
-		/// <param name="DefaultValueReader">Default value</param>
-		/// <param name="InnerDefaultValue">Destination builder</param>
+		/// <param name="property">Property in question</param>
+		/// <param name="enumObj">Referenced enumeration</param>
+		/// <param name="defaultValueReader">Default value</param>
+		/// <param name="innerDefaultValue">Destination builder</param>
 		/// <returns>True if the default value was parsed.</returns>
-		public static bool SanitizeEnumDefaultValue(UhtProperty Property, UhtEnum Enum, IUhtTokenReader DefaultValueReader, StringBuilder InnerDefaultValue)
+		public static bool SanitizeEnumDefaultValue(UhtProperty property, UhtEnum enumObj, IUhtTokenReader defaultValueReader, StringBuilder innerDefaultValue)
 		{
-			UhtTokenList CppIdentifier = DefaultValueReader.GetCppIdentifier(UhtCppIdentifierOptions.None);
-			UhtTokenList StartingPoint = CppIdentifier.Next != null ? CppIdentifier.Next : CppIdentifier;
-			StartingPoint.Join(InnerDefaultValue, "::");
-			UhtTokenListCache.Return(CppIdentifier);
+			UhtTokenList cppIdentifier = defaultValueReader.GetCppIdentifier(UhtCppIdentifierOptions.None);
+			UhtTokenList startingPoint = cppIdentifier.Next ?? cppIdentifier;
+			startingPoint.Join(innerDefaultValue, "::");
+			UhtTokenListCache.Return(cppIdentifier);
 
-			int EntryIndex = Enum.GetIndexByName(InnerDefaultValue.ToString());
-			if (EntryIndex == -1)
+			int entryIndex = enumObj.GetIndexByName(innerDefaultValue.ToString());
+			if (entryIndex == -1)
 			{
 				return false;
 			}
-			if (Enum.MetaData.ContainsKey(UhtNames.Hidden, EntryIndex))
+			if (enumObj.MetaData.ContainsKey(UhtNames.Hidden, entryIndex))
 			{
-				Property.LogError($"Hidden enum entries cannot be used as default values: '{Property.SourceName}' '{InnerDefaultValue}'");
+				property.LogError($"Hidden enum entries cannot be used as default values: '{property.SourceName}' '{innerDefaultValue}'");
 			}
 			return true;
 		}
 
 		/// <inheritdoc/>
-		public override StringBuilder AppendText(StringBuilder Builder, UhtPropertyTextType TextType, bool bIsTemplateArgument)
+		public override StringBuilder AppendText(StringBuilder builder, UhtPropertyTextType textType, bool isTemplateArgument)
 		{
-			return AppendEnumText(Builder, this, this.Enum, TextType, bIsTemplateArgument);
+			return AppendEnumText(builder, this, this.Enum, textType, isTemplateArgument);
 		}
 
 		/// <inheritdoc/>
-		public override StringBuilder AppendMemberDecl(StringBuilder Builder, IUhtPropertyMemberContext Context, string Name, string NameSuffix, int Tabs)
+		public override StringBuilder AppendMemberDecl(StringBuilder builder, IUhtPropertyMemberContext context, string name, string nameSuffix, int tabs)
 		{
 			if (this.UnderlyingProperty != null)
 			{
-				Builder.AppendMemberDecl(this.UnderlyingProperty, Context, Name, GetNameSuffix(NameSuffix, "_Underlying"), Tabs);
+				builder.AppendMemberDecl(this.UnderlyingProperty, context, name, GetNameSuffix(nameSuffix, "_Underlying"), tabs);
 			}
-			return AppendMemberDecl(Builder, Context, Name, NameSuffix, Tabs, this.Enum.CppForm != UhtEnumCppForm.EnumClass ? "FBytePropertyParams" : "FEnumPropertyParams");
+			return AppendMemberDecl(builder, context, name, nameSuffix, tabs, this.Enum.CppForm != UhtEnumCppForm.EnumClass ? "FBytePropertyParams" : "FEnumPropertyParams");
 		}
 
 		/// <inheritdoc/>
-		public override StringBuilder AppendMemberDef(StringBuilder Builder, IUhtPropertyMemberContext Context, string Name, string NameSuffix, string? Offset, int Tabs)
+		public override StringBuilder AppendMemberDef(StringBuilder builder, IUhtPropertyMemberContext context, string name, string nameSuffix, string? offset, int tabs)
 		{
 			if (this.Enum.CppForm == UhtEnumCppForm.EnumClass)
 			{
 				if (this.UnderlyingProperty != null)
 				{
-					Builder.AppendMemberDef(this.UnderlyingProperty, Context, Name, GetNameSuffix(NameSuffix, "_Underlying"), "0", Tabs);
+					builder.AppendMemberDef(this.UnderlyingProperty, context, name, GetNameSuffix(nameSuffix, "_Underlying"), "0", tabs);
 				}
-				AppendMemberDefStart(Builder, Context, Name, NameSuffix, Offset, Tabs, "FEnumPropertyParams", "UECodeGen_Private::EPropertyGenFlags::Enum");
-				AppendMemberDefRef(Builder, Context, this.Enum, true);
-				AppendMemberDefEnd(Builder, Context, Name, NameSuffix);
-				return Builder;
+				AppendMemberDefStart(builder, context, name, nameSuffix, offset, tabs, "FEnumPropertyParams", "UECodeGen_Private::EPropertyGenFlags::Enum");
+				AppendMemberDefRef(builder, context, this.Enum, true);
+				AppendMemberDefEnd(builder, context, name, nameSuffix);
+				return builder;
 			}
 			else
 			{
-				AppendMemberDefStart(Builder, Context, Name, NameSuffix, Offset, Tabs, "FBytePropertyParams", "UECodeGen_Private::EPropertyGenFlags::Byte", true, true);
-				AppendMemberDefRef(Builder, Context, this.Enum, true);
-				AppendMemberDefEnd(Builder, Context, Name, NameSuffix);
+				AppendMemberDefStart(builder, context, name, nameSuffix, offset, tabs, "FBytePropertyParams", "UECodeGen_Private::EPropertyGenFlags::Byte", true, true);
+				AppendMemberDefRef(builder, context, this.Enum, true);
+				AppendMemberDefEnd(builder, context, name, nameSuffix);
 			}
-			return Builder;
+			return builder;
 		}
 
 		/// <inheritdoc/>
-		public override StringBuilder AppendMemberPtr(StringBuilder Builder, IUhtPropertyMemberContext Context, string Name, string NameSuffix, int Tabs)
+		public override StringBuilder AppendMemberPtr(StringBuilder builder, IUhtPropertyMemberContext context, string name, string nameSuffix, int tabs)
 		{
 			if (this.UnderlyingProperty != null)
 			{
-				Builder.AppendMemberPtr(this.UnderlyingProperty, Context, Name, GetNameSuffix(NameSuffix, "_Underlying"), Tabs);
+				builder.AppendMemberPtr(this.UnderlyingProperty, context, name, GetNameSuffix(nameSuffix, "_Underlying"), tabs);
 			}
-			base.AppendMemberPtr(Builder, Context, Name, NameSuffix, Tabs);
-			return Builder;
+			base.AppendMemberPtr(builder, context, name, nameSuffix, tabs);
+			return builder;
 		}
 
 		/// <inheritdoc/>
-		public override StringBuilder AppendFunctionThunkParameterArg(StringBuilder Builder)
+		public override StringBuilder AppendFunctionThunkParameterArg(StringBuilder builder)
 		{
-			return AppendEnumFunctionThunkParameterArg(Builder, this, this.Enum);
+			return AppendEnumFunctionThunkParameterArg(builder, this, this.Enum);
 		}
 
 		/// <inheritdoc/>
-		public override void AppendObjectHashes(StringBuilder Builder, int StartingLength, IUhtPropertyMemberContext Context)
+		public override void AppendObjectHashes(StringBuilder builder, int startingLength, IUhtPropertyMemberContext context)
 		{
-			Builder.AppendObjectHash(StartingLength, this, Context, this.Enum);
+			builder.AppendObjectHash(startingLength, this, context, this.Enum);
 		}
 
 		/// <inheritdoc/>
-		public override StringBuilder AppendNullConstructorArg(StringBuilder Builder, bool bIsInitializer)
+		public override StringBuilder AppendNullConstructorArg(StringBuilder builder, bool isInitializer)
 		{
 			if (this.Enum.CppForm != UhtEnumCppForm.EnumClass)
 			{
-				Builder.Append('0');
+				builder.Append('0');
 			}
 			else
 			{
-				Builder.Append('(').Append(this.Enum.CppType).Append(")0");
+				builder.Append('(').Append(this.Enum.CppType).Append(")0");
 			}
-			return Builder;
+			return builder;
 		}
 
 		/// <inheritdoc/>
-		public override bool SanitizeDefaultValue(IUhtTokenReader DefaultValueReader, StringBuilder InnerDefaultValue)
+		public override bool SanitizeDefaultValue(IUhtTokenReader defaultValueReader, StringBuilder innerDefaultValue)
 		{
-			return SanitizeEnumDefaultValue(this, this.Enum, DefaultValueReader, InnerDefaultValue);
+			return SanitizeEnumDefaultValue(this, this.Enum, defaultValueReader, innerDefaultValue);
 		}
 
 		/// <inheritdoc/>
-		public override bool IsSameType(UhtProperty Other)
+		public override bool IsSameType(UhtProperty other)
 		{
-			if (Other is UhtEnumProperty OtherEnum)
+			if (other is UhtEnumProperty otherEnum)
 			{
-				return this.Enum == OtherEnum.Enum;
+				return this.Enum == otherEnum.Enum;
 			}
-			else if (Other is UhtByteProperty OtherByte)
+			else if (other is UhtByteProperty otherByte)
 			{
-				return this.Enum == OtherByte.Enum;
+				return this.Enum == otherByte.Enum;
 			}
 			return false;
 		}
 
 		private UhtProperty CreateUnderlyingProperty()
 		{
-			UhtPropertySettings PropertySettings = new UhtPropertySettings();
-			PropertySettings.Reset(this, 0, this.PropertyCategory, 0);
-			PropertySettings.SourceName = "UnderlyingType";
+			UhtPropertySettings propertySettings = new UhtPropertySettings();
+			propertySettings.Reset(this, 0, this.PropertyCategory, 0);
+			propertySettings.SourceName = "UnderlyingType";
 			switch (this.UnderlyingType)
 			{
 				case UhtEnumUnderlyingType.int8:
-					return new UhtInt8Property(PropertySettings, this.UnderlyingTypeSize);
+					return new UhtInt8Property(propertySettings, this.UnderlyingTypeSize);
 				case UhtEnumUnderlyingType.int16:
-					return new UhtInt16Property(PropertySettings, this.UnderlyingTypeSize);
+					return new UhtInt16Property(propertySettings, this.UnderlyingTypeSize);
 				case UhtEnumUnderlyingType.int32:
-					return new UhtIntProperty(PropertySettings, this.UnderlyingTypeSize);
+					return new UhtIntProperty(propertySettings, this.UnderlyingTypeSize);
 				case UhtEnumUnderlyingType.int64:
-					return new UhtInt64Property(PropertySettings, this.UnderlyingTypeSize);
+					return new UhtInt64Property(propertySettings, this.UnderlyingTypeSize);
 				case UhtEnumUnderlyingType.uint8:
-					return new UhtByteProperty(PropertySettings, this.UnderlyingTypeSize);
+					return new UhtByteProperty(propertySettings, this.UnderlyingTypeSize);
 				case UhtEnumUnderlyingType.uint16:
-					return new UhtUInt16Property(PropertySettings, this.UnderlyingTypeSize);
+					return new UhtUInt16Property(propertySettings, this.UnderlyingTypeSize);
 				case UhtEnumUnderlyingType.uint32:
-					return new UhtUInt32Property(PropertySettings, this.UnderlyingTypeSize);
+					return new UhtUInt32Property(propertySettings, this.UnderlyingTypeSize);
 				case UhtEnumUnderlyingType.uint64:
-					return new UhtUInt64Property(PropertySettings, this.UnderlyingTypeSize);
+					return new UhtUInt64Property(propertySettings, this.UnderlyingTypeSize);
 				default:
 					throw new UhtIceException("Unexpected underlying enum type");
 			}
@@ -335,23 +337,23 @@ namespace EpicGames.UHT.Types
 		[UhtPropertyType(Keyword = "TEnumAsByte")]
 		[SuppressMessage("CodeQuality", "IDE0051:Remove unused private members", Justification = "Attribute accessed method")]
 		[SuppressMessage("Style", "IDE0060:Remove unused parameter", Justification = "Attribute accessed method")]
-		private static UhtProperty? EnumProperty(UhtPropertyResolvePhase ResolvePhase, UhtPropertySettings PropertySettings, IUhtTokenReader TokenReader, UhtToken MatchedToken)
+		private static UhtProperty? EnumProperty(UhtPropertyResolvePhase resolvePhase, UhtPropertySettings propertySettings, IUhtTokenReader tokenReader, UhtToken matchedToken)
 		{
-			UhtType Outer = PropertySettings.Outer;
-			UhtEnum? Enum = null;
-			using (var TokenContext = new UhtMessageContext("TEnumAsByte"))
+			UhtType outer = propertySettings.Outer;
+			UhtEnum? enumObj = null;
+			using (UhtMessageContext tokenContext = new UhtMessageContext("TEnumAsByte"))
 			{
-				TokenReader
+				tokenReader
 					.Require("TEnumAsByte")
 					.Require('<')
 					.Optional("enum")
-					.RequireCppIdentifier(UhtCppIdentifierOptions.None, (UhtTokenList TokenList) =>
+					.RequireCppIdentifier(UhtCppIdentifierOptions.None, (UhtTokenList tokenList) =>
 					{
-						Enum = Outer.FindType(UhtFindOptions.Enum | UhtFindOptions.SourceName, TokenList, TokenReader) as UhtEnum;
+						enumObj = outer.FindType(UhtFindOptions.Enum | UhtFindOptions.SourceName, tokenList, tokenReader) as UhtEnum;
 					})
 					.Require('>');
 			}
-			return Enum != null ? new UhtByteProperty(PropertySettings, UhtPropertyIntType.None, Enum) : null;
+			return enumObj != null ? new UhtByteProperty(propertySettings, UhtPropertyIntType.None, enumObj) : null;
 		}
 		#endregion
 	}
diff --git a/Engine/Source/Programs/Shared/EpicGames.UHT/Types/Properties/UhtFieldPathProperty.cs b/Engine/Source/Programs/Shared/EpicGames.UHT/Types/Properties/UhtFieldPathProperty.cs
--- a/Engine/Source/Programs/Shared/EpicGames.UHT/Types/Properties/UhtFieldPathProperty.cs
+++ b/Engine/Source/Programs/Shared/EpicGames.UHT/Types/Properties/UhtFieldPathProperty.cs
@@ -1,11 +1,11 @@
 // Copyright Epic Games, Inc. All Rights Reserved.
 
+using System.Diagnostics.CodeAnalysis;
+using System.Text;
 using EpicGames.Core;
 using EpicGames.UHT.Tables;
 using EpicGames.UHT.Tokenizer;
 using EpicGames.UHT.Utils;
-using System.Diagnostics.CodeAnalysis;
-using System.Text;
 
 namespace EpicGames.UHT.Types
 {
@@ -17,16 +17,16 @@ namespace EpicGames.UHT.Types
 	public class UhtFieldPathProperty : UhtProperty
 	{
 		/// <inheritdoc/>
-		public override string EngineClassName { get => "FieldPathProperty"; }
+		public override string EngineClassName => "FieldPathProperty";
 
 		/// <inheritdoc/>
-		protected override string CppTypeText { get => "TFieldPath"; }
+		protected override string CppTypeText => "TFieldPath";
 
 		/// <inheritdoc/>
-		protected override string PGetMacroText { get => "TFIELDPATH"; }
+		protected override string PGetMacroText => "TFIELDPATH";
 
 		/// <inheritdoc/>
-		protected override UhtPGetArgumentType PGetTypeArgument { get => UhtPGetArgumentType.TypeText; }
+		protected override UhtPGetArgumentType PGetTypeArgument => UhtPGetArgumentType.TypeText;
 
 		/// <summary>
 		/// Class name without the prefix
@@ -36,11 +36,11 @@ namespace EpicGames.UHT.Types
 		/// <summary>
 		/// Construct new property
 		/// </summary>
-		/// <param name="PropertySettings">Property settings</param>
-		/// <param name="FieldClassName">Field class name</param>
-		public UhtFieldPathProperty(UhtPropertySettings PropertySettings, string FieldClassName) : base(PropertySettings)
+		/// <param name="propertySettings">Property settings</param>
+		/// <param name="fieldClassName">Field class name</param>
+		public UhtFieldPathProperty(UhtPropertySettings propertySettings, string fieldClassName) : base(propertySettings)
 		{
-			this.FieldClassName = FieldClassName;
+			this.FieldClassName = fieldClassName;
 			this.PropertyCaps |= UhtPropertyCaps.IsParameterSupportedByBlueprint | UhtPropertyCaps.IsMemberSupportedByBlueprint;
 		}
 
@@ -51,55 +51,55 @@ namespace EpicGames.UHT.Types
 		}
 
 		/// <inheritdoc/>
-		public override StringBuilder AppendText(StringBuilder Builder, UhtPropertyTextType TextType, bool bIsTemplateArgument)
+		public override StringBuilder AppendText(StringBuilder builder, UhtPropertyTextType textType, bool isTemplateArgument)
 		{
-			switch (TextType)
+			switch (textType)
 			{
 				case UhtPropertyTextType.SparseShort:
-					Builder.Append("TFieldPath");
+					builder.Append("TFieldPath");
 					break;
 
 				default:
-					Builder.Append("TFieldPath<F").Append(this.FieldClassName).Append('>');
+					builder.Append("TFieldPath<F").Append(this.FieldClassName).Append('>');
 					break;
 			}
-			return Builder;
+			return builder;
 		}
 
 		/// <inheritdoc/>
-		public override StringBuilder AppendMemberDecl(StringBuilder Builder, IUhtPropertyMemberContext Context, string Name, string NameSuffix, int Tabs)
+		public override StringBuilder AppendMemberDecl(StringBuilder builder, IUhtPropertyMemberContext context, string name, string nameSuffix, int tabs)
 		{
-			return AppendMemberDecl(Builder, Context, Name, NameSuffix, Tabs, "FFieldPathPropertyParams");
+			return AppendMemberDecl(builder, context, name, nameSuffix, tabs, "FFieldPathPropertyParams");
 		}
 
 		/// <inheritdoc/>
-		public override StringBuilder AppendMemberDef(StringBuilder Builder, IUhtPropertyMemberContext Context, string Name, string NameSuffix, string? Offset, int Tabs)
+		public override StringBuilder AppendMemberDef(StringBuilder builder, IUhtPropertyMemberContext context, string name, string nameSuffix, string? offset, int tabs)
 		{
-			AppendMemberDefStart(Builder, Context, Name, NameSuffix, Offset, Tabs, "FFieldPathPropertyParams", "UECodeGen_Private::EPropertyGenFlags::FieldPath");
-			Builder.Append("&F").Append(this.FieldClassName).Append("::StaticClass, ");
-			AppendMemberDefEnd(Builder, Context, Name, NameSuffix);
-			return Builder;
+			AppendMemberDefStart(builder, context, name, nameSuffix, offset, tabs, "FFieldPathPropertyParams", "UECodeGen_Private::EPropertyGenFlags::FieldPath");
+			builder.Append("&F").Append(this.FieldClassName).Append("::StaticClass, ");
+			AppendMemberDefEnd(builder, context, name, nameSuffix);
+			return builder;
 		}
 
 		/// <inheritdoc/>
-		public override StringBuilder AppendNullConstructorArg(StringBuilder Builder, bool bIsInitializer)
+		public override StringBuilder AppendNullConstructorArg(StringBuilder builder, bool isInitializer)
 		{
-			Builder.Append("nullptr");
-			return Builder;
+			builder.Append("nullptr");
+			return builder;
 		}
 
 		/// <inheritdoc/>
-		public override bool SanitizeDefaultValue(IUhtTokenReader DefaultValueReader, StringBuilder InnerDefaultValue)
+		public override bool SanitizeDefaultValue(IUhtTokenReader defaultValueReader, StringBuilder innerDefaultValue)
 		{
 			return false;
 		}
 
 		/// <inheritdoc/>
-		public override bool IsSameType(UhtProperty Other)
+		public override bool IsSameType(UhtProperty other)
 		{
-			if (Other is UhtFieldPathProperty OtherFieldPath)
+			if (other is UhtFieldPathProperty otherFieldPath)
 			{
-				return this.FieldClassName == OtherFieldPath.FieldClassName;
+				return this.FieldClassName == otherFieldPath.FieldClassName;
 			}
 			return false;
 		}
@@ -108,23 +108,23 @@ namespace EpicGames.UHT.Types
 		[UhtPropertyType(Keyword = "TFieldPath", Options = UhtPropertyTypeOptions.Immediate)]
 		[SuppressMessage("CodeQuality", "IDE0051:Remove unused private members", Justification = "Attribute accessed method")]
 		[SuppressMessage("Style", "IDE0060:Remove unused parameter", Justification = "Attribute accessed method")]
-		private static UhtProperty? FieldPathProperty(UhtPropertyResolvePhase ResolvePhase, UhtPropertySettings PropertySettings, IUhtTokenReader TokenReader, UhtToken MatchedToken)
+		private static UhtProperty? FieldPathProperty(UhtPropertyResolvePhase resolvePhase, UhtPropertySettings propertySettings, IUhtTokenReader tokenReader, UhtToken matchedToken)
 		{
-			using (var TokenContext = new UhtMessageContext("TFieldPath"))
+			using (UhtMessageContext tokenContext = new UhtMessageContext("TFieldPath"))
 			{
-				UhtToken Identifier = new UhtToken();
-				TokenReader
+				UhtToken identifier = new UhtToken();
+				tokenReader
 					.Require("TFieldPath")
 					.Require('<')
-					.RequireIdentifier((ref UhtToken Token) => Identifier = Token)
+					.RequireIdentifier((ref UhtToken token) => identifier = token)
 					.Require('>');
 
-				StringView FieldClassName = new StringView(Identifier.Value, 1);
-				if (!PropertySettings.Outer.Session.IsValidPropertyTypeName(FieldClassName))
+				StringView fieldClassName = new StringView(identifier.Value, 1);
+				if (!propertySettings.Outer.Session.IsValidPropertyTypeName(fieldClassName))
 				{
-					throw new UhtException($"Undefined property type: {Identifier.Value}");
+					throw new UhtException($"Undefined property type: {identifier.Value}");
 				}
-				return new UhtFieldPathProperty(PropertySettings, FieldClassName.ToString());
+				return new UhtFieldPathProperty(propertySettings, fieldClassName.ToString());
 			}
 		}
 		#endregion
diff --git a/Engine/Source/Programs/Shared/EpicGames.UHT/Types/Properties/UhtFloatProperty.cs b/Engine/Source/Programs/Shared/EpicGames.UHT/Types/Properties/UhtFloatProperty.cs
--- a/Engine/Source/Programs/Shared/EpicGames.UHT/Types/Properties/UhtFloatProperty.cs
+++ b/Engine/Source/Programs/Shared/EpicGames.UHT/Types/Properties/UhtFloatProperty.cs
@@ -1,10 +1,10 @@
 // Copyright Epic Games, Inc. All Rights Reserved.
 
+using System.Diagnostics.CodeAnalysis;
+using System.Text;
 using EpicGames.UHT.Tables;
 using EpicGames.UHT.Tokenizer;
 using EpicGames.UHT.Utils;
-using System.Diagnostics.CodeAnalysis;
-using System.Text;
 
 namespace EpicGames.UHT.Types
 {
@@ -16,54 +16,54 @@ namespace EpicGames.UHT.Types
 	public class UhtFloatProperty : UhtNumericProperty
 	{
 		/// <inheritdoc/>
-		public override string EngineClassName { get => "FloatProperty"; }
+		public override string EngineClassName => "FloatProperty";
 
 		/// <inheritdoc/>
-		protected override string CppTypeText { get => "float"; }
+		protected override string CppTypeText => "float";
 
 		/// <summary>
 		/// Construct a new property
 		/// </summary>
-		/// <param name="PropertySettings">Property settings</param>
-		public UhtFloatProperty(UhtPropertySettings PropertySettings) : base(PropertySettings, UhtPropertyIntType.None)
+		/// <param name="propertySettings">Property settings</param>
+		public UhtFloatProperty(UhtPropertySettings propertySettings) : base(propertySettings, UhtPropertyIntType.None)
 		{
 			this.PropertyCaps |= UhtPropertyCaps.CanExposeOnSpawn | UhtPropertyCaps.IsParameterSupportedByBlueprint | UhtPropertyCaps.IsMemberSupportedByBlueprint;
 		}
 
 		/// <inheritdoc/>
-		public override StringBuilder AppendMemberDecl(StringBuilder Builder, IUhtPropertyMemberContext Context, string Name, string NameSuffix, int Tabs)
+		public override StringBuilder AppendMemberDecl(StringBuilder builder, IUhtPropertyMemberContext context, string name, string nameSuffix, int tabs)
 		{
-			return AppendMemberDecl(Builder, Context, Name, NameSuffix, Tabs, "FFloatPropertyParams");
+			return AppendMemberDecl(builder, context, name, nameSuffix, tabs, "FFloatPropertyParams");
 		}
 
 		/// <inheritdoc/>
-		public override StringBuilder AppendMemberDef(StringBuilder Builder, IUhtPropertyMemberContext Context, string Name, string NameSuffix, string? Offset, int Tabs)
+		public override StringBuilder AppendMemberDef(StringBuilder builder, IUhtPropertyMemberContext context, string name, string nameSuffix, string? offset, int tabs)
 		{
-			AppendMemberDefStart(Builder, Context, Name, NameSuffix, Offset, Tabs, "FFloatPropertyParams", "UECodeGen_Private::EPropertyGenFlags::Float");
-			AppendMemberDefEnd(Builder, Context, Name, NameSuffix);
-			return Builder;
+			AppendMemberDefStart(builder, context, name, nameSuffix, offset, tabs, "FFloatPropertyParams", "UECodeGen_Private::EPropertyGenFlags::Float");
+			AppendMemberDefEnd(builder, context, name, nameSuffix);
+			return builder;
 		}
 
 		/// <inheritdoc/>
-		public override bool SanitizeDefaultValue(IUhtTokenReader DefaultValueReader, StringBuilder InnerDefaultValue)
+		public override bool SanitizeDefaultValue(IUhtTokenReader defaultValueReader, StringBuilder innerDefaultValue)
 		{
-			InnerDefaultValue.AppendFormat("{0:F6}", DefaultValueReader.GetConstFloatExpression());
+			innerDefaultValue.AppendFormat("{0:F6}", defaultValueReader.GetConstFloatExpression());
 			return true;
 		}
 
 		/// <inheritdoc/>
-		public override bool IsSameType(UhtProperty Other)
+		public override bool IsSameType(UhtProperty other)
 		{
-			return Other is UhtFloatProperty;
+			return other is UhtFloatProperty;
 		}
 
 		#region Keywords
 		[UhtPropertyType(Keyword = "float", Options = UhtPropertyTypeOptions.Simple | UhtPropertyTypeOptions.Immediate)]
 		[SuppressMessage("CodeQuality", "IDE0051:Remove unused private members", Justification = "Attribute accessed method")]
 		[SuppressMessage("Style", "IDE0060:Remove unused parameter", Justification = "Attribute accessed method")]
-		private static UhtProperty? FloatProperty(UhtPropertyResolvePhase ResolvePhase, UhtPropertySettings PropertySettings, IUhtTokenReader TokenReader, UhtToken MatchedToken)
+		private static UhtProperty? FloatProperty(UhtPropertyResolvePhase resolvePhase, UhtPropertySettings propertySettings, IUhtTokenReader tokenReader, UhtToken matchedToken)
 		{
-			return new UhtFloatProperty(PropertySettings);
+			return new UhtFloatProperty(propertySettings);
 		}
 		#endregion
 	}
diff --git a/Engine/Source/Programs/Shared/EpicGames.UHT/Types/Properties/UhtInt16Property.cs b/Engine/Source/Programs/Shared/EpicGames.UHT/Types/Properties/UhtInt16Property.cs
--- a/Engine/Source/Programs/Shared/EpicGames.UHT/Types/Properties/UhtInt16Property.cs
+++ b/Engine/Source/Programs/Shared/EpicGames.UHT/Types/Properties/UhtInt16Property.cs
@@ -1,10 +1,10 @@
 // Copyright Epic Games, Inc. All Rights Reserved.
 
+using System.Diagnostics.CodeAnalysis;
+using System.Text;
 using EpicGames.UHT.Tables;
 using EpicGames.UHT.Tokenizer;
 using EpicGames.UHT.Utils;
-using System.Diagnostics.CodeAnalysis;
-using System.Text;
 
 namespace EpicGames.UHT.Types
 {
@@ -16,53 +16,53 @@ namespace EpicGames.UHT.Types
 	public class UhtInt16Property : UhtNumericProperty
 	{
 		/// <inheritdoc/>
-		public override string EngineClassName { get => "Int16Property"; }
+		public override string EngineClassName => "Int16Property";
 
 		/// <inheritdoc/>
-		protected override string CppTypeText { get => "int16"; }
+		protected override string CppTypeText => "int16";
 
 		/// <summary>
 		/// Construct new property
 		/// </summary>
-		/// <param name="PropertySettings">Property settings</param>
-		/// <param name="IntType">Integer type</param>
-		public UhtInt16Property(UhtPropertySettings PropertySettings, UhtPropertyIntType IntType) : base(PropertySettings, IntType)
+		/// <param name="propertySettings">Property settings</param>
+		/// <param name="intType">Integer type</param>
+		public UhtInt16Property(UhtPropertySettings propertySettings, UhtPropertyIntType intType) : base(propertySettings, intType)
 		{
 		}
 
 		/// <inheritdoc/>
-		public override StringBuilder AppendMemberDecl(StringBuilder Builder, IUhtPropertyMemberContext Context, string Name, string NameSuffix, int Tabs)
+		public override StringBuilder AppendMemberDecl(StringBuilder builder, IUhtPropertyMemberContext context, string name, string nameSuffix, int tabs)
 		{
-			return AppendMemberDecl(Builder, Context, Name, NameSuffix, Tabs, "FInt16PropertyParams");
+			return AppendMemberDecl(builder, context, name, nameSuffix, tabs, "FInt16PropertyParams");
 		}
 
 		/// <inheritdoc/>
-		public override StringBuilder AppendMemberDef(StringBuilder Builder, IUhtPropertyMemberContext Context, string Name, string NameSuffix, string? Offset, int Tabs)
+		public override StringBuilder AppendMemberDef(StringBuilder builder, IUhtPropertyMemberContext context, string name, string nameSuffix, string? offset, int tabs)
 		{
-			AppendMemberDefStart(Builder, Context, Name, NameSuffix, Offset, Tabs, "FInt16PropertyParams", "UECodeGen_Private::EPropertyGenFlags::Int16");
-			AppendMemberDefEnd(Builder, Context, Name, NameSuffix);
-			return Builder;
+			AppendMemberDefStart(builder, context, name, nameSuffix, offset, tabs, "FInt16PropertyParams", "UECodeGen_Private::EPropertyGenFlags::Int16");
+			AppendMemberDefEnd(builder, context, name, nameSuffix);
+			return builder;
 		}
 
 		/// <inheritdoc/>
-		public override bool SanitizeDefaultValue(IUhtTokenReader DefaultValueReader, StringBuilder InnerDefaultValue)
+		public override bool SanitizeDefaultValue(IUhtTokenReader defaultValueReader, StringBuilder innerDefaultValue)
 		{
 			return false;
 		}
 
 		/// <inheritdoc/>
-		public override bool IsSameType(UhtProperty Other)
+		public override bool IsSameType(UhtProperty other)
 		{
-			return Other is UhtInt16Property;
+			return other is UhtInt16Property;
 		}
 
 		#region Keyword
 		[UhtPropertyType(Keyword = "int16", Options = UhtPropertyTypeOptions.Simple | UhtPropertyTypeOptions.Immediate)]
 		[SuppressMessage("CodeQuality", "IDE0051:Remove unused private members", Justification = "Attribute accessed method")]
 		[SuppressMessage("Style", "IDE0060:Remove unused parameter", Justification = "Attribute accessed method")]
-		private static UhtProperty? Int16Property(UhtPropertyResolvePhase ResolvePhase, UhtPropertySettings PropertySettings, IUhtTokenReader TokenReader, UhtToken MatchedToken)
+		private static UhtProperty? Int16Property(UhtPropertyResolvePhase resolvePhase, UhtPropertySettings propertySettings, IUhtTokenReader tokenReader, UhtToken matchedToken)
 		{
-			return new UhtInt16Property(PropertySettings, UhtPropertyIntType.Sized);
+			return new UhtInt16Property(propertySettings, UhtPropertyIntType.Sized);
 		}
 		#endregion
 	}
diff --git a/Engine/Source/Programs/Shared/EpicGames.UHT/Types/Properties/UhtInt64Property.cs b/Engine/Source/Programs/Shared/EpicGames.UHT/Types/Properties/UhtInt64Property.cs
--- a/Engine/Source/Programs/Shared/EpicGames.UHT/Types/Properties/UhtInt64Property.cs
+++ b/Engine/Source/Programs/Shared/EpicGames.UHT/Types/Properties/UhtInt64Property.cs
@@ -1,10 +1,10 @@
 // Copyright Epic Games, Inc. All Rights Reserved.
 
+using System.Diagnostics.CodeAnalysis;
+using System.Text;
 using EpicGames.UHT.Tables;
 using EpicGames.UHT.Tokenizer;
 using EpicGames.UHT.Utils;
-using System.Diagnostics.CodeAnalysis;
-using System.Text;
 
 namespace EpicGames.UHT.Types
 {
@@ -16,55 +16,55 @@ namespace EpicGames.UHT.Types
 	public class UhtInt64Property : UhtNumericProperty
 	{
 		/// <inheritdoc/>
-		public override string EngineClassName { get => "Int64Property"; }
+		public override string EngineClassName => "Int64Property";
 
 		/// <inheritdoc/>
-		protected override string CppTypeText { get => "int64"; }
+		protected override string CppTypeText => "int64";
 
 		/// <summary>
 		/// Construct new property
 		/// </summary>
-		/// <param name="PropertySettings">Property settings</param>
-		/// <param name="IntType">Integer type</param>
-		public UhtInt64Property(UhtPropertySettings PropertySettings, UhtPropertyIntType IntType) : base(PropertySettings, IntType)
+		/// <param name="propertySettings">Property settings</param>
+		/// <param name="intType">Integer type</param>
+		public UhtInt64Property(UhtPropertySettings propertySettings, UhtPropertyIntType intType) : base(propertySettings, intType)
 		{
 			this.PropertyCaps |= UhtPropertyCaps.CanExposeOnSpawn | UhtPropertyCaps.IsParameterSupportedByBlueprint | UhtPropertyCaps.IsMemberSupportedByBlueprint;
 		}
 
 		/// <inheritdoc/>
-		public override StringBuilder AppendMemberDecl(StringBuilder Builder, IUhtPropertyMemberContext Context, string Name, string NameSuffix, int Tabs)
+		public override StringBuilder AppendMemberDecl(StringBuilder builder, IUhtPropertyMemberContext context, string name, string nameSuffix, int tabs)
 		{
-			return AppendMemberDecl(Builder, Context, Name, NameSuffix, Tabs, "FInt64PropertyParams");
+			return AppendMemberDecl(builder, context, name, nameSuffix, tabs, "FInt64PropertyParams");
 		}
 
 		/// <inheritdoc/>
-		public override StringBuilder AppendMemberDef(StringBuilder Builder, IUhtPropertyMemberContext Context, string Name, string NameSuffix, string? Offset, int Tabs)
+		public override StringBuilder AppendMemberDef(StringBuilder builder, IUhtPropertyMemberContext context, string name, string nameSuffix, string? offset, int tabs)
 		{
-			AppendMemberDefStart(Builder, Context, Name, NameSuffix, Offset, Tabs, "FInt64PropertyParams", "UECodeGen_Private::EPropertyGenFlags::Int64");
-			AppendMemberDefEnd(Builder, Context, Name, NameSuffix);
-			return Builder;
+			AppendMemberDefStart(builder, context, name, nameSuffix, offset, tabs, "FInt64PropertyParams", "UECodeGen_Private::EPropertyGenFlags::Int64");
+			AppendMemberDefEnd(builder, context, name, nameSuffix);
+			return builder;
 		}
 
 		/// <inheritdoc/>
-		public override bool SanitizeDefaultValue(IUhtTokenReader DefaultValueReader, StringBuilder InnerDefaultValue)
+		public override bool SanitizeDefaultValue(IUhtTokenReader defaultValueReader, StringBuilder innerDefaultValue)
 		{
-			InnerDefaultValue.Append(DefaultValueReader.GetConstLongExpression());
+			innerDefaultValue.Append(defaultValueReader.GetConstLongExpression());
 			return true;
 		}
 
 		/// <inheritdoc/>
-		public override bool IsSameType(UhtProperty Other)
+		public override bool IsSameType(UhtProperty other)
 		{
-			return Other is UhtInt64Property;
+			return other is UhtInt64Property;
 		}
 
 		#region Keyword
 		[UhtPropertyType(Keyword = "int64", Options = UhtPropertyTypeOptions.Simple | UhtPropertyTypeOptions.Immediate)]
 		[SuppressMessage("CodeQuality", "IDE0051:Remove unused private members", Justification = "Attribute accessed method")]
 		[SuppressMessage("Style", "IDE0060:Remove unused parameter", Justification = "Attribute accessed method")]
-		private static UhtProperty? Int64Property(UhtPropertyResolvePhase ResolvePhase, UhtPropertySettings PropertySettings, IUhtTokenReader TokenReader, UhtToken MatchedToken)
+		private static UhtProperty? Int64Property(UhtPropertyResolvePhase resolvePhase, UhtPropertySettings propertySettings, IUhtTokenReader tokenReader, UhtToken matchedToken)
 		{
-			return new UhtInt64Property(PropertySettings, UhtPropertyIntType.Sized);
+			return new UhtInt64Property(propertySettings, UhtPropertyIntType.Sized);
 		}
 		#endregion
 	}
diff --git a/Engine/Source/Programs/Shared/EpicGames.UHT/Types/Properties/UhtInt8Property.cs b/Engine/Source/Programs/Shared/EpicGames.UHT/Types/Properties/UhtInt8Property.cs
--- a/Engine/Source/Programs/Shared/EpicGames.UHT/Types/Properties/UhtInt8Property.cs
+++ b/Engine/Source/Programs/Shared/EpicGames.UHT/Types/Properties/UhtInt8Property.cs
@@ -1,10 +1,10 @@
 // Copyright Epic Games, Inc. All Rights Reserved.
 
+using System.Diagnostics.CodeAnalysis;
+using System.Text;
 using EpicGames.UHT.Tables;
 using EpicGames.UHT.Tokenizer;
 using EpicGames.UHT.Utils;
-using System.Diagnostics.CodeAnalysis;
-using System.Text;
 
 namespace EpicGames.UHT.Types
 {
@@ -16,53 +16,53 @@ namespace EpicGames.UHT.Types
 	public class UhtInt8Property : UhtNumericProperty
 	{
 		/// <inheritdoc/>
-		public override string EngineClassName { get => "Int8Property"; }
+		public override string EngineClassName => "Int8Property";
 
 		/// <inheritdoc/>
-		protected override string CppTypeText { get => "int8"; }
+		protected override string CppTypeText => "int8";
 
 		/// <summary>
 		/// Construct a new property
 		/// </summary>
-		/// <param name="PropertySettings">Property settings</param>
-		/// <param name="IntType">Integer type</param>
-		public UhtInt8Property(UhtPropertySettings PropertySettings, UhtPropertyIntType IntType) : base(PropertySettings, IntType)
+		/// <param name="propertySettings">Property settings</param>
+		/// <param name="intType">Integer type</param>
+		public UhtInt8Property(UhtPropertySettings propertySettings, UhtPropertyIntType intType) : base(propertySettings, intType)
 		{
 		}
 
 		/// <inheritdoc/>
-		public override StringBuilder AppendMemberDecl(StringBuilder Builder, IUhtPropertyMemberContext Context, string Name, string NameSuffix, int Tabs)
+		public override StringBuilder AppendMemberDecl(StringBuilder builder, IUhtPropertyMemberContext context, string name, string nameSuffix, int tabs)
 		{
-			return AppendMemberDecl(Builder, Context, Name, NameSuffix, Tabs, "FInt8PropertyParams");
+			return AppendMemberDecl(builder, context, name, nameSuffix, tabs, "FInt8PropertyParams");
 		}
 
 		/// <inheritdoc/>
-		public override StringBuilder AppendMemberDef(StringBuilder Builder, IUhtPropertyMemberContext Context, string Name, string NameSuffix, string? Offset, int Tabs)
+		public override StringBuilder AppendMemberDef(StringBuilder builder, IUhtPropertyMemberContext context, string name, string nameSuffix, string? offset, int tabs)
 		{
-			AppendMemberDefStart(Builder, Context, Name, NameSuffix, Offset, Tabs, "FInt8PropertyParams", "UECodeGen_Private::EPropertyGenFlags::Int8");
-			AppendMemberDefEnd(Builder, Context, Name, NameSuffix);
-			return Builder;
+			AppendMemberDefStart(builder, context, name, nameSuffix, offset, tabs, "FInt8PropertyParams", "UECodeGen_Private::EPropertyGenFlags::Int8");
+			AppendMemberDefEnd(builder, context, name, nameSuffix);
+			return builder;
 		}
 
 		/// <inheritdoc/>
-		public override bool SanitizeDefaultValue(IUhtTokenReader DefaultValueReader, StringBuilder InnerDefaultValue)
+		public override bool SanitizeDefaultValue(IUhtTokenReader defaultValueReader, StringBuilder innerDefaultValue)
 		{
 			return false;
 		}
 
 		/// <inheritdoc/>
-		public override bool IsSameType(UhtProperty Other)
+		public override bool IsSameType(UhtProperty other)
 		{
-			return Other is UhtInt8Property;
+			return other is UhtInt8Property;
 		}
 
 		#region Keyword
 		[UhtPropertyType(Keyword = "int8", Options = UhtPropertyTypeOptions.Simple | UhtPropertyTypeOptions.Immediate)]
 		[SuppressMessage("CodeQuality", "IDE0051:Remove unused private members", Justification = "Attribute accessed method")]
 		[SuppressMessage("Style", "IDE0060:Remove unused parameter", Justification = "Attribute accessed method")]
-		private static UhtProperty? Int8Property(UhtPropertyResolvePhase ResolvePhase, UhtPropertySettings PropertySettings, IUhtTokenReader TokenReader, UhtToken MatchedToken)
+		private static UhtProperty? Int8Property(UhtPropertyResolvePhase resolvePhase, UhtPropertySettings propertySettings, IUhtTokenReader tokenReader, UhtToken matchedToken)
 		{
-			return new UhtInt8Property(PropertySettings, UhtPropertyIntType.Sized);
+			return new UhtInt8Property(propertySettings, UhtPropertyIntType.Sized);
 		}
 		#endregion
 	}
diff --git a/Engine/Source/Programs/Shared/EpicGames.UHT/Types/Properties/UhtIntProperty.cs b/Engine/Source/Programs/Shared/EpicGames.UHT/Types/Properties/UhtIntProperty.cs
--- a/Engine/Source/Programs/Shared/EpicGames.UHT/Types/Properties/UhtIntProperty.cs
+++ b/Engine/Source/Programs/Shared/EpicGames.UHT/Types/Properties/UhtIntProperty.cs
@@ -1,10 +1,10 @@
 // Copyright Epic Games, Inc. All Rights Reserved.
 
+using System.Diagnostics.CodeAnalysis;
+using System.Text;
 using EpicGames.UHT.Tables;
 using EpicGames.UHT.Tokenizer;
 using EpicGames.UHT.Utils;
-using System.Diagnostics.CodeAnalysis;
-using System.Text;
 
 namespace EpicGames.UHT.Types
 {
@@ -16,77 +16,77 @@ namespace EpicGames.UHT.Types
 	public class UhtIntProperty : UhtNumericProperty
 	{
 		/// <inheritdoc/>
-		public override string EngineClassName { get => "IntProperty"; }
+		public override string EngineClassName => "IntProperty";
 
 		/// <inheritdoc/>
-		protected override string CppTypeText { get => "int32"; }
+		protected override string CppTypeText => "int32";
 
 		/// <summary>
 		/// Construct new property
 		/// </summary>
-		/// <param name="PropertySettings">Property settings</param>
-		/// <param name="IntType">Integer type</param>
-		public UhtIntProperty(UhtPropertySettings PropertySettings, UhtPropertyIntType IntType) : base(PropertySettings, IntType)
+		/// <param name="propertySettings">Property settings</param>
+		/// <param name="intType">Integer type</param>
+		public UhtIntProperty(UhtPropertySettings propertySettings, UhtPropertyIntType intType) : base(propertySettings, intType)
 		{
 			this.PropertyCaps |= UhtPropertyCaps.CanExposeOnSpawn | UhtPropertyCaps.IsParameterSupportedByBlueprint | UhtPropertyCaps.IsMemberSupportedByBlueprint;
 		}
 
 		/// <inheritdoc/>
-		public override StringBuilder AppendMemberDecl(StringBuilder Builder, IUhtPropertyMemberContext Context, string Name, string NameSuffix, int Tabs)
+		public override StringBuilder AppendMemberDecl(StringBuilder builder, IUhtPropertyMemberContext context, string name, string nameSuffix, int tabs)
 		{
-			return AppendMemberDecl(Builder, Context, Name, NameSuffix, Tabs, 
+			return AppendMemberDecl(builder, context, name, nameSuffix, tabs,
 				this.IntType == UhtPropertyIntType.Unsized ? "FUnsizedIntPropertyParams" : "FIntPropertyParams");
 		}
 
 		/// <inheritdoc/>
-		public override StringBuilder AppendMemberDef(StringBuilder Builder, IUhtPropertyMemberContext Context, string Name, string NameSuffix, string? Offset, int Tabs)
+		public override StringBuilder AppendMemberDef(StringBuilder builder, IUhtPropertyMemberContext context, string name, string nameSuffix, string? offset, int tabs)
 		{
-			AppendMemberDefStart(Builder, Context, Name, NameSuffix, Offset, Tabs,
+			AppendMemberDefStart(builder, context, name, nameSuffix, offset, tabs,
 				this.IntType == UhtPropertyIntType.Unsized ? "FUnsizedIntPropertyParams" : "FIntPropertyParams",
 				"UECodeGen_Private::EPropertyGenFlags::Int");
-			AppendMemberDefEnd(Builder, Context, Name, NameSuffix);
-			return Builder;
+			AppendMemberDefEnd(builder, context, name, nameSuffix);
+			return builder;
 		}
 
 		/// <inheritdoc/>
-		public override bool SanitizeDefaultValue(IUhtTokenReader DefaultValueReader, StringBuilder InnerDefaultValue)
+		public override bool SanitizeDefaultValue(IUhtTokenReader defaultValueReader, StringBuilder innerDefaultValue)
 		{
-			InnerDefaultValue.Append(DefaultValueReader.GetConstIntExpression());
+			innerDefaultValue.Append(defaultValueReader.GetConstIntExpression());
 			return true;
 		}
 
 		/// <inheritdoc/>
-		public override bool IsSameType(UhtProperty Other)
+		public override bool IsSameType(UhtProperty other)
 		{
-			return Other is UhtIntProperty;
+			return other is UhtIntProperty;
 		}
 
 		#region Keywords
 		[UhtPropertyType(Keyword = "int32", Options = UhtPropertyTypeOptions.Simple | UhtPropertyTypeOptions.Immediate)]
 		[SuppressMessage("CodeQuality", "IDE0051:Remove unused private members", Justification = "Attribute accessed method")]
 		[SuppressMessage("Style", "IDE0060:Remove unused parameter", Justification = "Attribute accessed method")]
-		private static UhtProperty? Int32Property(UhtPropertyResolvePhase ResolvePhase, UhtPropertySettings PropertySettings, IUhtTokenReader TokenReader, UhtToken MatchedToken)
+		private static UhtProperty? Int32Property(UhtPropertyResolvePhase resolvePhase, UhtPropertySettings propertySettings, IUhtTokenReader tokenReader, UhtToken matchedToken)
 		{
-			return new UhtIntProperty(PropertySettings, UhtPropertyIntType.Sized);
+			return new UhtIntProperty(propertySettings, UhtPropertyIntType.Sized);
 		}
 
 		[UhtPropertyType(Keyword = "int", Options = UhtPropertyTypeOptions.Simple | UhtPropertyTypeOptions.Immediate)]
 		[SuppressMessage("CodeQuality", "IDE0051:Remove unused private members", Justification = "Attribute accessed method")]
 		[SuppressMessage("Style", "IDE0060:Remove unused parameter", Justification = "Attribute accessed method")]
-		private static UhtProperty? IntProperty(UhtPropertyResolvePhase ResolvePhase, UhtPropertySettings PropertySettings, IUhtTokenReader TokenReader, UhtToken MatchedToken)
+		private static UhtProperty? IntProperty(UhtPropertyResolvePhase resolvePhase, UhtPropertySettings propertySettings, IUhtTokenReader tokenReader, UhtToken matchedToken)
 		{
-			return new UhtIntProperty(PropertySettings, UhtPropertyIntType.Unsized);
+			return new UhtIntProperty(propertySettings, UhtPropertyIntType.Unsized);
 		}
 
 		[UhtPropertyType(Keyword = "signed", Options = UhtPropertyTypeOptions.Immediate)]
 		[SuppressMessage("CodeQuality", "IDE0051:Remove unused private members", Justification = "Attribute accessed method")]
 		[SuppressMessage("Style", "IDE0060:Remove unused parameter", Justification = "Attribute accessed method")]
-		private static UhtProperty? SignedProperty(UhtPropertyResolvePhase ResolvePhase, UhtPropertySettings PropertySettings, IUhtTokenReader TokenReader, UhtToken MatchedToken)
+		private static UhtProperty? SignedProperty(UhtPropertyResolvePhase resolvePhase, UhtPropertySettings propertySettings, IUhtTokenReader tokenReader, UhtToken matchedToken)
 		{
-			TokenReader
+			tokenReader
 				.Require("signed")
 				.Optional("int");
-			return new UhtIntProperty(PropertySettings, UhtPropertyIntType.Unsized);
+			return new UhtIntProperty(propertySettings, UhtPropertyIntType.Unsized);
 		}
 		#endregion
 	}
diff --git a/Engine/Source/Programs/Shared/EpicGames.UHT/Types/Properties/UhtInterfaceProperty.cs b/Engine/Source/Programs/Shared/EpicGames.UHT/Types/Properties/UhtInterfaceProperty.cs
--- a/Engine/Source/Programs/Shared/EpicGames.UHT/Types/Properties/UhtInterfaceProperty.cs
+++ b/Engine/Source/Programs/Shared/EpicGames.UHT/Types/Properties/UhtInterfaceProperty.cs
@@ -1,13 +1,13 @@
 // Copyright Epic Games, Inc. All Rights Reserved.
 
-using EpicGames.Core;
-using EpicGames.UHT.Tables;
-using EpicGames.UHT.Tokenizer;
-using EpicGames.UHT.Utils;
 using System.Collections.Generic;
 using System.Diagnostics.CodeAnalysis;
 using System.Text;
 using System.Text.Json.Serialization;
+using EpicGames.Core;
+using EpicGames.UHT.Tables;
+using EpicGames.UHT.Tokenizer;
+using EpicGames.UHT.Utils;
 
 namespace EpicGames.UHT.Types
 {
@@ -19,16 +19,16 @@ namespace EpicGames.UHT.Types
 	public class UhtInterfaceProperty : UhtProperty
 	{
 		/// <inheritdoc/>
-		public override string EngineClassName { get => "InterfaceProperty"; }
+		public override string EngineClassName => "InterfaceProperty";
 
 		/// <inheritdoc/>
-		protected override string CppTypeText { get => "TScriptInterface"; }
+		protected override string CppTypeText => "TScriptInterface";
 
 		/// <inheritdoc/>
-		protected override string PGetMacroText { get => "TINTERFACE"; }
+		protected override string PGetMacroText => "TINTERFACE";
 
 		/// <inheritdoc/>
-		protected override UhtPGetArgumentType PGetTypeArgument { get => UhtPGetArgumentType.TypeText; }
+		protected override UhtPGetArgumentType PGetTypeArgument => UhtPGetArgumentType.TypeText;
 
 		/// <summary>
 		/// Referenced interface class
@@ -39,11 +39,11 @@ namespace EpicGames.UHT.Types
 		/// <summary>
 		/// Create a new property
 		/// </summary>
-		/// <param name="PropertySettings">Property settings</param>
-		/// <param name="InterfaceClass">Referenced interface</param>
-		public UhtInterfaceProperty(UhtPropertySettings PropertySettings, UhtClass InterfaceClass) : base(PropertySettings)
+		/// <param name="propertySettings">Property settings</param>
+		/// <param name="interfaceClass">Referenced interface</param>
+		public UhtInterfaceProperty(UhtPropertySettings propertySettings, UhtClass interfaceClass) : base(propertySettings)
 		{
-			this.InterfaceClass = InterfaceClass;
+			this.InterfaceClass = interfaceClass;
 			this.PropertyFlags |= EPropertyFlags.UObjectWrapper;
 			this.PropertyCaps |= UhtPropertyCaps.CanExposeOnSpawn | UhtPropertyCaps.IsParameterSupportedByBlueprint | UhtPropertyCaps.IsMemberSupportedByBlueprint | UhtPropertyCaps.PassCppArgsByRef;
 			this.PropertyCaps &= ~(UhtPropertyCaps.CanHaveConfig | UhtPropertyCaps.CanBeContainerKey);
@@ -54,10 +54,10 @@ namespace EpicGames.UHT.Types
 		}
 
 		/// <inheritdoc/>
-		protected override bool ResolveSelf(UhtResolvePhase Phase)
+		protected override bool ResolveSelf(UhtResolvePhase phase)
 		{
-			bool bResults = base.ResolveSelf(Phase);
-			switch (Phase)
+			bool results = base.ResolveSelf(phase);
+			switch (phase)
 			{
 				case UhtResolvePhase.Final:
 					if (this.InterfaceClass.HierarchyHasAnyClassFlags(EClassFlags.DefaultToInstanced))
@@ -66,30 +66,30 @@ namespace EpicGames.UHT.Types
 					}
 					break;
 			}
-			return bResults;
+			return results;
 		}
 
 		/// <inheritdoc/>
-		public override bool ScanForInstancedReferenced(bool bDeepScan)
+		public override bool ScanForInstancedReferenced(bool deepScan)
 		{
 			return !this.DisallowPropertyFlags.HasAnyFlags(EPropertyFlags.InstancedReference) && this.InterfaceClass.HierarchyHasAnyClassFlags(EClassFlags.DefaultToInstanced);
 		}
 
 		/// <inheritdoc/>
-		public override void CollectReferencesInternal(IUhtReferenceCollector Collector, bool bTemplateProperty)
+		public override void CollectReferencesInternal(IUhtReferenceCollector collector, bool templateProperty)
 		{
-			Collector.AddCrossModuleReference(this.InterfaceClass, false);
+			collector.AddCrossModuleReference(this.InterfaceClass, false);
 		}
 
 		/// <inheritdoc/>
 		public override string? GetForwardDeclarations()
 		{
-			UhtClass? ExportClass = this.InterfaceClass;
-			while (ExportClass != null && !ExportClass.ClassFlags.HasAnyFlags(EClassFlags.Native))
+			UhtClass? exportClass = this.InterfaceClass;
+			while (exportClass != null && !exportClass.ClassFlags.HasAnyFlags(EClassFlags.Native))
 			{
-				ExportClass = ExportClass.SuperClass;
+				exportClass = exportClass.SuperClass;
 			}
-			return ExportClass != null ? $"class {ExportClass.SourceName};" : null;
+			return exportClass != null ? $"class {exportClass.SourceName};" : null;
 		}
 
 		/// <inheritdoc/>
@@ -99,58 +99,58 @@ namespace EpicGames.UHT.Types
 		}
 
 		/// <inheritdoc/>
-		public override StringBuilder AppendText(StringBuilder Builder, UhtPropertyTextType TextType, bool bIsTemplateArgument)
+		public override StringBuilder AppendText(StringBuilder builder, UhtPropertyTextType textType, bool isTemplateArgument)
 		{
-			switch (TextType)
+			switch (textType)
 			{
 				case UhtPropertyTextType.SparseShort:
-					Builder.Append("TScriptInterface");
+					builder.Append("TScriptInterface");
 					break;
 
 				case UhtPropertyTextType.FunctionThunkParameterArgType:
-					Builder.Append(this.InterfaceClass.SourceName);
+					builder.Append(this.InterfaceClass.SourceName);
 					break;
 
 				default:
-					Builder.Append("TScriptInterface<").Append(this.InterfaceClass.SourceName).Append('>');
+					builder.Append("TScriptInterface<").Append(this.InterfaceClass.SourceName).Append('>');
 					break;
 			}
-			return Builder;
+			return builder;
 		}
 
 		/// <inheritdoc/>
-		public override StringBuilder AppendMemberDecl(StringBuilder Builder, IUhtPropertyMemberContext Context, string Name, string NameSuffix, int Tabs)
+		public override StringBuilder AppendMemberDecl(StringBuilder builder, IUhtPropertyMemberContext context, string name, string nameSuffix, int tabs)
 		{
-			return AppendMemberDecl(Builder, Context, Name, NameSuffix, Tabs, "FInterfacePropertyParams");
+			return AppendMemberDecl(builder, context, name, nameSuffix, tabs, "FInterfacePropertyParams");
 		}
 
 		/// <inheritdoc/>
-		public override StringBuilder AppendMemberDef(StringBuilder Builder, IUhtPropertyMemberContext Context, string Name, string NameSuffix, string? Offset, int Tabs)
+		public override StringBuilder AppendMemberDef(StringBuilder builder, IUhtPropertyMemberContext context, string name, string nameSuffix, string? offset, int tabs)
 		{
 			// FScriptInterface<USomeInterface> is valid so in that case we need to pass in the interface class and not the alternate object (which in the end is the same object)
-			AppendMemberDefStart(Builder, Context, Name, NameSuffix, Offset, Tabs, "FInterfacePropertyParams", "UECodeGen_Private::EPropertyGenFlags::Interface");
-			AppendMemberDefRef(Builder, Context, this.InterfaceClass.AlternateObject != null ? this.InterfaceClass.AlternateObject : this.InterfaceClass, false);
-			AppendMemberDefEnd(Builder, Context, Name, NameSuffix);
-			return Builder;
+			AppendMemberDefStart(builder, context, name, nameSuffix, offset, tabs, "FInterfacePropertyParams", "UECodeGen_Private::EPropertyGenFlags::Interface");
+			AppendMemberDefRef(builder, context, this.InterfaceClass.AlternateObject ?? this.InterfaceClass, false);
+			AppendMemberDefEnd(builder, context, name, nameSuffix);
+			return builder;
 		}
 
 		/// <inheritdoc/>
-		public override StringBuilder AppendNullConstructorArg(StringBuilder Builder, bool bIsInitializer)
+		public override StringBuilder AppendNullConstructorArg(StringBuilder builder, bool isInitializer)
 		{
-			Builder.Append("NULL");
-			return Builder;
+			builder.Append("NULL");
+			return builder;
 		}
 
 		/// <inheritdoc/>
-		public override bool SanitizeDefaultValue(IUhtTokenReader DefaultValueReader, StringBuilder InnerDefaultValue)
+		public override bool SanitizeDefaultValue(IUhtTokenReader defaultValueReader, StringBuilder innerDefaultValue)
 		{
 			return false;
 		}
 
 		/// <inheritdoc/>
-		protected override void ValidateMember(UhtStruct Struct, UhtValidationOptions Options)
+		protected override void ValidateMember(UhtStruct structObj, UhtValidationOptions options)
 		{
-			base.ValidateMember(Struct, Options);
+			base.ValidateMember(structObj, options);
 
 			if (this.PointerType == UhtPointerType.Native)
 			{
@@ -159,19 +159,19 @@ namespace EpicGames.UHT.Types
 		}
 
 		/// <inheritdoc/>
-		public override bool IsSameType(UhtProperty Other)
+		public override bool IsSameType(UhtProperty other)
 		{
-			if (Other is UhtInterfaceProperty OtherObject)
+			if (other is UhtInterfaceProperty otherObject)
 			{
-				return this.InterfaceClass == OtherObject.InterfaceClass;
+				return this.InterfaceClass == otherObject.InterfaceClass;
 			}
 			return false;
 		}
 
 		/// <inheritdoc/>
-		public override bool MustBeConstArgument([NotNullWhen(true)] out UhtType? ErrorType)
+		public override bool MustBeConstArgument([NotNullWhen(true)] out UhtType? errorType)
 		{
-			ErrorType = this.InterfaceClass;
+			errorType = this.InterfaceClass;
 			return this.InterfaceClass.ClassFlags.HasAnyFlags(EClassFlags.Const);
 		}
 
@@ -179,27 +179,27 @@ namespace EpicGames.UHT.Types
 		[UhtPropertyType(Keyword = "FScriptInterface", Options = UhtPropertyTypeOptions.Simple)] // This can't be immediate due to the reference to UInterface
 		[SuppressMessage("CodeQuality", "IDE0051:Remove unused private members", Justification = "Attribute accessed method")]
 		[SuppressMessage("Style", "IDE0060:Remove unused parameter", Justification = "Attribute accessed method")]
-		private static UhtProperty? FScriptInterfaceProperty(UhtPropertyResolvePhase ResolvePhase, UhtPropertySettings PropertySettings, IUhtTokenReader TokenReader, UhtToken MatchedToken)
+		private static UhtProperty? FScriptInterfaceProperty(UhtPropertyResolvePhase resolvePhase, UhtPropertySettings propertySettings, IUhtTokenReader tokenReader, UhtToken matchedToken)
 		{
-			return new UhtInterfaceProperty(PropertySettings, PropertySettings.Outer.Session.IInterface);
+			return new UhtInterfaceProperty(propertySettings, propertySettings.Outer.Session.IInterface);
 		}
 
 		[UhtPropertyType(Keyword = "TScriptInterface")]
 		[SuppressMessage("CodeQuality", "IDE0051:Remove unused private members", Justification = "Attribute accessed method")]
 		[SuppressMessage("Style", "IDE0060:Remove unused parameter", Justification = "Attribute accessed method")]
-		private static UhtProperty? TScriptInterfaceProperty(UhtPropertyResolvePhase ResolvePhase, UhtPropertySettings PropertySettings, IUhtTokenReader TokenReader, UhtToken MatchedToken)
+		private static UhtProperty? TScriptInterfaceProperty(UhtPropertyResolvePhase resolvePhase, UhtPropertySettings propertySettings, IUhtTokenReader tokenReader, UhtToken matchedToken)
 		{
-			UhtClass? PropertyClass = UhtObjectPropertyBase.ParseTemplateObject(PropertySettings, TokenReader, MatchedToken, false);
-			if (PropertyClass == null)
+			UhtClass? propertyClass = UhtObjectPropertyBase.ParseTemplateObject(propertySettings, tokenReader, matchedToken, false);
+			if (propertyClass == null)
 			{
 				return null;
 			}
 
-			if (PropertyClass.ClassFlags.HasAnyFlags(EClassFlags.Interface))
+			if (propertyClass.ClassFlags.HasAnyFlags(EClassFlags.Interface))
 			{
-				return new UhtInterfaceProperty(PropertySettings, PropertyClass);
+				return new UhtInterfaceProperty(propertySettings, propertyClass);
 			}
-			return new UhtObjectProperty(PropertySettings, PropertyClass, null, EPropertyFlags.UObjectWrapper);
+			return new UhtObjectProperty(propertySettings, propertyClass, null, EPropertyFlags.UObjectWrapper);
 		}
 		#endregion
 	}
diff --git a/Engine/Source/Programs/Shared/EpicGames.UHT/Types/Properties/UhtLargeWorldCoordinatesRealProperty.cs b/Engine/Source/Programs/Shared/EpicGames.UHT/Types/Properties/UhtLargeWorldCoordinatesRealProperty.cs
--- a/Engine/Source/Programs/Shared/EpicGames.UHT/Types/Properties/UhtLargeWorldCoordinatesRealProperty.cs
+++ b/Engine/Source/Programs/Shared/EpicGames.UHT/Types/Properties/UhtLargeWorldCoordinatesRealProperty.cs
@@ -1,10 +1,10 @@
 // Copyright Epic Games, Inc. All Rights Reserved.
 
+using System.Diagnostics.CodeAnalysis;
+using System.Text;
 using EpicGames.UHT.Tables;
 using EpicGames.UHT.Tokenizer;
 using EpicGames.UHT.Utils;
-using System.Diagnostics.CodeAnalysis;
-using System.Text;
 
 namespace EpicGames.UHT.Types
 {
@@ -16,60 +16,60 @@ namespace EpicGames.UHT.Types
 	public class UhtLargeWorldCoordinatesRealProperty : UhtNumericProperty
 	{
 		/// <inheritdoc/>
-		public override string EngineClassName { get => "LargeWorldCoordinatesRealProperty"; }
+		public override string EngineClassName => "LargeWorldCoordinatesRealProperty";
 
 		/// <inheritdoc/>
-		protected override string CppTypeText { get => "double"; }
+		protected override string CppTypeText => "double";
 
 		/// <summary>
 		/// Construct new property
 		/// </summary>
-		/// <param name="PropertySettings">Property settings</param>
-		public UhtLargeWorldCoordinatesRealProperty(UhtPropertySettings PropertySettings) : base(PropertySettings, UhtPropertyIntType.None)
+		/// <param name="propertySettings">Property settings</param>
+		public UhtLargeWorldCoordinatesRealProperty(UhtPropertySettings propertySettings) : base(propertySettings, UhtPropertyIntType.None)
 		{
 			this.PropertyCaps |= UhtPropertyCaps.IsParameterSupportedByBlueprint | UhtPropertyCaps.IsMemberSupportedByBlueprint;
 		}
 
 		/// <inheritdoc/>
-		public override StringBuilder AppendMemberDecl(StringBuilder Builder, IUhtPropertyMemberContext Context, string Name, string NameSuffix, int Tabs)
+		public override StringBuilder AppendMemberDecl(StringBuilder builder, IUhtPropertyMemberContext context, string name, string nameSuffix, int tabs)
 		{
-			return AppendMemberDecl(Builder, Context, Name, NameSuffix, Tabs, "FLargeWorldCoordinatesRealPropertyParams");
+			return AppendMemberDecl(builder, context, name, nameSuffix, tabs, "FLargeWorldCoordinatesRealPropertyParams");
 		}
 
 		/// <inheritdoc/>
-		public override StringBuilder AppendMemberDef(StringBuilder Builder, IUhtPropertyMemberContext Context, string Name, string NameSuffix, string? Offset, int Tabs)
+		public override StringBuilder AppendMemberDef(StringBuilder builder, IUhtPropertyMemberContext context, string name, string nameSuffix, string? offset, int tabs)
 		{
-			AppendMemberDefStart(Builder, Context, Name, NameSuffix, Offset, Tabs, 
-				"FLargeWorldCoordinatesRealPropertyParams", 
+			AppendMemberDefStart(builder, context, name, nameSuffix, offset, tabs,
+				"FLargeWorldCoordinatesRealPropertyParams",
 				"UECodeGen_Private::EPropertyGenFlags::LargeWorldCoordinatesReal");
-			AppendMemberDefEnd(Builder, Context, Name, NameSuffix);
-			return Builder;
+			AppendMemberDefEnd(builder, context, name, nameSuffix);
+			return builder;
 		}
 
 		/// <inheritdoc/>
-		public override bool SanitizeDefaultValue(IUhtTokenReader DefaultValueReader, StringBuilder InnerDefaultValue)
+		public override bool SanitizeDefaultValue(IUhtTokenReader defaultValueReader, StringBuilder innerDefaultValue)
 		{
-			InnerDefaultValue.AppendFormat("{0:F6}", DefaultValueReader.GetConstFloatExpression());
+			innerDefaultValue.AppendFormat("{0:F6}", defaultValueReader.GetConstFloatExpression());
 			return true;
 		}
 
 		/// <inheritdoc/>
-		public override bool IsSameType(UhtProperty Other)
+		public override bool IsSameType(UhtProperty other)
 		{
-			return Other is UhtLargeWorldCoordinatesRealProperty;
+			return other is UhtLargeWorldCoordinatesRealProperty;
 		}
 
 		#region Keyword
 		[UhtPropertyType(Keyword = "FLargeWorldCoordinatesReal", Options = UhtPropertyTypeOptions.Simple | UhtPropertyTypeOptions.Immediate)]
 		[SuppressMessage("CodeQuality", "IDE0051:Remove unused private members", Justification = "Attribute accessed method")]
 		[SuppressMessage("Style", "IDE0060:Remove unused parameter", Justification = "Attribute accessed method")]
-		private static UhtProperty? LargeWorldCoordinatesRealProperty(UhtPropertyResolvePhase ResolvePhase, UhtPropertySettings PropertySettings, IUhtTokenReader TokenReader, UhtToken MatchedToken)
+		private static UhtProperty? LargeWorldCoordinatesRealProperty(UhtPropertyResolvePhase resolvePhase, UhtPropertySettings propertySettings, IUhtTokenReader tokenReader, UhtToken matchedToken)
 		{
-			if (!PropertySettings.Outer.HeaderFile.bIsNoExportTypes)
+			if (!propertySettings.Outer.HeaderFile.IsNoExportTypes)
 			{
-				TokenReader.LogError("FLargeWorldCoordinatesReal is intended for LWC support only and should not be used outside of NoExportTypes.h");
+				tokenReader.LogError("FLargeWorldCoordinatesReal is intended for LWC support only and should not be used outside of NoExportTypes.h");
 			}
-			return new UhtLargeWorldCoordinatesRealProperty(PropertySettings);
+			return new UhtLargeWorldCoordinatesRealProperty(propertySettings);
 		}
 		#endregion
 	}
diff --git a/Engine/Source/Programs/Shared/EpicGames.UHT/Types/Properties/UhtLazyObjectPtrProperty.cs b/Engine/Source/Programs/Shared/EpicGames.UHT/Types/Properties/UhtLazyObjectPtrProperty.cs
--- a/Engine/Source/Programs/Shared/EpicGames.UHT/Types/Properties/UhtLazyObjectPtrProperty.cs
+++ b/Engine/Source/Programs/Shared/EpicGames.UHT/Types/Properties/UhtLazyObjectPtrProperty.cs
@@ -1,11 +1,11 @@
 // Copyright Epic Games, Inc. All Rights Reserved.
 
+using System.Diagnostics.CodeAnalysis;
+using System.Text;
 using EpicGames.Core;
 using EpicGames.UHT.Tables;
 using EpicGames.UHT.Tokenizer;
 using EpicGames.UHT.Utils;
-using System.Diagnostics.CodeAnalysis;
-using System.Text;
 
 namespace EpicGames.UHT.Types
 {
@@ -18,90 +18,90 @@ namespace EpicGames.UHT.Types
 	public class UhtLazyObjectPtrProperty : UhtObjectPropertyBase
 	{
 		/// <inheritdoc/>
-		public override string EngineClassName { get => "LazyObjectProperty"; }
+		public override string EngineClassName => "LazyObjectProperty";
 
 		/// <inheritdoc/>
-		protected override string CppTypeText { get => "TLazyObjectPtr"; }
+		protected override string CppTypeText => "TLazyObjectPtr";
 
 		/// <inheritdoc/>
-		protected override string PGetMacroText { get => "LAZYOBJECT"; }
+		protected override string PGetMacroText => "LAZYOBJECT";
 
 		/// <inheritdoc/>
-		protected override bool bPGetPassAsNoPtr { get => true; }
+		protected override bool PGetPassAsNoPtr => true;
 
 		/// <inheritdoc/>
-		protected override UhtPGetArgumentType PGetTypeArgument { get => UhtPGetArgumentType.TypeText; }
+		protected override UhtPGetArgumentType PGetTypeArgument => UhtPGetArgumentType.TypeText;
 
 		/// <summary>
 		/// Construct new property
 		/// </summary>
-		/// <param name="PropertySettings">Property settings</param>
-		/// <param name="Class">Referenced class</param>
-		public UhtLazyObjectPtrProperty(UhtPropertySettings PropertySettings, UhtClass Class)
-			: base(PropertySettings, Class, null)
+		/// <param name="propertySettings">Property settings</param>
+		/// <param name="classObj">Referenced class</param>
+		public UhtLazyObjectPtrProperty(UhtPropertySettings propertySettings, UhtClass classObj)
+			: base(propertySettings, classObj, null)
 		{
 			this.PropertyFlags |= EPropertyFlags.UObjectWrapper;
 			this.PropertyCaps |= UhtPropertyCaps.PassCppArgsByRef | UhtPropertyCaps.RequiresNullConstructorArg;
 		}
 
 		/// <inheritdoc/>
-		public override StringBuilder AppendText(StringBuilder Builder, UhtPropertyTextType TextType, bool bIsTemplateArgument)
+		public override StringBuilder AppendText(StringBuilder builder, UhtPropertyTextType textType, bool isTemplateArgument)
 		{
-			switch (TextType)
+			switch (textType)
 			{
 				case UhtPropertyTextType.FunctionThunkParameterArgType:
-					Builder.Append("TLazyObjectPtr<").Append(this.Class.SourceName).Append("> "); //COMPATIBILITY-TODO - Extra space
+					builder.Append("TLazyObjectPtr<").Append(this.Class.SourceName).Append("> "); //COMPATIBILITY-TODO - Extra space
 					break;
 
 				default:
-					Builder.Append("TLazyObjectPtr<").Append(this.Class.SourceName).Append('>');
+					builder.Append("TLazyObjectPtr<").Append(this.Class.SourceName).Append('>');
 					break;
 			}
-			return Builder;
+			return builder;
 		}
 
 		/// <inheritdoc/>
-		public override StringBuilder AppendMemberDecl(StringBuilder Builder, IUhtPropertyMemberContext Context, string Name, string NameSuffix, int Tabs)
+		public override StringBuilder AppendMemberDecl(StringBuilder builder, IUhtPropertyMemberContext context, string name, string nameSuffix, int tabs)
 		{
-			return AppendMemberDecl(Builder, Context, Name, NameSuffix, Tabs, "FLazyObjectPropertyParams");
+			return AppendMemberDecl(builder, context, name, nameSuffix, tabs, "FLazyObjectPropertyParams");
 		}
 
 		/// <inheritdoc/>
-		public override StringBuilder AppendMemberDef(StringBuilder Builder, IUhtPropertyMemberContext Context, string Name, string NameSuffix, string? Offset, int Tabs)
+		public override StringBuilder AppendMemberDef(StringBuilder builder, IUhtPropertyMemberContext context, string name, string nameSuffix, string? offset, int tabs)
 		{
-			AppendMemberDefStart(Builder, Context, Name, NameSuffix, Offset, Tabs, "FLazyObjectPropertyParams", "UECodeGen_Private::EPropertyGenFlags::LazyObject");
-			AppendMemberDefRef(Builder, Context, this.Class, false);
-			AppendMemberDefEnd(Builder, Context, Name, NameSuffix);
-			return Builder;
+			AppendMemberDefStart(builder, context, name, nameSuffix, offset, tabs, "FLazyObjectPropertyParams", "UECodeGen_Private::EPropertyGenFlags::LazyObject");
+			AppendMemberDefRef(builder, context, this.Class, false);
+			AppendMemberDefEnd(builder, context, name, nameSuffix);
+			return builder;
 		}
 
 		/// <inheritdoc/>
-		public override StringBuilder AppendNullConstructorArg(StringBuilder Builder, bool bIsInitializer)
+		public override StringBuilder AppendNullConstructorArg(StringBuilder builder, bool isInitializer)
 		{
-			Builder.Append("NULL");
-			return Builder;
+			builder.Append("NULL");
+			return builder;
 		}
 
 		/// <inheritdoc/>
-		public override bool IsSameType(UhtProperty Other)
+		public override bool IsSameType(UhtProperty other)
 		{
-			if (Other is UhtLazyObjectPtrProperty OtherObject)
+			if (other is UhtLazyObjectPtrProperty otherObject)
 			{
-				return this.Class == OtherObject.Class && this.MetaClass == OtherObject.MetaClass;
+				return this.Class == otherObject.Class && this.MetaClass == otherObject.MetaClass;
 			}
 			return false;
 		}
 
 		/// <inheritdoc/>
-		public override void Validate(UhtStruct OuterStruct, UhtProperty OutermostProperty, UhtValidationOptions Options)
+		public override void Validate(UhtStruct outerStruct, UhtProperty outermostProperty, UhtValidationOptions options)
 		{
-			base.Validate(OuterStruct, OutermostProperty, Options);
+			base.Validate(outerStruct, outermostProperty, options);
 
 			// UFunctions with a smart pointer as input parameter wont compile anyway, because of missing P_GET_... macro.
 			// UFunctions with a smart pointer as return type will crash when called via blueprint, because they are not supported in VM.
 			if (this.PropertyCategory == UhtPropertyCategory.RegularParameter || this.PropertyCategory == UhtPropertyCategory.ReplicatedParameter)
 			{
-				OuterStruct.LogError("UFunctions cannot take a lazy pointer as a parameter.");
+				outerStruct.LogError("UFunctions cannot take a lazy pointer as a parameter.");
 			}
 		}
 
@@ -109,20 +109,20 @@ namespace EpicGames.UHT.Types
 		[UhtPropertyType(Keyword = "TLazyObjectPtr")]
 		[SuppressMessage("CodeQuality", "IDE0051:Remove unused private members", Justification = "Attribute accessed method")]
 		[SuppressMessage("Style", "IDE0060:Remove unused parameter", Justification = "Attribute accessed method")]
-		private static UhtProperty? LazyObjectPtrProperty(UhtPropertyResolvePhase ResolvePhase, UhtPropertySettings PropertySettings, IUhtTokenReader TokenReader, UhtToken MatchedToken)
+		private static UhtProperty? LazyObjectPtrProperty(UhtPropertyResolvePhase resolvePhase, UhtPropertySettings propertySettings, IUhtTokenReader tokenReader, UhtToken matchedToken)
 		{
-			UhtClass? PropertyClass = ParseTemplateObject(PropertySettings, TokenReader, MatchedToken, true);
-			if (PropertyClass == null)
+			UhtClass? propertyClass = ParseTemplateObject(propertySettings, tokenReader, matchedToken, true);
+			if (propertyClass == null)
 			{
 				return null;
 			}
 
-			if (PropertyClass.IsChildOf(PropertyClass.Session.UClass))
+			if (propertyClass.IsChildOf(propertyClass.Session.UClass))
 			{
-				TokenReader.LogError("Class variables cannot be lazy, they are always strong.");
+				tokenReader.LogError("Class variables cannot be lazy, they are always strong.");
 			}
 
-			return new UhtLazyObjectPtrProperty(PropertySettings, PropertyClass);
+			return new UhtLazyObjectPtrProperty(propertySettings, propertyClass);
 		}
 		#endregion
 	}
