More coding standard fixes

diff --git a/Engine/Source/Programs/Shared/EpicGames.UHT/Types/Properties/UhtMapProperty.cs b/Engine/Source/Programs/Shared/EpicGames.UHT/Types/Properties/UhtMapProperty.cs
--- a/Engine/Source/Programs/Shared/EpicGames.UHT/Types/Properties/UhtMapProperty.cs
+++ b/Engine/Source/Programs/Shared/EpicGames.UHT/Types/Properties/UhtMapProperty.cs
@@ -1,13 +1,13 @@
 // Copyright Epic Games, Inc. All Rights Reserved.
 
+using System.Collections.Generic;
+using System.Diagnostics.CodeAnalysis;
+using System.Text;
 using EpicGames.Core;
 using EpicGames.UHT.Parsers;
 using EpicGames.UHT.Tables;
 using EpicGames.UHT.Tokenizer;
 using EpicGames.UHT.Utils;
-using System.Collections.Generic;
-using System.Diagnostics.CodeAnalysis;
-using System.Text;
 
 namespace EpicGames.UHT.Types
 {
@@ -20,16 +20,16 @@ namespace EpicGames.UHT.Types
 	public class UhtMapProperty : UhtContainerBaseProperty
 	{
 		/// <inheritdoc/>
-		public override string EngineClassName { get => "MapProperty"; }
+		public override string EngineClassName => "MapProperty";
 
 		/// <inheritdoc/>
-		protected override string CppTypeText { get => "TMap"; }
+		protected override string CppTypeText => "TMap";
 
 		/// <inheritdoc/>
-		protected override string PGetMacroText { get => "TMAP"; }
+		protected override string PGetMacroText => "TMAP";
 
 		/// <inheritdoc/>
-		protected override UhtPGetArgumentType PGetTypeArgument { get => UhtPGetArgumentType.TypeText; }
+		protected override UhtPGetArgumentType PGetTypeArgument => UhtPGetArgumentType.TypeText;
 
 		/// <summary>
 		/// Key property
@@ -39,25 +39,25 @@ namespace EpicGames.UHT.Types
 		/// <summary>
 		/// Construct new property
 		/// </summary>
-		/// <param name="PropertySettings">Property settings</param>
-		/// <param name="Key">Key property</param>
-		/// <param name="Value">Value property</param>
-		public UhtMapProperty(UhtPropertySettings PropertySettings, UhtProperty Key, UhtProperty Value) : base(PropertySettings, Value)
+		/// <param name="propertySettings">Property settings</param>
+		/// <param name="key">Key property</param>
+		/// <param name="value">Value property</param>
+		public UhtMapProperty(UhtPropertySettings propertySettings, UhtProperty key, UhtProperty value) : base(propertySettings, value)
 		{
 			// If the creation of the value property set more flags, then copy those flags to ourselves
-			this.PropertyFlags |= Value.PropertyFlags & EPropertyFlags.UObjectWrapper;
+			this.PropertyFlags |= value.PropertyFlags & EPropertyFlags.UObjectWrapper;
 
 			// Make sure the 'UObjectWrapper' flag is maintained so that both 'TMap<TSubclassOf<...>, ...>' and 'TMap<UClass*, TSubclassOf<...>>' works correctly
-			Key.PropertyFlags = (Value.PropertyFlags & ~EPropertyFlags.UObjectWrapper) | (Key.PropertyFlags & EPropertyFlags.UObjectWrapper);
-			Key.DisallowPropertyFlags = ~(EPropertyFlags.ContainsInstancedReference | EPropertyFlags.InstancedReference | EPropertyFlags.UObjectWrapper);
+			key.PropertyFlags = (value.PropertyFlags & ~EPropertyFlags.UObjectWrapper) | (key.PropertyFlags & EPropertyFlags.UObjectWrapper);
+			key.DisallowPropertyFlags = ~(EPropertyFlags.ContainsInstancedReference | EPropertyFlags.InstancedReference | EPropertyFlags.UObjectWrapper);
 
-			this.KeyProperty = Key;
+			this.KeyProperty = key;
 			this.PropertyCaps |= UhtPropertyCaps.PassCppArgsByRef;
 			this.PropertyCaps &= ~(UhtPropertyCaps.CanBeContainerValue | UhtPropertyCaps.CanBeContainerKey);
-			this.PropertyCaps = (this.PropertyCaps & ~(UhtPropertyCaps.IsParameterSupportedByBlueprint | UhtPropertyCaps.IsMemberSupportedByBlueprint | UhtPropertyCaps.CanExposeOnSpawn));
+			this.PropertyCaps &= ~(UhtPropertyCaps.IsParameterSupportedByBlueprint | UhtPropertyCaps.IsMemberSupportedByBlueprint | UhtPropertyCaps.CanExposeOnSpawn);
 			this.PropertyCaps |= (this.ValueProperty.PropertyCaps & (UhtPropertyCaps.IsParameterSupportedByBlueprint | UhtPropertyCaps.IsMemberSupportedByBlueprint | UhtPropertyCaps.CanExposeOnSpawn)) &
 				(this.KeyProperty.PropertyCaps & (UhtPropertyCaps.IsParameterSupportedByBlueprint | UhtPropertyCaps.IsMemberSupportedByBlueprint));
-			this.PropertyFlags = Value.PropertyFlags;
+			this.PropertyFlags = value.PropertyFlags;
 			this.ValueProperty.SourceName = this.SourceName;
 			this.ValueProperty.EngineName = this.EngineName;
 			this.ValueProperty.PropertyFlags &= EPropertyFlags.PropagateToMapValue;
@@ -76,18 +76,18 @@ namespace EpicGames.UHT.Types
 		}
 
 		/// <inheritdoc/>
-		protected override bool ResolveSelf(UhtResolvePhase Phase)
+		protected override bool ResolveSelf(UhtResolvePhase phase)
 		{
-			bool bResults = base.ResolveSelf(Phase);
-			switch (Phase)
+			bool results = base.ResolveSelf(phase);
+			switch (phase)
 			{
 				case UhtResolvePhase.Final:
-					this.KeyProperty.Resolve(Phase);
+					this.KeyProperty.Resolve(phase);
 					this.KeyProperty.MetaData.Clear();
 
-					EPropertyFlags NewFlags = ResolveAndReturnNewFlags(this.ValueProperty, Phase);
-					this.PropertyFlags |= NewFlags;
-					this.KeyProperty.PropertyFlags |= NewFlags;
+					EPropertyFlags newFlags = ResolveAndReturnNewFlags(this.ValueProperty, phase);
+					this.PropertyFlags |= newFlags;
+					this.KeyProperty.PropertyFlags |= newFlags;
 					this.MetaData.Add(this.ValueProperty.MetaData);
 					this.ValueProperty.MetaData.Clear();
 					this.ValueProperty.PropertyFlags = this.PropertyFlags & EPropertyFlags.PropagateToMapValue;
@@ -96,14 +96,14 @@ namespace EpicGames.UHT.Types
 					PropagateFlagsFromInnerAndHandlePersistentInstanceMetadata(this, this.MetaData, this.ValueProperty);
 					break;
 			}
-			return bResults;
+			return results;
 		}
 
 		/// <inheritdoc/>
-		public override void CollectReferencesInternal(IUhtReferenceCollector Collector, bool bTemplateProperty)
+		public override void CollectReferencesInternal(IUhtReferenceCollector collector, bool templateProperty)
 		{
-			this.ValueProperty.CollectReferencesInternal(Collector, true);
-			this.KeyProperty.CollectReferencesInternal(Collector, true);
+			this.ValueProperty.CollectReferencesInternal(collector, true);
+			this.KeyProperty.CollectReferencesInternal(collector, true);
 		}
 
 		/// <inheritdoc/>
@@ -115,103 +115,103 @@ namespace EpicGames.UHT.Types
 		/// <inheritdoc/>
 		public override IEnumerable<UhtType> EnumerateReferencedTypes()
 		{
-			foreach (UhtType Type in this.ValueProperty.EnumerateReferencedTypes())
+			foreach (UhtType type in this.ValueProperty.EnumerateReferencedTypes())
 			{
-				yield return Type;
+				yield return type;
 			}
-			foreach (UhtType Type in this.KeyProperty.EnumerateReferencedTypes())
+			foreach (UhtType type in this.KeyProperty.EnumerateReferencedTypes())
 			{
-				yield return Type;
+				yield return type;
 			}
 		}
 
 		/// <inheritdoc/>
-		public override StringBuilder AppendText(StringBuilder Builder, UhtPropertyTextType TextType, bool bIsTemplateArgument)
+		public override StringBuilder AppendText(StringBuilder builder, UhtPropertyTextType textType, bool isTemplateArgument)
 		{
-			switch (TextType)
+			switch (textType)
 			{
 				case UhtPropertyTextType.SparseShort:
-					Builder.Append("TMap");
+					builder.Append("TMap");
 					break;
 
 				case UhtPropertyTextType.FunctionThunkParameterArgType:
-					Builder.AppendFunctionThunkParameterArrayType(this.KeyProperty).Append(',').AppendFunctionThunkParameterArrayType(this.ValueProperty);
+					builder.AppendFunctionThunkParameterArrayType(this.KeyProperty).Append(',').AppendFunctionThunkParameterArrayType(this.ValueProperty);
 					break;
 
 				default:
-					Builder.Append("TMap<").AppendPropertyText(this.KeyProperty, TextType, true);
-					if (Builder[Builder.Length - 1] == '>')
+					builder.Append("TMap<").AppendPropertyText(this.KeyProperty, textType, true);
+					if (builder[^1] == '>')
 					{
 						// if our internal property type is a template class, add a space between the closing brackets b/c VS.NET cannot parse this correctly
-						Builder.Append(' ');
+						builder.Append(' ');
 					}
-					Builder.Append(',').AppendPropertyText(this.ValueProperty, TextType, true);
-					if (Builder[Builder.Length - 1] == '>')
+					builder.Append(',').AppendPropertyText(this.ValueProperty, textType, true);
+					if (builder[^1] == '>')
 					{
 						// if our internal property type is a template class, add a space between the closing brackets b/c VS.NET cannot parse this correctly
-						Builder.Append(' ');
+						builder.Append(' ');
 					}
-					Builder.Append('>');
+					builder.Append('>');
 					break;
 			}
-			return Builder;
+			return builder;
 		}
 
 		/// <inheritdoc/>
-		public override StringBuilder AppendMemberDecl(StringBuilder Builder, IUhtPropertyMemberContext Context, string Name, string NameSuffix, int Tabs)
+		public override StringBuilder AppendMemberDecl(StringBuilder builder, IUhtPropertyMemberContext context, string name, string nameSuffix, int tabs)
 		{
-			Builder.AppendMemberDecl(this.ValueProperty, Context, Name, GetNameSuffix(NameSuffix, "_ValueProp"), Tabs);
-			Builder.AppendMemberDecl(this.KeyProperty, Context, Name, GetNameSuffix(NameSuffix, "_Key_KeyProp"), Tabs);
-			return AppendMemberDecl(Builder, Context, Name, NameSuffix, Tabs, "FMapPropertyParams");
+			builder.AppendMemberDecl(this.ValueProperty, context, name, GetNameSuffix(nameSuffix, "_ValueProp"), tabs);
+			builder.AppendMemberDecl(this.KeyProperty, context, name, GetNameSuffix(nameSuffix, "_Key_KeyProp"), tabs);
+			return AppendMemberDecl(builder, context, name, nameSuffix, tabs, "FMapPropertyParams");
 		}
 
 		/// <inheritdoc/>
-		public override StringBuilder AppendMemberDef(StringBuilder Builder, IUhtPropertyMemberContext Context, string Name, string NameSuffix, string? Offset, int Tabs)
+		public override StringBuilder AppendMemberDef(StringBuilder builder, IUhtPropertyMemberContext context, string name, string nameSuffix, string? offset, int tabs)
 		{
-			Builder.AppendMemberDef(this.ValueProperty, Context, Name, GetNameSuffix(NameSuffix, "_ValueProp"), "1", Tabs);
-			Builder.AppendMemberDef(this.KeyProperty, Context, Name, GetNameSuffix(NameSuffix, "_Key_KeyProp"), "0", Tabs);
-			AppendMemberDefStart(Builder, Context, Name, NameSuffix, Offset, Tabs, "FMapPropertyParams", "UECodeGen_Private::EPropertyGenFlags::Map");
-			Builder.Append(this.Allocator == UhtPropertyAllocator.MemoryImage ? "EMapPropertyFlags::UsesMemoryImageAllocator" : "EMapPropertyFlags::None").Append(", ");
-			AppendMemberDefEnd(Builder, Context, Name, NameSuffix);
-			return Builder;
+			builder.AppendMemberDef(this.ValueProperty, context, name, GetNameSuffix(nameSuffix, "_ValueProp"), "1", tabs);
+			builder.AppendMemberDef(this.KeyProperty, context, name, GetNameSuffix(nameSuffix, "_Key_KeyProp"), "0", tabs);
+			AppendMemberDefStart(builder, context, name, nameSuffix, offset, tabs, "FMapPropertyParams", "UECodeGen_Private::EPropertyGenFlags::Map");
+			builder.Append(this.Allocator == UhtPropertyAllocator.MemoryImage ? "EMapPropertyFlags::UsesMemoryImageAllocator" : "EMapPropertyFlags::None").Append(", ");
+			AppendMemberDefEnd(builder, context, name, nameSuffix);
+			return builder;
 		}
 
 		/// <inheritdoc/>
-		public override StringBuilder AppendMemberPtr(StringBuilder Builder, IUhtPropertyMemberContext Context, string Name, string NameSuffix, int Tabs)
+		public override StringBuilder AppendMemberPtr(StringBuilder builder, IUhtPropertyMemberContext context, string name, string nameSuffix, int tabs)
 		{
-			Builder.AppendMemberPtr(this.ValueProperty, Context, Name, GetNameSuffix(NameSuffix, "_ValueProp"), Tabs);
-			Builder.AppendMemberPtr(this.KeyProperty, Context, Name, GetNameSuffix(NameSuffix, "_Key_KeyProp"), Tabs);
-			base.AppendMemberPtr(Builder, Context, Name, NameSuffix, Tabs);
-			return Builder;
+			builder.AppendMemberPtr(this.ValueProperty, context, name, GetNameSuffix(nameSuffix, "_ValueProp"), tabs);
+			builder.AppendMemberPtr(this.KeyProperty, context, name, GetNameSuffix(nameSuffix, "_Key_KeyProp"), tabs);
+			base.AppendMemberPtr(builder, context, name, nameSuffix, tabs);
+			return builder;
 		}
 
 		/// <inheritdoc/>
-		public override void AppendObjectHashes(StringBuilder Builder, int StartingLength, IUhtPropertyMemberContext Context)
+		public override void AppendObjectHashes(StringBuilder builder, int startingLength, IUhtPropertyMemberContext context)
 		{
-			this.KeyProperty.AppendObjectHashes(Builder, StartingLength, Context);
-			this.ValueProperty.AppendObjectHashes(Builder, StartingLength, Context);
+			this.KeyProperty.AppendObjectHashes(builder, startingLength, context);
+			this.ValueProperty.AppendObjectHashes(builder, startingLength, context);
 		}
 
 		/// <inheritdoc/>
-		public override StringBuilder AppendNullConstructorArg(StringBuilder Builder, bool bIsInitializer)
+		public override StringBuilder AppendNullConstructorArg(StringBuilder builder, bool isInitializer)
 		{
-			Builder.AppendPropertyText(this, UhtPropertyTextType.Construction).Append("()");
-			return Builder;
+			builder.AppendPropertyText(this, UhtPropertyTextType.Construction).Append("()");
+			return builder;
 		}
 
 		/// <inheritdoc/>
-		public override bool SanitizeDefaultValue(IUhtTokenReader DefaultValueReader, StringBuilder InnerDefaultValue)
+		public override bool SanitizeDefaultValue(IUhtTokenReader defaultValueReader, StringBuilder innerDefaultValue)
 		{
 			return false;
 		}
 
 		/// <inheritdoc/>
-		public override bool IsSameType(UhtProperty Other)
+		public override bool IsSameType(UhtProperty other)
 		{
-			if (Other is UhtMapProperty OtherMap)
+			if (other is UhtMapProperty otherMap)
 			{
-				return this.ValueProperty.IsSameType(OtherMap.ValueProperty) &&
-					this.KeyProperty.IsSameType(OtherMap.KeyProperty);
+				return this.ValueProperty.IsSameType(otherMap.ValueProperty) &&
+					this.KeyProperty.IsSameType(otherMap.KeyProperty);
 			}
 			return false;
 		}
@@ -224,31 +224,31 @@ namespace EpicGames.UHT.Types
 		}
 
 		/// <inheritdoc/>
-		public override void Validate(UhtStruct OuterStruct, UhtProperty OutermostProperty, UhtValidationOptions Options)
+		public override void Validate(UhtStruct outerStruct, UhtProperty outermostProperty, UhtValidationOptions options)
 		{
-			base.Validate(OuterStruct, OutermostProperty, Options);
-			this.KeyProperty.Validate(OuterStruct, OutermostProperty, Options | UhtValidationOptions.IsKey);
+			base.Validate(outerStruct, outermostProperty, options);
+			this.KeyProperty.Validate(outerStruct, outermostProperty, options | UhtValidationOptions.IsKey);
 		}
 
 		///<inheritdoc/>
-		public override bool ValidateStructPropertyOkForNet(UhtProperty ReferencingProperty)
+		public override bool ValidateStructPropertyOkForNet(UhtProperty referencingProperty)
 		{
 			if (!this.PropertyFlags.HasAnyFlags(EPropertyFlags.RepSkip))
 			{
-				ReferencingProperty.LogError($"Maps are not supported for Replication or RPCs.  Map '{this.SourceName}' in '{this.Outer?.SourceName}'.  Origin '{ReferencingProperty.SourceName}'");
+				referencingProperty.LogError($"Maps are not supported for Replication or RPCs.  Map '{this.SourceName}' in '{this.Outer?.SourceName}'.  Origin '{referencingProperty.SourceName}'");
 				return false;
 			}
 			return true;
 		}
 
 		/// <inheritdoc/>
-		protected override void ValidateFunctionArgument(UhtFunction Function, UhtValidationOptions Options)
+		protected override void ValidateFunctionArgument(UhtFunction function, UhtValidationOptions options)
 		{
-			base.ValidateFunctionArgument(Function, Options);
+			base.ValidateFunctionArgument(function, options);
 
-			if (Function.FunctionFlags.HasAnyFlags(EFunctionFlags.Net))
+			if (function.FunctionFlags.HasAnyFlags(EFunctionFlags.Net))
 			{
-				if (!Function.FunctionFlags.HasAnyFlags(EFunctionFlags.NetRequest | EFunctionFlags.NetResponse))
+				if (!function.FunctionFlags.HasAnyFlags(EFunctionFlags.NetRequest | EFunctionFlags.NetResponse))
 				{
 					this.LogError("Maps are not supported in an RPC.");
 				}
@@ -258,64 +258,64 @@ namespace EpicGames.UHT.Types
 		#region Keyword
 		[UhtPropertyType(Keyword = "TMap")]
 		[SuppressMessage("CodeQuality", "IDE0051:Remove unused private members", Justification = "Attribute accessed method")]
-		private static UhtProperty? MapProperty(UhtPropertyResolvePhase ResolvePhase, UhtPropertySettings PropertySettings, IUhtTokenReader TokenReader, UhtToken MatchedToken)
+		private static UhtProperty? MapProperty(UhtPropertyResolvePhase resolvePhase, UhtPropertySettings propertySettings, IUhtTokenReader tokenReader, UhtToken matchedToken)
 		{
-			using (var TokenContext = new UhtMessageContext("TMap"))
+			using (UhtMessageContext tokenContext = new UhtMessageContext("TMap"))
 			{
-				if (!TokenReader.SkipExpectedType(MatchedToken.Value, PropertySettings.PropertyCategory == UhtPropertyCategory.Member))
+				if (!tokenReader.SkipExpectedType(matchedToken.Value, propertySettings.PropertyCategory == UhtPropertyCategory.Member))
 				{
 					return null;
 				}
-				TokenReader.Require('<');
+				tokenReader.Require('<');
 
 				// Parse the key type
-				UhtProperty? Key = UhtPropertyParser.ParseTemplateParam(ResolvePhase, PropertySettings, "Key", TokenReader);
-				if (Key == null)
+				UhtProperty? key = UhtPropertyParser.ParseTemplateParam(resolvePhase, propertySettings, "Key", tokenReader);
+				if (key == null)
 				{
 					return null;
 				}
 
-				if (!Key.PropertyCaps.HasAnyFlags(UhtPropertyCaps.CanBeContainerKey))
+				if (!key.PropertyCaps.HasAnyFlags(UhtPropertyCaps.CanBeContainerKey))
 				{
-					TokenReader.LogError($"The type \'{Key.GetUserFacingDecl()}\' can not be used as a key in a TMap");
+					tokenReader.LogError($"The type \'{key.GetUserFacingDecl()}\' can not be used as a key in a TMap");
 				}
 
-				if (PropertySettings.PropertyFlags.HasAnyFlags(EPropertyFlags.Net))
+				if (propertySettings.PropertyFlags.HasAnyFlags(EPropertyFlags.Net))
 				{
-					TokenReader.LogError("Replicated maps are not supported.");
+					tokenReader.LogError("Replicated maps are not supported.");
 				}
 
-				TokenReader.Require(',');
+				tokenReader.Require(',');
 
 				// Parse the value type
-				UhtProperty? Value = UhtPropertyParser.ParseTemplateParam(ResolvePhase, PropertySettings, "Value", TokenReader);
-				if (Value == null)
+				UhtProperty? value = UhtPropertyParser.ParseTemplateParam(resolvePhase, propertySettings, "Value", tokenReader);
+				if (value == null)
 				{
 					return null;
 				}
 
-				if (!Value.PropertyCaps.HasAnyFlags(UhtPropertyCaps.CanBeContainerValue))
+				if (!value.PropertyCaps.HasAnyFlags(UhtPropertyCaps.CanBeContainerValue))
 				{
-					TokenReader.LogError($"The type \'{Value.GetUserFacingDecl()}\' can not be used as a value in a TMap");
+					tokenReader.LogError($"The type \'{value.GetUserFacingDecl()}\' can not be used as a value in a TMap");
 				}
 
-				if (TokenReader.TryOptional(','))
+				if (tokenReader.TryOptional(','))
 				{
-					UhtToken AllocatorToken = TokenReader.GetIdentifier();
-					if (AllocatorToken.IsIdentifier("FMemoryImageSetAllocator"))
+					UhtToken allocatorToken = tokenReader.GetIdentifier();
+					if (allocatorToken.IsIdentifier("FMemoryImageSetAllocator"))
 					{
-						PropertySettings.Allocator = UhtPropertyAllocator.MemoryImage;
+						propertySettings.Allocator = UhtPropertyAllocator.MemoryImage;
 					}
 					else
 					{
-						TokenReader.LogError($"Found '{AllocatorToken.Value}' - explicit allocators are not supported in TMap properties.");
+						tokenReader.LogError($"Found '{allocatorToken.Value}' - explicit allocators are not supported in TMap properties.");
 					}
 				}
-				TokenReader.Require('>');
+				tokenReader.Require('>');
 
 				//@TODO: Prevent sparse delegate types from being used in a container
 
-				return new UhtMapProperty(PropertySettings, Key, Value);
+				return new UhtMapProperty(propertySettings, key, value);
 			}
 		}
 		#endregion
diff --git a/Engine/Source/Programs/Shared/EpicGames.UHT/Types/Properties/UhtMulticastDelegateProperty.cs b/Engine/Source/Programs/Shared/EpicGames.UHT/Types/Properties/UhtMulticastDelegateProperty.cs
--- a/Engine/Source/Programs/Shared/EpicGames.UHT/Types/Properties/UhtMulticastDelegateProperty.cs
+++ b/Engine/Source/Programs/Shared/EpicGames.UHT/Types/Properties/UhtMulticastDelegateProperty.cs
@@ -1,11 +1,11 @@
 // Copyright Epic Games, Inc. All Rights Reserved.
 
+using System.Text;
+using System.Text.Json.Serialization;
 using EpicGames.Core;
 using EpicGames.UHT.Tables;
 using EpicGames.UHT.Tokenizer;
 using EpicGames.UHT.Utils;
-using System.Text;
-using System.Text.Json.Serialization;
 
 namespace EpicGames.UHT.Types
 {
@@ -16,16 +16,16 @@ namespace EpicGames.UHT.Types
 	public abstract class UhtMulticastDelegateProperty : UhtProperty
 	{
 		/// <inheritdoc/>
-		public override string EngineClassName { get => "MulticastDelegateProperty"; }
+		public override string EngineClassName => "MulticastDelegateProperty";
 
 		/// <inheritdoc/>
-		protected override string CppTypeText { get => "FMulticastScriptDelegate"; }
+		protected override string CppTypeText => "FMulticastScriptDelegate";
 
 		/// <inheritdoc/>
-		protected override string PGetMacroText { get => "PROPERTY"; }
+		protected override string PGetMacroText => "PROPERTY";
 
 		/// <inheritdoc/>
-		protected override UhtPGetArgumentType PGetTypeArgument { get => UhtPGetArgumentType.TypeText; }
+		protected override UhtPGetArgumentType PGetTypeArgument => UhtPGetArgumentType.TypeText;
 
 		/// <summary>
 		/// Referenced function
@@ -36,92 +36,92 @@ namespace EpicGames.UHT.Types
 		/// <summary>
 		/// Construct property
 		/// </summary>
-		/// <param name="PropertySettings">Property settings</param>
-		/// <param name="Function">Referenced function</param>
-		protected UhtMulticastDelegateProperty(UhtPropertySettings PropertySettings, UhtFunction Function) : base(PropertySettings)
+		/// <param name="propertySettings">Property settings</param>
+		/// <param name="function">Referenced function</param>
+		protected UhtMulticastDelegateProperty(UhtPropertySettings propertySettings, UhtFunction function) : base(propertySettings)
 		{
-			this.Function = Function;
-			this.HeaderFile.AddReferencedHeader(Function);
+			this.Function = function;
+			this.HeaderFile.AddReferencedHeader(function);
 			this.PropertyCaps |= UhtPropertyCaps.PassCppArgsByRef | UhtPropertyCaps.CanBeBlueprintAssignable | UhtPropertyCaps.CanBeBlueprintCallable |
 				UhtPropertyCaps.CanBeBlueprintAuthorityOnly | UhtPropertyCaps.IsMemberSupportedByBlueprint | UhtPropertyCaps.SupportsRigVM;
 			this.PropertyCaps &= ~(UhtPropertyCaps.IsParameterSupportedByBlueprint);
 		}
 
 		/// <inheritdoc/>
-		protected override bool ResolveSelf(UhtResolvePhase Phase)
+		protected override bool ResolveSelf(UhtResolvePhase phase)
 		{
-			bool bResults = base.ResolveSelf(Phase);
-			switch (Phase)
+			bool results = base.ResolveSelf(phase);
+			switch (phase)
 			{
 				case UhtResolvePhase.Final:
 					this.PropertyFlags |= EPropertyFlags.InstancedReference & ~this.DisallowPropertyFlags;
 					break;
 			}
-			return bResults;
+			return results;
 		}
 
 		/// <inheritdoc/>
-		public override bool ScanForInstancedReferenced(bool bDeepScan)
+		public override bool ScanForInstancedReferenced(bool deepScan)
 		{
 			return true;
 		}
 
 		/// <inheritdoc/>
-		public override void CollectReferencesInternal(IUhtReferenceCollector Collector, bool bTemplateProperty)
+		public override void CollectReferencesInternal(IUhtReferenceCollector collector, bool templateProperty)
 		{
-			Collector.AddCrossModuleReference(this.Function, true);
+			collector.AddCrossModuleReference(this.Function, true);
 		}
 
 		/// <inheritdoc/>
-		public override StringBuilder AppendText(StringBuilder Builder, UhtPropertyTextType TextType, bool bIsTemplateArgument)
+		public override StringBuilder AppendText(StringBuilder builder, UhtPropertyTextType textType, bool isTemplateArgument)
 		{
-			switch (TextType)
+			switch (textType)
 			{
 				case UhtPropertyTextType.EventParameterFunctionMember:
 				case UhtPropertyTextType.ExportMember:
 				case UhtPropertyTextType.RigVMType:
 				case UhtPropertyTextType.FunctionThunkParameterArgType:
-					Builder.Append(this.CppTypeText);
+					builder.Append(this.CppTypeText);
 					break;
 
 				default:
-					Builder.Append(this.Function.SourceName);
+					builder.Append(this.Function.SourceName);
 					break;
 			}
-			return Builder;
+			return builder;
 		}
 
 		/// <inheritdoc/>
-		public override void AppendObjectHashes(StringBuilder Builder, int StartingLength, IUhtPropertyMemberContext Context)
+		public override void AppendObjectHashes(StringBuilder builder, int startingLength, IUhtPropertyMemberContext context)
 		{
-			Builder.AppendObjectHash(StartingLength, this, Context, this.Function);
+			builder.AppendObjectHash(startingLength, this, context, this.Function);
 		}
 
 		/// <inheritdoc/>
-		public override StringBuilder AppendFunctionThunkParameterArg(StringBuilder Builder)
+		public override StringBuilder AppendFunctionThunkParameterArg(StringBuilder builder)
 		{
-			return Builder.Append(this.Function.SourceName).Append('(').AppendFunctionThunkParameterName(this).Append(')');
+			return builder.Append(this.Function.SourceName).Append('(').AppendFunctionThunkParameterName(this).Append(')');
 		}
 
 		/// <inheritdoc/>
-		public override StringBuilder AppendNullConstructorArg(StringBuilder Builder, bool bIsInitializer)
+		public override StringBuilder AppendNullConstructorArg(StringBuilder builder, bool isInitializer)
 		{
-			Builder.AppendPropertyText(this, UhtPropertyTextType.Construction).Append("()");
-			return Builder;
+			builder.AppendPropertyText(this, UhtPropertyTextType.Construction).Append("()");
+			return builder;
 		}
 
 		/// <inheritdoc/>
-		public override bool SanitizeDefaultValue(IUhtTokenReader DefaultValueReader, StringBuilder InnerDefaultValue)
+		public override bool SanitizeDefaultValue(IUhtTokenReader defaultValueReader, StringBuilder innerDefaultValue)
 		{
 			return false;
 		}
 
 		/// <inheritdoc/>
-		public override bool IsSameType(UhtProperty Other)
+		public override bool IsSameType(UhtProperty other)
 		{
-			if (Other is UhtMulticastDelegateProperty OtherObject)
+			if (other is UhtMulticastDelegateProperty otherObject)
 			{
-				return this.Function == OtherObject.Function;
+				return this.Function == otherObject.Function;
 			}
 			return false;
 		}
diff --git a/Engine/Source/Programs/Shared/EpicGames.UHT/Types/Properties/UhtMulticastInlineDelegateProperty.cs b/Engine/Source/Programs/Shared/EpicGames.UHT/Types/Properties/UhtMulticastInlineDelegateProperty.cs
--- a/Engine/Source/Programs/Shared/EpicGames.UHT/Types/Properties/UhtMulticastInlineDelegateProperty.cs
+++ b/Engine/Source/Programs/Shared/EpicGames.UHT/Types/Properties/UhtMulticastInlineDelegateProperty.cs
@@ -1,7 +1,7 @@
 // Copyright Epic Games, Inc. All Rights Reserved.
 
-using EpicGames.UHT.Tables;
 using System.Text;
+using EpicGames.UHT.Tables;
 
 namespace EpicGames.UHT.Types
 {
@@ -12,46 +12,46 @@ namespace EpicGames.UHT.Types
 	public class UhtMulticastInlineDelegateProperty : UhtMulticastDelegateProperty
 	{
 		/// <inheritdoc/>
-		public override string EngineClassName { get => "MulticastInlineDelegateProperty"; }
+		public override string EngineClassName => "MulticastInlineDelegateProperty";
 
 		/// <summary>
 		/// Construct new property
 		/// </summary>
-		/// <param name="PropertySettings">Property settings</param>
-		/// <param name="Function">Referenced function</param>
-		public UhtMulticastInlineDelegateProperty(UhtPropertySettings PropertySettings, UhtFunction Function) : base(PropertySettings, Function)
+		/// <param name="propertySettings">Property settings</param>
+		/// <param name="function">Referenced function</param>
+		public UhtMulticastInlineDelegateProperty(UhtPropertySettings propertySettings, UhtFunction function) : base(propertySettings, function)
 		{
 		}
 
 		/// <inheritdoc/>
-		public override StringBuilder AppendText(StringBuilder Builder, UhtPropertyTextType TextType, bool bIsTemplateArgument)
+		public override StringBuilder AppendText(StringBuilder builder, UhtPropertyTextType textType, bool isTemplateArgument)
 		{
-			switch (TextType)
+			switch (textType)
 			{
 				case UhtPropertyTextType.FunctionThunkParameterArgType:
-					Builder.Append("FMulticastInlineDelegateProperty");
+					builder.Append("FMulticastInlineDelegateProperty");
 					break;
 
 				default:
-					base.AppendText(Builder, TextType, bIsTemplateArgument);
+					base.AppendText(builder, textType, isTemplateArgument);
 					break;
 			}
-			return Builder;
+			return builder;
 		}
 
 		/// <inheritdoc/>
-		public override StringBuilder AppendMemberDecl(StringBuilder Builder, IUhtPropertyMemberContext Context, string Name, string NameSuffix, int Tabs)
+		public override StringBuilder AppendMemberDecl(StringBuilder builder, IUhtPropertyMemberContext context, string name, string nameSuffix, int tabs)
 		{
-			return AppendMemberDecl(Builder, Context, Name, NameSuffix, Tabs, "FMulticastDelegatePropertyParams");
+			return AppendMemberDecl(builder, context, name, nameSuffix, tabs, "FMulticastDelegatePropertyParams");
 		}
 
 		/// <inheritdoc/>
-		public override StringBuilder AppendMemberDef(StringBuilder Builder, IUhtPropertyMemberContext Context, string Name, string NameSuffix, string? Offset, int Tabs)
+		public override StringBuilder AppendMemberDef(StringBuilder builder, IUhtPropertyMemberContext context, string name, string nameSuffix, string? offset, int tabs)
 		{
-			AppendMemberDefStart(Builder, Context, Name, NameSuffix, Offset, Tabs, "FMulticastDelegatePropertyParams", "UECodeGen_Private::EPropertyGenFlags::InlineMulticastDelegate");
-			AppendMemberDefRef(Builder, Context, this.Function, true);
-			AppendMemberDefEnd(Builder, Context, Name, NameSuffix);
-			return Builder;
+			AppendMemberDefStart(builder, context, name, nameSuffix, offset, tabs, "FMulticastDelegatePropertyParams", "UECodeGen_Private::EPropertyGenFlags::InlineMulticastDelegate");
+			AppendMemberDefRef(builder, context, this.Function, true);
+			AppendMemberDefEnd(builder, context, name, nameSuffix);
+			return builder;
 		}
 	}
 }
diff --git a/Engine/Source/Programs/Shared/EpicGames.UHT/Types/Properties/UhtMulticastSparseDelegateProperty.cs b/Engine/Source/Programs/Shared/EpicGames.UHT/Types/Properties/UhtMulticastSparseDelegateProperty.cs
--- a/Engine/Source/Programs/Shared/EpicGames.UHT/Types/Properties/UhtMulticastSparseDelegateProperty.cs
+++ b/Engine/Source/Programs/Shared/EpicGames.UHT/Types/Properties/UhtMulticastSparseDelegateProperty.cs
@@ -1,7 +1,7 @@
 // Copyright Epic Games, Inc. All Rights Reserved.
 
-using EpicGames.UHT.Tables;
 using System.Text;
+using EpicGames.UHT.Tables;
 
 namespace EpicGames.UHT.Types
 {
@@ -13,46 +13,46 @@ namespace EpicGames.UHT.Types
 	public class UhtMulticastSparseDelegateProperty : UhtMulticastDelegateProperty
 	{
 		/// <inheritdoc/>
-		public override string EngineClassName { get => "MulticastSparseDelegateProperty"; }
+		public override string EngineClassName => "MulticastSparseDelegateProperty";
 
 		/// <summary>
 		/// Construct a new property
 		/// </summary>
-		/// <param name="PropertySettings">Property settings</param>
-		/// <param name="Function">Referenced function</param>
-		public UhtMulticastSparseDelegateProperty(UhtPropertySettings PropertySettings, UhtFunction Function) : base(PropertySettings, Function)
+		/// <param name="propertySettings">Property settings</param>
+		/// <param name="function">Referenced function</param>
+		public UhtMulticastSparseDelegateProperty(UhtPropertySettings propertySettings, UhtFunction function) : base(propertySettings, function)
 		{
 		}
 
 		/// <inheritdoc/>
-		public override StringBuilder AppendText(StringBuilder Builder, UhtPropertyTextType TextType, bool bIsTemplateArgument)
+		public override StringBuilder AppendText(StringBuilder builder, UhtPropertyTextType textType, bool isTemplateArgument)
 		{
-			switch (TextType)
+			switch (textType)
 			{
 				case UhtPropertyTextType.FunctionThunkParameterArgType:
-					Builder.Append("FMulticastSparseDelegateProperty");
+					builder.Append("FMulticastSparseDelegateProperty");
 					break;
 
 				default:
-					base.AppendText(Builder, TextType, bIsTemplateArgument);
+					base.AppendText(builder, textType, isTemplateArgument);
 					break;
 			}
-			return Builder;
+			return builder;
 		}
 
 		/// <inheritdoc/>
-		public override StringBuilder AppendMemberDecl(StringBuilder Builder, IUhtPropertyMemberContext Context, string Name, string NameSuffix, int Tabs)
+		public override StringBuilder AppendMemberDecl(StringBuilder builder, IUhtPropertyMemberContext context, string name, string nameSuffix, int tabs)
 		{
-			return AppendMemberDecl(Builder, Context, Name, NameSuffix, Tabs, "FMulticastDelegatePropertyParams");
+			return AppendMemberDecl(builder, context, name, nameSuffix, tabs, "FMulticastDelegatePropertyParams");
 		}
 
 		/// <inheritdoc/>
-		public override StringBuilder AppendMemberDef(StringBuilder Builder, IUhtPropertyMemberContext Context, string Name, string NameSuffix, string? Offset, int Tabs)
+		public override StringBuilder AppendMemberDef(StringBuilder builder, IUhtPropertyMemberContext context, string name, string nameSuffix, string? offset, int tabs)
 		{
-			AppendMemberDefStart(Builder, Context, Name, NameSuffix, Offset, Tabs, "FMulticastDelegatePropertyParams", "UECodeGen_Private::EPropertyGenFlags::SparseMulticastDelegate");
-			AppendMemberDefRef(Builder, Context, this.Function, true);
-			AppendMemberDefEnd(Builder, Context, Name, NameSuffix);
-			return Builder;
+			AppendMemberDefStart(builder, context, name, nameSuffix, offset, tabs, "FMulticastDelegatePropertyParams", "UECodeGen_Private::EPropertyGenFlags::SparseMulticastDelegate");
+			AppendMemberDefRef(builder, context, this.Function, true);
+			AppendMemberDefEnd(builder, context, name, nameSuffix);
+			return builder;
 		}
 	}
 }
diff --git a/Engine/Source/Programs/Shared/EpicGames.UHT/Types/Properties/UhtNameProperty.cs b/Engine/Source/Programs/Shared/EpicGames.UHT/Types/Properties/UhtNameProperty.cs
--- a/Engine/Source/Programs/Shared/EpicGames.UHT/Types/Properties/UhtNameProperty.cs
+++ b/Engine/Source/Programs/Shared/EpicGames.UHT/Types/Properties/UhtNameProperty.cs
@@ -1,11 +1,11 @@
 // Copyright Epic Games, Inc. All Rights Reserved.
 
+using System.Diagnostics.CodeAnalysis;
+using System.Text;
 using EpicGames.Core;
 using EpicGames.UHT.Tables;
 using EpicGames.UHT.Tokenizer;
 using EpicGames.UHT.Utils;
-using System.Diagnostics.CodeAnalysis;
-using System.Text;
 
 namespace EpicGames.UHT.Types
 {
@@ -18,83 +18,83 @@ namespace EpicGames.UHT.Types
 	public class UhtNameProperty : UhtProperty
 	{
 		/// <inheritdoc/>
-		public override string EngineClassName { get => "NameProperty"; }
+		public override string EngineClassName => "NameProperty";
 
 		/// <inheritdoc/>
-		protected override string CppTypeText { get => "FName"; }
+		protected override string CppTypeText => "FName";
 
 		/// <inheritdoc/>
-		protected override string PGetMacroText { get => "PROPERTY"; }
+		protected override string PGetMacroText => "PROPERTY";
 
 		/// <inheritdoc/>
-		protected override UhtPGetArgumentType PGetTypeArgument { get => UhtPGetArgumentType.EngineClass; }
+		protected override UhtPGetArgumentType PGetTypeArgument => UhtPGetArgumentType.EngineClass;
 
 		/// <summary>
 		/// Construct a new property
 		/// </summary>
-		/// <param name="PropertySettings">Property settings</param>
-		public UhtNameProperty(UhtPropertySettings PropertySettings) : base(PropertySettings)
+		/// <param name="propertySettings">Property settings</param>
+		public UhtNameProperty(UhtPropertySettings propertySettings) : base(propertySettings)
 		{
 			this.PropertyCaps |= UhtPropertyCaps.CanExposeOnSpawn | UhtPropertyCaps.IsParameterSupportedByBlueprint |
 				UhtPropertyCaps.IsMemberSupportedByBlueprint | UhtPropertyCaps.SupportsRigVM;
 		}
 
 		/// <inheritdoc/>
-		public override StringBuilder AppendNullConstructorArg(StringBuilder Builder, bool bIsInitializer)
+		public override StringBuilder AppendNullConstructorArg(StringBuilder builder, bool isInitializer)
 		{
-			Builder.Append("NAME_None");
-			return Builder;
+			builder.Append("NAME_None");
+			return builder;
 		}
 
 		/// <inheritdoc/>
-		public override StringBuilder AppendMemberDecl(StringBuilder Builder, IUhtPropertyMemberContext Context, string Name, string NameSuffix, int Tabs)
+		public override StringBuilder AppendMemberDecl(StringBuilder builder, IUhtPropertyMemberContext context, string name, string nameSuffix, int tabs)
 		{
-			return AppendMemberDecl(Builder, Context, Name, NameSuffix, Tabs, "FNamePropertyParams");
+			return AppendMemberDecl(builder, context, name, nameSuffix, tabs, "FNamePropertyParams");
 		}
 
 		/// <inheritdoc/>
-		public override StringBuilder AppendMemberDef(StringBuilder Builder, IUhtPropertyMemberContext Context, string Name, string NameSuffix, string? Offset, int Tabs)
+		public override StringBuilder AppendMemberDef(StringBuilder builder, IUhtPropertyMemberContext context, string name, string nameSuffix, string? offset, int tabs)
 		{
-			AppendMemberDefStart(Builder, Context, Name, NameSuffix, Offset, Tabs, "FNamePropertyParams", "UECodeGen_Private::EPropertyGenFlags::Name");
-			AppendMemberDefEnd(Builder, Context, Name, NameSuffix);
-			return Builder;
+			AppendMemberDefStart(builder, context, name, nameSuffix, offset, tabs, "FNamePropertyParams", "UECodeGen_Private::EPropertyGenFlags::Name");
+			AppendMemberDefEnd(builder, context, name, nameSuffix);
+			return builder;
 		}
 
 		/// <inheritdoc/>
-		public override bool SanitizeDefaultValue(IUhtTokenReader DefaultValueReader, StringBuilder InnerDefaultValue)
+		public override bool SanitizeDefaultValue(IUhtTokenReader defaultValueReader, StringBuilder innerDefaultValue)
 		{
-			if (DefaultValueReader.TryOptional("NAME_None"))
+			if (defaultValueReader.TryOptional("NAME_None"))
 			{
-				InnerDefaultValue.Append("None");
+				innerDefaultValue.Append("None");
 			}
-			else if (DefaultValueReader.TryOptional("FName"))
+			else if (defaultValueReader.TryOptional("FName"))
 			{
-				DefaultValueReader.Require('(');
-				StringView Value = DefaultValueReader.GetWrappedConstString();
-				DefaultValueReader.Require(')');
-				InnerDefaultValue.Append(Value);
+				defaultValueReader.Require('(');
+				StringView value = defaultValueReader.GetWrappedConstString();
+				defaultValueReader.Require(')');
+				innerDefaultValue.Append(value);
 			}
 			else
 			{
-				StringView Value = DefaultValueReader.GetWrappedConstString();
-				InnerDefaultValue.Append(Value);
+				StringView value = defaultValueReader.GetWrappedConstString();
+				innerDefaultValue.Append(value);
 			}
 			return true;
 		}
 
 		/// <inheritdoc/>
-		public override bool IsSameType(UhtProperty Other)
+		public override bool IsSameType(UhtProperty other)
 		{
-			return Other is UhtNameProperty;
+			return other is UhtNameProperty;
 		}
 
 		#region Keyword
 		[UhtPropertyType(Keyword = "FName", Options = UhtPropertyTypeOptions.Simple | UhtPropertyTypeOptions.Immediate)]
 		[SuppressMessage("CodeQuality", "IDE0051:Remove unused private members", Justification = "Attribute accessed method")]
 		[SuppressMessage("Style", "IDE0060:Remove unused parameter", Justification = "Attribute accessed method")]
-		private static UhtProperty? NameProperty(UhtPropertyResolvePhase ResolvePhase, UhtPropertySettings PropertySettings, IUhtTokenReader TokenReader, UhtToken MatchedToken)
+		private static UhtProperty? NameProperty(UhtPropertyResolvePhase resolvePhase, UhtPropertySettings propertySettings, IUhtTokenReader tokenReader, UhtToken matchedToken)
 		{
-			return new UhtNameProperty(PropertySettings);
+			return new UhtNameProperty(propertySettings);
 		}
 		#endregion
 	}
diff --git a/Engine/Source/Programs/Shared/EpicGames.UHT/Types/Properties/UhtNumericProperty.cs b/Engine/Source/Programs/Shared/EpicGames.UHT/Types/Properties/UhtNumericProperty.cs
--- a/Engine/Source/Programs/Shared/EpicGames.UHT/Types/Properties/UhtNumericProperty.cs
+++ b/Engine/Source/Programs/Shared/EpicGames.UHT/Types/Properties/UhtNumericProperty.cs
@@ -11,10 +11,10 @@ namespace EpicGames.UHT.Types
 	public abstract class UhtNumericProperty : UhtProperty
 	{
 		/// <inheritdoc/>
-		protected override string PGetMacroText { get => "PROPERTY"; }
+		protected override string PGetMacroText => "PROPERTY";
 
 		/// <inheritdoc/>
-		protected override UhtPGetArgumentType PGetTypeArgument { get => UhtPGetArgumentType.EngineClass; }
+		protected override UhtPGetArgumentType PGetTypeArgument => UhtPGetArgumentType.EngineClass;
 
 		/// <summary>
 		/// Describes the integer as either being sized on unsized
@@ -24,19 +24,19 @@ namespace EpicGames.UHT.Types
 		/// <summary>
 		/// Construct new property
 		/// </summary>
-		/// <param name="PropertySettings">Property settings</param>
-		/// <param name="IntType">Type of integer</param>
-		protected UhtNumericProperty(UhtPropertySettings PropertySettings, UhtPropertyIntType IntType) : base(PropertySettings)
+		/// <param name="propertySettings">Property settings</param>
+		/// <param name="intType">Type of integer</param>
+		protected UhtNumericProperty(UhtPropertySettings propertySettings, UhtPropertyIntType intType) : base(propertySettings)
 		{
-			this.IntType = IntType;
+			this.IntType = intType;
 			this.PropertyCaps |= UhtPropertyCaps.RequiresNullConstructorArg | UhtPropertyCaps.SupportsRigVM;
 		}
 
 		/// <inheritdoc/>
-		public override StringBuilder AppendNullConstructorArg(StringBuilder Builder, bool bIsInitializer)
+		public override StringBuilder AppendNullConstructorArg(StringBuilder builder, bool isInitializer)
 		{
-			Builder.Append('0');
-			return Builder;
+			builder.Append('0');
+			return builder;
 		}
 	}
 }
diff --git a/Engine/Source/Programs/Shared/EpicGames.UHT/Types/Properties/UhtObjectProperty.cs b/Engine/Source/Programs/Shared/EpicGames.UHT/Types/Properties/UhtObjectProperty.cs
--- a/Engine/Source/Programs/Shared/EpicGames.UHT/Types/Properties/UhtObjectProperty.cs
+++ b/Engine/Source/Programs/Shared/EpicGames.UHT/Types/Properties/UhtObjectProperty.cs
@@ -1,9 +1,9 @@
 // Copyright Epic Games, Inc. All Rights Reserved.
 
+using System.Text;
 using EpicGames.Core;
 using EpicGames.UHT.Tables;
 using EpicGames.UHT.Utils;
-using System.Text;
 
 namespace EpicGames.UHT.Types
 {
@@ -17,28 +17,28 @@ namespace EpicGames.UHT.Types
 		public override string EngineClassName => "ObjectProperty";
 
 		/// <inheritdoc/>
-		protected override UhtPGetArgumentType PGetTypeArgument { get => UhtPGetArgumentType.TypeText; }
+		protected override UhtPGetArgumentType PGetTypeArgument => UhtPGetArgumentType.TypeText;
 
 		/// <summary>
 		/// Construct a new property
 		/// </summary>
-		/// <param name="PropertySettings">Property settings</param>
-		/// <param name="Class">Referenced class</param>
-		/// <param name="MetaClass">Optional reference class (used by class properties)</param>
-		/// <param name="ExtraFlags">Extra flags to add to the property</param>
-		public UhtObjectProperty(UhtPropertySettings PropertySettings, UhtClass Class, UhtClass? MetaClass = null, EPropertyFlags ExtraFlags = EPropertyFlags.None)
-			: base(PropertySettings, Class, MetaClass)
+		/// <param name="propertySettings">Property settings</param>
+		/// <param name="classObj">Referenced class</param>
+		/// <param name="metaClass">Optional reference class (used by class properties)</param>
+		/// <param name="extraFlags">Extra flags to add to the property</param>
+		public UhtObjectProperty(UhtPropertySettings propertySettings, UhtClass classObj, UhtClass? metaClass = null, EPropertyFlags extraFlags = EPropertyFlags.None)
+			: base(propertySettings, classObj, metaClass)
 		{
-			this.PropertyFlags |= ExtraFlags;
+			this.PropertyFlags |= extraFlags;
 			this.PropertyCaps |= UhtPropertyCaps.RequiresNullConstructorArg | UhtPropertyCaps.CanBeInstanced | UhtPropertyCaps.CanExposeOnSpawn |
 				UhtPropertyCaps.IsParameterSupportedByBlueprint | UhtPropertyCaps.IsMemberSupportedByBlueprint;
 		}
 
 		/// <inheritdoc/>
-		protected override bool ResolveSelf(UhtResolvePhase Phase)
+		protected override bool ResolveSelf(UhtResolvePhase phase)
 		{
-			bool bResults = base.ResolveSelf(Phase);
-			switch (Phase)
+			bool results = base.ResolveSelf(phase);
+			switch (phase)
 			{
 				case UhtResolvePhase.Final:
 					if (this.Class.HierarchyHasAnyClassFlags(EClassFlags.DefaultToInstanced))
@@ -48,67 +48,67 @@ namespace EpicGames.UHT.Types
 					}
 					break;
 			}
-			return bResults;
+			return results;
 		}
 
 		/// <inheritdoc/>
-		public override bool ScanForInstancedReferenced(bool bDeepScan)
+		public override bool ScanForInstancedReferenced(bool deepScan)
 		{
 			return this.Class.HierarchyHasAnyClassFlags(EClassFlags.DefaultToInstanced);
 		}
 
 		/// <inheritdoc/>
-		public override StringBuilder AppendText(StringBuilder Builder, UhtPropertyTextType TextType, bool bIsTemplateArgument)
+		public override StringBuilder AppendText(StringBuilder builder, UhtPropertyTextType textType, bool isTemplateArgument)
 		{
-			switch (TextType)
+			switch (textType)
 			{
 				case UhtPropertyTextType.FunctionThunkRetVal:
 					if (this.PropertyFlags.HasAnyFlags(EPropertyFlags.ConstParm))
 					{
-						Builder.Append("const ");
+						builder.Append("const ");
 					}
-					Builder.Append(this.Class.SourceName).Append('*');
+					builder.Append(this.Class.SourceName).Append('*');
 					break;
 
 				case UhtPropertyTextType.FunctionThunkParameterArgType:
-					Builder.Append(this.Class.SourceName);
+					builder.Append(this.Class.SourceName);
 					break;
 
 				default:
-					Builder.Append(this.Class.SourceName).Append('*');
+					builder.Append(this.Class.SourceName).Append('*');
 					break;
 			}
-			return Builder;
+			return builder;
 		}
 
 		/// <inheritdoc/>
-		public override StringBuilder AppendMemberDecl(StringBuilder Builder, IUhtPropertyMemberContext Context, string Name, string NameSuffix, int Tabs)
+		public override StringBuilder AppendMemberDecl(StringBuilder builder, IUhtPropertyMemberContext context, string name, string nameSuffix, int tabs)
 		{
-			return AppendMemberDecl(Builder, Context, Name, NameSuffix, Tabs, "FObjectPropertyParams");
+			return AppendMemberDecl(builder, context, name, nameSuffix, tabs, "FObjectPropertyParams");
 		}
 
 		/// <inheritdoc/>
-		public override StringBuilder AppendMemberDef(StringBuilder Builder, IUhtPropertyMemberContext Context, string Name, string NameSuffix, string? Offset, int Tabs)
+		public override StringBuilder AppendMemberDef(StringBuilder builder, IUhtPropertyMemberContext context, string name, string nameSuffix, string? offset, int tabs)
 		{
-			AppendMemberDefStart(Builder, Context, Name, NameSuffix, Offset, Tabs, "FObjectPropertyParams", "UECodeGen_Private::EPropertyGenFlags::Object");
-			AppendMemberDefRef(Builder, Context, this.Class, false);
-			AppendMemberDefEnd(Builder, Context, Name, NameSuffix);
-			return Builder;
+			AppendMemberDefStart(builder, context, name, nameSuffix, offset, tabs, "FObjectPropertyParams", "UECodeGen_Private::EPropertyGenFlags::Object");
+			AppendMemberDefRef(builder, context, this.Class, false);
+			AppendMemberDefEnd(builder, context, name, nameSuffix);
+			return builder;
 		}
 
 		/// <inheritdoc/>
-		public override StringBuilder AppendNullConstructorArg(StringBuilder Builder, bool bIsInitializer)
+		public override StringBuilder AppendNullConstructorArg(StringBuilder builder, bool isInitializer)
 		{
-			Builder.Append("NULL");
-			return Builder;
+			builder.Append("NULL");
+			return builder;
 		}
 
 		/// <inheritdoc/>
-		public override bool IsSameType(UhtProperty Other)
+		public override bool IsSameType(UhtProperty other)
 		{
-			if (Other is UhtObjectProperty OtherObject)
+			if (other is UhtObjectProperty otherObject)
 			{
-				return this.Class == OtherObject.Class && this.MetaClass == OtherObject.MetaClass;
+				return this.Class == otherObject.Class && this.MetaClass == otherObject.MetaClass;
 			}
 			return false;
 		}
diff --git a/Engine/Source/Programs/Shared/EpicGames.UHT/Types/Properties/UhtObjectPropertyBase.cs b/Engine/Source/Programs/Shared/EpicGames.UHT/Types/Properties/UhtObjectPropertyBase.cs
--- a/Engine/Source/Programs/Shared/EpicGames.UHT/Types/Properties/UhtObjectPropertyBase.cs
+++ b/Engine/Source/Programs/Shared/EpicGames.UHT/Types/Properties/UhtObjectPropertyBase.cs
@@ -1,14 +1,14 @@
 // Copyright Epic Games, Inc. All Rights Reserved.
 
-using EpicGames.Core;
-using EpicGames.UHT.Tables;
-using EpicGames.UHT.Tokenizer;
-using EpicGames.UHT.Utils;
 using System;
 using System.Collections.Generic;
 using System.Diagnostics.CodeAnalysis;
 using System.Text;
 using System.Text.Json.Serialization;
+using EpicGames.Core;
+using EpicGames.UHT.Tables;
+using EpicGames.UHT.Tokenizer;
+using EpicGames.UHT.Utils;
 
 namespace EpicGames.UHT.Types
 {
@@ -19,13 +19,13 @@ namespace EpicGames.UHT.Types
 	public abstract class UhtObjectPropertyBase : UhtProperty
 	{
 		/// <inheritdoc/>
-		public override string EngineClassName { get => "ObjectPropertyBase"; }
+		public override string EngineClassName => "ObjectPropertyBase";
 
 		/// <inheritdoc/>
-		protected override string CppTypeText { get => "Object"; }
+		protected override string CppTypeText => "Object";
 
 		/// <inheritdoc/>
-		protected override string PGetMacroText { get => "OBJECT"; }
+		protected override string PGetMacroText => "OBJECT";
 
 		/// <summary>
 		/// Referenced UCLASS
@@ -42,20 +42,20 @@ namespace EpicGames.UHT.Types
 		/// <summary>
 		/// Construct a property
 		/// </summary>
-		/// <param name="PropertySettings">Property settings</param>
-		/// <param name="Class">Referenced UCLASS</param>
-		/// <param name="MetaClass">Referenced UCLASS used by class properties</param>
-		protected UhtObjectPropertyBase(UhtPropertySettings PropertySettings, UhtClass Class, UhtClass? MetaClass) : base(PropertySettings)
+		/// <param name="propertySettings">Property settings</param>
+		/// <param name="classObj">Referenced UCLASS</param>
+		/// <param name="metaClass">Referenced UCLASS used by class properties</param>
+		protected UhtObjectPropertyBase(UhtPropertySettings propertySettings, UhtClass classObj, UhtClass? metaClass) : base(propertySettings)
 		{
-			this.Class = Class;
-			this.MetaClass = MetaClass;
+			this.Class = classObj;
+			this.MetaClass = metaClass;
 
 			// This applies to EVERYTHING including raw pointer
 			// Imply const if it's a parameter that is a pointer to a const class
 			// NOTE: We shouldn't be automatically adding const param because in some cases with functions and blueprint native event, the 
 			// generated code won't match.  For now, just disabled the auto add in that case and check for the error in the validation code.
 			// Otherwise, the user is not warned and they will get compile errors.
-			if (PropertySettings.PropertyCategory != UhtPropertyCategory.Member && this.Class.ClassFlags.HasAnyFlags(EClassFlags.Const) && !PropertySettings.Options.HasAnyFlags(UhtPropertyOptions.NoAutoConst))
+			if (propertySettings.PropertyCategory != UhtPropertyCategory.Member && this.Class.ClassFlags.HasAnyFlags(EClassFlags.Const) && !propertySettings.Options.HasAnyFlags(UhtPropertyOptions.NoAutoConst))
 			{
 				this.PropertyFlags |= EPropertyFlags.ConstParm;
 			}
@@ -68,10 +68,10 @@ namespace EpicGames.UHT.Types
 		}
 
 		/// <inheritdoc/>
-		protected override bool ResolveSelf(UhtResolvePhase Phase)
+		protected override bool ResolveSelf(UhtResolvePhase phase)
 		{
-			bool bResults = base.ResolveSelf(Phase);
-			switch (Phase)
+			bool results = base.ResolveSelf(phase);
+			switch (phase)
 			{
 				case UhtResolvePhase.Final:
 					if (this.Class.HierarchyHasAnyClassFlags(EClassFlags.DefaultToInstanced))
@@ -80,20 +80,20 @@ namespace EpicGames.UHT.Types
 					}
 					break;
 			}
-			return bResults;
+			return results;
 		}
 
 		/// <inheritdoc/>
-		public override bool ScanForInstancedReferenced(bool bDeepScan)
+		public override bool ScanForInstancedReferenced(bool deepScan)
 		{
 			return !this.DisallowPropertyFlags.HasAnyFlags(EPropertyFlags.InstancedReference) && this.Class.HierarchyHasAnyClassFlags(EClassFlags.DefaultToInstanced);
 		}
 
 		/// <inheritdoc/>
-		public override void CollectReferencesInternal(IUhtReferenceCollector Collector, bool bTemplateProperty)
+		public override void CollectReferencesInternal(IUhtReferenceCollector collector, bool templateProperty)
 		{
-			Collector.AddCrossModuleReference(this.Class, false);
-			Collector.AddCrossModuleReference(this.MetaClass, false);
+			collector.AddCrossModuleReference(this.Class, false);
+			collector.AddCrossModuleReference(this.MetaClass, false);
 		}
 
 		/// <inheritdoc/>
@@ -113,14 +113,14 @@ namespace EpicGames.UHT.Types
 		}
 
 		/// <inheritdoc/>
-		public override bool SanitizeDefaultValue(IUhtTokenReader DefaultValueReader, StringBuilder InnerDefaultValue)
+		public override bool SanitizeDefaultValue(IUhtTokenReader defaultValueReader, StringBuilder innerDefaultValue)
 		{
-			int Value;
-			if (DefaultValueReader.TryOptional("NULL") ||
-				DefaultValueReader.TryOptional("nullptr") ||
-				(DefaultValueReader.TryOptionalConstInt(out Value) && Value == 0))
+			int value;
+			if (defaultValueReader.TryOptional("NULL") ||
+				defaultValueReader.TryOptional("nullptr") ||
+				(defaultValueReader.TryOptionalConstInt(out value) && value == 0))
 			{
-				InnerDefaultValue.Append("None");
+				innerDefaultValue.Append("None");
 				return true;
 			}
 			return false;
@@ -134,9 +134,9 @@ namespace EpicGames.UHT.Types
 		}
 
 		/// <inheritdoc/>
-		public override bool MustBeConstArgument([NotNullWhen(true)] out UhtType? ErrorType)
+		public override bool MustBeConstArgument([NotNullWhen(true)] out UhtType? errorType)
 		{
-			ErrorType = this.Class;
+			errorType = this.Class;
 			return this.Class.ClassFlags.HasAnyFlags(EClassFlags.Const);
 		}
 
@@ -144,138 +144,137 @@ namespace EpicGames.UHT.Types
 		/// <summary>
 		/// Parse a template type
 		/// </summary>
-		/// <param name="PropertySettings">Property settings</param>
-		/// <param name="TokenReader">Token reader</param>
-		/// <param name="MatchedToken">Token matched for type</param>
-		/// <param name="bReturnUInterface">If true, return the UInterface instead of the type listed</param>
+		/// <param name="propertySettings">Property settings</param>
+		/// <param name="tokenReader">Token reader</param>
+		/// <param name="matchedToken">Token matched for type</param>
+		/// <param name="returnUInterface">If true, return the UInterface instead of the type listed</param>
 		/// <returns>Referenced class</returns>
-		public static UhtClass? ParseTemplateObject(UhtPropertySettings PropertySettings, IUhtTokenReader TokenReader, UhtToken MatchedToken, bool bReturnUInterface)
+		public static UhtClass? ParseTemplateObject(UhtPropertySettings propertySettings, IUhtTokenReader tokenReader, UhtToken matchedToken, bool returnUInterface)
 		{
-			UhtSession Session = PropertySettings.Outer.Session;
-			if (TokenReader.TryOptional("const"))
+			UhtSession session = propertySettings.Outer.Session;
+			if (tokenReader.TryOptional("const"))
 			{
-				PropertySettings.MetaData.Add(UhtNames.NativeConst, "");
+				propertySettings.MetaData.Add(UhtNames.NativeConst, "");
 			}
-			if (!TokenReader.SkipExpectedType(MatchedToken.Value, PropertySettings.PropertyCategory == UhtPropertyCategory.Member))
+			if (!tokenReader.SkipExpectedType(matchedToken.Value, propertySettings.PropertyCategory == UhtPropertyCategory.Member))
 			{
 				return null;
 			}
-			bool bIsNativeConstTemplateArg = false;
-			UhtToken Identifier = new UhtToken();
-			TokenReader
+			bool isNativeConstTemplateArg = false;
+			UhtToken identifier = new UhtToken();
+			tokenReader
 				.Require('<')
-				.Optional("const", () => { bIsNativeConstTemplateArg = true; })
+				.Optional("const", () => { isNativeConstTemplateArg = true; })
 				.Optional("class")
-				.RequireIdentifier((ref UhtToken Token) => { Identifier = Token; })
-				.Optional("const", () => { bIsNativeConstTemplateArg = true; })
+				.RequireIdentifier((ref UhtToken token) => { identifier = token; })
+				.Optional("const", () => { isNativeConstTemplateArg = true; })
 				.Require('>');
 
-			if (bIsNativeConstTemplateArg)
+			if (isNativeConstTemplateArg)
 			{
-				PropertySettings.MetaData.Add(UhtNames.NativeConstTemplateArg, "");
+				propertySettings.MetaData.Add(UhtNames.NativeConstTemplateArg, "");
 			}
-			Session.Config!.RedirectTypeIdentifier(ref Identifier);
-			UhtClass? Return = PropertySettings.Outer.FindType(UhtFindOptions.SourceName | UhtFindOptions.Class, ref Identifier, TokenReader) as UhtClass;
-			if (Return != null && Return.AlternateObject != null && bReturnUInterface)
+			session.Config!.RedirectTypeIdentifier(ref identifier);
+			UhtClass? returnClass = propertySettings.Outer.FindType(UhtFindOptions.SourceName | UhtFindOptions.Class, ref identifier, tokenReader) as UhtClass;
+			if (returnClass != null && returnClass.AlternateObject != null && returnUInterface)
 			{
-				Return = Return.AlternateObject as UhtClass;
+				returnClass = returnClass.AlternateObject as UhtClass;
 			}
-			return Return;
+			return returnClass;
 		}
 
 		/// <summary>
 		/// Parse a template type
 		/// </summary>
-		/// <param name="PropertySettings">Property settings</param>
-		/// <param name="TokenReader">Token reader</param>
-		/// <param name="MatchedToken">Token matched for type</param>
+		/// <param name="propertySettings">Property settings</param>
+		/// <param name="tokenReader">Token reader</param>
+		/// <param name="matchedToken">Token matched for type</param>
 		/// <returns>Referenced class</returns>
-		public static UhtClass? ParseTemplateClass(UhtPropertySettings PropertySettings, IUhtTokenReader TokenReader, UhtToken MatchedToken)
+		public static UhtClass? ParseTemplateClass(UhtPropertySettings propertySettings, IUhtTokenReader tokenReader, UhtToken matchedToken)
 		{
-			UhtSession Session = PropertySettings.Outer.Session;
-			UhtToken Identifier = new UhtToken();
+			UhtSession session = propertySettings.Outer.Session;
+			UhtToken identifier = new UhtToken();
 
-			if (TokenReader.TryOptional("const"))
+			if (tokenReader.TryOptional("const"))
 			{
-				PropertySettings.MetaData.Add(UhtNames.NativeConst, "");
+				propertySettings.MetaData.Add(UhtNames.NativeConst, "");
 			}
 
-			TokenReader.Optional("class");
+			tokenReader.Optional("class");
 
-			if (!TokenReader.SkipExpectedType(MatchedToken.Value, PropertySettings.PropertyCategory == UhtPropertyCategory.Member))
+			if (!tokenReader.SkipExpectedType(matchedToken.Value, propertySettings.PropertyCategory == UhtPropertyCategory.Member))
 			{
 				return null;
 			}
 
-			TokenReader
+			tokenReader
 				.Require('<')
 				.Optional("class")
-				.RequireIdentifier((ref UhtToken Token) => { Identifier = Token; })
+				.RequireIdentifier((ref UhtToken token) => { identifier = token; })
 				.Require('>');
 
-			Session.Config!.RedirectTypeIdentifier(ref Identifier);
-			UhtClass? Return = PropertySettings.Outer.FindType(UhtFindOptions.SourceName | UhtFindOptions.Class, ref Identifier, TokenReader) as UhtClass;
-			if (Return != null && Return.AlternateObject != null)
+			session.Config!.RedirectTypeIdentifier(ref identifier);
+			UhtClass? returnClass = propertySettings.Outer.FindType(UhtFindOptions.SourceName | UhtFindOptions.Class, ref identifier, tokenReader) as UhtClass;
+			if (returnClass != null && returnClass.AlternateObject != null)
 			{
-				Return = Return.AlternateObject as UhtClass;
+				returnClass = returnClass.AlternateObject as UhtClass;
 			}
-			return Return;
+			return returnClass;
 		}
 
 		/// <summary>
 		/// Log point usage warning/error
 		/// </summary>
-		/// <param name="PropertySettings">Property settings</param>
-		/// <param name="EngineBehavior">Expected behavior for engine types</param>
-		/// <param name="NonEngineBehavior">Expected behavior for non-engine types</param>
-		/// <param name="PointerTypeDesc">Description of the pointer type</param>
-		/// <param name="TokenReader">Token reader for type being parsed</param>
-		/// <param name="TypeStartPos">Starting character position of the type</param>
-		/// <param name="AlternativeTypeDesc">Suggested alternate declaration</param>
+		/// <param name="propertySettings">Property settings</param>
+		/// <param name="engineBehavior">Expected behavior for engine types</param>
+		/// <param name="nonEngineBehavior">Expected behavior for non-engine types</param>
+		/// <param name="pointerTypeDesc">Description of the pointer type</param>
+		/// <param name="tokenReader">Token reader for type being parsed</param>
+		/// <param name="typeStartPos">Starting character position of the type</param>
+		/// <param name="alternativeTypeDesc">Suggested alternate declaration</param>
 		/// <exception cref="UhtIceException">Thrown if the behavior type is unexpected</exception>
-		public static void ConditionalLogPointerUsage(UhtPropertySettings PropertySettings, UhtPointerMemberBehavior EngineBehavior, UhtPointerMemberBehavior NonEngineBehavior,
-			string PointerTypeDesc, IUhtTokenReader TokenReader, int TypeStartPos, string? AlternativeTypeDesc)
+		public static void ConditionalLogPointerUsage(UhtPropertySettings propertySettings, UhtPointerMemberBehavior engineBehavior, UhtPointerMemberBehavior nonEngineBehavior,
+			string pointerTypeDesc, IUhtTokenReader tokenReader, int typeStartPos, string? alternativeTypeDesc)
 		{
-			if (PropertySettings.PropertyCategory != UhtPropertyCategory.Member)
+			if (propertySettings.PropertyCategory != UhtPropertyCategory.Member)
 			{
 				return;
 			}
 
-			UhtPackage Package = PropertySettings.Outer.Package;
-			bool bEngineBehavior = Package.bIsPartOfEngine && !Package.bIsPlugin;
-			UhtPointerMemberBehavior Behavior = bEngineBehavior ? EngineBehavior : NonEngineBehavior;
+			UhtPackage package = propertySettings.Outer.Package;
+			UhtPointerMemberBehavior behavior = package.IsPartOfEngine && !package.IsPlugin ? engineBehavior : nonEngineBehavior;
 
-			if (Behavior == UhtPointerMemberBehavior.AllowSilently)
+			if (behavior == UhtPointerMemberBehavior.AllowSilently)
 			{
 				return;
 			}
 
-			string Type = TokenReader.GetStringView(TypeStartPos, TokenReader.InputPos - TypeStartPos).ToString();
-			Type = Type.Replace("\n", " ", StringComparison.Ordinal);
-			Type = Type.Replace("\r", "", StringComparison.Ordinal);
-			Type = Type.Replace("\t", " ", StringComparison.Ordinal);
+			string type = tokenReader.GetStringView(typeStartPos, tokenReader.InputPos - typeStartPos).ToString();
+			type = type.Replace("\n", " ", StringComparison.Ordinal);
+			type = type.Replace("\r", "", StringComparison.Ordinal);
+			type = type.Replace("\t", " ", StringComparison.Ordinal);
 
-			switch (Behavior)
+			switch (behavior)
 			{
 				case UhtPointerMemberBehavior.Disallow:
-					if (!string.IsNullOrEmpty(AlternativeTypeDesc))
+					if (!String.IsNullOrEmpty(alternativeTypeDesc))
 					{
-						TokenReader.LogError($"{PointerTypeDesc} usage in member declaration detected [[[{Type}]]].  This is disallowed for the target/module, consider {AlternativeTypeDesc} as an alternative.");
+						tokenReader.LogError($"{pointerTypeDesc} usage in member declaration detected [[[{type}]]].  This is disallowed for the target/module, consider {alternativeTypeDesc} as an alternative.");
 					}
 					else
 					{
-						TokenReader.LogError($"{PointerTypeDesc} usage in member declaration detected [[[{Type}]]].");
+						tokenReader.LogError($"{pointerTypeDesc} usage in member declaration detected [[[{type}]]].");
 					}
 					break;
 
 				case UhtPointerMemberBehavior.AllowAndLog:
-					if (!string.IsNullOrEmpty(AlternativeTypeDesc))
+					if (!String.IsNullOrEmpty(alternativeTypeDesc))
 					{
-						TokenReader.LogTrace($"{PointerTypeDesc} usage in member declaration detected [[[{Type}]]].  Consider {AlternativeTypeDesc} as an alternative.");
+						tokenReader.LogTrace($"{pointerTypeDesc} usage in member declaration detected [[[{type}]]].  Consider {alternativeTypeDesc} as an alternative.");
 					}
 					else
 					{
-						TokenReader.LogTrace("{PointerTypeDesc} usage in member declaration detected [[[{Type}]]].");
+						tokenReader.LogTrace("{PointerTypeDesc} usage in member declaration detected [[[{Type}]]].");
 					}
 					break;
 
diff --git a/Engine/Source/Programs/Shared/EpicGames.UHT/Types/Properties/UhtObjectPtrProperty.cs b/Engine/Source/Programs/Shared/EpicGames.UHT/Types/Properties/UhtObjectPtrProperty.cs
--- a/Engine/Source/Programs/Shared/EpicGames.UHT/Types/Properties/UhtObjectPtrProperty.cs
+++ b/Engine/Source/Programs/Shared/EpicGames.UHT/Types/Properties/UhtObjectPtrProperty.cs
@@ -1,11 +1,11 @@
 // Copyright Epic Games, Inc. All Rights Reserved.
 
+using System.Diagnostics.CodeAnalysis;
+using System.Text;
 using EpicGames.Core;
 using EpicGames.UHT.Tables;
 using EpicGames.UHT.Tokenizer;
 using EpicGames.UHT.Utils;
-using System.Diagnostics.CodeAnalysis;
-using System.Text;
 
 namespace EpicGames.UHT.Types
 {
@@ -17,86 +17,86 @@ namespace EpicGames.UHT.Types
 	public class UhtObjectPtrProperty : UhtObjectProperty
 	{
 		/// <inheritdoc/>
-		public override string EngineClassName { get => "ObjectPtrProperty"; }
+		public override string EngineClassName => "ObjectPtrProperty";
 
 		/// <inheritdoc/>
-		protected override string CppTypeText { get => "ObjectPtr"; }
+		protected override string CppTypeText => "ObjectPtr";
 
 		/// <inheritdoc/>
-		protected override string PGetMacroText { get => "OBJECTPTR"; }
+		protected override string PGetMacroText => "OBJECTPTR";
 
 		/// <inheritdoc/>
-		protected override UhtPGetArgumentType PGetTypeArgument { get => UhtPGetArgumentType.TypeText; }
+		protected override UhtPGetArgumentType PGetTypeArgument => UhtPGetArgumentType.TypeText;
 
 		/// <summary>
 		/// Construct a new property
 		/// </summary>
-		/// <param name="PropertySettings">Property settings</param>
-		/// <param name="Class">Referenced class</param>
-		/// <param name="ExtraFlags">Extra property flags to apply</param>
-		public UhtObjectPtrProperty(UhtPropertySettings PropertySettings, UhtClass Class, EPropertyFlags ExtraFlags = EPropertyFlags.None)
-			: base(PropertySettings, Class)
+		/// <param name="propertySettings">Property settings</param>
+		/// <param name="classObj">Referenced class</param>
+		/// <param name="extraFlags">Extra property flags to apply</param>
+		public UhtObjectPtrProperty(UhtPropertySettings propertySettings, UhtClass classObj, EPropertyFlags extraFlags = EPropertyFlags.None)
+			: base(propertySettings, classObj)
 		{
-			this.PropertyFlags |= ExtraFlags | EPropertyFlags.UObjectWrapper;
+			this.PropertyFlags |= extraFlags | EPropertyFlags.UObjectWrapper;
 			this.PropertyCaps |= UhtPropertyCaps.PassCppArgsByRef;
 		}
 
 		/// <inheritdoc/>
-		public override StringBuilder AppendText(StringBuilder Builder, UhtPropertyTextType TextType, bool bIsTemplateArgument)
+		public override StringBuilder AppendText(StringBuilder builder, UhtPropertyTextType textType, bool isTemplateArgument)
 		{
-			switch (TextType)
+			switch (textType)
 			{
 				case UhtPropertyTextType.GetterSetterArg:
-					if (bIsTemplateArgument)
+					if (isTemplateArgument)
 					{
-						Builder.Append("TObjectPtr<").Append(this.Class.SourceName).Append('>');
+						builder.Append("TObjectPtr<").Append(this.Class.SourceName).Append('>');
 					}
 					else
 					{
-						Builder.Append(this.Class.SourceName).Append('*');
+						builder.Append(this.Class.SourceName).Append('*');
 					}
 					break;
 
 				case UhtPropertyTextType.FunctionThunkRetVal:
 					if (this.PropertyFlags.HasAnyFlags(EPropertyFlags.ConstParm))
 					{
-						Builder.Append("const ");
+						builder.Append("const ");
 					}
-					Builder.Append("TObjectPtr<").Append(this.Class.SourceName).Append('>');
+					builder.Append("TObjectPtr<").Append(this.Class.SourceName).Append('>');
 					break;
 
 				default:
-					Builder.Append("TObjectPtr<").Append(this.Class.SourceName).Append('>');
+					builder.Append("TObjectPtr<").Append(this.Class.SourceName).Append('>');
 					break;
 			}
-			return Builder;
+			return builder;
 		}
-		
+
 		/// <inheritdoc/>
-		public override StringBuilder AppendMemberDecl(StringBuilder Builder, IUhtPropertyMemberContext Context, string Name, string NameSuffix, int Tabs)
+		public override StringBuilder AppendMemberDecl(StringBuilder builder, IUhtPropertyMemberContext context, string name, string nameSuffix, int tabs)
 		{
-			return AppendMemberDecl(Builder, Context, Name, NameSuffix, Tabs, "FObjectPtrPropertyParams");
+			return AppendMemberDecl(builder, context, name, nameSuffix, tabs, "FObjectPtrPropertyParams");
 		}
 
 		/// <inheritdoc/>
-		public override StringBuilder AppendMemberDef(StringBuilder Builder, IUhtPropertyMemberContext Context, string Name, string NameSuffix, string? Offset, int Tabs)
+		public override StringBuilder AppendMemberDef(StringBuilder builder, IUhtPropertyMemberContext context, string name, string nameSuffix, string? offset, int tabs)
 		{
-			AppendMemberDefStart(Builder, Context, Name, NameSuffix, Offset, Tabs, "FObjectPtrPropertyParams", "UECodeGen_Private::EPropertyGenFlags::Object | UECodeGen_Private::EPropertyGenFlags::ObjectPtr");
-			AppendMemberDefRef(Builder, Context, this.Class, false);
-			AppendMemberDefEnd(Builder, Context, Name, NameSuffix);
-			return Builder;
+			AppendMemberDefStart(builder, context, name, nameSuffix, offset, tabs, "FObjectPtrPropertyParams", "UECodeGen_Private::EPropertyGenFlags::Object | UECodeGen_Private::EPropertyGenFlags::ObjectPtr");
+			AppendMemberDefRef(builder, context, this.Class, false);
+			AppendMemberDefEnd(builder, context, name, nameSuffix);
+			return builder;
 		}
 
 		/// <inheritdoc/>
-		public override void Validate(UhtStruct OuterStruct, UhtProperty OutermostProperty, UhtValidationOptions Options)
+		public override void Validate(UhtStruct outerStruct, UhtProperty outermostProperty, UhtValidationOptions options)
 		{
-			base.Validate(OuterStruct, OutermostProperty, Options);
+			base.Validate(outerStruct, outermostProperty, options);
 
 			// UFunctions with a smart pointer as input parameter wont compile anyway, because of missing P_GET_... macro.
 			// UFunctions with a smart pointer as return type will crash when called via blueprint, because they are not supported in VM.
 			if (this.PropertyCategory == UhtPropertyCategory.RegularParameter || this.PropertyCategory == UhtPropertyCategory.ReplicatedParameter)
 			{
-				OuterStruct.LogError("UFunctions cannot take a TObjectPtr as a parameter.");
+				outerStruct.LogError("UFunctions cannot take a TObjectPtr as a parameter.");
 			}
 		}
 
@@ -104,28 +104,28 @@ namespace EpicGames.UHT.Types
 		[UhtPropertyType(Keyword = "TObjectPtr")]
 		[SuppressMessage("CodeQuality", "IDE0051:Remove unused private members", Justification = "Attribute accessed method")]
 		[SuppressMessage("Style", "IDE0060:Remove unused parameter", Justification = "Attribute accessed method")]
-		private static UhtProperty? ObjectPtrProperty(UhtPropertyResolvePhase ResolvePhase, UhtPropertySettings PropertySettings, IUhtTokenReader TokenReader, UhtToken MatchedToken)
+		private static UhtProperty? ObjectPtrProperty(UhtPropertyResolvePhase resolvePhase, UhtPropertySettings propertySettings, IUhtTokenReader tokenReader, UhtToken matchedToken)
 		{
-			UhtSession Session = PropertySettings.Outer.Session;
-			int TypeStartPos = TokenReader.PeekToken().InputStartPos;
+			UhtSession session = propertySettings.Outer.Session;
+			int typeStartPos = tokenReader.PeekToken().InputStartPos;
 
-			UhtClass? PropertyClass = ParseTemplateObject(PropertySettings, TokenReader, MatchedToken, true);
-			if (PropertyClass == null)
+			UhtClass? propertyClass = ParseTemplateObject(propertySettings, tokenReader, matchedToken, true);
+			if (propertyClass == null)
 			{
 				return null;
 			}
 
-			ConditionalLogPointerUsage(PropertySettings, Session.Config!.EngineObjectPtrMemberBehavior, 
-				Session.Config!.NonEngineObjectPtrMemberBehavior, "ObjectPtr", TokenReader, TypeStartPos, null);
+			ConditionalLogPointerUsage(propertySettings, session.Config!.EngineObjectPtrMemberBehavior,
+				session.Config!.NonEngineObjectPtrMemberBehavior, "ObjectPtr", tokenReader, typeStartPos, null);
 
-			if (PropertyClass.IsChildOf(PropertyClass.Session.UClass))
+			if (propertyClass.IsChildOf(propertyClass.Session.UClass))
 			{
 				// UObject specifies that there is no limiter
-				return new UhtClassPtrProperty(PropertySettings, PropertyClass, PropertyClass.Session.UObject);
+				return new UhtClassPtrProperty(propertySettings, propertyClass, propertyClass.Session.UObject);
 			}
 			else
 			{
-				return new UhtObjectPtrProperty(PropertySettings, PropertyClass);
+				return new UhtObjectPtrProperty(propertySettings, propertyClass);
 			}
 		}
 		#endregion
diff --git a/Engine/Source/Programs/Shared/EpicGames.UHT/Types/Properties/UhtSetProperty.cs b/Engine/Source/Programs/Shared/EpicGames.UHT/Types/Properties/UhtSetProperty.cs
--- a/Engine/Source/Programs/Shared/EpicGames.UHT/Types/Properties/UhtSetProperty.cs
+++ b/Engine/Source/Programs/Shared/EpicGames.UHT/Types/Properties/UhtSetProperty.cs
@@ -1,13 +1,13 @@
 // Copyright Epic Games, Inc. All Rights Reserved.
 
+using System.Collections.Generic;
+using System.Diagnostics.CodeAnalysis;
+using System.Text;
 using EpicGames.Core;
 using EpicGames.UHT.Parsers;
 using EpicGames.UHT.Tables;
 using EpicGames.UHT.Tokenizer;
 using EpicGames.UHT.Utils;
-using System.Collections.Generic;
-using System.Diagnostics.CodeAnalysis;
-using System.Text;
 
 namespace EpicGames.UHT.Types
 {
@@ -19,26 +19,26 @@ namespace EpicGames.UHT.Types
 	public class UhtSetProperty : UhtContainerBaseProperty
 	{
 		/// <inheritdoc/>
-		public override string EngineClassName { get => "SetProperty"; }
+		public override string EngineClassName => "SetProperty";
 
 		/// <inheritdoc/>
-		protected override string CppTypeText { get => "TSet"; }
+		protected override string CppTypeText => "TSet";
 
 		/// <inheritdoc/>
-		protected override string PGetMacroText { get => "TSET"; }
+		protected override string PGetMacroText => "TSET";
 
 		/// <inheritdoc/>
-		protected override UhtPGetArgumentType PGetTypeArgument { get => UhtPGetArgumentType.TypeText; }
+		protected override UhtPGetArgumentType PGetTypeArgument => UhtPGetArgumentType.TypeText;
 
 		/// <summary>
 		/// Construct a new property
 		/// </summary>
-		/// <param name="PropertySettings">Property settings</param>
-		/// <param name="Value">Property key</param>
-		public UhtSetProperty(UhtPropertySettings PropertySettings, UhtProperty Value) : base(PropertySettings, Value)
+		/// <param name="propertySettings">Property settings</param>
+		/// <param name="value">Property key</param>
+		public UhtSetProperty(UhtPropertySettings propertySettings, UhtProperty value) : base(propertySettings, value)
 		{
 			// If the creation of the value property set more flags, then copy those flags to ourselves
-			this.PropertyFlags |= Value.PropertyFlags & EPropertyFlags.UObjectWrapper;
+			this.PropertyFlags |= value.PropertyFlags & EPropertyFlags.UObjectWrapper;
 
 			this.PropertyCaps |= UhtPropertyCaps.PassCppArgsByRef;
 			this.PropertyCaps &= ~(UhtPropertyCaps.CanBeContainerValue | UhtPropertyCaps.CanBeContainerKey);
@@ -52,26 +52,26 @@ namespace EpicGames.UHT.Types
 		}
 
 		/// <inheritdoc/>
-		protected override bool ResolveSelf(UhtResolvePhase Phase)
+		protected override bool ResolveSelf(UhtResolvePhase phase)
 		{
-			bool bResults = base.ResolveSelf(Phase);
-			switch (Phase)
+			bool results = base.ResolveSelf(phase);
+			switch (phase)
 			{
 				case UhtResolvePhase.Final:
-					this.PropertyFlags |= ResolveAndReturnNewFlags(this.ValueProperty, Phase);
+					this.PropertyFlags |= ResolveAndReturnNewFlags(this.ValueProperty, phase);
 					this.MetaData.Add(this.ValueProperty.MetaData);
 					this.ValueProperty.MetaData.Clear();
 					this.ValueProperty.PropertyFlags = this.PropertyFlags & EPropertyFlags.PropagateToSetElement;
 					PropagateFlagsFromInnerAndHandlePersistentInstanceMetadata(this, this.MetaData, this.ValueProperty);
 					break;
 			}
-			return bResults;
+			return results;
 		}
 
 		/// <inheritdoc/>
-		public override void CollectReferencesInternal(IUhtReferenceCollector Collector, bool bTemplateProperty)
+		public override void CollectReferencesInternal(IUhtReferenceCollector collector, bool templateProperty)
 		{
-			this.ValueProperty.CollectReferencesInternal(Collector, true);
+			this.ValueProperty.CollectReferencesInternal(collector, true);
 		}
 
 		/// <inheritdoc/>
@@ -83,123 +83,123 @@ namespace EpicGames.UHT.Types
 		/// <inheritdoc/>
 		public override IEnumerable<UhtType> EnumerateReferencedTypes()
 		{
-			foreach (UhtType Type in this.ValueProperty.EnumerateReferencedTypes())
+			foreach (UhtType type in this.ValueProperty.EnumerateReferencedTypes())
 			{
-				yield return Type;
+				yield return type;
 			}
 		}
 
 		/// <inheritdoc/>
-		public override StringBuilder AppendText(StringBuilder Builder, UhtPropertyTextType TextType, bool bIsTemplateArgument)
+		public override StringBuilder AppendText(StringBuilder builder, UhtPropertyTextType textType, bool isTemplateArgument)
 		{
-			switch (TextType)
+			switch (textType)
 			{
 				case UhtPropertyTextType.SparseShort:
-					Builder.Append("TSet");
+					builder.Append("TSet");
 					break;
 
 				case UhtPropertyTextType.FunctionThunkParameterArgType:
-					Builder.AppendFunctionThunkParameterArrayType(this.ValueProperty);
+					builder.AppendFunctionThunkParameterArrayType(this.ValueProperty);
 					break;
 
 				default:
-					Builder.Append("TSet<").AppendPropertyText(this.ValueProperty, TextType, true);
-					if (Builder[Builder.Length - 1] == '>')
+					builder.Append("TSet<").AppendPropertyText(this.ValueProperty, textType, true);
+					if (builder[^1] == '>')
 					{
 						// if our internal property type is a template class, add a space between the closing brackets b/c VS.NET cannot parse this correctly
-						Builder.Append(' ');
+						builder.Append(' ');
 					}
-					Builder.Append('>');
+					builder.Append('>');
 					break;
 			}
-			return Builder;
+			return builder;
 		}
 
 		/// <inheritdoc/>
-		public override StringBuilder AppendMemberDecl(StringBuilder Builder, IUhtPropertyMemberContext Context, string Name, string NameSuffix, int Tabs)
+		public override StringBuilder AppendMemberDecl(StringBuilder builder, IUhtPropertyMemberContext context, string name, string nameSuffix, int tabs)
 		{
-			Builder.AppendMemberDecl(this.ValueProperty, Context, Name, GetNameSuffix(NameSuffix, "_ElementProp"), Tabs);
-			return AppendMemberDecl(Builder, Context, Name, NameSuffix, Tabs, "FSetPropertyParams");
+			builder.AppendMemberDecl(this.ValueProperty, context, name, GetNameSuffix(nameSuffix, "_ElementProp"), tabs);
+			return AppendMemberDecl(builder, context, name, nameSuffix, tabs, "FSetPropertyParams");
 		}
 
 		/// <inheritdoc/>
-		public override StringBuilder AppendMemberDef(StringBuilder Builder, IUhtPropertyMemberContext Context, string Name, string NameSuffix, string? Offset, int Tabs)
+		public override StringBuilder AppendMemberDef(StringBuilder builder, IUhtPropertyMemberContext context, string name, string nameSuffix, string? offset, int tabs)
 		{
-			Builder.AppendMemberDef(this.ValueProperty, Context, Name, GetNameSuffix(NameSuffix, "_ElementProp"), "0", Tabs);
-			Builder.AppendMetaDataDef(this, Context, Name, NameSuffix, Tabs);
+			builder.AppendMemberDef(this.ValueProperty, context, name, GetNameSuffix(nameSuffix, "_ElementProp"), "0", tabs);
+			builder.AppendMetaDataDef(this, context, name, nameSuffix, tabs);
 
-			if (this.ValueProperty is UhtStructProperty StructProperty)
+			if (this.ValueProperty is UhtStructProperty structProperty)
 			{
-				Builder
-					.AppendTabs(Tabs)
+				builder
+					.AppendTabs(tabs)
 					.Append("static_assert(TModels<CGetTypeHashable, ")
-					.Append(StructProperty.ScriptStruct.SourceName)
+					.Append(structProperty.ScriptStruct.SourceName)
 					.Append(">::Value, \"The structure '")
-					.Append(StructProperty.ScriptStruct.SourceName)
+					.Append(structProperty.ScriptStruct.SourceName)
 					.Append("' is used in a TSet but does not have a GetValueTypeHash defined\");\r\n");
 			}
 
-			AppendMemberDefStart(Builder, Context, Name, NameSuffix, Offset, Tabs, "FSetPropertyParams", "UECodeGen_Private::EPropertyGenFlags::Set", false);
-			AppendMemberDefEnd(Builder, Context, Name, NameSuffix);
-			return Builder;
+			AppendMemberDefStart(builder, context, name, nameSuffix, offset, tabs, "FSetPropertyParams", "UECodeGen_Private::EPropertyGenFlags::Set", false);
+			AppendMemberDefEnd(builder, context, name, nameSuffix);
+			return builder;
 		}
 
 		/// <inheritdoc/>
-		public override StringBuilder AppendMemberPtr(StringBuilder Builder, IUhtPropertyMemberContext Context, string Name, string NameSuffix, int Tabs)
+		public override StringBuilder AppendMemberPtr(StringBuilder builder, IUhtPropertyMemberContext context, string name, string nameSuffix, int tabs)
 		{
-			Builder.AppendMemberPtr(this.ValueProperty, Context, Name, GetNameSuffix(NameSuffix, "_ElementProp"), Tabs);
-			base.AppendMemberPtr(Builder, Context, Name, NameSuffix, Tabs);
-			return Builder;
+			builder.AppendMemberPtr(this.ValueProperty, context, name, GetNameSuffix(nameSuffix, "_ElementProp"), tabs);
+			base.AppendMemberPtr(builder, context, name, nameSuffix, tabs);
+			return builder;
 		}
 
 		/// <inheritdoc/>
-		public override void AppendObjectHashes(StringBuilder Builder, int StartingLength, IUhtPropertyMemberContext Context)
+		public override void AppendObjectHashes(StringBuilder builder, int startingLength, IUhtPropertyMemberContext context)
 		{
-			this.ValueProperty.AppendObjectHashes(Builder, StartingLength, Context);
+			this.ValueProperty.AppendObjectHashes(builder, startingLength, context);
 		}
 
 		/// <inheritdoc/>
-		public override StringBuilder AppendNullConstructorArg(StringBuilder Builder, bool bIsInitializer)
+		public override StringBuilder AppendNullConstructorArg(StringBuilder builder, bool isInitializer)
 		{
-			Builder.AppendPropertyText(this, UhtPropertyTextType.Construction).Append("()");
-			return Builder;
+			builder.AppendPropertyText(this, UhtPropertyTextType.Construction).Append("()");
+			return builder;
 		}
 
 		/// <inheritdoc/>
-		public override bool SanitizeDefaultValue(IUhtTokenReader DefaultValueReader, StringBuilder InnerDefaultValue)
+		public override bool SanitizeDefaultValue(IUhtTokenReader defaultValueReader, StringBuilder innerDefaultValue)
 		{
 			return false;
 		}
 
 		/// <inheritdoc/>
-		public override bool IsSameType(UhtProperty Other)
+		public override bool IsSameType(UhtProperty other)
 		{
-			if (Other is UhtSetProperty OtherSet)
+			if (other is UhtSetProperty otherSet)
 			{
-				return this.ValueProperty.IsSameType(OtherSet.ValueProperty);
+				return this.ValueProperty.IsSameType(otherSet.ValueProperty);
 			}
 			return false;
 		}
 
 		///<inheritdoc/>
-		public override bool ValidateStructPropertyOkForNet(UhtProperty ReferencingProperty)
+		public override bool ValidateStructPropertyOkForNet(UhtProperty referencingProperty)
 		{
 			if (!this.PropertyFlags.HasAnyFlags(EPropertyFlags.RepSkip))
 			{
-				ReferencingProperty.LogError($"Sets are not supported for Replication or RPCs.  Set '{this.SourceName}' in '{this.Outer?.SourceName}'.  Origin '{ReferencingProperty.SourceName}'");
+				referencingProperty.LogError($"Sets are not supported for Replication or RPCs.  Set '{this.SourceName}' in '{this.Outer?.SourceName}'.  Origin '{referencingProperty.SourceName}'");
 				return false;
 			}
 			return true;
 		}
 
 		/// <inheritdoc/>
-		protected override void ValidateFunctionArgument(UhtFunction Function, UhtValidationOptions Options)
+		protected override void ValidateFunctionArgument(UhtFunction function, UhtValidationOptions options)
 		{
-			base.ValidateFunctionArgument(Function, Options);
+			base.ValidateFunctionArgument(function, options);
 
-			if (Function.FunctionFlags.HasAnyFlags(EFunctionFlags.Net))
+			if (function.FunctionFlags.HasAnyFlags(EFunctionFlags.Net))
 			{
-				if (!Function.FunctionFlags.HasAnyFlags(EFunctionFlags.NetRequest | EFunctionFlags.NetResponse))
+				if (!function.FunctionFlags.HasAnyFlags(EFunctionFlags.NetRequest | EFunctionFlags.NetResponse))
 				{
 					this.LogError("Sets are not supported in an RPC.");
 				}
@@ -209,45 +209,45 @@ namespace EpicGames.UHT.Types
 		#region Keyword
 		[UhtPropertyType(Keyword = "TSet")]
 		[SuppressMessage("CodeQuality", "IDE0051:Remove unused private members", Justification = "Attribute accessed method")]
-		private static UhtProperty? SetProperty(UhtPropertyResolvePhase ResolvePhase, UhtPropertySettings PropertySettings, IUhtTokenReader TokenReader, UhtToken MatchedToken)
+		private static UhtProperty? SetProperty(UhtPropertyResolvePhase resolvePhase, UhtPropertySettings propertySettings, IUhtTokenReader tokenReader, UhtToken matchedToken)
 		{
-			using (var TokenContext = new UhtMessageContext("TSet"))
+			using (UhtMessageContext tokenContext = new UhtMessageContext("TSet"))
 			{
-				if (!TokenReader.SkipExpectedType(MatchedToken.Value, PropertySettings.PropertyCategory == UhtPropertyCategory.Member))
+				if (!tokenReader.SkipExpectedType(matchedToken.Value, propertySettings.PropertyCategory == UhtPropertyCategory.Member))
 				{
 					return null;
 				}
-				TokenReader.Require('<');
+				tokenReader.Require('<');
 
 				// Parse the value type
-				UhtProperty? Value = UhtPropertyParser.ParseTemplateParam(ResolvePhase, PropertySettings, "Value", TokenReader);
-				if (Value == null)
+				UhtProperty? value = UhtPropertyParser.ParseTemplateParam(resolvePhase, propertySettings, "Value", tokenReader);
+				if (value == null)
 				{
 					return null;
 				}
 
-				if (!Value.PropertyCaps.HasAnyFlags(UhtPropertyCaps.CanBeContainerKey))
+				if (!value.PropertyCaps.HasAnyFlags(UhtPropertyCaps.CanBeContainerKey))
 				{
-					TokenReader.LogError($"The type \'{Value.GetUserFacingDecl()}\' can not be used as a key in a TSet");
+					tokenReader.LogError($"The type \'{value.GetUserFacingDecl()}\' can not be used as a key in a TSet");
 				}
 
-				if (PropertySettings.PropertyFlags.HasAnyFlags(EPropertyFlags.Net))
+				if (propertySettings.PropertyFlags.HasAnyFlags(EPropertyFlags.Net))
 				{
-					TokenReader.LogError("Replicated sets are not supported.");
+					tokenReader.LogError("Replicated sets are not supported.");
 				}
 
-				if (TokenReader.TryOptional(','))
+				if (tokenReader.TryOptional(','))
 				{
 					// If we found a comma, read the next thing, assume it's a keyfuncs, and report that
-					UhtToken KeyFuncToken = TokenReader.GetIdentifier();
-					throw new UhtException(TokenReader, $"Found '{KeyFuncToken.Value}' - explicit KeyFuncs are not supported in TSet properties.");
+					UhtToken keyFuncToken = tokenReader.GetIdentifier();
+					throw new UhtException(tokenReader, $"Found '{keyFuncToken.Value}' - explicit KeyFuncs are not supported in TSet properties.");
 				}
 
-				TokenReader.Require('>');
+				tokenReader.Require('>');
 
 				//@TODO: Prevent sparse delegate types from being used in a container
 
-				return new UhtSetProperty(PropertySettings, Value);
+				return new UhtSetProperty(propertySettings, value);
 			}
 		}
 		#endregion
diff --git a/Engine/Source/Programs/Shared/EpicGames.UHT/Types/Properties/UhtSoftClassProperty.cs b/Engine/Source/Programs/Shared/EpicGames.UHT/Types/Properties/UhtSoftClassProperty.cs
--- a/Engine/Source/Programs/Shared/EpicGames.UHT/Types/Properties/UhtSoftClassProperty.cs
+++ b/Engine/Source/Programs/Shared/EpicGames.UHT/Types/Properties/UhtSoftClassProperty.cs
@@ -1,10 +1,10 @@
 // Copyright Epic Games, Inc. All Rights Reserved.
 
+using System.Diagnostics.CodeAnalysis;
+using System.Text;
 using EpicGames.UHT.Tables;
 using EpicGames.UHT.Tokenizer;
 using EpicGames.UHT.Utils;
-using System.Diagnostics.CodeAnalysis;
-using System.Text;
 
 namespace EpicGames.UHT.Types
 {
@@ -17,25 +17,25 @@ namespace EpicGames.UHT.Types
 	public class UhtSoftClassProperty : UhtSoftObjectProperty
 	{
 		/// <inheritdoc/>
-		public override string EngineClassName { get => "SoftClassProperty"; }
+		public override string EngineClassName => "SoftClassProperty";
 
 		/// <inheritdoc/>
-		protected override string CppTypeText { get => "SoftClassPtr"; }
+		protected override string CppTypeText => "SoftClassPtr";
 
 		/// <inheritdoc/>
-		protected override string PGetMacroText { get => "SOFTCLASS"; }
+		protected override string PGetMacroText => "SOFTCLASS";
 
 		/// <inheritdoc/>
-		protected override UhtPGetArgumentType PGetTypeArgument { get => UhtPGetArgumentType.TypeText; }
+		protected override UhtPGetArgumentType PGetTypeArgument => UhtPGetArgumentType.TypeText;
 
 		/// <summary>
 		/// Construct a new class property
 		/// </summary>
-		/// <param name="PropertySettings">Property setting</param>
-		/// <param name="Class">Referenced class (Always UClass)</param>
-		/// <param name="MetaClass">Referenced meta class</param>
-		public UhtSoftClassProperty(UhtPropertySettings PropertySettings, UhtClass Class, UhtClass MetaClass)
-			: base(PropertySettings, Class, MetaClass)
+		/// <param name="propertySettings">Property setting</param>
+		/// <param name="classObj">Referenced class (Always UClass)</param>
+		/// <param name="metaClass">Referenced meta class</param>
+		public UhtSoftClassProperty(UhtPropertySettings propertySettings, UhtClass classObj, UhtClass metaClass)
+			: base(propertySettings, classObj, metaClass)
 		{
 		}
 
@@ -46,52 +46,52 @@ namespace EpicGames.UHT.Types
 		}
 
 		/// <inheritdoc/>
-		public override void CollectReferencesInternal(IUhtReferenceCollector Collector, bool bTemplateProperty)
+		public override void CollectReferencesInternal(IUhtReferenceCollector collector, bool templateProperty)
 		{
-			Collector.AddCrossModuleReference(this.MetaClass, false);
+			collector.AddCrossModuleReference(this.MetaClass, false);
 		}
 
 		/// <inheritdoc/>
-		public override StringBuilder AppendText(StringBuilder Builder, UhtPropertyTextType TextType, bool bIsTemplateArgument)
+		public override StringBuilder AppendText(StringBuilder builder, UhtPropertyTextType textType, bool isTemplateArgument)
 		{
-			switch (TextType)
+			switch (textType)
 			{
 				default:
-					Builder.Append("TSoftClassPtr<").Append(this.MetaClass?.SourceName).Append("> ");
+					builder.Append("TSoftClassPtr<").Append(this.MetaClass?.SourceName).Append("> ");
 					break;
 			}
-			return Builder;
+			return builder;
 		}
 
 		/// <inheritdoc/>
-		public override StringBuilder AppendMemberDecl(StringBuilder Builder, IUhtPropertyMemberContext Context, string Name, string NameSuffix, int Tabs)
+		public override StringBuilder AppendMemberDecl(StringBuilder builder, IUhtPropertyMemberContext context, string name, string nameSuffix, int tabs)
 		{
-			return AppendMemberDecl(Builder, Context, Name, NameSuffix, Tabs, "FSoftClassPropertyParams");
+			return AppendMemberDecl(builder, context, name, nameSuffix, tabs, "FSoftClassPropertyParams");
 		}
 
 		/// <inheritdoc/>
-		public override StringBuilder AppendMemberDef(StringBuilder Builder, IUhtPropertyMemberContext Context, string Name, string NameSuffix, string? Offset, int Tabs)
+		public override StringBuilder AppendMemberDef(StringBuilder builder, IUhtPropertyMemberContext context, string name, string nameSuffix, string? offset, int tabs)
 		{
-			AppendMemberDefStart(Builder, Context, Name, NameSuffix, Offset, Tabs, "FSoftClassPropertyParams", "UECodeGen_Private::EPropertyGenFlags::SoftClass");
-			AppendMemberDefRef(Builder, Context, this.MetaClass, false);
-			AppendMemberDefEnd(Builder, Context, Name, NameSuffix);
-			return Builder;
+			AppendMemberDefStart(builder, context, name, nameSuffix, offset, tabs, "FSoftClassPropertyParams", "UECodeGen_Private::EPropertyGenFlags::SoftClass");
+			AppendMemberDefRef(builder, context, this.MetaClass, false);
+			AppendMemberDefEnd(builder, context, name, nameSuffix);
+			return builder;
 		}
 
 		#region Keyword
 		[UhtPropertyType(Keyword = "TSoftClassPtr")]
 		[SuppressMessage("CodeQuality", "IDE0051:Remove unused private members", Justification = "Attribute accessed method")]
 		[SuppressMessage("Style", "IDE0060:Remove unused parameter", Justification = "Attribute accessed method")]
-		private static UhtProperty? SoftClassPtrProperty(UhtPropertyResolvePhase ResolvePhase, UhtPropertySettings PropertySettings, IUhtTokenReader TokenReader, UhtToken MatchedToken)
+		private static UhtProperty? SoftClassPtrProperty(UhtPropertyResolvePhase resolvePhase, UhtPropertySettings propertySettings, IUhtTokenReader tokenReader, UhtToken matchedToken)
 		{
-			UhtClass? MetaClass = ParseTemplateClass(PropertySettings, TokenReader, MatchedToken);
-			if (MetaClass == null)
+			UhtClass? metaClass = ParseTemplateClass(propertySettings, tokenReader, matchedToken);
+			if (metaClass == null)
 			{
 				return null;
 			}
 
 			// With TSubclassOf, MetaClass is used as a class limiter.  
-			return new UhtSoftClassProperty(PropertySettings, MetaClass.Session.UClass, MetaClass);
+			return new UhtSoftClassProperty(propertySettings, metaClass.Session.UClass, metaClass);
 		}
 		#endregion
 	}
