More coding standard fixes

diff --git a/Engine/Source/Programs/Shared/EpicGames.UHT/Types/Properties/UhtSoftObjectProperty.cs b/Engine/Source/Programs/Shared/EpicGames.UHT/Types/Properties/UhtSoftObjectProperty.cs
--- a/Engine/Source/Programs/Shared/EpicGames.UHT/Types/Properties/UhtSoftObjectProperty.cs
+++ b/Engine/Source/Programs/Shared/EpicGames.UHT/Types/Properties/UhtSoftObjectProperty.cs
@@ -1,11 +1,11 @@
 // Copyright Epic Games, Inc. All Rights Reserved.
 
+using System.Diagnostics.CodeAnalysis;
+using System.Text;
 using EpicGames.Core;
 using EpicGames.UHT.Tables;
 using EpicGames.UHT.Tokenizer;
 using EpicGames.UHT.Utils;
-using System.Diagnostics.CodeAnalysis;
-using System.Text;
 
 namespace EpicGames.UHT.Types
 {
@@ -17,25 +17,25 @@ namespace EpicGames.UHT.Types
 	public class UhtSoftObjectProperty : UhtObjectPropertyBase
 	{
 		/// <inheritdoc/>
-		public override string EngineClassName { get => "SoftObjectProperty"; }
+		public override string EngineClassName => "SoftObjectProperty";
 
 		/// <inheritdoc/>
-		protected override string CppTypeText { get => "SoftObjectPtr"; }
+		protected override string CppTypeText => "SoftObjectPtr";
 
 		/// <inheritdoc/>
-		protected override string PGetMacroText { get => "SOFTOBJECT"; }
+		protected override string PGetMacroText => "SOFTOBJECT";
 
 		/// <inheritdoc/>
-		protected override UhtPGetArgumentType PGetTypeArgument { get => UhtPGetArgumentType.TypeText; }
+		protected override UhtPGetArgumentType PGetTypeArgument => UhtPGetArgumentType.TypeText;
 
 		/// <summary>
 		/// Construct a new property
 		/// </summary>
-		/// <param name="PropertySettings">Property settings</param>
-		/// <param name="Class">UCLASS being referenced</param>
-		/// <param name="MetaClass">Optional meta class (used by SoftClassProperty)</param>
-		public UhtSoftObjectProperty(UhtPropertySettings PropertySettings, UhtClass Class, UhtClass? MetaClass = null)
-			: base(PropertySettings, Class, MetaClass)
+		/// <param name="propertySettings">Property settings</param>
+		/// <param name="classObj">UCLASS being referenced</param>
+		/// <param name="metaClass">Optional meta class (used by SoftClassProperty)</param>
+		public UhtSoftObjectProperty(UhtPropertySettings propertySettings, UhtClass classObj, UhtClass? metaClass = null)
+			: base(propertySettings, classObj, metaClass)
 		{
 			this.PropertyFlags |= EPropertyFlags.UObjectWrapper;
 			this.PropertyCaps |= UhtPropertyCaps.PassCppArgsByRef | UhtPropertyCaps.RequiresNullConstructorArg | UhtPropertyCaps.CanExposeOnSpawn |
@@ -43,45 +43,45 @@ namespace EpicGames.UHT.Types
 		}
 
 		/// <inheritdoc/>
-		public override StringBuilder AppendNullConstructorArg(StringBuilder Builder, bool bIsInitializer)
+		public override StringBuilder AppendNullConstructorArg(StringBuilder builder, bool isInitializer)
 		{
-			Builder.Append("NULL");
-			return Builder;
+			builder.Append("NULL");
+			return builder;
 		}
 
 		/// <inheritdoc/>
-		public override StringBuilder AppendText(StringBuilder Builder, UhtPropertyTextType TextType, bool bIsTemplateArgument)
+		public override StringBuilder AppendText(StringBuilder builder, UhtPropertyTextType textType, bool isTemplateArgument)
 		{
-			switch (TextType)
+			switch (textType)
 			{
 				default:
-					Builder.Append("TSoftObjectPtr<").Append(this.Class.SourceName).Append('>');
+					builder.Append("TSoftObjectPtr<").Append(this.Class.SourceName).Append('>');
 					break;
 			}
-			return Builder;
+			return builder;
 		}
 
 		/// <inheritdoc/>
-		public override StringBuilder AppendMemberDecl(StringBuilder Builder, IUhtPropertyMemberContext Context, string Name, string NameSuffix, int Tabs)
+		public override StringBuilder AppendMemberDecl(StringBuilder builder, IUhtPropertyMemberContext context, string name, string nameSuffix, int tabs)
 		{
-			return AppendMemberDecl(Builder, Context, Name, NameSuffix, Tabs, "FSoftObjectPropertyParams");
+			return AppendMemberDecl(builder, context, name, nameSuffix, tabs, "FSoftObjectPropertyParams");
 		}
 
 		/// <inheritdoc/>
-		public override StringBuilder AppendMemberDef(StringBuilder Builder, IUhtPropertyMemberContext Context, string Name, string NameSuffix, string? Offset, int Tabs)
+		public override StringBuilder AppendMemberDef(StringBuilder builder, IUhtPropertyMemberContext context, string name, string nameSuffix, string? offset, int tabs)
 		{
-			AppendMemberDefStart(Builder, Context, Name, NameSuffix, Offset, Tabs, "FSoftObjectPropertyParams", "UECodeGen_Private::EPropertyGenFlags::SoftObject");
-			AppendMemberDefRef(Builder, Context, this.Class, false);
-			AppendMemberDefEnd(Builder, Context, Name, NameSuffix);
-			return Builder;
+			AppendMemberDefStart(builder, context, name, nameSuffix, offset, tabs, "FSoftObjectPropertyParams", "UECodeGen_Private::EPropertyGenFlags::SoftObject");
+			AppendMemberDefRef(builder, context, this.Class, false);
+			AppendMemberDefEnd(builder, context, name, nameSuffix);
+			return builder;
 		}
 
 		/// <inheritdoc/>
-		public override bool IsSameType(UhtProperty Other)
+		public override bool IsSameType(UhtProperty other)
 		{
-			if (Other is UhtSoftObjectProperty OtherObject)
+			if (other is UhtSoftObjectProperty otherObject)
 			{
-				return this.Class == OtherObject.Class && this.MetaClass == OtherObject.MetaClass;
+				return this.Class == otherObject.Class && this.MetaClass == otherObject.MetaClass;
 			}
 			return false;
 		}
@@ -90,20 +90,20 @@ namespace EpicGames.UHT.Types
 		[UhtPropertyType(Keyword = "TSoftObjectPtr")]
 		[SuppressMessage("CodeQuality", "IDE0051:Remove unused private members", Justification = "Attribute accessed method")]
 		[SuppressMessage("Style", "IDE0060:Remove unused parameter", Justification = "Attribute accessed method")]
-		private static UhtProperty? SoftObjectPtrProperty(UhtPropertyResolvePhase ResolvePhase, UhtPropertySettings PropertySettings, IUhtTokenReader TokenReader, UhtToken MatchedToken)
+		private static UhtProperty? SoftObjectPtrProperty(UhtPropertyResolvePhase resolvePhase, UhtPropertySettings propertySettings, IUhtTokenReader tokenReader, UhtToken matchedToken)
 		{
-			UhtClass? PropertyClass = UhtObjectPropertyBase.ParseTemplateObject(PropertySettings, TokenReader, MatchedToken, true);
-			if (PropertyClass == null)
+			UhtClass? propertyClass = UhtObjectPropertyBase.ParseTemplateObject(propertySettings, tokenReader, matchedToken, true);
+			if (propertyClass == null)
 			{
 				return null;
 			}
 
-			if (PropertyClass.IsChildOf(PropertyClass.Session.UClass))
+			if (propertyClass.IsChildOf(propertyClass.Session.UClass))
 			{
-				TokenReader.LogError("Class variables cannot be stored in TSoftObjectPtr, use TSoftClassPtr instead.");
+				tokenReader.LogError("Class variables cannot be stored in TSoftObjectPtr, use TSoftClassPtr instead.");
 			}
 
-			return new UhtSoftObjectProperty(PropertySettings, PropertyClass);
+			return new UhtSoftObjectProperty(propertySettings, propertyClass);
 		}
 		#endregion
 	}
diff --git a/Engine/Source/Programs/Shared/EpicGames.UHT/Types/Properties/UhtStrProperty.cs b/Engine/Source/Programs/Shared/EpicGames.UHT/Types/Properties/UhtStrProperty.cs
--- a/Engine/Source/Programs/Shared/EpicGames.UHT/Types/Properties/UhtStrProperty.cs
+++ b/Engine/Source/Programs/Shared/EpicGames.UHT/Types/Properties/UhtStrProperty.cs
@@ -1,11 +1,11 @@
 // Copyright Epic Games, Inc. All Rights Reserved.
 
+using System.Diagnostics.CodeAnalysis;
+using System.Text;
 using EpicGames.Core;
 using EpicGames.UHT.Tables;
 using EpicGames.UHT.Tokenizer;
 using EpicGames.UHT.Utils;
-using System.Diagnostics.CodeAnalysis;
-using System.Text;
 
 namespace EpicGames.UHT.Types
 {
@@ -18,82 +18,82 @@ namespace EpicGames.UHT.Types
 	public class UhtStrProperty : UhtProperty
 	{
 		/// <inheritdoc/>
-		public override string EngineClassName { get => "StrProperty"; }
+		public override string EngineClassName => "StrProperty";
 
 		/// <inheritdoc/>
-		protected override string CppTypeText { get => "FString"; }
+		protected override string CppTypeText => "FString";
 
 		/// <inheritdoc/>
-		protected override string PGetMacroText { get => "PROPERTY"; }
+		protected override string PGetMacroText => "PROPERTY";
 
 		/// <inheritdoc/>
-		protected override UhtPGetArgumentType PGetTypeArgument { get => UhtPGetArgumentType.EngineClass; }
+		protected override UhtPGetArgumentType PGetTypeArgument => UhtPGetArgumentType.EngineClass;
 
 		/// <summary>
 		/// Construct a new property
 		/// </summary>
-		/// <param name="PropertySettings">Property settings</param>
-		public UhtStrProperty(UhtPropertySettings PropertySettings) : base(PropertySettings)
+		/// <param name="propertySettings">Property settings</param>
+		public UhtStrProperty(UhtPropertySettings propertySettings) : base(propertySettings)
 		{
-			this.PropertyCaps |= UhtPropertyCaps.PassCppArgsByRef | UhtPropertyCaps.CanExposeOnSpawn | UhtPropertyCaps.IsParameterSupportedByBlueprint | 
+			this.PropertyCaps |= UhtPropertyCaps.PassCppArgsByRef | UhtPropertyCaps.CanExposeOnSpawn | UhtPropertyCaps.IsParameterSupportedByBlueprint |
 				UhtPropertyCaps.IsMemberSupportedByBlueprint | UhtPropertyCaps.SupportsRigVM;
 		}
 
 		/// <inheritdoc/>
-		public override StringBuilder AppendNullConstructorArg(StringBuilder Builder, bool bIsInitializer)
+		public override StringBuilder AppendNullConstructorArg(StringBuilder builder, bool isInitializer)
 		{
-			Builder.Append("TEXT(\"\")");
-			return Builder;
+			builder.Append("TEXT(\"\")");
+			return builder;
 		}
 
 		/// <inheritdoc/>
-		public override StringBuilder AppendMemberDecl(StringBuilder Builder, IUhtPropertyMemberContext Context, string Name, string NameSuffix, int Tabs)
+		public override StringBuilder AppendMemberDecl(StringBuilder builder, IUhtPropertyMemberContext context, string name, string nameSuffix, int tabs)
 		{
-			return AppendMemberDecl(Builder, Context, Name, NameSuffix, Tabs, "FStrPropertyParams");
+			return AppendMemberDecl(builder, context, name, nameSuffix, tabs, "FStrPropertyParams");
 		}
 
 		/// <inheritdoc/>
-		public override StringBuilder AppendMemberDef(StringBuilder Builder, IUhtPropertyMemberContext Context, string Name, string NameSuffix, string? Offset, int Tabs)
+		public override StringBuilder AppendMemberDef(StringBuilder builder, IUhtPropertyMemberContext context, string name, string nameSuffix, string? offset, int tabs)
 		{
-			AppendMemberDefStart(Builder, Context, Name, NameSuffix, Offset, Tabs, "FStrPropertyParams", "UECodeGen_Private::EPropertyGenFlags::Str");
-			AppendMemberDefEnd(Builder, Context, Name, NameSuffix);
-			return Builder;
+			AppendMemberDefStart(builder, context, name, nameSuffix, offset, tabs, "FStrPropertyParams", "UECodeGen_Private::EPropertyGenFlags::Str");
+			AppendMemberDefEnd(builder, context, name, nameSuffix);
+			return builder;
 		}
 
 		/// <inheritdoc/>
-		public override bool SanitizeDefaultValue(IUhtTokenReader DefaultValueReader, StringBuilder InnerDefaultValue)
+		public override bool SanitizeDefaultValue(IUhtTokenReader defaultValueReader, StringBuilder innerDefaultValue)
 		{
-			if (DefaultValueReader.TryOptional("FString"))
+			if (defaultValueReader.TryOptional("FString"))
 			{
-				DefaultValueReader.Require('(');
-				StringView Value = DefaultValueReader.GetWrappedConstString();
-				DefaultValueReader.Require(')');
-				InnerDefaultValue.Append(Value);
+				defaultValueReader.Require('(');
+				StringView value = defaultValueReader.GetWrappedConstString();
+				defaultValueReader.Require(')');
+				innerDefaultValue.Append(value);
 			}
 			else
 			{
-				StringView Value = DefaultValueReader.GetWrappedConstString();
-				InnerDefaultValue.Append(Value);
+				StringView value = defaultValueReader.GetWrappedConstString();
+				innerDefaultValue.Append(value);
 			}
 			return true;
 		}
 
 		/// <inheritdoc/>
-		public override bool IsSameType(UhtProperty Other)
+		public override bool IsSameType(UhtProperty other)
 		{
-			return Other is UhtStrProperty;
+			return other is UhtStrProperty;
 		}
 
 		/// <inheritdoc/>
-		protected override void ValidateFunctionArgument(UhtFunction Function, UhtValidationOptions Options)
+		protected override void ValidateFunctionArgument(UhtFunction function, UhtValidationOptions options)
 		{
-			base.ValidateFunctionArgument(Function, Options);
+			base.ValidateFunctionArgument(function, options);
 
-			if (Function.FunctionFlags.HasAnyFlags(EFunctionFlags.Net))
+			if (function.FunctionFlags.HasAnyFlags(EFunctionFlags.Net))
 			{
-				if (!Function.FunctionFlags.HasAnyFlags(EFunctionFlags.NetRequest))
+				if (!function.FunctionFlags.HasAnyFlags(EFunctionFlags.NetRequest))
 				{
-					if (this.RefQualifier != UhtPropertyRefQualifier.ConstRef && !this.bIsStaticArray)
+					if (this.RefQualifier != UhtPropertyRefQualifier.ConstRef && !this.IsStaticArray)
 					{
 						this.LogError("Replicated FString parameters must be passed by const reference");
 					}
@@ -105,36 +105,36 @@ namespace EpicGames.UHT.Types
 		[UhtPropertyType(Keyword = "FMemoryImageString", Options = UhtPropertyTypeOptions.Immediate)]
 		[SuppressMessage("CodeQuality", "IDE0051:Remove unused private members", Justification = "Attribute accessed method")]
 		[SuppressMessage("Style", "IDE0060:Remove unused parameter", Justification = "Attribute accessed method")]
-		private static UhtProperty? StrProperty(UhtPropertyResolvePhase ResolvePhase, UhtPropertySettings PropertySettings, IUhtTokenReader TokenReader, UhtToken MatchedToken)
+		private static UhtProperty? StrProperty(UhtPropertyResolvePhase resolvePhase, UhtPropertySettings propertySettings, IUhtTokenReader tokenReader, UhtToken matchedToken)
 		{
-			if (!TokenReader.SkipExpectedType(MatchedToken.Value, PropertySettings.PropertyCategory == UhtPropertyCategory.Member))
+			if (!tokenReader.SkipExpectedType(matchedToken.Value, propertySettings.PropertyCategory == UhtPropertyCategory.Member))
 			{
 				return null;
 			}
-			UhtStrProperty Out = new UhtStrProperty(PropertySettings);
-			if (Out.PropertyCategory != UhtPropertyCategory.Member)
+			UhtStrProperty property = new UhtStrProperty(propertySettings);
+			if (property.PropertyCategory != UhtPropertyCategory.Member)
 			{
-				if (TokenReader.TryOptional('&'))
+				if (tokenReader.TryOptional('&'))
 				{
-					if (Out.PropertyFlags.HasAnyFlags(EPropertyFlags.ConstParm))
+					if (property.PropertyFlags.HasAnyFlags(EPropertyFlags.ConstParm))
 					{
 						// 'const FString& Foo' came from 'FString' in .uc, no flags
-						Out.PropertyFlags &= ~EPropertyFlags.ConstParm;
+						property.PropertyFlags &= ~EPropertyFlags.ConstParm;
 
 						// We record here that we encountered a const reference, because we need to remove that information from flags for code generation purposes.
-						Out.RefQualifier = UhtPropertyRefQualifier.ConstRef;
+						property.RefQualifier = UhtPropertyRefQualifier.ConstRef;
 					}
 					else
 					{
 						// 'FString& Foo' came from 'out FString' in .uc
-						Out.PropertyFlags |= EPropertyFlags.OutParm;
+						property.PropertyFlags |= EPropertyFlags.OutParm;
 
 						// And we record here that we encountered a non-const reference here too.
-						Out.RefQualifier = UhtPropertyRefQualifier.NonConstRef;
+						property.RefQualifier = UhtPropertyRefQualifier.NonConstRef;
 					}
 				}
 			}
-			return Out;
+			return property;
 		}
 	}
 }
diff --git a/Engine/Source/Programs/Shared/EpicGames.UHT/Types/Properties/UhtStructProperty.cs b/Engine/Source/Programs/Shared/EpicGames.UHT/Types/Properties/UhtStructProperty.cs
--- a/Engine/Source/Programs/Shared/EpicGames.UHT/Types/Properties/UhtStructProperty.cs
+++ b/Engine/Source/Programs/Shared/EpicGames.UHT/Types/Properties/UhtStructProperty.cs
@@ -1,13 +1,13 @@
 // Copyright Epic Games, Inc. All Rights Reserved.
 
-using EpicGames.Core;
-using EpicGames.UHT.Tables;
-using EpicGames.UHT.Tokenizer;
-using EpicGames.UHT.Utils;
 using System.Collections.Generic;
 using System.Diagnostics.CodeAnalysis;
 using System.Text;
 using System.Text.Json.Serialization;
+using EpicGames.Core;
+using EpicGames.UHT.Tables;
+using EpicGames.UHT.Tokenizer;
+using EpicGames.UHT.Utils;
 
 namespace EpicGames.UHT.Types
 {
@@ -19,16 +19,16 @@ namespace EpicGames.UHT.Types
 	public class UhtStructProperty : UhtProperty
 	{
 		/// <inheritdoc/>
-		public override string EngineClassName { get => "StructProperty"; }
+		public override string EngineClassName => "StructProperty";
 
 		/// <inheritdoc/>
-		protected override string CppTypeText { get => "invalid"; }
+		protected override string CppTypeText => "invalid";
 
 		/// <inheritdoc/>
-		protected override string PGetMacroText { get => "STRUCT"; }
+		protected override string PGetMacroText => "STRUCT";
 
 		/// <inheritdoc/>
-		protected override UhtPGetArgumentType PGetTypeArgument { get => UhtPGetArgumentType.TypeText; }
+		protected override UhtPGetArgumentType PGetTypeArgument => UhtPGetArgumentType.TypeText;
 
 		/// <summary>
 		/// USTRUCT referenced by the property
@@ -39,14 +39,14 @@ namespace EpicGames.UHT.Types
 		/// <summary>
 		/// Construct property
 		/// </summary>
-		/// <param name="PropertySettings">Property settings</param>
-		/// <param name="ScriptStruct">USTRUCT being referenced</param>
-		public UhtStructProperty(UhtPropertySettings PropertySettings, UhtScriptStruct ScriptStruct) : base(PropertySettings)
+		/// <param name="propertySettings">Property settings</param>
+		/// <param name="scriptStruct">USTRUCT being referenced</param>
+		public UhtStructProperty(UhtPropertySettings propertySettings, UhtScriptStruct scriptStruct) : base(propertySettings)
 		{
-			this.ScriptStruct = ScriptStruct;
-			this.HeaderFile.AddReferencedHeader(ScriptStruct);
+			this.ScriptStruct = scriptStruct;
+			this.HeaderFile.AddReferencedHeader(scriptStruct);
 			this.PropertyCaps |= UhtPropertyCaps.SupportsRigVM;
-			if (this.ScriptStruct.bHasNoOpConstructor)
+			if (this.ScriptStruct.HasNoOpConstructor)
 			{
 				this.PropertyCaps |= UhtPropertyCaps.RequiresNullConstructorArg;
 			}
@@ -61,10 +61,10 @@ namespace EpicGames.UHT.Types
 		}
 
 		/// <inheritdoc/>
-		protected override bool ResolveSelf(UhtResolvePhase Phase)
+		protected override bool ResolveSelf(UhtResolvePhase phase)
 		{
-			bool bResults = base.ResolveSelf(Phase);
-			switch (Phase)
+			bool results = base.ResolveSelf(phase);
+			switch (phase)
 			{
 				case UhtResolvePhase.Final:
 					if (ScanForInstancedReferenced(true))
@@ -73,25 +73,25 @@ namespace EpicGames.UHT.Types
 					}
 					break;
 			}
-			return bResults;
+			return results;
 		}
 
 		/// <inheritdoc/>
-		public override bool ScanForInstancedReferenced(bool bDeepScan)
+		public override bool ScanForInstancedReferenced(bool deepScan)
 		{
-			return this.ScriptStruct.ScanForInstancedReferenced(bDeepScan);
+			return this.ScriptStruct.ScanForInstancedReferenced(deepScan);
 		}
 
 		/// <inheritdoc/>
-		public override void CollectReferencesInternal(IUhtReferenceCollector Collector, bool bTemplateProperty)
+		public override void CollectReferencesInternal(IUhtReferenceCollector collector, bool templateProperty)
 		{
-			Collector.AddCrossModuleReference(this.ScriptStruct, true);
+			collector.AddCrossModuleReference(this.ScriptStruct, true);
 		}
 
 		/// <inheritdoc/>
 		public override string? GetForwardDeclarations()
 		{
-			return !this.ScriptStruct.bIsCoreType ? $"struct {this.ScriptStruct.SourceName};" : null;
+			return !this.ScriptStruct.IsCoreType ? $"struct {this.ScriptStruct.SourceName};" : null;
 		}
 
 		/// <inheritdoc/>
@@ -101,90 +101,90 @@ namespace EpicGames.UHT.Types
 		}
 
 		/// <inheritdoc/>
-		public override StringBuilder AppendText(StringBuilder Builder, UhtPropertyTextType TextType, bool bIsTemplateArgument)
+		public override StringBuilder AppendText(StringBuilder builder, UhtPropertyTextType textType, bool isTemplateArgument)
 		{
-			switch (TextType)
+			switch (textType)
 			{
 				default:
-					Builder.Append(this.ScriptStruct.SourceName);
+					builder.Append(this.ScriptStruct.SourceName);
 					break;
 			}
-			return Builder;
+			return builder;
 		}
 
 		/// <inheritdoc/>
-		public override StringBuilder AppendMemberDecl(StringBuilder Builder, IUhtPropertyMemberContext Context, string Name, string NameSuffix, int Tabs)
+		public override StringBuilder AppendMemberDecl(StringBuilder builder, IUhtPropertyMemberContext context, string name, string nameSuffix, int tabs)
 		{
-			return AppendMemberDecl(Builder, Context, Name, NameSuffix, Tabs, "FStructPropertyParams");
+			return AppendMemberDecl(builder, context, name, nameSuffix, tabs, "FStructPropertyParams");
 		}
 
 		/// <inheritdoc/>
-		public override StringBuilder AppendMemberDef(StringBuilder Builder, IUhtPropertyMemberContext Context, string Name, string NameSuffix, string? Offset, int Tabs)
+		public override StringBuilder AppendMemberDef(StringBuilder builder, IUhtPropertyMemberContext context, string name, string nameSuffix, string? offset, int tabs)
 		{
-			AppendMemberDefStart(Builder, Context, Name, NameSuffix, Offset, Tabs, "FStructPropertyParams", "UECodeGen_Private::EPropertyGenFlags::Struct");
-			AppendMemberDefRef(Builder, Context, this.ScriptStruct, true);
-			AppendMemberDefEnd(Builder, Context, Name, NameSuffix);
-			return Builder;
+			AppendMemberDefStart(builder, context, name, nameSuffix, offset, tabs, "FStructPropertyParams", "UECodeGen_Private::EPropertyGenFlags::Struct");
+			AppendMemberDefRef(builder, context, this.ScriptStruct, true);
+			AppendMemberDefEnd(builder, context, name, nameSuffix);
+			return builder;
 		}
 
 		/// <inheritdoc/>
-		public override void AppendObjectHashes(StringBuilder Builder, int StartingLength, IUhtPropertyMemberContext Context)
+		public override void AppendObjectHashes(StringBuilder builder, int startingLength, IUhtPropertyMemberContext context)
 		{
-			Builder.AppendObjectHash(StartingLength, this, Context, this.ScriptStruct);
+			builder.AppendObjectHash(startingLength, this, context, this.ScriptStruct);
 		}
 
 		/// <inheritdoc/>
-		public override StringBuilder AppendNullConstructorArg(StringBuilder Builder, bool bIsInitializer)
+		public override StringBuilder AppendNullConstructorArg(StringBuilder builder, bool isInitializer)
 		{
-			bool bHasNoOpConstructor = this.ScriptStruct.bHasNoOpConstructor;
-			if (bIsInitializer && bHasNoOpConstructor)
+			bool hasNoOpConstructor = this.ScriptStruct.HasNoOpConstructor;
+			if (isInitializer && hasNoOpConstructor)
 			{
-				Builder.Append("ForceInit");
+				builder.Append("ForceInit");
 			}
 			else
 			{
-				Builder.AppendPropertyText(this, UhtPropertyTextType.Construction);
-				if (bHasNoOpConstructor)
+				builder.AppendPropertyText(this, UhtPropertyTextType.Construction);
+				if (hasNoOpConstructor)
 				{
-					Builder.Append("(ForceInit)");
+					builder.Append("(ForceInit)");
 				}
 				else
 				{
-					Builder.Append("()");
+					builder.Append("()");
 				}
 			}
-			return Builder;
+			return builder;
 		}
 
 		/// <inheritdoc/>
-		public override bool SanitizeDefaultValue(IUhtTokenReader DefaultValueReader, StringBuilder InnerDefaultValue)
+		public override bool SanitizeDefaultValue(IUhtTokenReader defaultValueReader, StringBuilder innerDefaultValue)
 		{
-			UhtStructDefaultValue StructDefaultValue;
-			if (!this.Session.TryGetStructDefaultValue(this.ScriptStruct.SourceName, out StructDefaultValue))
+			UhtStructDefaultValue structDefaultValue;
+			if (!this.Session.TryGetStructDefaultValue(this.ScriptStruct.SourceName, out structDefaultValue))
 			{
-				StructDefaultValue = this.Session.DefaultStructDefaultValue;
+				structDefaultValue = this.Session.DefaultStructDefaultValue;
 			}
-			return StructDefaultValue.Delegate(this, DefaultValueReader, InnerDefaultValue);
+			return structDefaultValue.Delegate(this, defaultValueReader, innerDefaultValue);
 		}
 
 		/// <inheritdoc/>
-		public override bool IsSameType(UhtProperty Other)
+		public override bool IsSameType(UhtProperty other)
 		{
-			if (Other is UhtStructProperty OtherObject)
+			if (other is UhtStructProperty otherObject)
 			{
-				return this.ScriptStruct == OtherObject.ScriptStruct;
+				return this.ScriptStruct == otherObject.ScriptStruct;
 			}
 			return false;
 		}
 
 		/// <inheritdoc/>
-		protected override void ValidateFunctionArgument(UhtFunction Function, UhtValidationOptions Options)
+		protected override void ValidateFunctionArgument(UhtFunction function, UhtValidationOptions options)
 		{
-			base.ValidateFunctionArgument(Function, Options);
+			base.ValidateFunctionArgument(function, options);
 
-			if (Function.FunctionFlags.HasAnyFlags(EFunctionFlags.Net))
+			if (function.FunctionFlags.HasAnyFlags(EFunctionFlags.Net))
 			{
-				if (!Function.FunctionFlags.HasAnyFlags(EFunctionFlags.NetRequest | EFunctionFlags.NetResponse))
+				if (!function.FunctionFlags.HasAnyFlags(EFunctionFlags.NetRequest | EFunctionFlags.NetResponse))
 				{
 					this.Session.ValidateScriptStructOkForNet(this, this.ScriptStruct);
 				}
@@ -192,7 +192,7 @@ namespace EpicGames.UHT.Types
 		}
 
 		/// <inheritdoc/>
-		protected override void ValidateMember(UhtStruct Struct, UhtValidationOptions Options)
+		protected override void ValidateMember(UhtStruct structObj, UhtValidationOptions options)
 		{
 			if (this.PropertyFlags.HasAnyFlags(EPropertyFlags.Net))
 			{
@@ -201,50 +201,50 @@ namespace EpicGames.UHT.Types
 		}
 
 		///<inheritdoc/>
-		public override bool ValidateStructPropertyOkForNet(UhtProperty ReferencingProperty)
+		public override bool ValidateStructPropertyOkForNet(UhtProperty referencingProperty)
 		{
-			return ReferencingProperty.Session.ValidateScriptStructOkForNet(ReferencingProperty, this.ScriptStruct);
+			return referencingProperty.Session.ValidateScriptStructOkForNet(referencingProperty, this.ScriptStruct);
 		}
 
 		#region Structure default value sanitizers
 		[UhtStructDefaultValue(Name = "FVector")]
 		[SuppressMessage("CodeQuality", "IDE0051:Remove unused private members", Justification = "Attribute accessed method")]
 		[SuppressMessage("Style", "IDE0060:Remove unused parameter", Justification = "Attribute accessed method")]
-		private static bool VectorStructDefaultValue(UhtStructProperty Property, IUhtTokenReader DefaultValueReader, StringBuilder InnerDefaultValue)
+		private static bool VectorStructDefaultValue(UhtStructProperty property, IUhtTokenReader defaultValueReader, StringBuilder innerDefaultValue)
 		{
-			string Format = "{0:F6},{1:F6},{2:F6}";
+			const string Format = "{0:F6},{1:F6},{2:F6}";
 
-			DefaultValueReader.Require("FVector");
-			if (DefaultValueReader.TryOptional("::"))
+			defaultValueReader.Require("FVector");
+			if (defaultValueReader.TryOptional("::"))
 			{
-				switch (DefaultValueReader.GetIdentifier().Value.ToString())
+				switch (defaultValueReader.GetIdentifier().Value.ToString())
 				{
 					case "ZeroVector": return true;
-					case "UpVector": InnerDefaultValue.AppendFormat(Format, 0, 0, 1); return true;
-					case "ForwardVector": InnerDefaultValue.AppendFormat(Format, 1, 0, 0); return true;
-					case "RightVector": InnerDefaultValue.AppendFormat(Format, 0, 1, 0); return true;
+					case "UpVector": innerDefaultValue.AppendFormat(Format, 0, 0, 1); return true;
+					case "ForwardVector": innerDefaultValue.AppendFormat(Format, 1, 0, 0); return true;
+					case "RightVector": innerDefaultValue.AppendFormat(Format, 0, 1, 0); return true;
 					default: return false;
 				}
 			}
 			else
 			{
-				DefaultValueReader.Require("(");
-				if (DefaultValueReader.TryOptional("ForceInit"))
+				defaultValueReader.Require("(");
+				if (defaultValueReader.TryOptional("ForceInit"))
 				{
 				}
 				else
 				{
-					double X, Y, Z;
-					X = Y = Z = DefaultValueReader.GetConstDoubleExpression();
-					if (DefaultValueReader.TryOptional(','))
+					double x, y, z;
+					x = y = z = defaultValueReader.GetConstDoubleExpression();
+					if (defaultValueReader.TryOptional(','))
 					{
-						Y = DefaultValueReader.GetConstDoubleExpression();
-						DefaultValueReader.Require(',');
-						Z = DefaultValueReader.GetConstDoubleExpression();
+						y = defaultValueReader.GetConstDoubleExpression();
+						defaultValueReader.Require(',');
+						z = defaultValueReader.GetConstDoubleExpression();
 					}
-					InnerDefaultValue.AppendFormat(Format, X, Y, Z);
+					innerDefaultValue.AppendFormat(Format, x, y, z);
 				}
-				DefaultValueReader.Require(")");
+				defaultValueReader.Require(")");
 				return true;
 			}
 		}
@@ -252,12 +252,12 @@ namespace EpicGames.UHT.Types
 		[UhtStructDefaultValue(Name = "FRotator")]
 		[SuppressMessage("CodeQuality", "IDE0051:Remove unused private members", Justification = "Attribute accessed method")]
 		[SuppressMessage("Style", "IDE0060:Remove unused parameter", Justification = "Attribute accessed method")]
-		private static bool RotatorStructDefaultValue(UhtStructProperty Property, IUhtTokenReader DefaultValueReader, StringBuilder InnerDefaultValue)
+		private static bool RotatorStructDefaultValue(UhtStructProperty property, IUhtTokenReader defaultValueReader, StringBuilder innerDefaultValue)
 		{
-			DefaultValueReader.Require("FRotator");
-			if (DefaultValueReader.TryOptional("::"))
+			defaultValueReader.Require("FRotator");
+			if (defaultValueReader.TryOptional("::"))
 			{
-				switch (DefaultValueReader.GetIdentifier().Value.ToString())
+				switch (defaultValueReader.GetIdentifier().Value.ToString())
 				{
 					case "ZeroRotator": return true;
 					default: return false;
@@ -265,20 +265,20 @@ namespace EpicGames.UHT.Types
 			}
 			else
 			{
-				DefaultValueReader.Require("(");
-				if (DefaultValueReader.TryOptional("ForceInit"))
+				defaultValueReader.Require("(");
+				if (defaultValueReader.TryOptional("ForceInit"))
 				{
 				}
 				else
 				{
-					double X = DefaultValueReader.GetConstDoubleExpression();
-					DefaultValueReader.Require(',');
-					double Y = DefaultValueReader.GetConstDoubleExpression();
-					DefaultValueReader.Require(',');
-					double Z = DefaultValueReader.GetConstDoubleExpression();
-					InnerDefaultValue.AppendFormat("{0:F6},{1:F6},{2:F6}", X, Y, Z);
+					double x = defaultValueReader.GetConstDoubleExpression();
+					defaultValueReader.Require(',');
+					double y = defaultValueReader.GetConstDoubleExpression();
+					defaultValueReader.Require(',');
+					double z = defaultValueReader.GetConstDoubleExpression();
+					innerDefaultValue.AppendFormat("{0:F6},{1:F6},{2:F6}", x, y, z);
 				}
-				DefaultValueReader.Require(")");
+				defaultValueReader.Require(")");
 				return true;
 			}
 		}
@@ -286,34 +286,34 @@ namespace EpicGames.UHT.Types
 		[UhtStructDefaultValue(Name = "FVector2D")]
 		[SuppressMessage("CodeQuality", "IDE0051:Remove unused private members", Justification = "Attribute accessed method")]
 		[SuppressMessage("Style", "IDE0060:Remove unused parameter", Justification = "Attribute accessed method")]
-		private static bool Vector2DStructDefaultValue(UhtStructProperty Property, IUhtTokenReader DefaultValueReader, StringBuilder InnerDefaultValue)
+		private static bool Vector2DStructDefaultValue(UhtStructProperty property, IUhtTokenReader defaultValueReader, StringBuilder innerDefaultValue)
 		{
-			string Format = "(X={0:F3},Y={1:F3})";
+			const string Format = "(X={0:F3},Y={1:F3})";
 
-			DefaultValueReader.Require("FVector2D");
-			if (DefaultValueReader.TryOptional("::"))
+			defaultValueReader.Require("FVector2D");
+			if (defaultValueReader.TryOptional("::"))
 			{
-				switch (DefaultValueReader.GetIdentifier().Value.ToString())
+				switch (defaultValueReader.GetIdentifier().Value.ToString())
 				{
 					case "ZeroVector": return true;
-					case "UnitVector": InnerDefaultValue.AppendFormat(Format, 1.0, 1.0); return true;
+					case "UnitVector": innerDefaultValue.AppendFormat(Format, 1.0, 1.0); return true;
 					default: return false;
 				}
 			}
 			else
 			{
-				DefaultValueReader.Require("(");
-				if (DefaultValueReader.TryOptional("ForceInit"))
+				defaultValueReader.Require("(");
+				if (defaultValueReader.TryOptional("ForceInit"))
 				{
 				}
 				else
 				{
-					double X = DefaultValueReader.GetConstDoubleExpression();
-					DefaultValueReader.Require(',');
-					double Y = DefaultValueReader.GetConstDoubleExpression();
-					InnerDefaultValue.AppendFormat(Format, X, Y);
+					double x = defaultValueReader.GetConstDoubleExpression();
+					defaultValueReader.Require(',');
+					double y = defaultValueReader.GetConstDoubleExpression();
+					innerDefaultValue.AppendFormat(Format, x, y);
 				}
-				DefaultValueReader.Require(")");
+				defaultValueReader.Require(")");
 				return true;
 			}
 		}
@@ -321,47 +321,47 @@ namespace EpicGames.UHT.Types
 		[UhtStructDefaultValue(Name = "FLinearColor")]
 		[SuppressMessage("CodeQuality", "IDE0051:Remove unused private members", Justification = "Attribute accessed method")]
 		[SuppressMessage("Style", "IDE0060:Remove unused parameter", Justification = "Attribute accessed method")]
-		private static bool LinearColorStructDefaultValue(UhtStructProperty Property, IUhtTokenReader DefaultValueReader, StringBuilder InnerDefaultValue)
+		private static bool LinearColorStructDefaultValue(UhtStructProperty property, IUhtTokenReader defaultValueReader, StringBuilder innerDefaultValue)
 		{
-			string Format = "(R={0:F6},G={1:F6},B={2:F6},A={3:F6})";
+			const string Format = "(R={0:F6},G={1:F6},B={2:F6},A={3:F6})";
 
-			DefaultValueReader.Require("FLinearColor");
-			if (DefaultValueReader.TryOptional("::"))
+			defaultValueReader.Require("FLinearColor");
+			if (defaultValueReader.TryOptional("::"))
 			{
-				switch (DefaultValueReader.GetIdentifier().Value.ToString())
+				switch (defaultValueReader.GetIdentifier().Value.ToString())
 				{
-					case "White": InnerDefaultValue.AppendFormat(Format, 1.0, 1.0, 1.0, 1.0); ; return true;
-					case "Gray": InnerDefaultValue.AppendFormat(Format, 0.5, 0.5, 0.5, 1.0); ; return true;
-					case "Black": InnerDefaultValue.AppendFormat(Format, 0.0, 0.0, 0.0, 1.0); ; return true;
-					case "Transparent": InnerDefaultValue.AppendFormat(Format, 0.0, 0.0, 0.0, 0.0); ; return true;
-					case "Red": InnerDefaultValue.AppendFormat(Format, 1.0, 0.0, 0.0, 1.0); ; return true;
-					case "Green": InnerDefaultValue.AppendFormat(Format, 0.0, 1.0, 0.0, 1.0); ; return true;
-					case "Blue": InnerDefaultValue.AppendFormat(Format, 0.0, 0.0, 1.0, 1.0); ; return true;
-					case "Yellow": InnerDefaultValue.AppendFormat(Format, 1.0, 1.0, 0.0, 1.0); ; return true;
+					case "White": innerDefaultValue.AppendFormat(Format, 1.0, 1.0, 1.0, 1.0); ; return true;
+					case "Gray": innerDefaultValue.AppendFormat(Format, 0.5, 0.5, 0.5, 1.0); ; return true;
+					case "Black": innerDefaultValue.AppendFormat(Format, 0.0, 0.0, 0.0, 1.0); ; return true;
+					case "Transparent": innerDefaultValue.AppendFormat(Format, 0.0, 0.0, 0.0, 0.0); ; return true;
+					case "Red": innerDefaultValue.AppendFormat(Format, 1.0, 0.0, 0.0, 1.0); ; return true;
+					case "Green": innerDefaultValue.AppendFormat(Format, 0.0, 1.0, 0.0, 1.0); ; return true;
+					case "Blue": innerDefaultValue.AppendFormat(Format, 0.0, 0.0, 1.0, 1.0); ; return true;
+					case "Yellow": innerDefaultValue.AppendFormat(Format, 1.0, 1.0, 0.0, 1.0); ; return true;
 					default: return false;
 				}
 			}
 			else
 			{
-				DefaultValueReader.Require("(");
-				if (DefaultValueReader.TryOptional("ForceInit"))
+				defaultValueReader.Require("(");
+				if (defaultValueReader.TryOptional("ForceInit"))
 				{
 				}
 				else
 				{
-					double R = DefaultValueReader.GetConstDoubleExpression();
-					DefaultValueReader.Require(',');
-					double G = DefaultValueReader.GetConstDoubleExpression();
-					DefaultValueReader.Require(',');
-					double B = DefaultValueReader.GetConstDoubleExpression();
-					double A = 1.0;
-					if (DefaultValueReader.TryOptional(','))
+					double r = defaultValueReader.GetConstDoubleExpression();
+					defaultValueReader.Require(',');
+					double g = defaultValueReader.GetConstDoubleExpression();
+					defaultValueReader.Require(',');
+					double b = defaultValueReader.GetConstDoubleExpression();
+					double a = 1.0;
+					if (defaultValueReader.TryOptional(','))
 					{
-						A = DefaultValueReader.GetConstDoubleExpression();
+						a = defaultValueReader.GetConstDoubleExpression();
 					}
-					InnerDefaultValue.AppendFormat(Format, R, G, B, A);
+					innerDefaultValue.AppendFormat(Format, r, g, b, a);
 				}
-				DefaultValueReader.Require(")");
+				defaultValueReader.Require(")");
 				return true;
 			}
 		}
@@ -369,61 +369,60 @@ namespace EpicGames.UHT.Types
 		[UhtStructDefaultValue(Name = "FColor")]
 		[SuppressMessage("CodeQuality", "IDE0051:Remove unused private members", Justification = "Attribute accessed method")]
 		[SuppressMessage("Style", "IDE0060:Remove unused parameter", Justification = "Attribute accessed method")]
-		private static bool ColorStructDefaultValue(UhtStructProperty Property, IUhtTokenReader DefaultValueReader, StringBuilder InnerDefaultValue)
+		private static bool ColorStructDefaultValue(UhtStructProperty property, IUhtTokenReader defaultValueReader, StringBuilder innerDefaultValue)
 		{
-			string Format = "(R={0},G={1},B={2},A={3})";
+			const string Format = "(R={0},G={1},B={2},A={3})";
 
-			DefaultValueReader.Require("FColor");
-			if (DefaultValueReader.TryOptional("::"))
+			defaultValueReader.Require("FColor");
+			if (defaultValueReader.TryOptional("::"))
 			{
-				switch (DefaultValueReader.GetIdentifier().Value.ToString())
+				switch (defaultValueReader.GetIdentifier().Value.ToString())
 				{
-					case "White": InnerDefaultValue.AppendFormat(Format, 255, 255, 255, 255); ; return true;
-					case "Black": InnerDefaultValue.AppendFormat(Format, 0, 0, 0, 255); ; return true;
-					case "Red": InnerDefaultValue.AppendFormat(Format, 255, 0, 0, 255); ; return true;
-					case "Green": InnerDefaultValue.AppendFormat(Format, 0, 255, 0, 255); ; return true;
-					case "Blue": InnerDefaultValue.AppendFormat(Format, 0, 0, 255, 255); ; return true;
-					case "Yellow": InnerDefaultValue.AppendFormat(Format, 255, 255, 0, 255); ; return true;
-					case "Cyan": InnerDefaultValue.AppendFormat(Format, 0, 255, 255, 255); ; return true;
-					case "Magenta": InnerDefaultValue.AppendFormat(Format, 255, 0, 255, 255); ; return true;
+					case "White": innerDefaultValue.AppendFormat(Format, 255, 255, 255, 255); ; return true;
+					case "Black": innerDefaultValue.AppendFormat(Format, 0, 0, 0, 255); ; return true;
+					case "Red": innerDefaultValue.AppendFormat(Format, 255, 0, 0, 255); ; return true;
+					case "Green": innerDefaultValue.AppendFormat(Format, 0, 255, 0, 255); ; return true;
+					case "Blue": innerDefaultValue.AppendFormat(Format, 0, 0, 255, 255); ; return true;
+					case "Yellow": innerDefaultValue.AppendFormat(Format, 255, 255, 0, 255); ; return true;
+					case "Cyan": innerDefaultValue.AppendFormat(Format, 0, 255, 255, 255); ; return true;
+					case "Magenta": innerDefaultValue.AppendFormat(Format, 255, 0, 255, 255); ; return true;
 					default: return false;
 				}
 			}
 			else
 			{
-				DefaultValueReader.Require("(");
-				if (DefaultValueReader.TryOptional("ForceInit"))
+				defaultValueReader.Require("(");
+				if (defaultValueReader.TryOptional("ForceInit"))
 				{
 				}
 				else
 				{
-					int R = DefaultValueReader.GetConstIntExpression();
-					DefaultValueReader.Require(',');
-					int G = DefaultValueReader.GetConstIntExpression();
-					DefaultValueReader.Require(',');
-					int B = DefaultValueReader.GetConstIntExpression();
-					int A = 255;
-					if (DefaultValueReader.TryOptional(','))
+					int r = defaultValueReader.GetConstIntExpression();
+					defaultValueReader.Require(',');
+					int g = defaultValueReader.GetConstIntExpression();
+					defaultValueReader.Require(',');
+					int b = defaultValueReader.GetConstIntExpression();
+					int a = 255;
+					if (defaultValueReader.TryOptional(','))
 					{
-						A = DefaultValueReader.GetConstIntExpression();
+						a = defaultValueReader.GetConstIntExpression();
 					}
-					InnerDefaultValue.AppendFormat(Format, R, G, B, A);
+					innerDefaultValue.AppendFormat(Format, r, g, b, a);
 				}
-				DefaultValueReader.Require(")");
+				defaultValueReader.Require(")");
 				return true;
 			}
 		}
 
 		[UhtStructDefaultValue(Options = UhtStructDefaultValueOptions.Default)]
 		[SuppressMessage("CodeQuality", "IDE0051:Remove unused private members", Justification = "Attribute accessed method")]
-		[SuppressMessage("Style", "IDE0060:Remove unused parameter", Justification = "Attribute accessed method")]
-		private static bool DefaultStructDefaultValue(UhtStructProperty Property, IUhtTokenReader DefaultValueReader, StringBuilder InnerDefaultValue)
+		private static bool DefaultStructDefaultValue(UhtStructProperty property, IUhtTokenReader defaultValueReader, StringBuilder innerDefaultValue)
 		{
-			DefaultValueReader
-				.Require(Property.ScriptStruct.SourceName)
+			defaultValueReader
+				.Require(property.ScriptStruct.SourceName)
 				.Require('(')
 				.Require(')');
-			InnerDefaultValue.Append("()");
+			innerDefaultValue.Append("()");
 			return true;
 		}
 		#endregion
diff --git a/Engine/Source/Programs/Shared/EpicGames.UHT/Types/Properties/UhtTextProperty.cs b/Engine/Source/Programs/Shared/EpicGames.UHT/Types/Properties/UhtTextProperty.cs
--- a/Engine/Source/Programs/Shared/EpicGames.UHT/Types/Properties/UhtTextProperty.cs
+++ b/Engine/Source/Programs/Shared/EpicGames.UHT/Types/Properties/UhtTextProperty.cs
@@ -1,12 +1,12 @@
 // Copyright Epic Games, Inc. All Rights Reserved.
 
+using System;
+using System.Diagnostics.CodeAnalysis;
+using System.Text;
 using EpicGames.Core;
 using EpicGames.UHT.Tables;
 using EpicGames.UHT.Tokenizer;
 using EpicGames.UHT.Utils;
-using System;
-using System.Diagnostics.CodeAnalysis;
-using System.Text;
 
 namespace EpicGames.UHT.Types
 {
@@ -18,74 +18,74 @@ namespace EpicGames.UHT.Types
 	public class UhtTextProperty : UhtProperty
 	{
 		/// <inheritdoc/>
-		public override string EngineClassName { get => "TextProperty"; }
+		public override string EngineClassName => "TextProperty";
 
 		/// <inheritdoc/>
-		protected override string CppTypeText { get => "FText"; }
+		protected override string CppTypeText => "FText";
 
 		/// <inheritdoc/>
-		protected override string PGetMacroText { get => "PROPERTY"; }
+		protected override string PGetMacroText => "PROPERTY";
 
 		/// <inheritdoc/>
-		protected override UhtPGetArgumentType PGetTypeArgument { get => UhtPGetArgumentType.EngineClass; }
+		protected override UhtPGetArgumentType PGetTypeArgument => UhtPGetArgumentType.EngineClass;
 
 		/// <summary>
 		/// Construct a new property
 		/// </summary>
-		/// <param name="PropertySettings">Property settings</param>
-		public UhtTextProperty(UhtPropertySettings PropertySettings) : base(PropertySettings)
+		/// <param name="propertySettings">Property settings</param>
+		public UhtTextProperty(UhtPropertySettings propertySettings) : base(propertySettings)
 		{
-			this.PropertyCaps |= UhtPropertyCaps.PassCppArgsByRef | UhtPropertyCaps.CanExposeOnSpawn | UhtPropertyCaps.IsParameterSupportedByBlueprint | 
+			this.PropertyCaps |= UhtPropertyCaps.PassCppArgsByRef | UhtPropertyCaps.CanExposeOnSpawn | UhtPropertyCaps.IsParameterSupportedByBlueprint |
 				UhtPropertyCaps.IsMemberSupportedByBlueprint | UhtPropertyCaps.SupportsRigVM;
 			this.PropertyCaps &= ~(UhtPropertyCaps.CanBeContainerKey);
 		}
 
 		/// <inheritdoc/>
-		public override StringBuilder AppendMemberDecl(StringBuilder Builder, IUhtPropertyMemberContext Context, string Name, string NameSuffix, int Tabs)
+		public override StringBuilder AppendMemberDecl(StringBuilder builder, IUhtPropertyMemberContext context, string name, string nameSuffix, int tabs)
 		{
-			return AppendMemberDecl(Builder, Context, Name, NameSuffix, Tabs, "FTextPropertyParams");
+			return AppendMemberDecl(builder, context, name, nameSuffix, tabs, "FTextPropertyParams");
 		}
 
 		/// <inheritdoc/>
-		public override StringBuilder AppendMemberDef(StringBuilder Builder, IUhtPropertyMemberContext Context, string Name, string NameSuffix, string? Offset, int Tabs)
+		public override StringBuilder AppendMemberDef(StringBuilder builder, IUhtPropertyMemberContext context, string name, string nameSuffix, string? offset, int tabs)
 		{
-			AppendMemberDefStart(Builder, Context, Name, NameSuffix, Offset, Tabs, "FTextPropertyParams", "UECodeGen_Private::EPropertyGenFlags::Text");
-			AppendMemberDefEnd(Builder, Context, Name, NameSuffix);
-			return Builder;
+			AppendMemberDefStart(builder, context, name, nameSuffix, offset, tabs, "FTextPropertyParams", "UECodeGen_Private::EPropertyGenFlags::Text");
+			AppendMemberDefEnd(builder, context, name, nameSuffix);
+			return builder;
 		}
 
 		/// <inheritdoc/>
-		public override StringBuilder AppendNullConstructorArg(StringBuilder Builder, bool bIsInitializer)
+		public override StringBuilder AppendNullConstructorArg(StringBuilder builder, bool isInitializer)
 		{
-			Builder.Append("FText::GetEmpty()");
-			return Builder;
+			builder.Append("FText::GetEmpty()");
+			return builder;
 		}
 
 		/// <inheritdoc/>
-		public override bool SanitizeDefaultValue(IUhtTokenReader DefaultValueReader, StringBuilder InnerDefaultValue)
+		public override bool SanitizeDefaultValue(IUhtTokenReader defaultValueReader, StringBuilder innerDefaultValue)
 		{
-			if (DefaultValueReader.TryOptional("FText"))
+			if (defaultValueReader.TryOptional("FText"))
 			{
-				if (DefaultValueReader.TryOptional('('))
+				if (defaultValueReader.TryOptional('('))
 				{
-					DefaultValueReader.Require(')');
+					defaultValueReader.Require(')');
 					return true;
 				}
-				else if (DefaultValueReader.TryOptional("::"))
+				else if (defaultValueReader.TryOptional("::"))
 				{
 					// Handle legacy cases of FText::FromString being used as default values
 					// These should be replaced with INVTEXT as FText::FromString can produce inconsistent keys
-					if (DefaultValueReader.TryOptional("FromString"))
+					if (defaultValueReader.TryOptional("FromString"))
 					{
-						DefaultValueReader.Require('(');
-						StringView Value = DefaultValueReader.GetWrappedConstString();
-						DefaultValueReader.Require(')');
-						InnerDefaultValue.Append('\"').Append(Value).Append('\"');
+						defaultValueReader.Require('(');
+						StringView value = defaultValueReader.GetWrappedConstString();
+						defaultValueReader.Require(')');
+						innerDefaultValue.Append('\"').Append(value).Append('\"');
 						this.LogWarning("FText::FromString should be replaced with INVTEXT for default parameter values");
 					}
 					else
 					{
-						DefaultValueReader.Require("GetEmpty").Require('(').Require(')');
+						defaultValueReader.Require("GetEmpty").Require('(').Require(')');
 					}
 					return true;
 				}
@@ -95,9 +95,9 @@ namespace EpicGames.UHT.Types
 				}
 			}
 
-			UhtToken Token = DefaultValueReader.GetToken();
+			UhtToken token = defaultValueReader.GetToken();
 
-			if (Token.IsIdentifier("LOCTEXT"))
+			if (token.IsIdentifier("LOCTEXT"))
 			{
 				//ETSTODO - Add a test case for this
 				this.LogError($"LOCTEXT default parameter values are not supported; use NSLOCTEXT instead: {this.SourceName}");
@@ -105,349 +105,347 @@ namespace EpicGames.UHT.Types
 			}
 			else
 			{
-				return SanitizeDefaultValue(this, DefaultValueReader, ref Token, InnerDefaultValue, false);
+				return SanitizeDefaultValue(this, defaultValueReader, ref token, innerDefaultValue, false);
 			}
 		}
 
 		/// <inheritdoc/>
-		public override bool IsSameType(UhtProperty Other)
+		public override bool IsSameType(UhtProperty other)
 		{
-			return Other is UhtTextProperty;
+			return other is UhtTextProperty;
 		}
 
 		#region Parsing keywords and default parsers		
 		[UhtPropertyType(Keyword = "FText", Options = UhtPropertyTypeOptions.Simple | UhtPropertyTypeOptions.Immediate)]
 		[SuppressMessage("CodeQuality", "IDE0051:Remove unused private members", Justification = "Attribute accessed method")]
 		[SuppressMessage("Style", "IDE0060:Remove unused parameter", Justification = "Attribute accessed method")]
-		private static UhtProperty? TextProperty(UhtPropertyResolvePhase ResolvePhase, UhtPropertySettings PropertySettings, IUhtTokenReader TokenReader, UhtToken MatchedToken)
+		private static UhtProperty? TextProperty(UhtPropertyResolvePhase resolvePhase, UhtPropertySettings propertySettings, IUhtTokenReader tokenReader, UhtToken matchedToken)
 		{
-			return new UhtTextProperty(PropertySettings);
+			return new UhtTextProperty(propertySettings);
 		}
 
 		[UhtPropertyType(Keyword = "Text", Options = UhtPropertyTypeOptions.Simple | UhtPropertyTypeOptions.Immediate | UhtPropertyTypeOptions.CaseInsensitive)]
 		[SuppressMessage("CodeQuality", "IDE0051:Remove unused private members", Justification = "Attribute accessed method")]
 		[SuppressMessage("Style", "IDE0060:Remove unused parameter", Justification = "Attribute accessed method")]
-		private static UhtProperty? MissingPrefixTextProperty(UhtPropertyResolvePhase ResolvePhase, UhtPropertySettings PropertySettings, IUhtTokenReader TokenReader, UhtToken MatchedToken)
+		private static UhtProperty? MissingPrefixTextProperty(UhtPropertyResolvePhase resolvePhase, UhtPropertySettings propertySettings, IUhtTokenReader tokenReader, UhtToken matchedToken)
 		{
-			throw new UhtException(TokenReader, "'Text' is missing a prefix, expecting 'FText'");
+			throw new UhtException(tokenReader, "'Text' is missing a prefix, expecting 'FText'");
 		}
 
 		[UhtLocTextDefaultValue(Name = "INVTEXT")]
 		[SuppressMessage("CodeQuality", "IDE0051:Remove unused private members", Justification = "Attribute accessed method")]
 		[SuppressMessage("Style", "IDE0060:Remove unused parameter", Justification = "Attribute accessed method")]
-		private static bool InvTextDefaultValue(UhtTextProperty Property, IUhtTokenReader DefaultValueReader, ref UhtToken MacroToken, StringBuilder InnerDefaultValue)
+		private static bool InvTextDefaultValue(UhtTextProperty property, IUhtTokenReader defaultValueReader, ref UhtToken macroToken, StringBuilder innerDefaultValue)
 		{
-			DefaultValueReader.Require('(');
-			StringView String = DefaultValueReader.GetConstQuotedString();
-			DefaultValueReader.Require(')');
-			InnerDefaultValue.Append("INVTEXT(").Append(String).Append(')');
+			defaultValueReader.Require('(');
+			StringView value = defaultValueReader.GetConstQuotedString();
+			defaultValueReader.Require(')');
+			innerDefaultValue.Append("INVTEXT(").Append(value).Append(')');
 			return true;
 		}
 
 		[UhtLocTextDefaultValue(Name = "LOCTEXT")]
 		[SuppressMessage("CodeQuality", "IDE0051:Remove unused private members", Justification = "Attribute accessed method")]
 		[SuppressMessage("Style", "IDE0060:Remove unused parameter", Justification = "Attribute accessed method")]
-		private static bool LocTextDefaultValue(UhtTextProperty Property, IUhtTokenReader DefaultValueReader, ref UhtToken MacroToken, StringBuilder InnerDefaultValue)
+		private static bool LocTextDefaultValue(UhtTextProperty property, IUhtTokenReader defaultValueReader, ref UhtToken macroToken, StringBuilder innerDefaultValue)
 		{
-			Property.LogError($"LOCTEXT default parameter values are not supported; use NSLOCTEXT instead: {Property.SourceName}");
+			property.LogError($"LOCTEXT default parameter values are not supported; use NSLOCTEXT instead: {property.SourceName}");
 			return false;
 		}
 
 		[UhtLocTextDefaultValue(Name = "NSLOCTEXT")]
 		[SuppressMessage("CodeQuality", "IDE0051:Remove unused private members", Justification = "Attribute accessed method")]
 		[SuppressMessage("Style", "IDE0060:Remove unused parameter", Justification = "Attribute accessed method")]
-		private static bool NsLocTextDefaultValue(UhtTextProperty Property, IUhtTokenReader DefaultValueReader, ref UhtToken MacroToken, StringBuilder InnerDefaultValue)
+		private static bool NsLocTextDefaultValue(UhtTextProperty property, IUhtTokenReader defaultValueReader, ref UhtToken macroToken, StringBuilder innerDefaultValue)
 		{
-			DefaultValueReader.Require('(');
-			StringView NamespaceString = DefaultValueReader.GetConstQuotedString();
-			DefaultValueReader.Require(',');
-			StringView KeyString = DefaultValueReader.GetConstQuotedString();
-			DefaultValueReader.Require(',');
-			StringView SourceString = DefaultValueReader.GetConstQuotedString();
-			DefaultValueReader.Require(')');
+			defaultValueReader.Require('(');
+			StringView namespaceString = defaultValueReader.GetConstQuotedString();
+			defaultValueReader.Require(',');
+			StringView keyString = defaultValueReader.GetConstQuotedString();
+			defaultValueReader.Require(',');
+			StringView sourceString = defaultValueReader.GetConstQuotedString();
+			defaultValueReader.Require(')');
 
 			// Strip out the package name
-			ReadOnlySpan<char> StrippedNS = NamespaceString.Span;
-			if (StrippedNS.Length > 1)
+			ReadOnlySpan<char> strippedNS = namespaceString.Span;
+			if (strippedNS.Length > 1)
 			{
-				StrippedNS = StrippedNS.Slice(1, StrippedNS.Length - 2).Trim();
-				if (StrippedNS.Length > 0)
+				strippedNS = strippedNS.Slice(1, strippedNS.Length - 2).Trim();
+				if (strippedNS.Length > 0)
 				{
-					if (StrippedNS[StrippedNS.Length - 1] == ']')
+					if (strippedNS[^1] == ']')
 					{
-						int Index = StrippedNS.LastIndexOf('[');
-						if (Index != -1)
+						int index = strippedNS.LastIndexOf('[');
+						if (index != -1)
 						{
-							StrippedNS = StrippedNS.Slice(0, Index).TrimEnd();
+							strippedNS = strippedNS.Slice(0, index).TrimEnd();
 						}
 					}
 				}
 			}
 
-			InnerDefaultValue.Append("NSLOCTEXT(\"").Append(StrippedNS).Append("\", ").Append(KeyString).Append(", ").Append(SourceString).Append(')');
+			innerDefaultValue.Append("NSLOCTEXT(\"").Append(strippedNS).Append("\", ").Append(keyString).Append(", ").Append(sourceString).Append(')');
 			return true;
 		}
 
 		[UhtLocTextDefaultValue(Name = "LOCTABLE")]
 		[SuppressMessage("CodeQuality", "IDE0051:Remove unused private members", Justification = "Attribute accessed method")]
 		[SuppressMessage("Style", "IDE0060:Remove unused parameter", Justification = "Attribute accessed method")]
-		private static bool LocTableDefaultValue(UhtTextProperty Property, IUhtTokenReader DefaultValueReader, ref UhtToken MacroToken, StringBuilder InnerDefaultValue)
+		private static bool LocTableDefaultValue(UhtTextProperty property, IUhtTokenReader defaultValueReader, ref UhtToken macroToken, StringBuilder innerDefaultValue)
 		{
-			DefaultValueReader.Require('(');
-			StringView NamespaceString = DefaultValueReader.GetConstQuotedString();
-			DefaultValueReader.Require(',');
-			StringView KeyString = DefaultValueReader.GetConstQuotedString();
-			DefaultValueReader.Require(')');
+			defaultValueReader.Require('(');
+			StringView namespaceString = defaultValueReader.GetConstQuotedString();
+			defaultValueReader.Require(',');
+			StringView keyString = defaultValueReader.GetConstQuotedString();
+			defaultValueReader.Require(')');
 
 			//ETSTODO - Validate the namespace string?
-			InnerDefaultValue.Append("LOCTABLE(").Append(NamespaceString).Append(", ").Append(KeyString).Append(')');
+			innerDefaultValue.Append("LOCTABLE(").Append(namespaceString).Append(", ").Append(keyString).Append(')');
 			return true;
 		}
 
 		[UhtLocTextDefaultValue(Name = "LOCGEN_FORMAT_NAMED")]
 		[SuppressMessage("CodeQuality", "IDE0051:Remove unused private members", Justification = "Attribute accessed method")]
 		[SuppressMessage("Style", "IDE0060:Remove unused parameter", Justification = "Attribute accessed method")]
-		private static bool LocGenFormatNamedDefaultValue(UhtTextProperty Property, IUhtTokenReader DefaultValueReader, ref UhtToken MacroToken, StringBuilder InnerDefaultValue)
+		private static bool LocGenFormatNamedDefaultValue(UhtTextProperty property, IUhtTokenReader defaultValueReader, ref UhtToken macroToken, StringBuilder innerDefaultValue)
 		{
-			InnerDefaultValue.Append("LOCGEN_FORMAT_NAMED(");
-			DefaultValueReader.Require('(');
-			if (!SanitizeDefaultValue(Property, DefaultValueReader, InnerDefaultValue, false))
+			innerDefaultValue.Append("LOCGEN_FORMAT_NAMED(");
+			defaultValueReader.Require('(');
+			if (!SanitizeDefaultValue(property, defaultValueReader, innerDefaultValue, false))
 			{
 				return false;
 			}
 
 			// RequireList assumes that we have already parsed the ','.  So if it is there, consume it.
-			DefaultValueReader.TryOptional(',');
+			defaultValueReader.TryOptional(',');
 
-			bool bSuccess = true;
-			DefaultValueReader.RequireList(')', ',', false, () =>
+			bool success = true;
+			defaultValueReader.RequireList(')', ',', false, () =>
 			{
-				StringView Value = DefaultValueReader.GetWrappedConstString();
-				DefaultValueReader.Require(',');
-				InnerDefaultValue.Append(", \"").Append(Value).Append("\", ");
-				if (!SanitizeDefaultValue(Property, DefaultValueReader, InnerDefaultValue, true))
+				StringView value = defaultValueReader.GetWrappedConstString();
+				defaultValueReader.Require(',');
+				innerDefaultValue.Append(", \"").Append(value).Append("\", ");
+				if (!SanitizeDefaultValue(property, defaultValueReader, innerDefaultValue, true))
 				{
-					bSuccess = false;
+					success = false;
 				}
 			});
-			InnerDefaultValue.Append(')');
-			return bSuccess;
+			innerDefaultValue.Append(')');
+			return success;
 		}
 
 		[UhtLocTextDefaultValue(Name = "LOCGEN_FORMAT_ORDERED")]
 		[SuppressMessage("CodeQuality", "IDE0051:Remove unused private members", Justification = "Attribute accessed method")]
 		[SuppressMessage("Style", "IDE0060:Remove unused parameter", Justification = "Attribute accessed method")]
-		private static bool LocGenFormatOrderedDefaultValue(UhtTextProperty Property, IUhtTokenReader DefaultValueReader, ref UhtToken MacroToken, StringBuilder InnerDefaultValue)
+		private static bool LocGenFormatOrderedDefaultValue(UhtTextProperty property, IUhtTokenReader defaultValueReader, ref UhtToken macroToken, StringBuilder innerDefaultValue)
 		{
-			InnerDefaultValue.Append("LOCGEN_FORMAT_ORDERED(");
-			DefaultValueReader.Require('(');
-			if (!SanitizeDefaultValue(Property, DefaultValueReader, InnerDefaultValue, false))
+			innerDefaultValue.Append("LOCGEN_FORMAT_ORDERED(");
+			defaultValueReader.Require('(');
+			if (!SanitizeDefaultValue(property, defaultValueReader, innerDefaultValue, false))
 			{
 				return false;
 			}
 
 			// RequireList assumes that we have already parsed the ','.  So if it is there, consume it.
-			DefaultValueReader.TryOptional(',');
+			defaultValueReader.TryOptional(',');
 
-			bool bSuccess = true;
-			DefaultValueReader.RequireList(')', ',', false, () =>
+			bool success = true;
+			defaultValueReader.RequireList(')', ',', false, () =>
 			{
-				InnerDefaultValue.Append(", ");
-				if (!SanitizeDefaultValue(Property, DefaultValueReader, InnerDefaultValue, true))
+				innerDefaultValue.Append(", ");
+				if (!SanitizeDefaultValue(property, defaultValueReader, innerDefaultValue, true))
 				{
-					bSuccess = false;
+					success = false;
 				}
 			});
-			InnerDefaultValue.Append(')');
-			return bSuccess;
+			innerDefaultValue.Append(')');
+			return success;
 		}
 
 		[UhtLocTextDefaultValue(Name = "LOCGEN_NUMBER")]
 		[SuppressMessage("CodeQuality", "IDE0051:Remove unused private members", Justification = "Attribute accessed method")]
 		[SuppressMessage("Style", "IDE0060:Remove unused parameter", Justification = "Attribute accessed method")]
-		private static bool LocGenNumberDefaultValue(UhtTextProperty Property, IUhtTokenReader DefaultValueReader, ref UhtToken MacroToken, StringBuilder InnerDefaultValue)
+		private static bool LocGenNumberDefaultValue(UhtTextProperty property, IUhtTokenReader defaultValueReader, ref UhtToken macroToken, StringBuilder innerDefaultValue)
 		{
-			return LocGenNumberOrPercentDefaultValue(DefaultValueReader, ref MacroToken, InnerDefaultValue, NumberStyle.None);
+			return LocGenNumberOrPercentDefaultValue(defaultValueReader, ref macroToken, innerDefaultValue, NumberStyle.None);
 		}
 
 		[UhtLocTextDefaultValue(Name = "LOCGEN_NUMBER_GROUPED")]
 		[SuppressMessage("CodeQuality", "IDE0051:Remove unused private members", Justification = "Attribute accessed method")]
 		[SuppressMessage("Style", "IDE0060:Remove unused parameter", Justification = "Attribute accessed method")]
-		private static bool LocGenNumberGroupedDefaultValue(UhtTextProperty Property, IUhtTokenReader DefaultValueReader, ref UhtToken MacroToken, StringBuilder InnerDefaultValue)
+		private static bool LocGenNumberGroupedDefaultValue(UhtTextProperty property, IUhtTokenReader defaultValueReader, ref UhtToken macroToken, StringBuilder innerDefaultValue)
 		{
-			return LocGenNumberOrPercentDefaultValue(DefaultValueReader, ref MacroToken, InnerDefaultValue, NumberStyle.Grouped);
+			return LocGenNumberOrPercentDefaultValue(defaultValueReader, ref macroToken, innerDefaultValue, NumberStyle.Grouped);
 		}
 
 		[UhtLocTextDefaultValue(Name = "LOCGEN_NUMBER_UNGROUPED")]
 		[SuppressMessage("CodeQuality", "IDE0051:Remove unused private members", Justification = "Attribute accessed method")]
 		[SuppressMessage("Style", "IDE0060:Remove unused parameter", Justification = "Attribute accessed method")]
-		private static bool LocGenNumberUngroupedDefaultValue(UhtTextProperty Property, IUhtTokenReader DefaultValueReader, ref UhtToken MacroToken, StringBuilder InnerDefaultValue)
+		private static bool LocGenNumberUngroupedDefaultValue(UhtTextProperty property, IUhtTokenReader defaultValueReader, ref UhtToken macroToken, StringBuilder innerDefaultValue)
 		{
-			return LocGenNumberOrPercentDefaultValue(DefaultValueReader, ref MacroToken, InnerDefaultValue, NumberStyle.Ungrouped);
+			return LocGenNumberOrPercentDefaultValue(defaultValueReader, ref macroToken, innerDefaultValue, NumberStyle.Ungrouped);
 		}
 
 		[UhtLocTextDefaultValue(Name = "LOCGEN_NUMBER_CUSTOM")]
 		[SuppressMessage("CodeQuality", "IDE0051:Remove unused private members", Justification = "Attribute accessed method")]
 		[SuppressMessage("Style", "IDE0060:Remove unused parameter", Justification = "Attribute accessed method")]
-		private static bool LocGenNumberCustomDefaultValue(UhtTextProperty Property, IUhtTokenReader DefaultValueReader, ref UhtToken MacroToken, StringBuilder InnerDefaultValue)
+		private static bool LocGenNumberCustomDefaultValue(UhtTextProperty property, IUhtTokenReader defaultValueReader, ref UhtToken macroToken, StringBuilder innerDefaultValue)
 		{
-			return LocGenNumberOrPercentDefaultValue(DefaultValueReader, ref MacroToken, InnerDefaultValue, NumberStyle.Custom);
+			return LocGenNumberOrPercentDefaultValue(defaultValueReader, ref macroToken, innerDefaultValue, NumberStyle.Custom);
 		}
 
 		[UhtLocTextDefaultValue(Name = "LOCGEN_PERCENT")]
 		[SuppressMessage("CodeQuality", "IDE0051:Remove unused private members", Justification = "Attribute accessed method")]
 		[SuppressMessage("Style", "IDE0060:Remove unused parameter", Justification = "Attribute accessed method")]
-		private static bool LocGenPercentDefaultValue(UhtTextProperty Property, IUhtTokenReader DefaultValueReader, ref UhtToken MacroToken, StringBuilder InnerDefaultValue)
+		private static bool LocGenPercentDefaultValue(UhtTextProperty property, IUhtTokenReader defaultValueReader, ref UhtToken macroToken, StringBuilder innerDefaultValue)
 		{
-			return LocGenNumberOrPercentDefaultValue(DefaultValueReader, ref MacroToken, InnerDefaultValue, NumberStyle.None);
+			return LocGenNumberOrPercentDefaultValue(defaultValueReader, ref macroToken, innerDefaultValue, NumberStyle.None);
 		}
 
 		[UhtLocTextDefaultValue(Name = "LOCGEN_PERCENT_GROUPED")]
 		[SuppressMessage("CodeQuality", "IDE0051:Remove unused private members", Justification = "Attribute accessed method")]
 		[SuppressMessage("Style", "IDE0060:Remove unused parameter", Justification = "Attribute accessed method")]
-		private static bool LocGenPercentGroupedDefaultValue(UhtTextProperty Property, IUhtTokenReader DefaultValueReader, ref UhtToken MacroToken, StringBuilder InnerDefaultValue)
+		private static bool LocGenPercentGroupedDefaultValue(UhtTextProperty property, IUhtTokenReader defaultValueReader, ref UhtToken macroToken, StringBuilder innerDefaultValue)
 		{
-			return LocGenNumberOrPercentDefaultValue(DefaultValueReader, ref MacroToken, InnerDefaultValue, NumberStyle.Grouped);
+			return LocGenNumberOrPercentDefaultValue(defaultValueReader, ref macroToken, innerDefaultValue, NumberStyle.Grouped);
 		}
 
 		[UhtLocTextDefaultValue(Name = "LOCGEN_PERCENT_UNGROUPED")]
 		[SuppressMessage("CodeQuality", "IDE0051:Remove unused private members", Justification = "Attribute accessed method")]
 		[SuppressMessage("Style", "IDE0060:Remove unused parameter", Justification = "Attribute accessed method")]
-		private static bool LocGenPercentUngroupedDefaultValue(UhtTextProperty Property, IUhtTokenReader DefaultValueReader, ref UhtToken MacroToken, StringBuilder InnerDefaultValue)
+		private static bool LocGenPercentUngroupedDefaultValue(UhtTextProperty property, IUhtTokenReader defaultValueReader, ref UhtToken macroToken, StringBuilder innerDefaultValue)
 		{
-			return LocGenNumberOrPercentDefaultValue(DefaultValueReader, ref MacroToken, InnerDefaultValue, NumberStyle.Ungrouped);
+			return LocGenNumberOrPercentDefaultValue(defaultValueReader, ref macroToken, innerDefaultValue, NumberStyle.Ungrouped);
 		}
 
 		[UhtLocTextDefaultValue(Name = "LOCGEN_PERCENT_CUSTOM")]
 		[SuppressMessage("CodeQuality", "IDE0051:Remove unused private members", Justification = "Attribute accessed method")]
 		[SuppressMessage("Style", "IDE0060:Remove unused parameter", Justification = "Attribute accessed method")]
-		private static bool LocGenPercentCustomDefaultValue(UhtTextProperty Property, IUhtTokenReader DefaultValueReader, ref UhtToken MacroToken, StringBuilder InnerDefaultValue)
+		private static bool LocGenPercentCustomDefaultValue(UhtTextProperty property, IUhtTokenReader defaultValueReader, ref UhtToken macroToken, StringBuilder innerDefaultValue)
 		{
-			return LocGenNumberOrPercentDefaultValue(DefaultValueReader, ref MacroToken, InnerDefaultValue, NumberStyle.Custom);
+			return LocGenNumberOrPercentDefaultValue(defaultValueReader, ref macroToken, innerDefaultValue, NumberStyle.Custom);
 		}
 
 		[UhtLocTextDefaultValue(Name = "LOCGEN_CURRENCY")]
 		[SuppressMessage("CodeQuality", "IDE0051:Remove unused private members", Justification = "Attribute accessed method")]
 		[SuppressMessage("Style", "IDE0060:Remove unused parameter", Justification = "Attribute accessed method")]
-		private static bool LocGenCurrencyDefaultValue(UhtTextProperty Property, IUhtTokenReader DefaultValueReader, ref UhtToken MacroToken, StringBuilder InnerDefaultValue)
+		private static bool LocGenCurrencyDefaultValue(UhtTextProperty property, IUhtTokenReader defaultValueReader, ref UhtToken macroToken, StringBuilder innerDefaultValue)
 		{
-			DefaultValueReader.Require('(');
+			defaultValueReader.Require('(');
 
-			double BaseValue;
-			UhtToken Token = DefaultValueReader.GetToken();
-			if (!Token.GetConstDouble(out BaseValue))
+			double baseValue;
+			UhtToken token = defaultValueReader.GetToken();
+			if (!token.GetConstDouble(out baseValue))
 			{
 				return false;
 			}
 
-			DefaultValueReader.Require(',');
-			StringView CurrencyCode = DefaultValueReader.GetConstQuotedString();
-			DefaultValueReader.Require(',');
-			StringView CultureNameString = DefaultValueReader.GetConstQuotedString();
-			DefaultValueReader.Require(')');
+			defaultValueReader.Require(',');
+			StringView currencyCode = defaultValueReader.GetConstQuotedString();
+			defaultValueReader.Require(',');
+			StringView cultureNameString = defaultValueReader.GetConstQuotedString();
+			defaultValueReader.Require(')');
 
 			// With UHT, we end up outputting just the integer part of the specified value.  It is using default locale information in UHT
-			InnerDefaultValue.Append("LOCGEN_CURRENCY(").Append((int)BaseValue).Append(", ").Append(CurrencyCode).Append(", ").Append(CultureNameString).Append(')');
+			innerDefaultValue.Append("LOCGEN_CURRENCY(").Append((int)baseValue).Append(", ").Append(currencyCode).Append(", ").Append(cultureNameString).Append(')');
 			return true;
 		}
 
 		[UhtLocTextDefaultValue(Name = "LOCGEN_DATE_LOCAL")]
 		[SuppressMessage("CodeQuality", "IDE0051:Remove unused private members", Justification = "Attribute accessed method")]
 		[SuppressMessage("Style", "IDE0060:Remove unused parameter", Justification = "Attribute accessed method")]
-		private static bool LocGenDateLocalDefaultValue(UhtTextProperty Property, IUhtTokenReader DefaultValueReader, ref UhtToken MacroToken, StringBuilder InnerDefaultValue)
+		private static bool LocGenDateLocalDefaultValue(UhtTextProperty property, IUhtTokenReader defaultValueReader, ref UhtToken macroToken, StringBuilder innerDefaultValue)
 		{
-			return LocGenDateTimeDefaultValue(DefaultValueReader, ref MacroToken, InnerDefaultValue, true, false, false, false);
+			return LocGenDateTimeDefaultValue(defaultValueReader, ref macroToken, innerDefaultValue, true, false, false, false);
 		}
 
 		[UhtLocTextDefaultValue(Name = "LOCGEN_DATE_UTC")]
 		[SuppressMessage("CodeQuality", "IDE0051:Remove unused private members", Justification = "Attribute accessed method")]
 		[SuppressMessage("Style", "IDE0060:Remove unused parameter", Justification = "Attribute accessed method")]
-		private static bool LocGenDateUtcDefaultValue(UhtTextProperty Property, IUhtTokenReader DefaultValueReader, ref UhtToken MacroToken, StringBuilder InnerDefaultValue)
+		private static bool LocGenDateUtcDefaultValue(UhtTextProperty property, IUhtTokenReader defaultValueReader, ref UhtToken macroToken, StringBuilder innerDefaultValue)
 		{
-			return LocGenDateTimeDefaultValue(DefaultValueReader, ref MacroToken, InnerDefaultValue, true, false, true, false);
+			return LocGenDateTimeDefaultValue(defaultValueReader, ref macroToken, innerDefaultValue, true, false, true, false);
 		}
 
 		[UhtLocTextDefaultValue(Name = "LOCGEN_TIME_LOCAL")]
 		[SuppressMessage("CodeQuality", "IDE0051:Remove unused private members", Justification = "Attribute accessed method")]
 		[SuppressMessage("Style", "IDE0060:Remove unused parameter", Justification = "Attribute accessed method")]
-		private static bool LocGenTimeLocalDefaultValue(UhtTextProperty Property, IUhtTokenReader DefaultValueReader, ref UhtToken MacroToken, StringBuilder InnerDefaultValue)
+		private static bool LocGenTimeLocalDefaultValue(UhtTextProperty property, IUhtTokenReader defaultValueReader, ref UhtToken macroToken, StringBuilder innerDefaultValue)
 		{
-			return LocGenDateTimeDefaultValue(DefaultValueReader, ref MacroToken, InnerDefaultValue, false, true, false, false);
+			return LocGenDateTimeDefaultValue(defaultValueReader, ref macroToken, innerDefaultValue, false, true, false, false);
 		}
 
 		[UhtLocTextDefaultValue(Name = "LOCGEN_TIME_UTC")]
 		[SuppressMessage("CodeQuality", "IDE0051:Remove unused private members", Justification = "Attribute accessed method")]
 		[SuppressMessage("Style", "IDE0060:Remove unused parameter", Justification = "Attribute accessed method")]
-		private static bool LocGenTimeUtcDefaultValue(UhtTextProperty Property, IUhtTokenReader DefaultValueReader, ref UhtToken MacroToken, StringBuilder InnerDefaultValue)
+		private static bool LocGenTimeUtcDefaultValue(UhtTextProperty property, IUhtTokenReader defaultValueReader, ref UhtToken macroToken, StringBuilder innerDefaultValue)
 		{
-			return LocGenDateTimeDefaultValue(DefaultValueReader, ref MacroToken, InnerDefaultValue, false, true, true, false);
+			return LocGenDateTimeDefaultValue(defaultValueReader, ref macroToken, innerDefaultValue, false, true, true, false);
 		}
 
 		[UhtLocTextDefaultValue(Name = "LOCGEN_DATETIME_LOCAL")]
 		[SuppressMessage("CodeQuality", "IDE0051:Remove unused private members", Justification = "Attribute accessed method")]
 		[SuppressMessage("Style", "IDE0060:Remove unused parameter", Justification = "Attribute accessed method")]
-		private static bool LocGenDateTimeLocalDefaultValue(UhtTextProperty Property, IUhtTokenReader DefaultValueReader, ref UhtToken MacroToken, StringBuilder InnerDefaultValue)
+		private static bool LocGenDateTimeLocalDefaultValue(UhtTextProperty property, IUhtTokenReader defaultValueReader, ref UhtToken macroToken, StringBuilder innerDefaultValue)
 		{
-			return LocGenDateTimeDefaultValue(DefaultValueReader, ref MacroToken, InnerDefaultValue, true, true, false, false);
+			return LocGenDateTimeDefaultValue(defaultValueReader, ref macroToken, innerDefaultValue, true, true, false, false);
 		}
 
 		[UhtLocTextDefaultValue(Name = "LOCGEN_DATETIME_UTC")]
 		[SuppressMessage("CodeQuality", "IDE0051:Remove unused private members", Justification = "Attribute accessed method")]
 		[SuppressMessage("Style", "IDE0060:Remove unused parameter", Justification = "Attribute accessed method")]
-		private static bool LocGenDateTimeUtcDefaultValue(UhtTextProperty Property, IUhtTokenReader DefaultValueReader, ref UhtToken MacroToken, StringBuilder InnerDefaultValue)
+		private static bool LocGenDateTimeUtcDefaultValue(UhtTextProperty property, IUhtTokenReader defaultValueReader, ref UhtToken macroToken, StringBuilder innerDefaultValue)
 		{
-			return LocGenDateTimeDefaultValue(DefaultValueReader, ref MacroToken, InnerDefaultValue, true, true, true, false);
+			return LocGenDateTimeDefaultValue(defaultValueReader, ref macroToken, innerDefaultValue, true, true, true, false);
 		}
 
 		[UhtLocTextDefaultValue(Name = "LOCGEN_DATETIME_CUSTOM_LOCAL")]
 		[SuppressMessage("CodeQuality", "IDE0051:Remove unused private members", Justification = "Attribute accessed method")]
 		[SuppressMessage("Style", "IDE0060:Remove unused parameter", Justification = "Attribute accessed method")]
-		private static bool LocGenDateTimeCustomLocalDefaultValue(UhtTextProperty Property, IUhtTokenReader DefaultValueReader, ref UhtToken MacroToken, StringBuilder InnerDefaultValue)
+		private static bool LocGenDateTimeCustomLocalDefaultValue(UhtTextProperty property, IUhtTokenReader defaultValueReader, ref UhtToken macroToken, StringBuilder innerDefaultValue)
 		{
-			return LocGenDateTimeDefaultValue(DefaultValueReader, ref MacroToken, InnerDefaultValue, true, true, false, true);
+			return LocGenDateTimeDefaultValue(defaultValueReader, ref macroToken, innerDefaultValue, true, true, false, true);
 		}
 
 		[UhtLocTextDefaultValue(Name = "LOCGEN_DATETIME_CUSTOM_UTC")]
 		[SuppressMessage("CodeQuality", "IDE0051:Remove unused private members", Justification = "Attribute accessed method")]
 		[SuppressMessage("Style", "IDE0060:Remove unused parameter", Justification = "Attribute accessed method")]
-		private static bool LocGenDateTimeCustomUtcDefaultValue(UhtTextProperty Property, IUhtTokenReader DefaultValueReader, ref UhtToken MacroToken, StringBuilder InnerDefaultValue)
+		private static bool LocGenDateTimeCustomUtcDefaultValue(UhtTextProperty property, IUhtTokenReader defaultValueReader, ref UhtToken macroToken, StringBuilder innerDefaultValue)
 		{
-			return LocGenDateTimeDefaultValue(DefaultValueReader, ref MacroToken, InnerDefaultValue, true, true, true, true);
+			return LocGenDateTimeDefaultValue(defaultValueReader, ref macroToken, innerDefaultValue, true, true, true, true);
 		}
 
 		[UhtLocTextDefaultValue(Name = "LOCGEN_TOUPPER")]
 		[SuppressMessage("CodeQuality", "IDE0051:Remove unused private members", Justification = "Attribute accessed method")]
-		[SuppressMessage("Style", "IDE0060:Remove unused parameter", Justification = "Attribute accessed method")]
-		private static bool LocGenToUpperDefaultValue(UhtTextProperty Property, IUhtTokenReader DefaultValueReader, ref UhtToken MacroToken, StringBuilder InnerDefaultValue)
+		private static bool LocGenToUpperDefaultValue(UhtTextProperty property, IUhtTokenReader defaultValueReader, ref UhtToken macroToken, StringBuilder innerDefaultValue)
 		{
-			return LocGenTransformDefaultValue(Property, DefaultValueReader, ref MacroToken, InnerDefaultValue);
+			return LocGenTransformDefaultValue(property, defaultValueReader, ref macroToken, innerDefaultValue);
 		}
 
 		[UhtLocTextDefaultValue(Name = "LOCGEN_TOLOWER")]
 		[SuppressMessage("CodeQuality", "IDE0051:Remove unused private members", Justification = "Attribute accessed method")]
-		[SuppressMessage("Style", "IDE0060:Remove unused parameter", Justification = "Attribute accessed method")]
-		private static bool LocGenToLowerDefaultValue(UhtTextProperty Property, IUhtTokenReader DefaultValueReader, ref UhtToken MacroToken, StringBuilder InnerDefaultValue)
+		private static bool LocGenToLowerDefaultValue(UhtTextProperty property, IUhtTokenReader defaultValueReader, ref UhtToken macroToken, StringBuilder innerDefaultValue)
 		{
-			return LocGenTransformDefaultValue(Property, DefaultValueReader, ref MacroToken, InnerDefaultValue);
+			return LocGenTransformDefaultValue(property, defaultValueReader, ref macroToken, innerDefaultValue);
 		}
 		#endregion
 
 		#region Default value helper methods
-		private static bool SanitizeDefaultValue(UhtTextProperty Property, IUhtTokenReader DefaultValueReader, StringBuilder InnerDefaultValue, bool bAllowNumerics)
+		private static bool SanitizeDefaultValue(UhtTextProperty property, IUhtTokenReader defaultValueReader, StringBuilder innerDefaultValue, bool allowNumerics)
 		{
-			UhtToken Token = DefaultValueReader.GetToken();
-			return SanitizeDefaultValue(Property, DefaultValueReader, ref Token, InnerDefaultValue, bAllowNumerics);
+			UhtToken token = defaultValueReader.GetToken();
+			return SanitizeDefaultValue(property, defaultValueReader, ref token, innerDefaultValue, allowNumerics);
 		}
 
-		private static bool SanitizeDefaultValue(UhtTextProperty Property, IUhtTokenReader DefaultValueReader, ref UhtToken Token, StringBuilder InnerDefaultValue, bool bAllowNumerics)
+		private static bool SanitizeDefaultValue(UhtTextProperty property, IUhtTokenReader defaultValueReader, ref UhtToken token, StringBuilder innerDefaultValue, bool allowNumerics)
 		{
-			switch (Token.TokenType)
+			switch (token.TokenType)
 			{
 				case UhtTokenType.Identifier:
-					UhtLocTextDefaultValue LocTextDefaultValue;
-					if (Property.Session.TryGetLocTextDefaultValue(Token.Value, out LocTextDefaultValue))
+					UhtLocTextDefaultValue locTextDefaultValue;
+					if (property.Session.TryGetLocTextDefaultValue(token.Value, out locTextDefaultValue))
 					{
-						if (LocTextDefaultValue.Delegate(Property, DefaultValueReader, ref Token, InnerDefaultValue))
+						if (locTextDefaultValue.Delegate(property, defaultValueReader, ref token, innerDefaultValue))
 						{
 							return true;
 						}
@@ -455,21 +453,21 @@ namespace EpicGames.UHT.Types
 					return false;
 
 				case UhtTokenType.StringConst:
-					InnerDefaultValue.Append(Token.Value);
+					innerDefaultValue.Append(token.Value);
 					return true;
 
 				case UhtTokenType.DecimalConst:
-					if (bAllowNumerics)
+					if (allowNumerics)
 					{
-						FormatDecimal(ref Token, InnerDefaultValue);
+						FormatDecimal(ref token, innerDefaultValue);
 						return true;
 					}
 					return false;
 
 				case UhtTokenType.FloatConst:
-					if (bAllowNumerics)
+					if (allowNumerics)
 					{
-						FormatFloat(ref Token, InnerDefaultValue);
+						FormatFloat(ref token, innerDefaultValue);
 						return true;
 					}
 					return false;
@@ -479,51 +477,51 @@ namespace EpicGames.UHT.Types
 			}
 		}
 
-		private static void FormatDecimal(ref UhtToken Token, StringBuilder InnerDefaultValue)
+		private static void FormatDecimal(ref UhtToken token, StringBuilder innerDefaultValue)
 		{
-			ReadOnlySpan<char> Span = Token.Value.Span;
-			int SuffixStart = Span.Length;
-			bool bIsUnsigned = false;
-			while (SuffixStart > 0)
+			ReadOnlySpan<char> span = token.Value.Span;
+			int suffixStart = span.Length;
+			bool isUnsigned = false;
+			while (suffixStart > 0)
 			{
-				char C = Span[SuffixStart - 1];
-				if (UhtFCString.IsUnsignedMarker(C))
+				char c = span[suffixStart - 1];
+				if (UhtFCString.IsUnsignedMarker(c))
 				{
-					bIsUnsigned = true;
-					SuffixStart--;
+					isUnsigned = true;
+					suffixStart--;
 				}
-				else if (UhtFCString.IsLongMarker(C))
+				else if (UhtFCString.IsLongMarker(c))
 				{
-					SuffixStart--;
+					suffixStart--;
 				}
 				else
 				{
 					break;
 				}
 			}
-			if (Token.GetConstLong(out long Value))
+			if (token.GetConstLong(out long value))
 			{
-				if (bIsUnsigned)
+				if (isUnsigned)
 				{
-					InnerDefaultValue.Append((ulong)Value).Append('u');
+					innerDefaultValue.Append((ulong)value).Append('u');
 				}
 				else
 				{
-					InnerDefaultValue.Append(Value);
+					innerDefaultValue.Append(value);
 				}
 			}
 		}
 
-		private static void FormatFloat(ref UhtToken Token, StringBuilder InnerDefaultValue)
+		private static void FormatFloat(ref UhtToken token, StringBuilder innerDefaultValue)
 		{
-			char C = Token.Value.Span[Token.Value.Length - 1];
-			bool bIsFloat = UhtFCString.IsFloatMarker(C);
-			if (Token.GetConstDouble(out double Value))
+			char c = token.Value.Span[^1];
+			bool isFloat = UhtFCString.IsFloatMarker(c);
+			if (token.GetConstDouble(out double value))
 			{
-				InnerDefaultValue.AppendFormat("{0:F6}", Value);
-				if (bIsFloat)
+				innerDefaultValue.AppendFormat("{0:F6}", value);
+				if (isFloat)
 				{
-					InnerDefaultValue.Append(C);
+					innerDefaultValue.Append(c);
 				}
 			}
 		}
@@ -536,56 +534,56 @@ namespace EpicGames.UHT.Types
 			Custom,
 		}
 
-		private static bool LocGenNumberOrPercentDefaultValue(IUhtTokenReader DefaultValueReader, ref UhtToken MacroToken, StringBuilder InnerDefaultValue, NumberStyle NumberStyle)
+		private static bool LocGenNumberOrPercentDefaultValue(IUhtTokenReader defaultValueReader, ref UhtToken macroToken, StringBuilder innerDefaultValue, NumberStyle numberStyle)
 		{
-			InnerDefaultValue.Append(MacroToken.Value.ToString());
-			DefaultValueReader.Require('(');
-			InnerDefaultValue.Append('(');
+			innerDefaultValue.Append(macroToken.Value.ToString());
+			defaultValueReader.Require('(');
+			innerDefaultValue.Append('(');
 
-			UhtToken Token = DefaultValueReader.GetToken();
-			switch (Token.TokenType)
+			UhtToken token = defaultValueReader.GetToken();
+			switch (token.TokenType)
 			{
 				case UhtTokenType.DecimalConst:
-					FormatDecimal(ref Token, InnerDefaultValue);
+					FormatDecimal(ref token, innerDefaultValue);
 					break;
 
 				case UhtTokenType.FloatConst:
-					FormatFloat(ref Token, InnerDefaultValue);
+					FormatFloat(ref token, innerDefaultValue);
 					break;
 
 				default:
 					return false;
 			}
-			DefaultValueReader.Require(',');
-			InnerDefaultValue.Append(", ");
+			defaultValueReader.Require(',');
+			innerDefaultValue.Append(", ");
 
-			if (NumberStyle == NumberStyle.Custom)
+			if (numberStyle == NumberStyle.Custom)
 			{
 				while (true)
 				{
-					UhtToken CustomToken = DefaultValueReader.GetToken();
-					if (!CustomToken.IsIdentifier())
+					UhtToken customToken = defaultValueReader.GetToken();
+					if (!customToken.IsIdentifier())
 					{
 						return false;
 					}
-					InnerDefaultValue.Append(CustomToken.Value);
+					innerDefaultValue.Append(customToken.Value);
 
-					DefaultValueReader.Require('(');
-					InnerDefaultValue.Append('(');
+					defaultValueReader.Require('(');
+					innerDefaultValue.Append('(');
 
-					switch (CustomToken.Value.ToString())
+					switch (customToken.Value.ToString())
 					{
 						case "SetAlwaysSign":
 						case "SetUseGrouping":
 							{
-								UhtToken BooleanToken = DefaultValueReader.GetToken();
-								if (BooleanToken.IsIdentifier("true"))
+								UhtToken booleanToken = defaultValueReader.GetToken();
+								if (booleanToken.IsIdentifier("true"))
 								{
-									InnerDefaultValue.Append("true");
+									innerDefaultValue.Append("true");
 								}
-								else if (BooleanToken.IsIdentifier("false"))
+								else if (booleanToken.IsIdentifier("false"))
 								{
-									InnerDefaultValue.Append("false");
+									innerDefaultValue.Append("false");
 								}
 								else
 								{
@@ -596,10 +594,10 @@ namespace EpicGames.UHT.Types
 
 						case "SetRoundingMode":
 							{
-								DefaultValueReader.Require("ERoundingMode");
-								DefaultValueReader.Require("::");
-								StringView Identifier = DefaultValueReader.GetIdentifier().Value;
-								InnerDefaultValue.Append("ERoundingMode::").Append(Identifier);
+								defaultValueReader.Require("ERoundingMode");
+								defaultValueReader.Require("::");
+								StringView identifier = defaultValueReader.GetIdentifier().Value;
+								innerDefaultValue.Append("ERoundingMode::").Append(identifier);
 							}
 							break;
 
@@ -608,10 +606,10 @@ namespace EpicGames.UHT.Types
 						case "SetMinimumFractionalDigits":
 						case "MaximumFractionalDigits":
 							{
-								UhtToken NumericToken = DefaultValueReader.GetToken();
-								if (NumericToken.GetConstInt(out int Value))
+								UhtToken numericToken = defaultValueReader.GetToken();
+								if (numericToken.GetConstInt(out int value))
 								{
-									InnerDefaultValue.Append(Value);
+									innerDefaultValue.Append(value);
 								}
 								else
 								{
@@ -624,50 +622,50 @@ namespace EpicGames.UHT.Types
 							return false;
 					}
 
-					DefaultValueReader.Require(')');
-					InnerDefaultValue.Append(')');
+					defaultValueReader.Require(')');
+					innerDefaultValue.Append(')');
 
-					if (!DefaultValueReader.TryOptional('.'))
+					if (!defaultValueReader.TryOptional('.'))
 					{
 						break;
 					}
-					InnerDefaultValue.Append('.');
+					innerDefaultValue.Append('.');
 				}
 
-				DefaultValueReader.Require(',');
-				InnerDefaultValue.Append(", ");
+				defaultValueReader.Require(',');
+				innerDefaultValue.Append(", ");
 			}
 
-			StringView _ = DefaultValueReader.GetConstQuotedString();
-			InnerDefaultValue.Append("\"\""); // UHT doesn't write out the culture
+			StringView _ = defaultValueReader.GetConstQuotedString();
+			innerDefaultValue.Append("\"\""); // UHT doesn't write out the culture
 
-			DefaultValueReader.Require(')');
-			InnerDefaultValue.Append(')');
+			defaultValueReader.Require(')');
+			innerDefaultValue.Append(')');
 			return true;
 		}
 
-		private static void FormatDateTimeStyle(IUhtTokenReader DefaultValueReader, StringBuilder InnerDefaultValue)
+		private static void FormatDateTimeStyle(IUhtTokenReader defaultValueReader, StringBuilder innerDefaultValue)
 		{
-			DefaultValueReader.Require(',');
-			DefaultValueReader.Require("EDateTimeStyle");
-			DefaultValueReader.Require("::");
-			StringView Identifier = DefaultValueReader.GetIdentifier().Value;
-			InnerDefaultValue.Append(", EDateTimeStyle::").Append(Identifier);
+			defaultValueReader.Require(',');
+			defaultValueReader.Require("EDateTimeStyle");
+			defaultValueReader.Require("::");
+			StringView identifier = defaultValueReader.GetIdentifier().Value;
+			innerDefaultValue.Append(", EDateTimeStyle::").Append(identifier);
 		}
 
-		private static bool LocGenDateTimeDefaultValue(IUhtTokenReader DefaultValueReader, ref UhtToken MacroToken, StringBuilder InnerDefaultValue, bool bDate, bool bTime, bool bUtc, bool bCustom)
+		private static bool LocGenDateTimeDefaultValue(IUhtTokenReader defaultValueReader, ref UhtToken macroToken, StringBuilder innerDefaultValue, bool isDate, bool isTime, bool isUtc, bool isCustom)
 		{
-			DefaultValueReader.Require('(');
-			InnerDefaultValue.Append(MacroToken.Value).Append('(');
+			defaultValueReader.Require('(');
+			innerDefaultValue.Append(macroToken.Value).Append('(');
 
-			long UnixTimestampValue;
-			UhtToken TimestampToken = DefaultValueReader.GetToken();
-			switch (TimestampToken.TokenType)
+			long unixTimestampValue;
+			UhtToken timestampToken = defaultValueReader.GetToken();
+			switch (timestampToken.TokenType)
 			{
 				case UhtTokenType.FloatConst:
-					if (TimestampToken.GetConstDouble(out double Value))
+					if (timestampToken.GetConstDouble(out double value))
 					{
-						UnixTimestampValue = (long)Value;
+						unixTimestampValue = (long)value;
 						break;
 					}
 					else
@@ -676,7 +674,7 @@ namespace EpicGames.UHT.Types
 					}
 
 				case UhtTokenType.DecimalConst:
-					if (!TimestampToken.GetConstLong(out UnixTimestampValue))
+					if (!timestampToken.GetConstLong(out unixTimestampValue))
 					{
 						return false;
 					}
@@ -686,52 +684,52 @@ namespace EpicGames.UHT.Types
 					return false;
 
 			}
-			InnerDefaultValue.Append(UnixTimestampValue);
+			innerDefaultValue.Append(unixTimestampValue);
 
-			if (bCustom)
+			if (isCustom)
 			{
-				DefaultValueReader.Require(',');
-				StringView CustomPattern = DefaultValueReader.GetConstQuotedString();
-				InnerDefaultValue.Append(", ").Append(CustomPattern);
+				defaultValueReader.Require(',');
+				StringView customPattern = defaultValueReader.GetConstQuotedString();
+				innerDefaultValue.Append(", ").Append(customPattern);
 			}
 			else
 			{
-				if (bDate)
+				if (isDate)
 				{
-					FormatDateTimeStyle(DefaultValueReader, InnerDefaultValue);
+					FormatDateTimeStyle(defaultValueReader, innerDefaultValue);
 				}
-				if (bTime)
+				if (isTime)
 				{
-					FormatDateTimeStyle(DefaultValueReader, InnerDefaultValue);
+					FormatDateTimeStyle(defaultValueReader, innerDefaultValue);
 				}
 			}
 
-			if (bUtc)
+			if (isUtc)
 			{
-				DefaultValueReader.Require(',');
-				StringView TimeZone = DefaultValueReader.GetConstQuotedString();
-				InnerDefaultValue.Append(", ").Append(TimeZone);
+				defaultValueReader.Require(',');
+				StringView timeZone = defaultValueReader.GetConstQuotedString();
+				innerDefaultValue.Append(", ").Append(timeZone);
 			}
 
-			DefaultValueReader.Require(',');
-			StringView _ = DefaultValueReader.GetConstQuotedString();
-			InnerDefaultValue.Append(", ").Append("\"\""); // We don't write out the culture
+			defaultValueReader.Require(',');
+			StringView _ = defaultValueReader.GetConstQuotedString();
+			innerDefaultValue.Append(", ").Append("\"\""); // We don't write out the culture
 
-			DefaultValueReader.Require(')');
-			InnerDefaultValue.Append(')');
+			defaultValueReader.Require(')');
+			innerDefaultValue.Append(')');
 			return true;
 		}
 
-		private static bool LocGenTransformDefaultValue(UhtTextProperty Property, IUhtTokenReader DefaultValueReader, ref UhtToken MacroToken, StringBuilder InnerDefaultValue)
+		private static bool LocGenTransformDefaultValue(UhtTextProperty property, IUhtTokenReader defaultValueReader, ref UhtToken macroToken, StringBuilder innerDefaultValue)
 		{
-			InnerDefaultValue.Append(MacroToken.Value).Append('(');
-			DefaultValueReader.Require('(');
-			if (!SanitizeDefaultValue(Property, DefaultValueReader, InnerDefaultValue, false))
+			innerDefaultValue.Append(macroToken.Value).Append('(');
+			defaultValueReader.Require('(');
+			if (!SanitizeDefaultValue(property, defaultValueReader, innerDefaultValue, false))
 			{
 				return false;
 			}
-			DefaultValueReader.Require(')');
-			InnerDefaultValue.Append(')');
+			defaultValueReader.Require(')');
+			innerDefaultValue.Append(')');
 			return true;
 		}
 
diff --git a/Engine/Source/Programs/Shared/EpicGames.UHT/Types/Properties/UhtUInt16Property.cs b/Engine/Source/Programs/Shared/EpicGames.UHT/Types/Properties/UhtUInt16Property.cs
--- a/Engine/Source/Programs/Shared/EpicGames.UHT/Types/Properties/UhtUInt16Property.cs
+++ b/Engine/Source/Programs/Shared/EpicGames.UHT/Types/Properties/UhtUInt16Property.cs
@@ -1,10 +1,10 @@
 // Copyright Epic Games, Inc. All Rights Reserved.
 
+using System.Diagnostics.CodeAnalysis;
+using System.Text;
 using EpicGames.UHT.Tables;
 using EpicGames.UHT.Tokenizer;
 using EpicGames.UHT.Utils;
-using System.Diagnostics.CodeAnalysis;
-using System.Text;
 
 namespace EpicGames.UHT.Types
 {
@@ -17,59 +17,59 @@ namespace EpicGames.UHT.Types
 	public class UhtUInt16Property : UhtNumericProperty
 	{
 		/// <inheritdoc/>
-		public override string EngineClassName { get => "UInt16Property"; }
+		public override string EngineClassName => "UInt16Property";
 
 		/// <inheritdoc/>
-		protected override string CppTypeText { get => "uint16"; }
+		protected override string CppTypeText => "uint16";
 
 		/// <summary>
 		/// Construct a new property
 		/// </summary>
-		/// <param name="PropertySettings">Property settings</param>
-		/// <param name="IntType">Integer type</param>
-		public UhtUInt16Property(UhtPropertySettings PropertySettings, UhtPropertyIntType IntType) : base(PropertySettings, IntType)
+		/// <param name="propertySettings">Property settings</param>
+		/// <param name="intType">Integer type</param>
+		public UhtUInt16Property(UhtPropertySettings propertySettings, UhtPropertyIntType intType) : base(propertySettings, intType)
 		{
 		}
 
 		/// <inheritdoc/>
-		public override StringBuilder AppendMemberDecl(StringBuilder Builder, IUhtPropertyMemberContext Context, string Name, string NameSuffix, int Tabs)
+		public override StringBuilder AppendMemberDecl(StringBuilder builder, IUhtPropertyMemberContext context, string name, string nameSuffix, int tabs)
 		{
-			return AppendMemberDecl(Builder, Context, Name, NameSuffix, Tabs, "FFInt16PropertyParams");
+			return AppendMemberDecl(builder, context, name, nameSuffix, tabs, "FFInt16PropertyParams");
 		}
 
 		/// <inheritdoc/>
-		public override StringBuilder AppendMemberDef(StringBuilder Builder, IUhtPropertyMemberContext Context, string Name, string NameSuffix, string? Offset, int Tabs)
+		public override StringBuilder AppendMemberDef(StringBuilder builder, IUhtPropertyMemberContext context, string name, string nameSuffix, string? offset, int tabs)
 		{
-			AppendMemberDefStart(Builder, Context, Name, NameSuffix, Offset, Tabs, "FFInt16PropertyParams", "UECodeGen_Private::EPropertyGenFlags::UInt16");
-			AppendMemberDefEnd(Builder, Context, Name, NameSuffix);
-			return Builder;
+			AppendMemberDefStart(builder, context, name, nameSuffix, offset, tabs, "FFInt16PropertyParams", "UECodeGen_Private::EPropertyGenFlags::UInt16");
+			AppendMemberDefEnd(builder, context, name, nameSuffix);
+			return builder;
 		}
 
 		/// <inheritdoc/>
-		public override bool SanitizeDefaultValue(IUhtTokenReader DefaultValueReader, StringBuilder InnerDefaultValue)
+		public override bool SanitizeDefaultValue(IUhtTokenReader defaultValueReader, StringBuilder innerDefaultValue)
 		{
 			return false;
 		}
 
 		/// <inheritdoc/>
-		public override bool IsSameType(UhtProperty Other)
+		public override bool IsSameType(UhtProperty other)
 		{
-			return Other is UhtUInt16Property;
+			return other is UhtUInt16Property;
 		}
 
 		#region Keyword
 		[UhtPropertyType(Keyword = "uint16", Options = UhtPropertyTypeOptions.Simple | UhtPropertyTypeOptions.Immediate)]
 		[SuppressMessage("CodeQuality", "IDE0051:Remove unused private members", Justification = "Attribute accessed method")]
 		[SuppressMessage("Style", "IDE0060:Remove unused parameter", Justification = "Attribute accessed method")]
-		private static UhtProperty? UInt16Property(UhtPropertyResolvePhase ResolvePhase, UhtPropertySettings PropertySettings, IUhtTokenReader TokenReader, UhtToken MatchedToken)
+		private static UhtProperty? UInt16Property(UhtPropertyResolvePhase resolvePhase, UhtPropertySettings propertySettings, IUhtTokenReader tokenReader, UhtToken matchedToken)
 		{
-			if (PropertySettings.bIsBitfield)
+			if (propertySettings.IsBitfield)
 			{
-				return new UhtBoolProperty(PropertySettings, UhtBoolType.UInt8);
+				return new UhtBoolProperty(propertySettings, UhtBoolType.UInt8);
 			}
 			else
 			{
-				return new UhtUInt16Property(PropertySettings, UhtPropertyIntType.Sized);
+				return new UhtUInt16Property(propertySettings, UhtPropertyIntType.Sized);
 			}
 		}
 		#endregion
