More coding standard fixes

diff --git a/Engine/Source/Programs/Shared/EpicGames.UHT/Types/Properties/UhtUInt32Property.cs b/Engine/Source/Programs/Shared/EpicGames.UHT/Types/Properties/UhtUInt32Property.cs
--- a/Engine/Source/Programs/Shared/EpicGames.UHT/Types/Properties/UhtUInt32Property.cs
+++ b/Engine/Source/Programs/Shared/EpicGames.UHT/Types/Properties/UhtUInt32Property.cs
@@ -1,10 +1,10 @@
 // Copyright Epic Games, Inc. All Rights Reserved.
 
+using System.Diagnostics.CodeAnalysis;
+using System.Text;
 using EpicGames.UHT.Tables;
 using EpicGames.UHT.Tokenizer;
 using EpicGames.UHT.Utils;
-using System.Diagnostics.CodeAnalysis;
-using System.Text;
 
 namespace EpicGames.UHT.Types
 {
@@ -17,74 +17,74 @@ namespace EpicGames.UHT.Types
 	public class UhtUInt32Property : UhtNumericProperty
 	{
 		/// <inheritdoc/>
-		public override string EngineClassName { get => "UInt32Property"; }
+		public override string EngineClassName => "UInt32Property";
 
 		/// <inheritdoc/>
-		protected override string CppTypeText { get => "uint32"; }
+		protected override string CppTypeText => "uint32";
 
 		/// <summary>
 		/// Construct a new property
 		/// </summary>
-		/// <param name="PropertySettings">Property settings</param>
-		/// <param name="IntType">Integer type</param>
-		public UhtUInt32Property(UhtPropertySettings PropertySettings, UhtPropertyIntType IntType) : base(PropertySettings, IntType)
+		/// <param name="propertySettings">Property settings</param>
+		/// <param name="intType">Integer type</param>
+		public UhtUInt32Property(UhtPropertySettings propertySettings, UhtPropertyIntType intType) : base(propertySettings, intType)
 		{
 		}
 
 		/// <inheritdoc/>
-		public override StringBuilder AppendMemberDecl(StringBuilder Builder, IUhtPropertyMemberContext Context, string Name, string NameSuffix, int Tabs)
+		public override StringBuilder AppendMemberDecl(StringBuilder builder, IUhtPropertyMemberContext context, string name, string nameSuffix, int tabs)
 		{
-			return AppendMemberDecl(Builder, Context, Name, NameSuffix, Tabs,
+			return AppendMemberDecl(builder, context, name, nameSuffix, tabs,
 				this.IntType == UhtPropertyIntType.Unsized ? "FUnsizedFIntPropertyParams" : "FUInt32PropertyParams");
 		}
 
 		/// <inheritdoc/>
-		public override StringBuilder AppendMemberDef(StringBuilder Builder, IUhtPropertyMemberContext Context, string Name, string NameSuffix, string? Offset, int Tabs)
+		public override StringBuilder AppendMemberDef(StringBuilder builder, IUhtPropertyMemberContext context, string name, string nameSuffix, string? offset, int tabs)
 		{
-			AppendMemberDefStart(Builder, Context, Name, NameSuffix, Offset, Tabs,
+			AppendMemberDefStart(builder, context, name, nameSuffix, offset, tabs,
 				this.IntType == UhtPropertyIntType.Unsized ? "FUnsizedFIntPropertyParams" : "FUInt32PropertyParams",
 				"UECodeGen_Private::EPropertyGenFlags::UInt32");
-			AppendMemberDefEnd(Builder, Context, Name, NameSuffix);
-			return Builder;
+			AppendMemberDefEnd(builder, context, name, nameSuffix);
+			return builder;
 		}
 
 		/// <inheritdoc/>
-		public override bool SanitizeDefaultValue(IUhtTokenReader DefaultValueReader, StringBuilder InnerDefaultValue)
+		public override bool SanitizeDefaultValue(IUhtTokenReader defaultValueReader, StringBuilder innerDefaultValue)
 		{
 			return false;
 		}
 
 		/// <inheritdoc/>
-		public override bool IsSameType(UhtProperty Other)
+		public override bool IsSameType(UhtProperty other)
 		{
-			return Other is UhtUInt32Property;
+			return other is UhtUInt32Property;
 		}
 
 		#region Keyword
 		[UhtPropertyType(Keyword = "uint32", Options = UhtPropertyTypeOptions.Simple | UhtPropertyTypeOptions.Immediate)]
 		[SuppressMessage("CodeQuality", "IDE0051:Remove unused private members", Justification = "Attribute accessed method")]
 		[SuppressMessage("Style", "IDE0060:Remove unused parameter", Justification = "Attribute accessed method")]
-		private static UhtProperty? UInt32Property(UhtPropertyResolvePhase ResolvePhase, UhtPropertySettings PropertySettings, IUhtTokenReader TokenReader, UhtToken MatchedToken)
+		private static UhtProperty? UInt32Property(UhtPropertyResolvePhase resolvePhase, UhtPropertySettings propertySettings, IUhtTokenReader tokenReader, UhtToken matchedToken)
 		{
-			if (PropertySettings.bIsBitfield)
+			if (propertySettings.IsBitfield)
 			{
-				return new UhtBoolProperty(PropertySettings, UhtBoolType.UInt8);
+				return new UhtBoolProperty(propertySettings, UhtBoolType.UInt8);
 			}
 			else
 			{
-				return new UhtUInt32Property(PropertySettings, UhtPropertyIntType.Sized);
+				return new UhtUInt32Property(propertySettings, UhtPropertyIntType.Sized);
 			}
 		}
 
 		[UhtPropertyType(Keyword = "unsigned", Options = UhtPropertyTypeOptions.Immediate)]
 		[SuppressMessage("CodeQuality", "IDE0051:Remove unused private members", Justification = "Attribute accessed method")]
 		[SuppressMessage("Style", "IDE0060:Remove unused parameter", Justification = "Attribute accessed method")]
-		private static UhtProperty? UnsignedProperty(UhtPropertyResolvePhase ResolvePhase, UhtPropertySettings PropertySettings, IUhtTokenReader TokenReader, UhtToken MatchedToken)
+		private static UhtProperty? UnsignedProperty(UhtPropertyResolvePhase resolvePhase, UhtPropertySettings propertySettings, IUhtTokenReader tokenReader, UhtToken matchedToken)
 		{
-			TokenReader
+			tokenReader
 				.Require("unsigned")
 				.Optional("int");
-			return new UhtUInt32Property(PropertySettings, UhtPropertyIntType.Unsized);
+			return new UhtUInt32Property(propertySettings, UhtPropertyIntType.Unsized);
 		}
 		#endregion
 	}
diff --git a/Engine/Source/Programs/Shared/EpicGames.UHT/Types/Properties/UhtUInt64Property.cs b/Engine/Source/Programs/Shared/EpicGames.UHT/Types/Properties/UhtUInt64Property.cs
--- a/Engine/Source/Programs/Shared/EpicGames.UHT/Types/Properties/UhtUInt64Property.cs
+++ b/Engine/Source/Programs/Shared/EpicGames.UHT/Types/Properties/UhtUInt64Property.cs
@@ -1,10 +1,10 @@
 // Copyright Epic Games, Inc. All Rights Reserved.
 
+using System.Diagnostics.CodeAnalysis;
+using System.Text;
 using EpicGames.UHT.Tables;
 using EpicGames.UHT.Tokenizer;
 using EpicGames.UHT.Utils;
-using System.Diagnostics.CodeAnalysis;
-using System.Text;
 
 namespace EpicGames.UHT.Types
 {
@@ -17,59 +17,59 @@ namespace EpicGames.UHT.Types
 	public class UhtUInt64Property : UhtNumericProperty
 	{
 		/// <inheritdoc/>
-		public override string EngineClassName { get => "UInt64Property"; }
+		public override string EngineClassName => "UInt64Property";
 
 		/// <inheritdoc/>
-		protected override string CppTypeText { get => "uint64"; }
+		protected override string CppTypeText => "uint64";
 
 		/// <summary>
 		/// Construct a new property
 		/// </summary>
-		/// <param name="PropertySettings">Property settings</param>
-		/// <param name="IntType">Integer type</param>
-		public UhtUInt64Property(UhtPropertySettings PropertySettings, UhtPropertyIntType IntType) : base(PropertySettings, IntType)
+		/// <param name="propertySettings">Property settings</param>
+		/// <param name="intType">Integer type</param>
+		public UhtUInt64Property(UhtPropertySettings propertySettings, UhtPropertyIntType intType) : base(propertySettings, intType)
 		{
 		}
 
 		/// <inheritdoc/>
-		public override StringBuilder AppendMemberDecl(StringBuilder Builder, IUhtPropertyMemberContext Context, string Name, string NameSuffix, int Tabs)
+		public override StringBuilder AppendMemberDecl(StringBuilder builder, IUhtPropertyMemberContext context, string name, string nameSuffix, int tabs)
 		{
-			return AppendMemberDecl(Builder, Context, Name, NameSuffix, Tabs, "FFInt64PropertyParams");
+			return AppendMemberDecl(builder, context, name, nameSuffix, tabs, "FFInt64PropertyParams");
 		}
 
 		/// <inheritdoc/>
-		public override StringBuilder AppendMemberDef(StringBuilder Builder, IUhtPropertyMemberContext Context, string Name, string NameSuffix, string? Offset, int Tabs)
+		public override StringBuilder AppendMemberDef(StringBuilder builder, IUhtPropertyMemberContext context, string name, string nameSuffix, string? offset, int tabs)
 		{
-			AppendMemberDefStart(Builder, Context, Name, NameSuffix, Offset, Tabs, "FFInt64PropertyParams", "UECodeGen_Private::EPropertyGenFlags::UInt64");
-			AppendMemberDefEnd(Builder, Context, Name, NameSuffix);
-			return Builder;
+			AppendMemberDefStart(builder, context, name, nameSuffix, offset, tabs, "FFInt64PropertyParams", "UECodeGen_Private::EPropertyGenFlags::UInt64");
+			AppendMemberDefEnd(builder, context, name, nameSuffix);
+			return builder;
 		}
 
 		/// <inheritdoc/>
-		public override bool SanitizeDefaultValue(IUhtTokenReader DefaultValueReader, StringBuilder InnerDefaultValue)
+		public override bool SanitizeDefaultValue(IUhtTokenReader defaultValueReader, StringBuilder innerDefaultValue)
 		{
 			return false;
 		}
 
 		/// <inheritdoc/>
-		public override bool IsSameType(UhtProperty Other)
+		public override bool IsSameType(UhtProperty other)
 		{
-			return Other is UhtUInt64Property;
+			return other is UhtUInt64Property;
 		}
 
 		#region Keyword
 		[UhtPropertyType(Keyword = "uint64", Options = UhtPropertyTypeOptions.Simple | UhtPropertyTypeOptions.Immediate)]
 		[SuppressMessage("CodeQuality", "IDE0051:Remove unused private members", Justification = "Attribute accessed method")]
 		[SuppressMessage("Style", "IDE0060:Remove unused parameter", Justification = "Attribute accessed method")]
-		private static UhtProperty? UInt64Property(UhtPropertyResolvePhase ResolvePhase, UhtPropertySettings PropertySettings, IUhtTokenReader TokenReader, UhtToken MatchedToken)
+		private static UhtProperty? UInt64Property(UhtPropertyResolvePhase resolvePhase, UhtPropertySettings propertySettings, IUhtTokenReader tokenReader, UhtToken matchedToken)
 		{
-			if (PropertySettings.bIsBitfield)
+			if (propertySettings.IsBitfield)
 			{
-				return new UhtBoolProperty(PropertySettings, UhtBoolType.UInt8);
+				return new UhtBoolProperty(propertySettings, UhtBoolType.UInt8);
 			}
 			else
 			{
-				return new UhtUInt64Property(PropertySettings, UhtPropertyIntType.Sized);
+				return new UhtUInt64Property(propertySettings, UhtPropertyIntType.Sized);
 			}
 		}
 		#endregion
diff --git a/Engine/Source/Programs/Shared/EpicGames.UHT/Types/Properties/UhtVoidProperty.cs b/Engine/Source/Programs/Shared/EpicGames.UHT/Types/Properties/UhtVoidProperty.cs
--- a/Engine/Source/Programs/Shared/EpicGames.UHT/Types/Properties/UhtVoidProperty.cs
+++ b/Engine/Source/Programs/Shared/EpicGames.UHT/Types/Properties/UhtVoidProperty.cs
@@ -1,11 +1,11 @@
 // Copyright Epic Games, Inc. All Rights Reserved.
 
-using EpicGames.UHT.Tables;
-using EpicGames.UHT.Tokenizer;
-using EpicGames.UHT.Utils;
 using System;
 using System.Diagnostics.CodeAnalysis;
 using System.Text;
+using EpicGames.UHT.Tables;
+using EpicGames.UHT.Tokenizer;
+using EpicGames.UHT.Utils;
 
 namespace EpicGames.UHT.Types
 {
@@ -17,70 +17,70 @@ namespace EpicGames.UHT.Types
 	public class UhtVoidProperty : UhtProperty
 	{
 		/// <inheritdoc/>
-		public override string EngineClassName { get => "UHTVoidProperty"; }
+		public override string EngineClassName => "UHTVoidProperty";
 
 		/// <inheritdoc/>
-		protected override string CppTypeText { get => "void"; }
+		protected override string CppTypeText => "void";
 
 		/// <inheritdoc/>
-		protected override string PGetMacroText { get => "invalid"; }
+		protected override string PGetMacroText => "invalid";
 
 		/// <summary>
 		/// Construct a new void property
 		/// </summary>
-		/// <param name="PropertySettings">Property settings</param>
-		public UhtVoidProperty(UhtPropertySettings PropertySettings) : base(PropertySettings)
+		/// <param name="propertySettings">Property settings</param>
+		public UhtVoidProperty(UhtPropertySettings propertySettings) : base(propertySettings)
 		{
 		}
 
 		/// <inheritdoc/>
-		public override StringBuilder AppendMemberDecl(StringBuilder Builder, IUhtPropertyMemberContext Context, string Name, string NameSuffix, int Tabs)
+		public override StringBuilder AppendMemberDecl(StringBuilder builder, IUhtPropertyMemberContext context, string name, string nameSuffix, int tabs)
 		{
 			throw new NotImplementedException();
 		}
 
 		/// <inheritdoc/>
-		public override StringBuilder AppendMemberDef(StringBuilder Builder, IUhtPropertyMemberContext Context, string Name, string NameSuffix, string? Offset, int Tabs)
+		public override StringBuilder AppendMemberDef(StringBuilder builder, IUhtPropertyMemberContext context, string name, string nameSuffix, string? offset, int tabs)
 		{
 			throw new NotImplementedException();
 		}
 
 		/// <inheritdoc/>
-		public override StringBuilder AppendFunctionThunkParameterGet(StringBuilder Builder)
+		public override StringBuilder AppendFunctionThunkParameterGet(StringBuilder builder)
 		{
 			throw new NotImplementedException();
 		}
 
 		/// <inheritdoc/>
-		public override StringBuilder AppendNullConstructorArg(StringBuilder Builder, bool bIsInitializer)
+		public override StringBuilder AppendNullConstructorArg(StringBuilder builder, bool isInitializer)
 		{
 			throw new NotImplementedException();
 		}
 
 		/// <inheritdoc/>
-		public override bool SanitizeDefaultValue(IUhtTokenReader DefaultValueReader, StringBuilder InnerDefaultValue)
+		public override bool SanitizeDefaultValue(IUhtTokenReader defaultValueReader, StringBuilder innerDefaultValue)
 		{
 			throw new NotImplementedException();
 		}
 
 		/// <inheritdoc/>
-		public override bool IsSameType(UhtProperty Other)
+		public override bool IsSameType(UhtProperty other)
 		{
-			return Other is UhtVoidProperty;
+			return other is UhtVoidProperty;
 		}
 
 		#region Keyword
 		[UhtPropertyType(Keyword = "void", Options = UhtPropertyTypeOptions.Simple | UhtPropertyTypeOptions.Immediate)]
 		[SuppressMessage("CodeQuality", "IDE0051:Remove unused private members", Justification = "Attribute accessed method")]
 		[SuppressMessage("Style", "IDE0060:Remove unused parameter", Justification = "Attribute accessed method")]
-		private static UhtProperty? VoidProperty(UhtPropertyResolvePhase ResolvePhase, UhtPropertySettings PropertySettings, IUhtTokenReader TokenReader, UhtToken MatchedToken)
+		private static UhtProperty? VoidProperty(UhtPropertyResolvePhase resolvePhase, UhtPropertySettings propertySettings, IUhtTokenReader tokenReader, UhtToken matchedToken)
 		{
-			if (PropertySettings.PropertyCategory != UhtPropertyCategory.Return)
+			if (propertySettings.PropertyCategory != UhtPropertyCategory.Return)
 			{
-				TokenReader.LogError("void type is only valid as a return type");
+				tokenReader.LogError("void type is only valid as a return type");
 				return null;
 			}
-			return new UhtVoidProperty(PropertySettings);
+			return new UhtVoidProperty(propertySettings);
 		}
 		#endregion
 	}
diff --git a/Engine/Source/Programs/Shared/EpicGames.UHT/Types/Properties/UhtWeakObjectPtrProperty.cs b/Engine/Source/Programs/Shared/EpicGames.UHT/Types/Properties/UhtWeakObjectPtrProperty.cs
--- a/Engine/Source/Programs/Shared/EpicGames.UHT/Types/Properties/UhtWeakObjectPtrProperty.cs
+++ b/Engine/Source/Programs/Shared/EpicGames.UHT/Types/Properties/UhtWeakObjectPtrProperty.cs
@@ -1,11 +1,11 @@
 // Copyright Epic Games, Inc. All Rights Reserved.
 
+using System.Diagnostics.CodeAnalysis;
+using System.Text;
 using EpicGames.Core;
 using EpicGames.UHT.Tables;
 using EpicGames.UHT.Tokenizer;
 using EpicGames.UHT.Utils;
-using System.Diagnostics.CodeAnalysis;
-using System.Text;
 
 namespace EpicGames.UHT.Types
 {
@@ -18,123 +18,123 @@ namespace EpicGames.UHT.Types
 	public class UhtWeakObjectPtrProperty : UhtObjectPropertyBase
 	{
 		/// <inheritdoc/>
-		public override string EngineClassName { get => "WeakObjectProperty"; }
+		public override string EngineClassName => "WeakObjectProperty";
 
 		/// <inheritdoc/>
-		protected override string PGetMacroText { get => this.PropertyFlags.HasAnyFlags(EPropertyFlags.AutoWeak) ? "AUTOWEAKOBJECT" : "WEAKOBJECT"; }
+		protected override string PGetMacroText => this.PropertyFlags.HasAnyFlags(EPropertyFlags.AutoWeak) ? "AUTOWEAKOBJECT" : "WEAKOBJECT";
 
 		/// <inheritdoc/>
-		protected override bool bPGetPassAsNoPtr { get => true; }
+		protected override bool PGetPassAsNoPtr => true;
 
 		/// <inheritdoc/>
-		protected override UhtPGetArgumentType PGetTypeArgument { get => UhtPGetArgumentType.TypeText; }
+		protected override UhtPGetArgumentType PGetTypeArgument => UhtPGetArgumentType.TypeText;
 
 		/// <summary>
 		/// Construct a new property
 		/// </summary>
-		/// <param name="PropertySettings">Property settings</param>
-		/// <param name="PropertyClass">Class being referenced</param>
-		/// <param name="ExtraFlags">Extra property flags to add to the definition</param>
-		public UhtWeakObjectPtrProperty(UhtPropertySettings PropertySettings, UhtClass PropertyClass, EPropertyFlags ExtraFlags = EPropertyFlags.None)
-			: base(PropertySettings, PropertyClass, null)
+		/// <param name="propertySettings">Property settings</param>
+		/// <param name="propertyClass">Class being referenced</param>
+		/// <param name="extraFlags">Extra property flags to add to the definition</param>
+		public UhtWeakObjectPtrProperty(UhtPropertySettings propertySettings, UhtClass propertyClass, EPropertyFlags extraFlags = EPropertyFlags.None)
+			: base(propertySettings, propertyClass, null)
 		{
-			this.PropertyFlags |= EPropertyFlags.UObjectWrapper | ExtraFlags;
+			this.PropertyFlags |= EPropertyFlags.UObjectWrapper | extraFlags;
 			this.PropertyCaps |= UhtPropertyCaps.PassCppArgsByRef | UhtPropertyCaps.RequiresNullConstructorArg | UhtPropertyCaps.IsMemberSupportedByBlueprint;
 			this.PropertyCaps &= ~(UhtPropertyCaps.IsParameterSupportedByBlueprint);
 		}
 
 		/// <inheritdoc/>
-		public override StringBuilder AppendText(StringBuilder Builder, UhtPropertyTextType TextType, bool bIsTemplateArgument)
+		public override StringBuilder AppendText(StringBuilder builder, UhtPropertyTextType textType, bool isTemplateArgument)
 		{
-			switch (TextType)
+			switch (textType)
 			{
 				default:
 					if (this.PropertyFlags.HasAnyFlags(EPropertyFlags.AutoWeak))
 					{
-						Builder.Append("TAutoWeakObjectPtr<").Append(this.Class.SourceName).Append('>');
+						builder.Append("TAutoWeakObjectPtr<").Append(this.Class.SourceName).Append('>');
 					}
 					else
 					{
-						Builder.Append("TWeakObjectPtr<").Append(this.Class.SourceName).Append('>');
+						builder.Append("TWeakObjectPtr<").Append(this.Class.SourceName).Append('>');
 					}
 					break;
 			}
-			return Builder;
+			return builder;
 		}
 
 		/// <inheritdoc/>
-		public override StringBuilder AppendMemberDecl(StringBuilder Builder, IUhtPropertyMemberContext Context, string Name, string NameSuffix, int Tabs)
+		public override StringBuilder AppendMemberDecl(StringBuilder builder, IUhtPropertyMemberContext context, string name, string nameSuffix, int tabs)
 		{
-			return AppendMemberDecl(Builder, Context, Name, NameSuffix, Tabs, "FWeakObjectPropertyParams");
+			return AppendMemberDecl(builder, context, name, nameSuffix, tabs, "FWeakObjectPropertyParams");
 		}
 
 		/// <inheritdoc/>
-		public override StringBuilder AppendMemberDef(StringBuilder Builder, IUhtPropertyMemberContext Context, string Name, string NameSuffix, string? Offset, int Tabs)
+		public override StringBuilder AppendMemberDef(StringBuilder builder, IUhtPropertyMemberContext context, string name, string nameSuffix, string? offset, int tabs)
 		{
-			AppendMemberDefStart(Builder, Context, Name, NameSuffix, Offset, Tabs, "FWeakObjectPropertyParams", "UECodeGen_Private::EPropertyGenFlags::WeakObject");
-			AppendMemberDefRef(Builder, Context, this.Class, false);
-			AppendMemberDefEnd(Builder, Context, Name, NameSuffix);
-			return Builder;
+			AppendMemberDefStart(builder, context, name, nameSuffix, offset, tabs, "FWeakObjectPropertyParams", "UECodeGen_Private::EPropertyGenFlags::WeakObject");
+			AppendMemberDefRef(builder, context, this.Class, false);
+			AppendMemberDefEnd(builder, context, name, nameSuffix);
+			return builder;
 		}
 
 		/// <inheritdoc/>
-		public override StringBuilder AppendNullConstructorArg(StringBuilder Builder, bool bIsInitializer)
+		public override StringBuilder AppendNullConstructorArg(StringBuilder builder, bool isInitializer)
 		{
-			Builder.Append("NULL");
-			return Builder;
+			builder.Append("NULL");
+			return builder;
 		}
 
 		/// <inheritdoc/>
-		public override bool IsSameType(UhtProperty Other)
+		public override bool IsSameType(UhtProperty other)
 		{
-			if (Other is UhtWeakObjectPtrProperty OtherObject)
+			if (other is UhtWeakObjectPtrProperty otherObject)
 			{
-				return this.Class == OtherObject.Class && this.MetaClass == OtherObject.MetaClass;
+				return this.Class == otherObject.Class && this.MetaClass == otherObject.MetaClass;
 			}
 			return false;
 		}
 
 		#region Keywords
-		private static UhtProperty CreateWeakProperty(UhtPropertySettings PropertySettings, IUhtTokenReader TokenReader, UhtClass Class, EPropertyFlags ExtraFlags = EPropertyFlags.None)
+		private static UhtProperty CreateWeakProperty(UhtPropertySettings propertySettings, IUhtTokenReader tokenReader, UhtClass classObj, EPropertyFlags extraFlags = EPropertyFlags.None)
 		{
-			if (Class.IsChildOf(Class.Session.UClass))
+			if (classObj.IsChildOf(classObj.Session.UClass))
 			{
-				TokenReader.LogError("Class variables cannot be weak, they are always strong.");
+				tokenReader.LogError("Class variables cannot be weak, they are always strong.");
 			}
 
-			if (PropertySettings.DisallowPropertyFlags.HasAnyFlags(EPropertyFlags.AutoWeak))
+			if (propertySettings.DisallowPropertyFlags.HasAnyFlags(EPropertyFlags.AutoWeak))
 			{
-				return new UhtObjectProperty(PropertySettings, Class, null, ExtraFlags | EPropertyFlags.UObjectWrapper);
+				return new UhtObjectProperty(propertySettings, classObj, null, extraFlags | EPropertyFlags.UObjectWrapper);
 			}
-			return new UhtWeakObjectPtrProperty(PropertySettings, Class, ExtraFlags);
+			return new UhtWeakObjectPtrProperty(propertySettings, classObj, extraFlags);
 		}
 
 		[UhtPropertyType(Keyword = "TWeakObjectPtr")]
 		[SuppressMessage("CodeQuality", "IDE0051:Remove unused private members", Justification = "Attribute accessed method")]
 		[SuppressMessage("Style", "IDE0060:Remove unused parameter", Justification = "Attribute accessed method")]
-		private static UhtProperty? WeakObjectProperty(UhtPropertyResolvePhase ResolvePhase, UhtPropertySettings PropertySettings, IUhtTokenReader TokenReader, UhtToken MatchedToken)
+		private static UhtProperty? WeakObjectProperty(UhtPropertyResolvePhase resolvePhase, UhtPropertySettings propertySettings, IUhtTokenReader tokenReader, UhtToken matchedToken)
 		{
-			UhtClass? PropertyClass = ParseTemplateObject(PropertySettings, TokenReader, MatchedToken, true);
-			if (PropertyClass == null)
+			UhtClass? propertyClass = ParseTemplateObject(propertySettings, tokenReader, matchedToken, true);
+			if (propertyClass == null)
 			{
 				return null;
 			}
 
-			return CreateWeakProperty(PropertySettings, TokenReader, PropertyClass);
+			return CreateWeakProperty(propertySettings, tokenReader, propertyClass);
 		}
 
 		[UhtPropertyType(Keyword = "TAutoWeakObjectPtr")]
 		[SuppressMessage("CodeQuality", "IDE0051:Remove unused private members", Justification = "Attribute accessed method")]
 		[SuppressMessage("Style", "IDE0060:Remove unused parameter", Justification = "Attribute accessed method")]
-		private static UhtProperty? AutoWeakObjectProperty(UhtPropertyResolvePhase ResolvePhase, UhtPropertySettings PropertySettings, IUhtTokenReader TokenReader, UhtToken MatchedToken)
+		private static UhtProperty? AutoWeakObjectProperty(UhtPropertyResolvePhase resolvePhase, UhtPropertySettings propertySettings, IUhtTokenReader tokenReader, UhtToken matchedToken)
 		{
-			UhtClass? PropertyClass = ParseTemplateObject(PropertySettings, TokenReader, MatchedToken, true);
-			if (PropertyClass == null)
+			UhtClass? propertyClass = ParseTemplateObject(propertySettings, tokenReader, matchedToken, true);
+			if (propertyClass == null)
 			{
 				return null;
 			}
 
-			return CreateWeakProperty(PropertySettings, TokenReader, PropertyClass, EPropertyFlags.AutoWeak);
+			return CreateWeakProperty(propertySettings, tokenReader, propertyClass, EPropertyFlags.AutoWeak);
 		}
 		#endregion
 	}
diff --git a/Engine/Source/Programs/Shared/EpicGames.UHT/Types/UhtClass.cs b/Engine/Source/Programs/Shared/EpicGames.UHT/Types/UhtClass.cs
--- a/Engine/Source/Programs/Shared/EpicGames.UHT/Types/UhtClass.cs
+++ b/Engine/Source/Programs/Shared/EpicGames.UHT/Types/UhtClass.cs
@@ -1,14 +1,14 @@
 // Copyright Epic Games, Inc. All Rights Reserved.
 
+using System;
+using System.Collections.Generic;
+using System.Text;
+using System.Text.Json.Serialization;
 using EpicGames.Core;
 using EpicGames.UHT.Parsers;
 using EpicGames.UHT.Tables;
 using EpicGames.UHT.Tokenizer;
 using EpicGames.UHT.Utils;
-using System;
-using System.Collections.Generic;
-using System.Text;
-using System.Text.Json.Serialization;
 
 namespace EpicGames.UHT.Types
 {
@@ -16,7 +16,7 @@ namespace EpicGames.UHT.Types
 	/// Series of flags not part of the engine's class flags that affect code generation or verification
 	/// </summary>
 	[Flags]
-	public enum UhtClassExportFlags : Int32
+	public enum UhtClassExportFlags : int
 	{
 
 		/// <summary>
@@ -95,35 +95,35 @@ namespace EpicGames.UHT.Types
 		/// <summary>
 		/// Test to see if any of the specified flags are set
 		/// </summary>
-		/// <param name="InFlags">Current flags</param>
-		/// <param name="TestFlags">Flags to test for</param>
+		/// <param name="inFlags">Current flags</param>
+		/// <param name="testFlags">Flags to test for</param>
 		/// <returns>True if any of the flags are set</returns>
-		public static bool HasAnyFlags(this UhtClassExportFlags InFlags, UhtClassExportFlags TestFlags)
+		public static bool HasAnyFlags(this UhtClassExportFlags inFlags, UhtClassExportFlags testFlags)
 		{
-			return (InFlags & TestFlags) != 0;
+			return (inFlags & testFlags) != 0;
 		}
 
 		/// <summary>
 		/// Test to see if all of the specified flags are set
 		/// </summary>
-		/// <param name="InFlags">Current flags</param>
-		/// <param name="TestFlags">Flags to test for</param>
+		/// <param name="inFlags">Current flags</param>
+		/// <param name="testFlags">Flags to test for</param>
 		/// <returns>True if all the flags are set</returns>
-		public static bool HasAllFlags(this UhtClassExportFlags InFlags, UhtClassExportFlags TestFlags)
+		public static bool HasAllFlags(this UhtClassExportFlags inFlags, UhtClassExportFlags testFlags)
 		{
-			return (InFlags & TestFlags) == TestFlags;
+			return (inFlags & testFlags) == testFlags;
 		}
 
 		/// <summary>
 		/// Test to see if a specific set of flags have a specific value.
 		/// </summary>
-		/// <param name="InFlags">Current flags</param>
-		/// <param name="TestFlags">Flags to test for</param>
-		/// <param name="MatchFlags">Expected value of the tested flags</param>
+		/// <param name="inFlags">Current flags</param>
+		/// <param name="testFlags">Flags to test for</param>
+		/// <param name="matchFlags">Expected value of the tested flags</param>
 		/// <returns>True if the given flags have a specific value.</returns>
-		public static bool HasExactFlags(this UhtClassExportFlags InFlags, UhtClassExportFlags TestFlags, UhtClassExportFlags MatchFlags)
+		public static bool HasExactFlags(this UhtClassExportFlags inFlags, UhtClassExportFlags testFlags, UhtClassExportFlags matchFlags)
 		{
-			return (InFlags & TestFlags) == MatchFlags;
+			return (inFlags & testFlags) == matchFlags;
 		}
 	}
 
@@ -186,35 +186,35 @@ namespace EpicGames.UHT.Types
 		/// <summary>
 		/// Test to see if any of the specified flags are set
 		/// </summary>
-		/// <param name="InFlags">Current flags</param>
-		/// <param name="TestFlags">Flags to test for</param>
+		/// <param name="inFlags">Current flags</param>
+		/// <param name="testFlags">Flags to test for</param>
 		/// <returns>True if any of the flags are set</returns>
-		public static bool HasAnyFlags(this UhtSerializerArchiveType InFlags, UhtSerializerArchiveType TestFlags)
+		public static bool HasAnyFlags(this UhtSerializerArchiveType inFlags, UhtSerializerArchiveType testFlags)
 		{
-			return (InFlags & TestFlags) != 0;
+			return (inFlags & testFlags) != 0;
 		}
 
 		/// <summary>
 		/// Test to see if all of the specified flags are set
 		/// </summary>
-		/// <param name="InFlags">Current flags</param>
-		/// <param name="TestFlags">Flags to test for</param>
+		/// <param name="inFlags">Current flags</param>
+		/// <param name="testFlags">Flags to test for</param>
 		/// <returns>True if all the flags are set</returns>
-		public static bool HasAllFlags(this UhtSerializerArchiveType InFlags, UhtSerializerArchiveType TestFlags)
+		public static bool HasAllFlags(this UhtSerializerArchiveType inFlags, UhtSerializerArchiveType testFlags)
 		{
-			return (InFlags & TestFlags) == TestFlags;
+			return (inFlags & testFlags) == testFlags;
 		}
 
 		/// <summary>
 		/// Test to see if a specific set of flags have a specific value.
 		/// </summary>
-		/// <param name="InFlags">Current flags</param>
-		/// <param name="TestFlags">Flags to test for</param>
-		/// <param name="MatchFlags">Expected value of the tested flags</param>
+		/// <param name="inFlags">Current flags</param>
+		/// <param name="testFlags">Flags to test for</param>
+		/// <param name="matchFlags">Expected value of the tested flags</param>
 		/// <returns>True if the given flags have a specific value.</returns>
-		public static bool HasExactFlags(this UhtSerializerArchiveType InFlags, UhtSerializerArchiveType TestFlags, UhtSerializerArchiveType MatchFlags)
+		public static bool HasExactFlags(this UhtSerializerArchiveType inFlags, UhtSerializerArchiveType testFlags, UhtSerializerArchiveType matchFlags)
 		{
-			return (InFlags & TestFlags) == MatchFlags;
+			return (inFlags & testFlags) == matchFlags;
 		}
 	}
 
@@ -233,6 +233,7 @@ namespace EpicGames.UHT.Types
 		/// Collection of tokens parsed in the declaration
 		/// </summary>
 		[System.Diagnostics.CodeAnalysis.SuppressMessage("Design", "CA1051:Do not declare visible instance fields", Justification = "<Pending>")]
+		[System.Diagnostics.CodeAnalysis.SuppressMessage("Style", "IDE1006:Naming Styles", Justification = "<Pending>")]
 		public UhtToken[] Tokens;
 
 		/// <summary>
@@ -256,6 +257,7 @@ namespace EpicGames.UHT.Types
 		/// Collection of tokens parsed in the declaration
 		/// </summary>
 		[System.Diagnostics.CodeAnalysis.SuppressMessage("Design", "CA1051:Do not declare visible instance fields", Justification = "<Pending>")]
+		[System.Diagnostics.CodeAnalysis.SuppressMessage("Style", "IDE1006:Naming Styles", Justification = "<Pending>")]
 		public UhtToken[] Tokens;
 
 		/// <summary>
@@ -266,7 +268,7 @@ namespace EpicGames.UHT.Types
 		/// <summary>
 		/// True if "virtual" was found prior to the matching name
 		/// </summary>
-		public bool bIsVirtual { get; set; }
+		public bool IsVirtual { get; set; }
 	}
 
 	/// <summary>
@@ -275,7 +277,7 @@ namespace EpicGames.UHT.Types
 	[UhtEngineClass(Name = "Class")]
 	public class UhtClass : UhtStruct
 	{
-		private List<UhtDeclaration>? DeclarationsInternal = null;
+		private List<UhtDeclaration>? _declarationsInternal = null;
 
 		/// <summary>
 		/// Configuration section
@@ -331,7 +333,7 @@ namespace EpicGames.UHT.Types
 		/// Collection of functions and other declarations found in the class
 		/// </summary>
 		[JsonIgnore]
-		public IList<UhtDeclaration>? Declarations => this.DeclarationsInternal;
+		public IList<UhtDeclaration>? Declarations => this._declarationsInternal;
 
 		/// <summary>
 		/// If this, this class is a UINTERFACE and NativeInterface is the associated native interface
@@ -352,7 +354,7 @@ namespace EpicGames.UHT.Types
 		/// <summary>
 		/// True if GENERATED_BODY was used.  If false, GENERATED_UCLASS_BODY was used.
 		/// </summary>
-		public bool bHasGeneratedBody { get; set; } = false;
+		public bool HasGeneratedBody { get; set; } = false;
 
 		/// <inheritdoc/>
 		[JsonIgnore]
@@ -375,11 +377,11 @@ namespace EpicGames.UHT.Types
 		}
 
 		/// <inheritdoc/>
-		public override string EngineClassName { get => "Class"; }
+		public override string EngineClassName => "Class";
 
 		///<inheritdoc/>
 		[JsonIgnore]
-		public override bool bDeprecated => this.ClassFlags.HasAnyFlags(EClassFlags.Deprecated);
+		public override bool Deprecated => this.ClassFlags.HasAnyFlags(EClassFlags.Deprecated);
 
 		///<inheritdoc/>
 		[JsonIgnore]
@@ -410,9 +412,9 @@ namespace EpicGames.UHT.Types
 		/// <summary>
 		/// Construct a new instance of the class
 		/// </summary>
-		/// <param name="Outer">The outer type</param>
-		/// <param name="LineNumber">Line number where class begins</param>
-		public UhtClass(UhtType Outer, int LineNumber) : base(Outer, LineNumber)
+		/// <param name="outer">The outer type</param>
+		/// <param name="lineNumber">Line number where class begins</param>
+		public UhtClass(UhtType outer, int lineNumber) : base(outer, lineNumber)
 		{
 			this.ClassWithin = this;
 		}
@@ -421,18 +423,18 @@ namespace EpicGames.UHT.Types
 		/// True if this class inherits from AActor
 		/// </summary>
 		[JsonIgnore]
-		public bool bIsActorClass => IsChildOf(this.Session.AActor);
+		public bool IsActorClass => IsChildOf(this.Session.AActor);
 
 		///<inheritdoc/>
 		[JsonIgnore]
 		public override string EngineNamePrefix
-		{ 
+		{
 			get
 			{
 				switch (this.ClassType)
 				{
-					case UhtClassType.Class:  
-						if (bIsActorClass)
+					case UhtClassType.Class:
+						if (IsActorClass)
 						{
 							return this.ClassFlags.HasAnyFlags(EClassFlags.Deprecated) ? "ADEPRECATED_" : "A";
 						}
@@ -456,47 +458,47 @@ namespace EpicGames.UHT.Types
 		/// <summary>
 		/// Add the given list of tokens as a possible declaration
 		/// </summary>
-		/// <param name="CompilerDirectives">Currently active compiler directives</param>
-		/// <param name="Tokens">List of declaration tokens</param>
-		/// <param name="Function">If parsed as part of a UFUNCTION, this will reference it</param>
-		public void AddDeclaration(UhtCompilerDirective CompilerDirectives, List<UhtToken> Tokens, UhtFunction? Function)
+		/// <param name="compilerDirectives">Currently active compiler directives</param>
+		/// <param name="tokens">List of declaration tokens</param>
+		/// <param name="function">If parsed as part of a UFUNCTION, this will reference it</param>
+		public void AddDeclaration(UhtCompilerDirective compilerDirectives, List<UhtToken> tokens, UhtFunction? function)
 		{
-			if (this.DeclarationsInternal == null)
+			if (this._declarationsInternal == null)
 			{
-				this.DeclarationsInternal = new List<UhtDeclaration>();
+				this._declarationsInternal = new List<UhtDeclaration>();
 			}
-			this.DeclarationsInternal.Add(new UhtDeclaration { CompilerDirectives = CompilerDirectives, Tokens = Tokens.ToArray(), Function = Function });
+			this._declarationsInternal.Add(new UhtDeclaration { CompilerDirectives = compilerDirectives, Tokens = tokens.ToArray(), Function = function });
 		}
 
 		/// <summary>
 		/// Search the declarations for a possible declaration match with the given name
 		/// </summary>
-		/// <param name="Name">Name to be located</param>
-		/// <param name="FoundDeclaration">Information about the matched declaration</param>
+		/// <param name="name">Name to be located</param>
+		/// <param name="foundDeclaration">Information about the matched declaration</param>
 		/// <returns>True if a declaration was found with the name.</returns>
-		public bool TryGetDeclaration(string Name, out UhtFoundDeclaration FoundDeclaration)
+		public bool TryGetDeclaration(string name, out UhtFoundDeclaration foundDeclaration)
 		{
-			if (Name.Length > 0)
+			if (name.Length > 0)
 			{
 				if (this.Declarations != null)
 				{
-					foreach (UhtDeclaration Declaration in this.Declarations)
+					foreach (UhtDeclaration declaration in this.Declarations)
 					{
-						bool bIsVirtual = false;
-						for (int Index = 0; Index < Declaration.Tokens.Length; ++Index)
+						bool isVirtual = false;
+						for (int index = 0; index < declaration.Tokens.Length; ++index)
 						{
-							if (Declaration.Tokens[Index].IsIdentifier("virtual"))
+							if (declaration.Tokens[index].IsIdentifier("virtual"))
 							{
-								bIsVirtual = true;
+								isVirtual = true;
 							}
-							else if (Declaration.Tokens[Index].IsIdentifier(Name))
+							else if (declaration.Tokens[index].IsIdentifier(name))
 							{
-								FoundDeclaration = new UhtFoundDeclaration
+								foundDeclaration = new UhtFoundDeclaration
 								{
-									CompilerDirectives = Declaration.CompilerDirectives,
-									Tokens = Declaration.Tokens,
-									NameTokenIndex = Index,
-									bIsVirtual = bIsVirtual,
+									CompilerDirectives = declaration.CompilerDirectives,
+									Tokens = declaration.Tokens,
+									NameTokenIndex = index,
+									IsVirtual = isVirtual,
 								};
 								return true;
 							}
@@ -506,22 +508,22 @@ namespace EpicGames.UHT.Types
 			}
 			if (this.NativeInterface != null)
 			{
-				return this.NativeInterface.TryGetDeclaration(Name, out FoundDeclaration);
+				return this.NativeInterface.TryGetDeclaration(name, out foundDeclaration);
 			}
-			FoundDeclaration = new UhtFoundDeclaration();
+			foundDeclaration = new UhtFoundDeclaration();
 			return false;
 		}
 
 		/// <summary>
 		/// Checks to see if the class or any super class has the given flags.
 		/// </summary>
-		/// <param name="FlagsToCheck"></param>
+		/// <param name="flagsToCheck"></param>
 		/// <returns>True if the flags are found</returns>
-		public bool HierarchyHasAnyClassFlags(EClassFlags FlagsToCheck)
+		public bool HierarchyHasAnyClassFlags(EClassFlags flagsToCheck)
 		{
-			for (UhtClass? Current = this; Current != null; Current = Current.SuperClass)
+			for (UhtClass? current = this; current != null; current = current.SuperClass)
 			{
-				if (Current.ClassFlags.HasAnyFlags(FlagsToCheck))
+				if (current.ClassFlags.HasAnyFlags(flagsToCheck))
 				{
 					return true;
 				}
@@ -532,27 +534,27 @@ namespace EpicGames.UHT.Types
 		#region Resolution support
 		struct GetterSetterToResolve
 		{
-			public UhtProperty Property;
-			public bool bSetter;
-			public bool bFound;
+			public UhtProperty Property { get; set; }
+			public bool Setter { get; set; }
+			public bool Found { get; set; }
 		}
 
 		/// <inheritdoc/>
-		protected override bool ResolveSelf(UhtResolvePhase Phase)
+		protected override bool ResolveSelf(UhtResolvePhase phase)
 		{
-			bool bResult = base.ResolveSelf(Phase);
+			bool result = base.ResolveSelf(phase);
 
-			switch (Phase)
+			switch (phase)
 			{
 				case UhtResolvePhase.Bases:
 
 					// Check to see if this class matches an entry in the ClassCastFlags.  If it does, that
 					// becomes our ClassCastFlags.  If not, then we inherit the flags from the super.
-					if (Enum.TryParse<EClassCastFlags>(this.SourceName, false, out EClassCastFlags Found))
+					if (Enum.TryParse<EClassCastFlags>(this.SourceName, false, out EClassCastFlags found))
 					{
-						if (Found != EClassCastFlags.None)
+						if (found != EClassCastFlags.None)
 						{
-							this.ClassCastFlags = Found;
+							this.ClassCastFlags = found;
 						}
 					}
 					if (this.ClassCastFlags == EClassCastFlags.None && this.SuperClass != null)
@@ -565,53 +567,53 @@ namespace EpicGames.UHT.Types
 					// and before properties have been finalized.
 					if (this.ClassType == UhtClassType.Class && this.ClassFlags.HasAnyFlags(EClassFlags.Const))
 					{
-						foreach (UhtType Child in this.Children)
+						foreach (UhtType child in this.Children)
 						{
-							if (Child is UhtProperty Property)
+							if (child is UhtProperty property)
 							{
-								Property.PropertyFlags |= EPropertyFlags.BlueprintReadOnly;
-								Property.PropertyExportFlags |= UhtPropertyExportFlags.ImpliedBlueprintPure;
+								property.PropertyFlags |= EPropertyFlags.BlueprintReadOnly;
+								property.PropertyExportFlags |= UhtPropertyExportFlags.ImpliedBlueprintPure;
 							}
 						}
 					}
 					break;
 
 				case UhtResolvePhase.Final:
-					Dictionary<string, List<GetterSetterToResolve>>? GSToResolve = null;
-					foreach (UhtType Child in this.Children)
+					Dictionary<string, List<GetterSetterToResolve>>? gsToResolve = null;
+					foreach (UhtType child in this.Children)
 					{
-						if (Child is UhtProperty Property)
+						if (child is UhtProperty property)
 						{
-							if (Property.PropertyExportFlags.HasAnyFlags(UhtPropertyExportFlags.FieldNotify))
+							if (property.PropertyExportFlags.HasAnyFlags(UhtPropertyExportFlags.FieldNotify))
 							{
 								this.ClassExportFlags |= UhtClassExportFlags.HasFieldNotify;
 							}
 
-							if (Property.PropertyFlags.HasAnyFlags(EPropertyFlags.Net))
+							if (property.PropertyFlags.HasAnyFlags(EPropertyFlags.Net))
 							{
 								this.ClassExportFlags |= UhtClassExportFlags.SelfHasReplicatedProperties;
 								break;
 							}
 
-							if (!Property.PropertyExportFlags.HasAnyFlags(UhtPropertyExportFlags.GetterSpecifiedNone))
+							if (!property.PropertyExportFlags.HasAnyFlags(UhtPropertyExportFlags.GetterSpecifiedNone))
 							{
-								if (Property.PropertyExportFlags.HasAnyFlags(UhtPropertyExportFlags.GetterSpecified))
+								if (property.PropertyExportFlags.HasAnyFlags(UhtPropertyExportFlags.GetterSpecified))
 								{
-									GSToResolve = AddGetterSetter(GSToResolve, Property.Getter == null ? $"Get{Property.SourceName}" : Property.Getter, Property, false);
+									gsToResolve = AddGetterSetter(gsToResolve, property.Getter ?? $"Get{property.SourceName}", property, false);
 								}
 							}
 
-							if (!Property.PropertyExportFlags.HasAnyFlags(UhtPropertyExportFlags.SetterSpecifiedNone))
+							if (!property.PropertyExportFlags.HasAnyFlags(UhtPropertyExportFlags.SetterSpecifiedNone))
 							{
-								if (Property.PropertyExportFlags.HasAnyFlags(UhtPropertyExportFlags.SetterSpecified))
+								if (property.PropertyExportFlags.HasAnyFlags(UhtPropertyExportFlags.SetterSpecified))
 								{
-									GSToResolve = AddGetterSetter(GSToResolve, Property.Setter == null ? $"Set{Property.SourceName}" : Property.Setter, Property, true);
+									gsToResolve = AddGetterSetter(gsToResolve, property.Setter ?? $"Set{property.SourceName}", property, true);
 								}
 							}
 						}
-						else if (Child is UhtFunction Function)
+						else if (child is UhtFunction function)
 						{
-							if (Function.FunctionExportFlags.HasAnyFlags(UhtFunctionExportFlags.FieldNotify))
+							if (function.FunctionExportFlags.HasAnyFlags(UhtFunctionExportFlags.FieldNotify))
 							{
 								this.ClassExportFlags |= UhtClassExportFlags.HasFieldNotify;
 							}
@@ -638,43 +640,43 @@ namespace EpicGames.UHT.Types
 					}
 
 					// If we have any getters/setters to resolve 
-					if (GSToResolve != null)
+					if (gsToResolve != null)
 					{
 						if (this.Declarations != null)
 						{
-							foreach (UhtDeclaration Declaration in this.Declarations)
+							foreach (UhtDeclaration declaration in this.Declarations)
 							{
 
 								// Locate the index of the name
-								int NameIndex = 0;
-								for (; NameIndex < Declaration.Tokens.Length - 1; ++NameIndex)
+								int nameIndex = 0;
+								for (; nameIndex < declaration.Tokens.Length - 1; ++nameIndex)
 								{
-									if (Declaration.Tokens[NameIndex + 1].IsValue('('))
+									if (declaration.Tokens[nameIndex + 1].IsValue('('))
 									{
 										break;
 									}
 								}
 
-								string Name = Declaration.Tokens[NameIndex].Value.ToString();
-								if (GSToResolve.TryGetValue(Name, out List<GetterSetterToResolve>? GSList))
+								string name = declaration.Tokens[nameIndex].Value.ToString();
+								if (gsToResolve.TryGetValue(name, out List<GetterSetterToResolve>? gsList))
 								{
-									for (int Index = 0; Index < GSList.Count; ++Index)
+									for (int index = 0; index < gsList.Count; ++index)
 									{
-										GetterSetterToResolve GS = GSList[Index];
-										if (!GS.bFound)
+										GetterSetterToResolve gs = gsList[index];
+										if (!gs.Found)
 										{
-											if (TryMatchGetterSetter(GS.Property, GS.bSetter, Declaration, NameIndex))
+											if (TryMatchGetterSetter(gs.Property, gs.Setter, declaration, nameIndex))
 											{
-												GS.bFound = true;
-												GSList[Index] = GS;
-												GS.Property.PropertyExportFlags |= GS.bSetter ? UhtPropertyExportFlags.SetterFound : UhtPropertyExportFlags.GetterFound;
-												if (GS.bSetter)
+												gs.Found = true;
+												gsList[index] = gs;
+												gs.Property.PropertyExportFlags |= gs.Setter ? UhtPropertyExportFlags.SetterFound : UhtPropertyExportFlags.GetterFound;
+												if (gs.Setter)
 												{
-													GS.Property.Setter = Name;
+													gs.Property.Setter = name;
 												}
 												else
 												{
-													GS.Property.Getter = Name;
+													gs.Property.Getter = name;
 												}
 											}
 										}
@@ -686,15 +688,15 @@ namespace EpicGames.UHT.Types
 					break;
 			}
 
-			return bResult;
+			return result;
 		}
 
 		/// <inheritdoc/>
-		protected override void ResolveChildren(UhtResolvePhase Phase)
+		protected override void ResolveChildren(UhtResolvePhase phase)
 		{
-			base.ResolveChildren(Phase);
+			base.ResolveChildren(phase);
 
-			switch (Phase)
+			switch (phase)
 			{
 				case UhtResolvePhase.Final:
 					if (ScanForInstancedReferenced(false))
@@ -706,211 +708,211 @@ namespace EpicGames.UHT.Types
 		}
 
 		/// <inheritdoc/>
-		public override bool ScanForInstancedReferenced(bool bDeepScan)
+		public override bool ScanForInstancedReferenced(bool deepScan)
 		{
 			if (this.ClassFlags.HasAnyFlags(EClassFlags.HasInstancedReference))
 			{
 				return true;
 			}
 
-			if (this.SuperClass != null && this.SuperClass.ScanForInstancedReferenced(bDeepScan))
+			if (this.SuperClass != null && this.SuperClass.ScanForInstancedReferenced(deepScan))
 			{
 				return true;
 			}
 
-			return base.ScanForInstancedReferenced(bDeepScan);
+			return base.ScanForInstancedReferenced(deepScan);
 		}
 
 		/// <summary>
 		/// Given a token stream, verify that it matches the expected signature for a getter/setter
 		/// </summary>
-		/// <param name="Property">Property requesting the setter/getter</param>
-		/// <param name="bSetter">If true, a setter is expected</param>
-		/// <param name="Declaration">The declaration being tested</param>
-		/// <param name="NameIndex">The index of the token with the expected getter/setter name</param>
+		/// <param name="property">Property requesting the setter/getter</param>
+		/// <param name="setter">If true, a setter is expected</param>
+		/// <param name="declaration">The declaration being tested</param>
+		/// <param name="nameIndex">The index of the token with the expected getter/setter name</param>
 		/// <returns>True if the declaration matches</returns>
-		private static bool TryMatchGetterSetter(UhtProperty Property, bool bSetter, UhtDeclaration Declaration, int NameIndex)
+		private static bool TryMatchGetterSetter(UhtProperty property, bool setter, UhtDeclaration declaration, int nameIndex)
 		{
-			UhtToken[] Tokens = Declaration.Tokens;
-			int TokenIndex = 0;
-			int TokenCount = Tokens.Length;
+			UhtToken[] tokens = declaration.Tokens;
+			int tokenIndex = 0;
+			int tokenCount = tokens.Length;
 
 			// Skip all of the API and virtual keywords
-			for (; TokenIndex < TokenCount && (Tokens[TokenIndex].IsValue("virtual") || Tokens[TokenIndex].Value.Span.EndsWith("_API")); ++TokenIndex)
+			for (; tokenIndex < tokenCount && (tokens[tokenIndex].IsValue("virtual") || tokens[tokenIndex].Value.Span.EndsWith("_API")); ++tokenIndex)
 			{
 				// Nothing to do in the body
 			}
 
-			int TypeIndex;
-			int TypeCount;
+			int typeIndex;
+			int typeCount;
 
 			// Verify the format of the function declaration and extract the type
-			if (bSetter)
+			if (setter)
 			{
 
 				// The return type must be 'void'
-				if (TokenIndex == TokenCount || !Tokens[TokenIndex].IsValue("void"))
+				if (tokenIndex == tokenCount || !tokens[tokenIndex].IsValue("void"))
 				{
 					return false;
 				}
-				TokenIndex++;
+				tokenIndex++;
 
 				// Followed by the name
-				if (TokenIndex == TokenCount || TokenIndex != NameIndex)
+				if (tokenIndex == tokenCount || tokenIndex != nameIndex)
 				{
 					return false;
 				}
-				TokenIndex++;
+				tokenIndex++;
 
 				// Followed by the '('
-				if (TokenIndex == TokenCount || !Tokens[TokenIndex].IsValue('('))
+				if (tokenIndex == tokenCount || !tokens[tokenIndex].IsValue('('))
 				{
 					return false;
 				}
-				TokenIndex++;
+				tokenIndex++;
 
 				// Find the ')'
-				TypeIndex = TokenIndex;
+				typeIndex = tokenIndex;
 				while (true)
 				{
-					if (TokenIndex == TokenCount)
+					if (tokenIndex == tokenCount)
 					{
 						return false;
 					}
-					if (Tokens[TokenIndex].IsValue(')'))
+					if (tokens[tokenIndex].IsValue(')'))
 					{
-						TypeCount = TokenIndex - TypeIndex;
+						typeCount = tokenIndex - typeIndex;
 						break;
 					}
-					TokenIndex++;
+					tokenIndex++;
 				}
 			}
 			else
 			{
 				// Get the type range
-				TypeIndex = TokenIndex;
-				TypeCount = NameIndex - TypeIndex;
-				TokenIndex = NameIndex + 1;
+				typeIndex = tokenIndex;
+				typeCount = nameIndex - typeIndex;
+				tokenIndex = nameIndex + 1;
 
 				// Followed by the '('
-				if (TokenIndex == TokenCount || !Tokens[TokenIndex].IsValue('('))
+				if (tokenIndex == tokenCount || !tokens[tokenIndex].IsValue('('))
 				{
 					return false;
 				}
-				TokenIndex++;
+				tokenIndex++;
 
 				// Followed by the ')'
-				if (TokenIndex == TokenCount || !Tokens[TokenIndex].IsValue(')'))
+				if (tokenIndex == tokenCount || !tokens[tokenIndex].IsValue(')'))
 				{
 					return false;
 				}
-				TokenIndex++;
+				tokenIndex++;
 
 				// Followed by the "const"
-				if (TokenIndex == TokenCount || !Tokens[TokenIndex].IsValue("const"))
+				if (tokenIndex == tokenCount || !tokens[tokenIndex].IsValue("const"))
 				{
 					return false;
 				}
 #pragma warning disable IDE0059 // Unnecessary assignment of a value
-				TokenIndex++;
+				tokenIndex++;
 #pragma warning restore IDE0059 // Unnecessary assignment of a value
 			}
 
 			// Verify the type
-			using (BorrowStringBuilder Borrower = new BorrowStringBuilder(StringBuilderCache.Small))
+			using (BorrowStringBuilder borrower = new BorrowStringBuilder(StringBuilderCache.Small))
 			{
-				StringBuilder Builder = Borrower.StringBuilder;
-				Builder.AppendPropertyText(Property, UhtPropertyTextType.GetterSetterArg);
-				if (Property.bIsStaticArray)
+				StringBuilder builder = borrower.StringBuilder;
+				builder.AppendPropertyText(property, UhtPropertyTextType.GetterSetterArg);
+				if (property.IsStaticArray)
 				{
-					if (Declaration.Function == null)
+					if (declaration.Function == null)
 					{
-						Builder.Append('*');
+						builder.Append('*');
 					}
 					else
 					{
-						Builder.Append('[').Append(Property.ArrayDimensions).Append(']');
+						builder.Append('[').Append(property.ArrayDimensions).Append(']');
 					}
 				}
-				string ExpectedType = Builder.ToString();
-				int TypeEnd = TypeIndex + TypeCount;
+				string expectedType = builder.ToString();
+				int typeEnd = typeIndex + typeCount;
 
 				// We will do a somewhat brute force string compare to avoid constructing an expected
 				// string or parsing the required string
-				ReadOnlySpan<char> ExpectedSpan = ExpectedType.AsSpan().Trim();
+				ReadOnlySpan<char> expectedSpan = expectedType.AsSpan().Trim();
 
 				// skip const. We allow const in parameter and return values even if the property is not const
-				if (TypeIndex < TypeEnd && Tokens[TypeIndex].IsIdentifier("const"))
+				if (typeIndex < typeEnd && tokens[typeIndex].IsIdentifier("const"))
 				{
-					TypeIndex++;
+					typeIndex++;
 				}
 
 				// Loop until we are done
-				while (TypeIndex < TypeEnd && ExpectedSpan.Length > 0)
+				while (typeIndex < typeEnd && expectedSpan.Length > 0)
 				{
-					UhtToken Token = Tokens[TypeIndex];
-					ReadOnlySpan<char> TokenSpan = Token.Value.Span;
+					UhtToken token = tokens[typeIndex];
+					ReadOnlySpan<char> tokenSpan = token.Value.Span;
 
 					// Make sure we match the expected type
-					if (TokenSpan.Length > ExpectedSpan.Length ||
-						!ExpectedSpan.StartsWith(TokenSpan))
+					if (tokenSpan.Length > expectedSpan.Length ||
+						!expectedSpan.StartsWith(tokenSpan))
 					{
 						return false;
 					}
 
 					// Extract the remaining part of the string we expect to match
-					ExpectedSpan = ExpectedSpan.Slice(TokenSpan.Length).TrimStart();
-					TypeIndex++;
+					expectedSpan = expectedSpan.Slice(tokenSpan.Length).TrimStart();
+					typeIndex++;
 				}
 
 				// Consume any '&' for reference support
-				if (TypeIndex < TypeEnd && Tokens[TypeIndex].IsSymbol('&'))
+				if (typeIndex < typeEnd && tokens[typeIndex].IsSymbol('&'))
 				{
-					TypeIndex++;
+					typeIndex++;
 				}
 
 				// For setters, there can be an identifier
-				if (bSetter && TypeIndex < TypeEnd && Tokens[TypeIndex].IsIdentifier())
+				if (setter && typeIndex < typeEnd && tokens[typeIndex].IsIdentifier())
 				{
-					TypeIndex++;
+					typeIndex++;
 				}
 
 				// We must be at the end for a match
-				if (TypeIndex != TypeEnd)
+				if (typeIndex != typeEnd)
 				{
 					return false;
 				}
 			}
 
 			// Validate the #if blocks.  This should really be done during validation
-			string FuncType = bSetter ? "setter" : "getter";
-			UhtPropertyExportFlags TestFlag = bSetter ? UhtPropertyExportFlags.SetterSpecified : UhtPropertyExportFlags.GetterSpecified;
-			if (Declaration.CompilerDirectives.HasAnyFlags(UhtCompilerDirective.WithEditor))
+			string funcType = setter ? "setter" : "getter";
+			UhtPropertyExportFlags testFlag = setter ? UhtPropertyExportFlags.SetterSpecified : UhtPropertyExportFlags.GetterSpecified;
+			if (declaration.CompilerDirectives.HasAnyFlags(UhtCompilerDirective.WithEditor))
 			{
-				if (Property.PropertyExportFlags.HasAnyFlags(TestFlag))
+				if (property.PropertyExportFlags.HasAnyFlags(testFlag))
 				{
-					Property.LogError(Tokens[0].InputLine, $"Property {Property.SourceName} {FuncType} function {Tokens[NameIndex].Value} "
+					property.LogError(tokens[0].InputLine, $"Property {property.SourceName} {funcType} function {tokens[nameIndex].Value} "
 						+ "cannot be declared within WITH_EDITOR block. Use WITH_EDITORONLY_DATA instead.");
 				}
 				return false;
 			}
-			if (Declaration.CompilerDirectives.HasAnyFlags(UhtCompilerDirective.WithEditorOnlyData) && !Property.PropertyFlags.HasAnyFlags(EPropertyFlags.EditorOnly))
+			if (declaration.CompilerDirectives.HasAnyFlags(UhtCompilerDirective.WithEditorOnlyData) && !property.PropertyFlags.HasAnyFlags(EPropertyFlags.EditorOnly))
 			{
-				if (Property.PropertyExportFlags.HasAnyFlags(TestFlag))
+				if (property.PropertyExportFlags.HasAnyFlags(testFlag))
 				{
-					Property.LogError(Tokens[0].InputLine, $"Property {Property.SourceName} is not editor-only but its {FuncType} function '{Tokens[NameIndex].Value}' is");
+					property.LogError(tokens[0].InputLine, $"Property {property.SourceName} is not editor-only but its {funcType} function '{tokens[nameIndex].Value}' is");
 				}
 				return false;
 			}
-			if (Declaration.Function != null)
+			if (declaration.Function != null)
 			{
-				if (!Declaration.Function.FunctionFlags.HasAnyFlags(EFunctionFlags.Native))
+				if (!declaration.Function.FunctionFlags.HasAnyFlags(EFunctionFlags.Native))
 				{
-					Property.LogError(Tokens[0].InputLine, $"Property {Property.SourceName} {FuncType} function '{Tokens[NameIndex].Value}' has to be native");
+					property.LogError(tokens[0].InputLine, $"Property {property.SourceName} {funcType} function '{tokens[nameIndex].Value}' has to be native");
 				}
-				if (Declaration.Function.FunctionFlags.HasAnyFlags(EFunctionFlags.Net | EFunctionFlags.Event))
+				if (declaration.Function.FunctionFlags.HasAnyFlags(EFunctionFlags.Net | EFunctionFlags.Event))
 				{
-					Property.LogError(Tokens[0].InputLine, $"Property {Property.SourceName} {FuncType} function '{Tokens[NameIndex].Value}' cannot be a net or an event");
+					property.LogError(tokens[0].InputLine, $"Property {property.SourceName} {funcType} function '{tokens[nameIndex].Value}' cannot be a net or an event");
 				}
 			}
 			return true;
@@ -919,38 +921,38 @@ namespace EpicGames.UHT.Types
 		/// <summary>
 		/// Add a requested getter/setter function name
 		/// </summary>
-		/// <param name="GSToResolve">Dictionary containing the lookup by function name.  Will be created if null.</param>
-		/// <param name="Name">Name of the getter/setter</param>
-		/// <param name="Property">Property requesting the getter/setter</param>
-		/// <param name="bSetter">True if this is a setter</param>
+		/// <param name="gsToResolve">Dictionary containing the lookup by function name.  Will be created if null.</param>
+		/// <param name="name">Name of the getter/setter</param>
+		/// <param name="property">Property requesting the getter/setter</param>
+		/// <param name="setter">True if this is a setter</param>
 		/// <returns>Resulting dictionary</returns>
-		private static Dictionary<string, List<GetterSetterToResolve>> AddGetterSetter(Dictionary<string, List<GetterSetterToResolve>>? GSToResolve, string Name, UhtProperty Property, bool bSetter)
+		private static Dictionary<string, List<GetterSetterToResolve>> AddGetterSetter(Dictionary<string, List<GetterSetterToResolve>>? gsToResolve, string name, UhtProperty property, bool setter)
 		{
-			if (GSToResolve == null)
+			if (gsToResolve == null)
 			{
-				GSToResolve = new Dictionary<string, List<GetterSetterToResolve>>();
+				gsToResolve = new Dictionary<string, List<GetterSetterToResolve>>();
 			}
-			if (!GSToResolve.TryGetValue(Name, out List<GetterSetterToResolve>? GSList))
+			if (!gsToResolve.TryGetValue(name, out List<GetterSetterToResolve>? gsList))
 			{
-				GSList = new List<GetterSetterToResolve>();
-				GSToResolve.Add(Name, GSList);
+				gsList = new List<GetterSetterToResolve>();
+				gsToResolve.Add(name, gsList);
 			}
-			GSList.Add(new GetterSetterToResolve { Property = Property, bSetter = bSetter });
-			return GSToResolve;
+			gsList.Add(new GetterSetterToResolve { Property = property, Setter = setter });
+			return gsToResolve;
 		}
 		#endregion
 
 		#region Validation support
 		/// <inheritdoc/>
-		protected override UhtValidationOptions Validate(UhtValidationOptions Options)
+		protected override UhtValidationOptions Validate(UhtValidationOptions options)
 		{
-			Options = base.Validate(Options);
+			options = base.Validate(options);
 
 			// Classes must start with a valid prefix
-			string ExpectedClassName = this.EngineNamePrefix + this.EngineName;
-			if (ExpectedClassName != this.SourceName)
+			string expectedClassName = this.EngineNamePrefix + this.EngineName;
+			if (expectedClassName != this.SourceName)
 			{
-				this.LogError($"Class '{this.SourceName}' has an invalid Unreal prefix, expecting '{ExpectedClassName}'");
+				this.LogError($"Class '{this.SourceName}' has an invalid Unreal prefix, expecting '{expectedClassName}'");
 			}
 
 			// If we have a super class
@@ -1023,82 +1025,82 @@ namespace EpicGames.UHT.Types
 
 						// Make sure that all interface functions are implemented property
 						// Iterate over all the bases looking for any interfaces
-						foreach (UhtStruct BaseStruct in this.Bases)
+						foreach (UhtStruct baseStruct in this.Bases)
 						{
 
 							// Skip children that aren't interfaces or if a common ancestor
-							UhtClass? BaseClass = BaseStruct as UhtClass;
-							if (BaseClass == null || BaseClass.ClassType == UhtClassType.Class || IsChildOf(BaseClass) )
+							UhtClass? baseClass = baseStruct as UhtClass;
+							if (baseClass == null || baseClass.ClassType == UhtClassType.Class || IsChildOf(baseClass))
 							{
 								continue;
 							}
 
 							// Loop through the function to make sure they are implemented
-							List<UhtType> BaseChildren = BaseClass.AlternateObject != null ? BaseClass.AlternateObject.Children : BaseClass.Children;
-							foreach (UhtType BaseChild in BaseChildren)
+							List<UhtType> baseChildren = baseClass.AlternateObject != null ? baseClass.AlternateObject.Children : baseClass.Children;
+							foreach (UhtType baseChild in baseChildren)
 							{
-								UhtFunction? BaseFunction = BaseChild as UhtFunction;
-								if (BaseFunction == null)
+								UhtFunction? baseFunction = baseChild as UhtFunction;
+								if (baseFunction == null)
 								{
-									continue; 
+									continue;
 								}
 
 								// Delegate signature functions are simple stubs and aren't required to be implemented (they are not callable)
-								bool bImplemented = BaseFunction.FunctionFlags.HasAnyFlags(EFunctionFlags.Delegate);
+								bool implemented = baseFunction.FunctionFlags.HasAnyFlags(EFunctionFlags.Delegate);
 
 								// Try to find an existing function
-								foreach (UhtType Child in this.Children)
+								foreach (UhtType child in this.Children)
 								{
-									UhtFunction? Function = Child as UhtFunction;
-									if (Function == null || Function.SourceName != BaseFunction.SourceName)
+									UhtFunction? function = child as UhtFunction;
+									if (function == null || function.SourceName != baseFunction.SourceName)
 									{
 										continue;
 									}
 
-									if (BaseFunction.FunctionFlags.HasAnyFlags(EFunctionFlags.Event) && !Function.FunctionFlags.HasAnyFlags(EFunctionFlags.Event))
+									if (baseFunction.FunctionFlags.HasAnyFlags(EFunctionFlags.Event) && !function.FunctionFlags.HasAnyFlags(EFunctionFlags.Event))
 									{
-										Function.LogError($"Implementation of function '{Function.SourceName}' must be declared as 'event' to match declaration in interface '{BaseClass.SourceName}'");
+										function.LogError($"Implementation of function '{function.SourceName}' must be declared as 'event' to match declaration in interface '{baseClass.SourceName}'");
 									}
 
-									if (BaseFunction.FunctionFlags.HasAnyFlags(EFunctionFlags.Delegate) && !Function.FunctionFlags.HasAnyFlags(EFunctionFlags.Delegate))
+									if (baseFunction.FunctionFlags.HasAnyFlags(EFunctionFlags.Delegate) && !function.FunctionFlags.HasAnyFlags(EFunctionFlags.Delegate))
 									{
-										Function.LogError($"Implementation of function '{Function.SourceName}' must be declared as 'delegate' to match declaration in interface '{BaseClass.SourceName}'");
+										function.LogError($"Implementation of function '{function.SourceName}' must be declared as 'delegate' to match declaration in interface '{baseClass.SourceName}'");
 									}
 
-									bImplemented = true;
+									implemented = true;
 
-									if (BaseFunction.Children.Count != Function.Children.Count)
+									if (baseFunction.Children.Count != function.Children.Count)
 									{
-										Function.LogError($"Implementation of function '{Function.SourceName}' conflicts with interface '{BaseClass.SourceName}' - different number of parameters ({Function.Children.Count}/{BaseFunction.Children.Count})");
+										function.LogError($"Implementation of function '{function.SourceName}' conflicts with interface '{baseClass.SourceName}' - different number of parameters ({function.Children.Count}/{baseFunction.Children.Count})");
 										continue;
 									}
 
-									for (int Index = 0; Index < Function.Children.Count; ++Index)
+									for (int index = 0; index < function.Children.Count; ++index)
 									{
-										UhtProperty BaseProperty = (UhtProperty)BaseFunction.Children[Index];
-										UhtProperty Property = (UhtProperty)Function.Children[Index];
-										if (!BaseProperty.MatchesType(Property))
+										UhtProperty baseProperty = (UhtProperty)baseFunction.Children[index];
+										UhtProperty property = (UhtProperty)function.Children[index];
+										if (!baseProperty.MatchesType(property))
 										{
-											if (BaseProperty.PropertyFlags.HasAnyFlags(EPropertyFlags.ReturnParm))
+											if (baseProperty.PropertyFlags.HasAnyFlags(EPropertyFlags.ReturnParm))
 											{
-												Function.LogError($"Implementation of function '{Function.SourceName}' conflicts only by return type with interface '{BaseClass.SourceName}'");
+												function.LogError($"Implementation of function '{function.SourceName}' conflicts only by return type with interface '{baseClass.SourceName}'");
 											}
 											else
 											{
-												Function.LogError($"Implementation of function '{Function.SourceName}' conflicts type with interface '{BaseClass.SourceName}' - parameter {Index} '{Property.SourceName}'");
+												function.LogError($"Implementation of function '{function.SourceName}' conflicts type with interface '{baseClass.SourceName}' - parameter {index} '{property.SourceName}'");
 											}
 										}
 									}
 								}
 
 								// Verify that if this has blueprint-callable functions that are not implementable events, we've implemented them as a UFunction in the target class
-								if (!bImplemented
-									&& BaseFunction.FunctionFlags.HasAnyFlags(EFunctionFlags.BlueprintCallable)
-									&& !BaseFunction.FunctionFlags.HasAnyFlags(EFunctionFlags.BlueprintEvent)
-									&& !BaseClass.MetaData.ContainsKey(UhtNames.CannotImplementInterfaceInBlueprint)
-									&& (BaseClass.AlternateObject == null || !BaseClass.AlternateObject.MetaData.ContainsKey(UhtNames.CannotImplementInterfaceInBlueprint)))
+								if (!implemented
+									&& baseFunction.FunctionFlags.HasAnyFlags(EFunctionFlags.BlueprintCallable)
+									&& !baseFunction.FunctionFlags.HasAnyFlags(EFunctionFlags.BlueprintEvent)
+									&& !baseClass.MetaData.ContainsKey(UhtNames.CannotImplementInterfaceInBlueprint)
+									&& (baseClass.AlternateObject == null || !baseClass.AlternateObject.MetaData.ContainsKey(UhtNames.CannotImplementInterfaceInBlueprint)))
 								{
-									this.LogError($"Missing UFunction implementation of function '{BaseFunction.SourceName}' from interface '{BaseClass.SourceName}'.  This function needs a UFUNCTION() declaration.");
+									this.LogError($"Missing UFunction implementation of function '{baseFunction.SourceName}' from interface '{baseClass.SourceName}'.  This function needs a UFUNCTION() declaration.");
 								}
 							}
 						}
@@ -1119,42 +1121,42 @@ namespace EpicGames.UHT.Types
 					{
 						// Interface with blueprint data should declare explicitly Blueprintable or NotBlueprintable to be clear
 						// In the backward compatible case where they declare neither, both of these bools are false
-						bool bCanImplementInBlueprints = this.MetaData.GetBoolean(UhtNames.IsBlueprintBase);
-						bool bCannotImplementInBlueprints = (!bCanImplementInBlueprints && this.MetaData.ContainsKey(UhtNames.IsBlueprintBase))
+						bool canImplementInBlueprints = this.MetaData.GetBoolean(UhtNames.IsBlueprintBase);
+						bool cannotImplementInBlueprints = (!canImplementInBlueprints && this.MetaData.ContainsKey(UhtNames.IsBlueprintBase))
 							|| this.MetaData.ContainsKey(UhtNames.CannotImplementInterfaceInBlueprint);
 
-						//bool bCanImplementInterfaceInBlueprint = !this.MetaData.ContainsKey(UhtNames.CannotImplementInterfaceInBlueprint);
-						foreach (UhtType Child in this.Children)
+						//bool canImplementInterfaceInBlueprint = !this.MetaData.ContainsKey(UhtNames.CannotImplementInterfaceInBlueprint);
+						foreach (UhtType child in this.Children)
 						{
-							if (Child is UhtFunction Function)
+							if (child is UhtFunction function)
 							{
 								// Verify interfaces with respect to their blueprint accessible functions
-								if (Function.FunctionFlags.HasAnyFlags(EFunctionFlags.BlueprintEvent) && !Function.MetaData.GetBoolean(UhtNames.BlueprintInternalUseOnly))
+								if (function.FunctionFlags.HasAnyFlags(EFunctionFlags.BlueprintEvent) && !function.MetaData.GetBoolean(UhtNames.BlueprintInternalUseOnly))
 								{
 									// Ensure that blueprint events are only allowed in implementable interfaces. Internal only functions allowed
-									if (bCannotImplementInBlueprints)
+									if (cannotImplementInBlueprints)
 									{
-										Function.LogError("Interfaces that are not implementable in blueprints cannot have Blueprint Event members.");
+										function.LogError("Interfaces that are not implementable in blueprints cannot have Blueprint Event members.");
 									}
-									if (!bCanImplementInBlueprints)
+									if (!canImplementInBlueprints)
 									{
 										// We do not currently warn about this case as there are a large number of existing interfaces that do not specify
 										// Function.LogWarning(TEXT("Interfaces with Blueprint Events should declare Blueprintable on the interface."));
 									}
 								}
 
-								if (Function.FunctionFlags.HasExactFlags(EFunctionFlags.BlueprintCallable | EFunctionFlags.BlueprintEvent, EFunctionFlags.BlueprintCallable))
+								if (function.FunctionFlags.HasExactFlags(EFunctionFlags.BlueprintCallable | EFunctionFlags.BlueprintEvent, EFunctionFlags.BlueprintCallable))
 								{
 									// Ensure that if this interface contains blueprint callable functions that are not blueprint defined, that it must be implemented natively
-									if (bCanImplementInBlueprints)
+									if (canImplementInBlueprints)
 									{
-										Function.LogError("Blueprint implementable interfaces cannot contain BlueprintCallable functions that are not "
+										function.LogError("Blueprint implementable interfaces cannot contain BlueprintCallable functions that are not "
 											+ "BlueprintImplementableEvents. Add NotBlueprintable to the interface if you wish to keep this function.");
 									}
-									if (!bCannotImplementInBlueprints)
+									if (!cannotImplementInBlueprints)
 									{
 										// Lowered this case to a warning instead of error, they will not show up as blueprintable unless they also have events
-										Function.LogWarning("Interfaces with BlueprintCallable functions but no events should explicitly declare NotBlueprintable on the interface.");
+										function.LogWarning("Interfaces with BlueprintCallable functions but no events should explicitly declare NotBlueprintable on the interface.");
 									}
 								}
 							}
@@ -1166,184 +1168,185 @@ namespace EpicGames.UHT.Types
 					break;
 			}
 
-			return Options |= UhtValidationOptions.Shadowing | UhtValidationOptions.Deprecated;
+			return options |= UhtValidationOptions.Shadowing | UhtValidationOptions.Deprecated;
 		}
 
-		static void ValidateBlueprintPopertyGetter(UhtProperty Property, string FunctionName, UhtFunction? TargetFunction)
+		static void ValidateBlueprintPopertyGetter(UhtProperty property, string functionName, UhtFunction? targetFunction)
 		{
-			if (TargetFunction == null)
+			if (targetFunction == null)
 			{
-				Property.LogError($"Blueprint Property getter function '{FunctionName}' not found");
+				property.LogError($"Blueprint Property getter function '{functionName}' not found");
 				return;
 			}
 
-			if (TargetFunction.ParameterProperties.Length != 0)
+			if (targetFunction.ParameterProperties.Length != 0)
 			{
-				TargetFunction.LogError($"Blueprint Property getter function '{TargetFunction.SourceName}' must not have parameters.");
+				targetFunction.LogError($"Blueprint Property getter function '{targetFunction.SourceName}' must not have parameters.");
 			}
 
-			UhtProperty? ReturnProperty = TargetFunction.ReturnProperty;
-			if (ReturnProperty == null || !Property.IsSameType(ReturnProperty))
+			UhtProperty? returnProperty = targetFunction.ReturnProperty;
+			if (returnProperty == null || !property.IsSameType(returnProperty))
 			{
-				TargetFunction.LogError($"Blueprint Property getter function '{TargetFunction.SourceName}' must have return value of type '{Property.GetUserFacingDecl()}'.");
+				targetFunction.LogError($"Blueprint Property getter function '{targetFunction.SourceName}' must have return value of type '{property.GetUserFacingDecl()}'.");
 			}
 
-			if (TargetFunction.FunctionFlags.HasAnyFlags(EFunctionFlags.Event))
+			if (targetFunction.FunctionFlags.HasAnyFlags(EFunctionFlags.Event))
 			{
-				TargetFunction.LogError($"Blueprint Property getter function '{TargetFunction.SourceName}' cannot be a blueprint event.");
+				targetFunction.LogError($"Blueprint Property getter function '{targetFunction.SourceName}' cannot be a blueprint event.");
 			}
 
-			if (!TargetFunction.FunctionFlags.HasAnyFlags(EFunctionFlags.BlueprintPure))
+			if (!targetFunction.FunctionFlags.HasAnyFlags(EFunctionFlags.BlueprintPure))
 			{
-				TargetFunction.LogError($"Blueprint Property getter function '{TargetFunction.SourceName}' must be pure.");
+				targetFunction.LogError($"Blueprint Property getter function '{targetFunction.SourceName}' must be pure.");
 			}
 		}
 
-		static void ValidateBlueprintPopertySetter(UhtProperty Property, string FunctionName, UhtFunction? TargetFunction)
+		static void ValidateBlueprintPopertySetter(UhtProperty property, string functionName, UhtFunction? targetFunction)
 		{
-			if (TargetFunction == null)
+			if (targetFunction == null)
 			{
-				Property.LogError($"Blueprint Property setter function '{FunctionName}' not found");
+				property.LogError($"Blueprint Property setter function '{functionName}' not found");
 				return;
 			}
 
-			if (TargetFunction.ReturnProperty != null)
+			if (targetFunction.ReturnProperty != null)
 			{
-				TargetFunction.LogError($"Blueprint Property setter function '{TargetFunction.SourceName}' must not have a return value.");
+				targetFunction.LogError($"Blueprint Property setter function '{targetFunction.SourceName}' must not have a return value.");
 			}
 
-			if (TargetFunction.ParameterProperties.Length != 1 || !Property.IsSameType((UhtProperty)TargetFunction.ParameterProperties.Span[0]))
+			if (targetFunction.ParameterProperties.Length != 1 || !property.IsSameType((UhtProperty)targetFunction.ParameterProperties.Span[0]))
 			{
-				TargetFunction.LogError($"Blueprint Property setter function '{TargetFunction.SourceName}' must have exactly one parameter of type '{Property.GetUserFacingDecl()}'.");
+				targetFunction.LogError($"Blueprint Property setter function '{targetFunction.SourceName}' must have exactly one parameter of type '{property.GetUserFacingDecl()}'.");
 			}
 
-			if (TargetFunction.FunctionFlags.HasAnyFlags(EFunctionFlags.Event))
+			if (targetFunction.FunctionFlags.HasAnyFlags(EFunctionFlags.Event))
 			{
-				TargetFunction.LogError($"Blueprint Property setter function '{TargetFunction.SourceName}' cannot be a blueprint event.");
+				targetFunction.LogError($"Blueprint Property setter function '{targetFunction.SourceName}' cannot be a blueprint event.");
 			}
 
-			if (!TargetFunction.FunctionFlags.HasAnyFlags(EFunctionFlags.BlueprintCallable))
+			if (!targetFunction.FunctionFlags.HasAnyFlags(EFunctionFlags.BlueprintCallable))
 			{
-				TargetFunction.LogError($"Blueprint Property setter function '{TargetFunction.SourceName}' must be a blueprint callable.");
+				targetFunction.LogError($"Blueprint Property setter function '{targetFunction.SourceName}' must be a blueprint callable.");
 			}
 
-			if (TargetFunction.FunctionFlags.HasAnyFlags(EFunctionFlags.BlueprintPure))
+			if (targetFunction.FunctionFlags.HasAnyFlags(EFunctionFlags.BlueprintPure))
 			{
-				TargetFunction.LogError($"Blueprint Property setter function '{TargetFunction.SourceName}' must not be pure.");
+				targetFunction.LogError($"Blueprint Property setter function '{targetFunction.SourceName}' must not be pure.");
 			}
 		}
 
-		static void ValidateRepNotifyCallback(UhtProperty Property, string FunctionName, UhtFunction? TargetFunction)
+		static void ValidateRepNotifyCallback(UhtProperty property, string functionName, UhtFunction? targetFunction)
 		{
-			if (TargetFunction == null)
+			if (targetFunction == null)
 			{
-				Property.LogError($"Replication notification function '{FunctionName}' not found");
+				property.LogError($"Replication notification function '{functionName}' not found");
 				return;
 			}
 
-			if (TargetFunction.ReturnProperty != null)
+			if (targetFunction.ReturnProperty != null)
 			{
-				TargetFunction.LogError($"Replication notification function '{TargetFunction.SourceName}' must not have a return value.");
+				targetFunction.LogError($"Replication notification function '{targetFunction.SourceName}' must not have a return value.");
 			}
 
-			bool bIsArrayProperty = Property.bIsStaticArray || Property is UhtArrayProperty;
-			int MaxParms = bIsArrayProperty ? 2 : 1;
-			ReadOnlyMemory<UhtType> Parameters = TargetFunction.ParameterProperties;
+			bool isArrayProperty = property.IsStaticArray || property is UhtArrayProperty;
+			int maxParms = isArrayProperty ? 2 : 1;
+			ReadOnlyMemory<UhtType> parameters = targetFunction.ParameterProperties;
 
-			if (Parameters.Length > MaxParms)
+			if (parameters.Length > maxParms)
 			{
-				TargetFunction.LogError($"Replication notification function '{TargetFunction.SourceName}' has too many parameters.");
+				targetFunction.LogError($"Replication notification function '{targetFunction.SourceName}' has too many parameters.");
 			}
 
-			if (Parameters.Length >= 1)
+			if (parameters.Length >= 1)
 			{
-				UhtProperty Parm = (UhtProperty)Parameters.Span[0];
+				UhtProperty parm = (UhtProperty)parameters.Span[0];
 				// First parameter is always the old value:
-				if (!Property.IsSameType(Parm))
+				if (!property.IsSameType(parm))
 				{
-					TargetFunction.LogError($"Replication notification function '{TargetFunction.SourceName}' first (optional) parameter must be of type '{Property.GetUserFacingDecl()}'.");
+					targetFunction.LogError($"Replication notification function '{targetFunction.SourceName}' first (optional) parameter must be of type '{property.GetUserFacingDecl()}'.");
 				}
 			}
 
-			if (Parameters.Length >= 2)
+			if (parameters.Length >= 2)
 			{
-				UhtProperty Parm = (UhtProperty)Parameters.Span[1];
+				UhtProperty parm = (UhtProperty)parameters.Span[1];
 				// First parameter is always the old value:
-				bool bIsValid = false;
-				if (Parm is UhtArrayProperty ArrayProperty)
+				bool isValid = false;
+				if (parm is UhtArrayProperty arrayProperty)
 				{
-					bIsValid = ArrayProperty.ValueProperty is UhtByteProperty && Parm.PropertyFlags.HasAllFlags(EPropertyFlags.ConstParm | EPropertyFlags.ReferenceParm);
+					isValid = arrayProperty.ValueProperty is UhtByteProperty && parm.PropertyFlags.HasAllFlags(EPropertyFlags.ConstParm | EPropertyFlags.ReferenceParm);
 				}
-				if (!bIsValid)
+				if (!isValid)
 				{
-					TargetFunction.LogError($"Replication notification function '{TargetFunction.SourceName}' second (optional) parameter must be of type 'const TArray<uint8>&'.");
+					targetFunction.LogError($"Replication notification function '{targetFunction.SourceName}' second (optional) parameter must be of type 'const TArray<uint8>&'.");
 				}
 			}
 		}
 
 		void ValidateProperties()
 		{
-			foreach (UhtType Child in this.Children)
+			foreach (UhtType child in this.Children)
 			{
-				if (Child is UhtProperty Property)
+				if (child is UhtProperty property)
 				{
-					if (Property.PropertyFlags.HasAnyFlags(EPropertyFlags.RepNotify))
+					if (property.PropertyFlags.HasAnyFlags(EPropertyFlags.RepNotify))
 					{
-						string Name = Property.RepNotifyName != null ? Property.RepNotifyName : String.Empty;
-						ValidateRepNotifyCallback(Property, Name, (UhtFunction?)FindType(UhtFindOptions.EngineName | UhtFindOptions.Function, Name.ToString()));
+						string name = property.RepNotifyName ?? String.Empty;
+						ValidateRepNotifyCallback(property, name, (UhtFunction?)FindType(UhtFindOptions.EngineName | UhtFindOptions.Function, name.ToString()));
 					}
 
-					if (Property.PropertyFlags.HasAnyFlags(EPropertyFlags.BlueprintVisible))
+					if (property.PropertyFlags.HasAnyFlags(EPropertyFlags.BlueprintVisible))
 					{
-						if (Property.MetaData.TryGetValue(UhtNames.BlueprintGetter, out string? Getter) && Getter.Length > 0)
+						if (property.MetaData.TryGetValue(UhtNames.BlueprintGetter, out string? getter) && getter.Length > 0)
 						{
-							ValidateBlueprintPopertyGetter(Property, Getter, (UhtFunction?)FindType(UhtFindOptions.EngineName | UhtFindOptions.Function, Getter));
+							ValidateBlueprintPopertyGetter(property, getter, (UhtFunction?)FindType(UhtFindOptions.EngineName | UhtFindOptions.Function, getter));
 						}
 
-						if (!Property.PropertyFlags.HasAnyFlags(EPropertyFlags.BlueprintReadOnly))
+						if (!property.PropertyFlags.HasAnyFlags(EPropertyFlags.BlueprintReadOnly))
 						{
-							if (Property.MetaData.TryGetValue(UhtNames.BlueprintSetter, out string? Setter) && Setter.Length > 0)
+							if (property.MetaData.TryGetValue(UhtNames.BlueprintSetter, out string? setter) && setter.Length > 0)
 							{
-								ValidateBlueprintPopertySetter(Property, Setter, (UhtFunction?)FindType(UhtFindOptions.EngineName | UhtFindOptions.Function, Setter));
+								ValidateBlueprintPopertySetter(property, setter, (UhtFunction?)FindType(UhtFindOptions.EngineName | UhtFindOptions.Function, setter));
 							}
 						}
 					}
 
 					// Validate if we are using editor only data in a class or struct definition
-					if (Property.PropertyFlags.HasAnyFlags(EPropertyFlags.EditorOnly) && this.ClassFlags.HasAnyFlags(EClassFlags.Optional))
+					if (property.PropertyFlags.HasAnyFlags(EPropertyFlags.EditorOnly) && this.ClassFlags.HasAnyFlags(EClassFlags.Optional))
 					{
-						Property.LogError("Cannot specify an editor only property inside an optional class.");
+						property.LogError("Cannot specify an editor only property inside an optional class.");
 					}
 
 					// Check for getter/setter
-					if (Property.PropertyExportFlags.HasExactFlags(UhtPropertyExportFlags.GetterSpecified | UhtPropertyExportFlags.GetterSpecifiedNone | UhtPropertyExportFlags.GetterFound, UhtPropertyExportFlags.GetterSpecified))
+					if (property.PropertyExportFlags.HasExactFlags(UhtPropertyExportFlags.GetterSpecified | UhtPropertyExportFlags.GetterSpecifiedNone | 
+						UhtPropertyExportFlags.GetterFound, UhtPropertyExportFlags.GetterSpecified))
 					{
-						string ExpectedName = Property.Getter == null ? $"Get{Property.SourceName}" : Property.Getter;
-						using (BorrowStringBuilder Borrower = new BorrowStringBuilder(StringBuilderCache.Small))
+						string expectedName = property.Getter ?? $"Get{property.SourceName}";
+						using (BorrowStringBuilder borrower = new BorrowStringBuilder(StringBuilderCache.Small))
 						{
-							StringBuilder Builder = Borrower.StringBuilder;
-							Builder.AppendPropertyText(Property, UhtPropertyTextType.GetterSetterArg);
-							if (Property.bIsStaticArray)
+							StringBuilder builder = borrower.StringBuilder;
+							builder.AppendPropertyText(property, UhtPropertyTextType.GetterSetterArg);
+							if (property.IsStaticArray)
 							{
-								Builder.Append($"*/[{Property.ArrayDimensions}]");
+								builder.Append($"*/[{property.ArrayDimensions}]");
 							}
-							string ExpectedType = Builder.ToString();
-							Property.LogError($"Property '{Property.SourceName}' expected a getter function '[const] {ExpectedType} [&] {ExpectedName}() const', but it was not found");
+							string expectedType = builder.ToString();
+							property.LogError($"Property '{property.SourceName}' expected a getter function '[const] {expectedType} [&] {expectedName}() const', but it was not found");
 						}
 					}
-					if (Property.PropertyExportFlags.HasExactFlags(UhtPropertyExportFlags.SetterSpecified | UhtPropertyExportFlags.SetterSpecifiedNone | UhtPropertyExportFlags.SetterFound, UhtPropertyExportFlags.SetterSpecified))
+					if (property.PropertyExportFlags.HasExactFlags(UhtPropertyExportFlags.SetterSpecified | UhtPropertyExportFlags.SetterSpecifiedNone | UhtPropertyExportFlags.SetterFound, UhtPropertyExportFlags.SetterSpecified))
 					{
-						string ExpectedName = Property.Setter == null ? $"Set{Property.SourceName}" : Property.Setter;
-						using (BorrowStringBuilder Borrower = new BorrowStringBuilder(StringBuilderCache.Small))
+						string expectedName = property.Setter ?? $"Set{property.SourceName}";
+						using (BorrowStringBuilder borrower = new BorrowStringBuilder(StringBuilderCache.Small))
 						{
-							StringBuilder Builder = Borrower.StringBuilder;
-							Builder.AppendPropertyText(Property, UhtPropertyTextType.GetterSetterArg);
-							if (Property.bIsStaticArray)
+							StringBuilder builder = borrower.StringBuilder;
+							builder.AppendPropertyText(property, UhtPropertyTextType.GetterSetterArg);
+							if (property.IsStaticArray)
 							{
-								Builder.Append($"*/[{Property.ArrayDimensions}]");
+								builder.Append($"*/[{property.ArrayDimensions}]");
 							}
-							string ExpectedType = Builder.ToString();
-							Property.LogError($"Property '{Property.SourceName}' expected a setter function 'void {ExpectedName}([const] {ExpectedType} [&] InArg)', but it was not found");
+							string expectedType = builder.ToString();
+							property.LogError($"Property '{property.SourceName}' expected a setter function 'void {expectedName}([const] {expectedType} [&] InArg)', but it was not found");
 						}
 					}
 				}
@@ -1352,7 +1355,7 @@ namespace EpicGames.UHT.Types
 		#endregion
 
 		/// <inheritdoc/>
-		public override void CollectReferences(IUhtReferenceCollector Collector)
+		public override void CollectReferences(IUhtReferenceCollector collector)
 		{
 
 			// Ignore any intrinsics
@@ -1362,41 +1365,41 @@ namespace EpicGames.UHT.Types
 			}
 
 			// In the original UHT, we don't export IInterface (it wasn't even in the header file)
-			if (this.HeaderFile.bIsNoExportTypes && this.SourceName == "IInterface")
+			if (this.HeaderFile.IsNoExportTypes && this.SourceName == "IInterface")
 			{
 				return;
 			}
 
 			// Add myself as declarations and cross module references
-			Collector.AddDeclaration(this, false);
-			Collector.AddDeclaration(this, true);
-			Collector.AddCrossModuleReference(this, false);
-			Collector.AddCrossModuleReference(this, true);
+			collector.AddDeclaration(this, false);
+			collector.AddDeclaration(this, true);
+			collector.AddCrossModuleReference(this, false);
+			collector.AddCrossModuleReference(this, true);
 
 			// Add the super class
 			if (this.SuperClass != null)
 			{
-				Collector.AddDeclaration(this.SuperClass, true);
-				Collector.AddCrossModuleReference(this.SuperClass, true);
+				collector.AddDeclaration(this.SuperClass, true);
+				collector.AddCrossModuleReference(this.SuperClass, true);
 			}
 
 			// Add any other base classes
-			foreach (UhtStruct Base in this.Bases)
+			foreach (UhtStruct @base in this.Bases)
 			{
-				if (Base is UhtClass BaseClass)
+				if (@base is UhtClass baseClass)
 				{
-					Collector.AddCrossModuleReference(BaseClass, false);
+					collector.AddCrossModuleReference(baseClass, false);
 				}
 			}
 
 			// Add the package
-			Collector.AddDeclaration(this.Package, true);
-			Collector.AddCrossModuleReference(this.Package, true);
+			collector.AddDeclaration(this.Package, true);
+			collector.AddCrossModuleReference(this.Package, true);
 
 			// Collect children references
-			foreach (UhtType Child in this.Children)
+			foreach (UhtType child in this.Children)
 			{
-				Child.CollectReferences(Collector);
+				child.CollectReferences(collector);
 			}
 
 			// Done at the end so any contained delegate functions are added first
@@ -1405,26 +1408,26 @@ namespace EpicGames.UHT.Types
 			switch (this.ClassType)
 			{
 				case UhtClassType.NativeInterface:
-					if (this.AlternateObject != null && this.AlternateObject is UhtField Field)
+					if (this.AlternateObject != null && this.AlternateObject is UhtField field)
 					{
-						Collector.AddExportType(Field);
+						collector.AddExportType(field);
 					}
 					break;
 				case UhtClassType.Interface:
 					break;
 				case UhtClassType.Class:
-					Collector.AddExportType(this);
+					collector.AddExportType(this);
 					break;
 			}
 		}
 
-		private bool ImplementsInterface(UhtClass Interface)
+		private bool ImplementsInterface(UhtClass interfaceObj)
 		{
-			for (UhtClass? SuperClass = this; SuperClass != null; SuperClass = SuperClass.SuperClass)
+			for (UhtClass? superClass = this; superClass != null; superClass = superClass.SuperClass)
 			{
-				foreach (UhtStruct Struct in SuperClass.Bases)
+				foreach (UhtStruct structObj in superClass.Bases)
 				{
-					if (Struct.IsChildOf(Interface))
+					if (structObj.IsChildOf(interfaceObj))
 					{
 						return true;
 					}
diff --git a/Engine/Source/Programs/Shared/EpicGames.UHT/Types/UhtEnum.cs b/Engine/Source/Programs/Shared/EpicGames.UHT/Types/UhtEnum.cs
--- a/Engine/Source/Programs/Shared/EpicGames.UHT/Types/UhtEnum.cs
+++ b/Engine/Source/Programs/Shared/EpicGames.UHT/Types/UhtEnum.cs
@@ -1,15 +1,15 @@
 // Copyright Epic Games, Inc. All Rights Reserved.
 
-using EpicGames.Core;
-using EpicGames.UHT.Tables;
-using EpicGames.UHT.Utils;
 using System;
 using System.Collections.Generic;
 using System.Text.Json.Serialization;
+using EpicGames.Core;
+using EpicGames.UHT.Tables;
+using EpicGames.UHT.Utils;
 
 namespace EpicGames.UHT.Types
 {
-	
+
 	/// <summary>
 	/// How the enumeration was declared
 	/// </summary>
@@ -131,7 +131,7 @@ namespace EpicGames.UHT.Types
 		/// <summary>
 		/// True if the enum was declared in an WITH_EDITORONLY_DATA block
 		/// </summary>
-		public bool bIsEditorOnly { get; set; } = false;
+		public bool IsEditorOnly { get; set; } = false;
 
 		/// <summary>
 		/// Collection of enumeration values
@@ -143,18 +143,18 @@ namespace EpicGames.UHT.Types
 		public override UhtEngineType EngineType => UhtEngineType.Enum;
 
 		/// <inheritdoc/>
-		public override string EngineClassName { get => "Enum"; }
+		public override string EngineClassName => "Enum";
 
 		///<inheritdoc/>
 		[JsonIgnore]
-		protected override UhtSpecifierValidatorTable? SpecifierValidatorTable { get => this.Session.GetSpecifierValidatorTable(UhtTableNames.Enum); }
+		protected override UhtSpecifierValidatorTable? SpecifierValidatorTable => this.Session.GetSpecifierValidatorTable(UhtTableNames.Enum);
 
 		/// <summary>
 		/// Construct a new enumeration
 		/// </summary>
-		/// <param name="Outer">Outer type</param>
-		/// <param name="LineNumber">Line number of declaration</param>
-		public UhtEnum(UhtType Outer, int LineNumber) : base(Outer, LineNumber)
+		/// <param name="outer">Outer type</param>
+		/// <param name="lineNumber">Line number of declaration</param>
+		public UhtEnum(UhtType outer, int lineNumber) : base(outer, lineNumber)
 		{
 			this.MetaData.KeyConversion = this;
 			this.EnumValues = new List<UhtEnumValue>();
@@ -163,13 +163,13 @@ namespace EpicGames.UHT.Types
 		/// <summary>
 		/// Test to see if the value is a known enum value
 		/// </summary>
-		/// <param name="Value">Value in question</param>
+		/// <param name="value">Value in question</param>
 		/// <returns>True if the value is known</returns>
-		public bool IsValidEnumValue(long Value)
+		public bool IsValidEnumValue(long value)
 		{
-			foreach (UhtEnumValue EnumValue in EnumValues)
+			foreach (UhtEnumValue enumValue in EnumValues)
 			{
-				if (EnumValue.Value == Value)
+				if (enumValue.Value == value)
 				{
 					return true;
 				}
@@ -180,16 +180,16 @@ namespace EpicGames.UHT.Types
 		/// <summary>
 		/// Return the index of the given enumeration value name
 		/// </summary>
-		/// <param name="Name">Value name in question</param>
+		/// <param name="name">Value name in question</param>
 		/// <returns>Index of the value or -1 if not found.</returns>
-		public int GetIndexByName(string Name)
+		public int GetIndexByName(string name)
 		{
-			Name = CleanEnumValueName(Name);
-			for (int Index = 0; Index < EnumValues.Count; ++Index)
+			name = CleanEnumValueName(name);
+			for (int index = 0; index < EnumValues.Count; ++index)
 			{
-				if (this.EnumValues[Index].Name == Name)
+				if (this.EnumValues[index].Name == name)
 				{
-					return Index;
+					return index;
 				}
 			}
 			return -1;
@@ -198,39 +198,39 @@ namespace EpicGames.UHT.Types
 		/// <summary>
 		/// Converts meta data name and index to a full meta data key name
 		/// </summary>
-		/// <param name="Name">Meta data key name</param>
-		/// <param name="NameIndex">Meta data key index</param>
+		/// <param name="name">Meta data key name</param>
+		/// <param name="nameIndex">Meta data key index</param>
 		/// <returns>Meta data name with the enum value name</returns>
-		public string GetMetaDataKey(string Name, int NameIndex)
+		public string GetMetaDataKey(string name, int nameIndex)
 		{
-			string EnumName = this.EnumValues[NameIndex].Name;
+			string enumName = this.EnumValues[nameIndex].Name;
 			if (this.CppForm != UhtEnumCppForm.Regular)
 			{
-				int ScopeIndex = EnumName.IndexOf("::", StringComparison.Ordinal);
-				if (ScopeIndex >= 0)
+				int scopeIndex = enumName.IndexOf("::", StringComparison.Ordinal);
+				if (scopeIndex >= 0)
 				{
-					EnumName = EnumName.Substring(ScopeIndex + 2);
+					enumName = enumName.Substring(scopeIndex + 2);
 				}
 			}
-			return $"{EnumName}.{Name}";
+			return $"{enumName}.{name}";
 		}
 
 		/// <summary>
 		/// Given an enumeration value name, return the full enumeration name
 		/// </summary>
-		/// <param name="ShortEnumName">Enum value name</param>
+		/// <param name="shortEnumName">Enum value name</param>
 		/// <returns>If required, enum type name combined with value name.  Otherwise just the value name.</returns>
 		/// <exception cref="UhtIceException">Unexpected enum form</exception>
-		public string GetFullEnumName(string ShortEnumName)
+		public string GetFullEnumName(string shortEnumName)
 		{
 			switch (this.CppForm)
 			{
 				case UhtEnumCppForm.Namespaced:
 				case UhtEnumCppForm.EnumClass:
-					return $"{this.SourceName}::{ShortEnumName}";
+					return $"{this.SourceName}::{shortEnumName}";
 
 				case UhtEnumCppForm.Regular:
-					return ShortEnumName;
+					return shortEnumName;
 
 				default:
 					throw new UhtIceException("Unexpected UhtEnumCppForm value");
@@ -240,33 +240,33 @@ namespace EpicGames.UHT.Types
 		/// <summary>
 		/// Add a new enum value.
 		/// </summary>
-		/// <param name="ShortEnumName">Name of the enum value.</param>
-		/// <param name="Value">Enumeration value or -1 if the value can't be determined.</param>
+		/// <param name="shortEnumName">Name of the enum value.</param>
+		/// <param name="value">Enumeration value or -1 if the value can't be determined.</param>
 		/// <returns></returns>
-		public int AddEnumValue(string ShortEnumName, long Value)
+		public int AddEnumValue(string shortEnumName, long value)
 		{
-			int EnumIndex = this.EnumValues.Count;
-			this.EnumValues.Add(new UhtEnumValue { Name = GetFullEnumName(ShortEnumName), Value = Value });
-			return EnumIndex;
+			int enumIndex = this.EnumValues.Count;
+			this.EnumValues.Add(new UhtEnumValue { Name = GetFullEnumName(shortEnumName), Value = value });
+			return enumIndex;
 		}
 
 		/// <summary>
 		/// Reconstruct the full enum name.  Any existing enumeration name will be stripped and replaced 
 		/// with this enumeration name.
 		/// </summary>
-		/// <param name="Name">Name to reconstruct.</param>
+		/// <param name="name">Name to reconstruct.</param>
 		/// <returns>Reconstructed enum name</returns>
-		private string CleanEnumValueName(string Name)
+		private string CleanEnumValueName(string name)
 		{
-			int LastColons = Name.LastIndexOf("::", StringComparison.Ordinal);
-			return LastColons == -1 ? GetFullEnumName(Name) : GetFullEnumName(Name.Substring(LastColons + 2));		
+			int lastColons = name.LastIndexOf("::", StringComparison.Ordinal);
+			return lastColons == -1 ? GetFullEnumName(name) : GetFullEnumName(name.Substring(lastColons + 2));
 		}
 
 		#region Validation support
 		///<inheritdoc/>
-		protected override void ValidateDocumentationPolicy(UhtDocumentationPolicy Policy)
+		protected override void ValidateDocumentationPolicy(UhtDocumentationPolicy policy)
 		{
-			if (Policy.bClassOrStructCommentRequired)
+			if (policy.ClassOrStructCommentRequired)
 			{
 				if (!this.MetaData.ContainsKey(UhtNames.ToolTip))
 				{
@@ -274,42 +274,42 @@ namespace EpicGames.UHT.Types
 				}
 			}
 
-			Dictionary<string, string> ToolTips = new Dictionary<string, string>(StringComparer.OrdinalIgnoreCase);
-			for (int EnumIndex = 0; EnumIndex < this.EnumValues.Count; ++EnumIndex)
+			Dictionary<string, string> toolTips = new Dictionary<string, string>(StringComparer.OrdinalIgnoreCase);
+			for (int enumIndex = 0; enumIndex < this.EnumValues.Count; ++enumIndex)
 			{
-				string? EntryName;
-				if (!this.MetaData.TryGetValue(UhtNames.Name, EnumIndex, out EntryName))
+				string? entryName;
+				if (!this.MetaData.TryGetValue(UhtNames.Name, enumIndex, out entryName))
 				{
 					continue;
 				}
 
-				string? ToolTip;
-				if (!this.MetaData.TryGetValue(UhtNames.ToolTip, EnumIndex, out ToolTip))
+				string? toolTip;
+				if (!this.MetaData.TryGetValue(UhtNames.ToolTip, enumIndex, out toolTip))
 				{
-					this.LogError(this.MetaData.LineNumber, $"Enum entry '{this.SourceName}::{EntryName}' does not provide a tooltip / comment (DocumentationPolicy)");
+					this.LogError(this.MetaData.LineNumber, $"Enum entry '{this.SourceName}::{entryName}' does not provide a tooltip / comment (DocumentationPolicy)");
 					continue;
 				}
 
-				string? DupName;
-				if (ToolTips.TryGetValue(ToolTip, out DupName))
+				string? dupName;
+				if (toolTips.TryGetValue(toolTip, out dupName))
 				{
-					this.LogError(this.MetaData.LineNumber, $"Enum entries '{this.SourceName}::{EntryName}' and '{this.SourceName}::{DupName}' have identical tooltips / comments (DocumentationPolicy)");
+					this.LogError(this.MetaData.LineNumber, $"Enum entries '{this.SourceName}::{entryName}' and '{this.SourceName}::{dupName}' have identical tooltips / comments (DocumentationPolicy)");
 				}
 				else
 				{
-					ToolTips.Add(ToolTip, EntryName);
+					toolTips.Add(toolTip, entryName);
 				}
 			}
 		}
 		#endregion
 
 		/// <inheritdoc/>
-		public override void CollectReferences(IUhtReferenceCollector Collector)
+		public override void CollectReferences(IUhtReferenceCollector collector)
 		{
-			Collector.AddExportType(this);
-			Collector.AddDeclaration(this, true);
-			Collector.AddCrossModuleReference(this, true);
-			Collector.AddCrossModuleReference(this.Package, true);
+			collector.AddExportType(this);
+			collector.AddDeclaration(this, true);
+			collector.AddCrossModuleReference(this, true);
+			collector.AddCrossModuleReference(this.Package, true);
 		}
 	}
 }
