More coding standard fixes

diff --git a/Engine/Source/Programs/Shared/EpicGames.UHT/Types/UhtField.cs b/Engine/Source/Programs/Shared/EpicGames.UHT/Types/UhtField.cs
--- a/Engine/Source/Programs/Shared/EpicGames.UHT/Types/UhtField.cs
+++ b/Engine/Source/Programs/Shared/EpicGames.UHT/Types/UhtField.cs
@@ -9,14 +9,14 @@ namespace EpicGames.UHT.Types
 	public abstract class UhtField : UhtObject
 	{
 		/// <inheritdoc/>
-		public override string EngineClassName { get => "Field"; }
+		public override string EngineClassName => "Field";
 
 		/// <summary>
 		/// Construct a new field
 		/// </summary>
-		/// <param name="Outer">Outer object</param>
-		/// <param name="LineNumber">Line number of declaration</param>
-		protected UhtField(UhtType Outer, int LineNumber) : base(Outer, LineNumber)
+		/// <param name="outer">Outer object</param>
+		/// <param name="lineNumber">Line number of declaration</param>
+		protected UhtField(UhtType outer, int lineNumber) : base(outer, lineNumber)
 		{
 		}
 	}
diff --git a/Engine/Source/Programs/Shared/EpicGames.UHT/Types/UhtFunction.cs b/Engine/Source/Programs/Shared/EpicGames.UHT/Types/UhtFunction.cs
--- a/Engine/Source/Programs/Shared/EpicGames.UHT/Types/UhtFunction.cs
+++ b/Engine/Source/Programs/Shared/EpicGames.UHT/Types/UhtFunction.cs
@@ -1,12 +1,12 @@
 // Copyright Epic Games, Inc. All Rights Reserved.
 
-using EpicGames.Core;
-using EpicGames.UHT.Tables;
-using EpicGames.UHT.Utils;
 using System;
 using System.Collections.Generic;
 using System.Text;
 using System.Text.Json.Serialization;
+using EpicGames.Core;
+using EpicGames.UHT.Tables;
+using EpicGames.UHT.Utils;
 
 namespace EpicGames.UHT.Types
 {
@@ -14,7 +14,7 @@ namespace EpicGames.UHT.Types
 	/// Series of flags not part of the engine's function flags that affect code generation or verification
 	/// </summary>
 	[Flags]
-	public enum UhtFunctionExportFlags : Int32
+	public enum UhtFunctionExportFlags : int
 	{
 
 		/// <summary>
@@ -120,35 +120,35 @@ namespace EpicGames.UHT.Types
 		/// <summary>
 		/// Test to see if any of the specified flags are set
 		/// </summary>
-		/// <param name="InFlags">Current flags</param>
-		/// <param name="TestFlags">Flags to test for</param>
+		/// <param name="inFlags">Current flags</param>
+		/// <param name="testFlags">Flags to test for</param>
 		/// <returns>True if any of the flags are set</returns>
-		public static bool HasAnyFlags(this UhtFunctionExportFlags InFlags, UhtFunctionExportFlags TestFlags)
+		public static bool HasAnyFlags(this UhtFunctionExportFlags inFlags, UhtFunctionExportFlags testFlags)
 		{
-			return (InFlags & TestFlags) != 0;
+			return (inFlags & testFlags) != 0;
 		}
 
 		/// <summary>
 		/// Test to see if all of the specified flags are set
 		/// </summary>
-		/// <param name="InFlags">Current flags</param>
-		/// <param name="TestFlags">Flags to test for</param>
+		/// <param name="inFlags">Current flags</param>
+		/// <param name="testFlags">Flags to test for</param>
 		/// <returns>True if all the flags are set</returns>
-		public static bool HasAllFlags(this UhtFunctionExportFlags InFlags, UhtFunctionExportFlags TestFlags)
+		public static bool HasAllFlags(this UhtFunctionExportFlags inFlags, UhtFunctionExportFlags testFlags)
 		{
-			return (InFlags & TestFlags) == TestFlags;
+			return (inFlags & testFlags) == testFlags;
 		}
 
 		/// <summary>
 		/// Test to see if a specific set of flags have a specific value.
 		/// </summary>
-		/// <param name="InFlags">Current flags</param>
-		/// <param name="TestFlags">Flags to test for</param>
-		/// <param name="MatchFlags">Expected value of the tested flags</param>
+		/// <param name="inFlags">Current flags</param>
+		/// <param name="testFlags">Flags to test for</param>
+		/// <param name="matchFlags">Expected value of the tested flags</param>
 		/// <returns>True if the given flags have a specific value.</returns>
-		public static bool HasExactFlags(this UhtFunctionExportFlags InFlags, UhtFunctionExportFlags TestFlags, UhtFunctionExportFlags MatchFlags)
+		public static bool HasExactFlags(this UhtFunctionExportFlags inFlags, UhtFunctionExportFlags testFlags, UhtFunctionExportFlags matchFlags)
 		{
-			return (InFlags & TestFlags) == MatchFlags;
+			return (inFlags & testFlags) == matchFlags;
 		}
 	}
 
@@ -183,11 +183,11 @@ namespace EpicGames.UHT.Types
 		/// <summary>
 		/// Test to see if the function type is a delegate type
 		/// </summary>
-		/// <param name="FunctionType">Function type being tested</param>
+		/// <param name="functionType">Function type being tested</param>
 		/// <returns>True if the type is a delegate</returns>
-		public static bool IsDelegate(this UhtFunctionType FunctionType)
+		public static bool IsDelegate(this UhtFunctionType functionType)
 		{
-			return FunctionType == UhtFunctionType.Delegate || FunctionType == UhtFunctionType.SparseDelegate;
+			return functionType == UhtFunctionType.Delegate || functionType == UhtFunctionType.SparseDelegate;
 		}
 	}
 
@@ -196,7 +196,7 @@ namespace EpicGames.UHT.Types
 	/// </summary>
 	public class UhtFunction : UhtStruct
 	{
-		private string? StrippedFunctionNameInternal = null;
+		private string? _strippedFunctionNameInternal = null;
 
 		/// <summary>
 		/// Suffix added to delegate engine names 
@@ -270,96 +270,90 @@ namespace EpicGames.UHT.Types
 		[JsonIgnore]
 		public string StrippedFunctionName
 		{
-			get
-			{
-				return this.StrippedFunctionNameInternal != null ? this.StrippedFunctionNameInternal : this.EngineName;
-			}
-			set 
-			{
-				this.StrippedFunctionNameInternal = value;
-			}
+			get => this._strippedFunctionNameInternal ?? this.EngineName;
+			set => this._strippedFunctionNameInternal = value;
 		}
 
 		///<inheritdoc/>
 		[JsonIgnore]
-		public override bool bDeprecated => this.MetaData.ContainsKey(UhtNames.DeprecatedFunction);
+		public override bool Deprecated => this.MetaData.ContainsKey(UhtNames.DeprecatedFunction);
 
 		///<inheritdoc/>
 		[JsonIgnore]
-		protected override UhtSpecifierValidatorTable? SpecifierValidatorTable { get => this.Session.GetSpecifierValidatorTable(UhtTableNames.Function); }
+		protected override UhtSpecifierValidatorTable? SpecifierValidatorTable => this.Session.GetSpecifierValidatorTable(UhtTableNames.Function);
 
 		/// <summary>
 		/// Identifier for an RPC call to a platform service
 		/// </summary>
-		public UInt16 RPCId { get; set; } = 0;
+		public ushort RPCId { get; set; } = 0;
 
 		/// <summary>
 		/// Identifier for an RPC call expecting a response
 		/// </summary>
-		public UInt16 RPCResponseId { get; set; } = 0;
+		public ushort RPCResponseId { get; set; } = 0;
 
 		/// <summary>
 		/// Name of the actual implementation
 		/// </summary>
-		public string CppImplName { get; set; } = string.Empty;
+		public string CppImplName { get; set; } = String.Empty;
 
 		/// <summary>
 		/// Name of the actual validation implementation
 		/// </summary>
-		public string CppValidationImplName { get; set; } = string.Empty;
+		public string CppValidationImplName { get; set; } = String.Empty;
 
 		/// <summary>
 		/// Name of the wrapper function that marshals the arguments and does the indirect call
 		/// </summary>
-		public string MarshalAndCallName { get; set; } = string.Empty;
+		public string MarshalAndCallName { get; set; } = String.Empty;
 
 		/// <summary>
 		/// Name for callback-style names
 		/// </summary>
-		public string UnMarshalAndCallName { get; set; } = string.Empty;
+		public string UnMarshalAndCallName { get; set; } = String.Empty;
 
 		/// <summary>
 		/// Endpoint name
 		/// </summary>
-		public string EndpointName { get; set; } = string.Empty;
+		public string EndpointName { get; set; } = String.Empty;
 
 		/// <summary>
 		/// True if the function has a return value.
 		/// </summary>
 		[JsonIgnore]
-		public bool bHasReturnProperty => this.Children.Count > 0 && ((UhtProperty)this.Children[this.Children.Count - 1]).PropertyFlags.HasAnyFlags(EPropertyFlags.ReturnParm);
+		public bool HasReturnProperty => this.Children.Count > 0 && ((UhtProperty)this.Children[^1]).PropertyFlags.HasAnyFlags(EPropertyFlags.ReturnParm);
 
 		/// <summary>
 		/// The return value property or null if the function doesn't have a return value
 		/// </summary>
 		[JsonIgnore]
-		public UhtProperty? ReturnProperty => this.bHasReturnProperty ? (UhtProperty)this.Children[this.Children.Count - 1] : null;
+		public UhtProperty? ReturnProperty => this.HasReturnProperty ? (UhtProperty)this.Children[^1] : null;
 
 		/// <summary>
 		/// True if the function has parameters.
 		/// </summary>
 		[JsonIgnore]
-		public bool bHasParameters => this.Children.Count > 0 && !((UhtProperty)this.Children[0]).PropertyFlags.HasAnyFlags(EPropertyFlags.ReturnParm);
+		public bool HasParameters => this.Children.Count > 0 && !((UhtProperty)this.Children[0]).PropertyFlags.HasAnyFlags(EPropertyFlags.ReturnParm);
 
 		/// <summary>
 		/// Return read only memory of all the function parameters
 		/// </summary>
 		[JsonIgnore]
-		public ReadOnlyMemory<UhtType> ParameterProperties => new ReadOnlyMemory<UhtType>(this.Children.ToArray(), 0, this.bHasReturnProperty ? this.Children.Count - 1 : this.Children.Count);
+		public ReadOnlyMemory<UhtType> ParameterProperties => new ReadOnlyMemory<UhtType>(this.Children.ToArray(), 0, this.HasReturnProperty ? this.Children.Count - 1 : this.Children.Count);
 
 		/// <summary>
 		/// True if the function has any outputs including a return value
 		/// </summary>
 		[JsonIgnore]
-		public bool bHasAnyOutputs
+		public bool HasAnyOutputs
 		{
 			get
 			{
-				foreach (UhtType Type in this.Children)
+				foreach (UhtType type in this.Children)
 				{
-					if (Type is UhtProperty Property)
+					if (type is UhtProperty property)
 					{
-						if (Property.PropertyFlags.HasAnyFlags(EPropertyFlags.ReturnParm | EPropertyFlags.OutParm))
+						if (property.PropertyFlags.HasAnyFlags(EPropertyFlags.ReturnParm | EPropertyFlags.OutParm))
 						{
 							return true;
 						}
@@ -372,22 +366,22 @@ namespace EpicGames.UHT.Types
 		/// <summary>
 		/// Construct a new instance of a function
 		/// </summary>
-		/// <param name="Outer">The parent object</param>
-		/// <param name="LineNumber">The line number where the function is defined</param>
-		public UhtFunction(UhtType Outer, int LineNumber) : base(Outer, LineNumber)
+		/// <param name="outer">The parent object</param>
+		/// <param name="lineNumber">The line number where the function is defined</param>
+		public UhtFunction(UhtType outer, int lineNumber) : base(outer, lineNumber)
 		{
 		}
 
 		/// <inheritdoc/>
-		public override void AddChild(UhtType Type)
+		public override void AddChild(UhtType type)
 		{
-			if (Type is UhtProperty Property)
+			if (type is UhtProperty property)
 			{
-				if (!Property.PropertyFlags.HasAnyFlags(EPropertyFlags.Parm))
+				if (!property.PropertyFlags.HasAnyFlags(EPropertyFlags.Parm))
 				{
 					throw new UhtIceException("Only properties marked as parameters can be added to functions");
 				}
-				if (Property.PropertyFlags.HasAnyFlags(EPropertyFlags.ReturnParm) && this.bHasReturnProperty)
+				if (property.PropertyFlags.HasAnyFlags(EPropertyFlags.ReturnParm) && this.HasReturnProperty)
 				{
 					throw new UhtIceException("Attempt to set multiple return properties on function");
 				}
@@ -396,22 +390,22 @@ namespace EpicGames.UHT.Types
 			{
 				throw new UhtIceException("Invalid type added");
 			}
-			base.AddChild(Type);
+			base.AddChild(type);
 		}
 
 		#region Resolution support
 		/// <inheritdoc/>
-		protected override bool ResolveSelf(UhtResolvePhase Phase)
+		protected override bool ResolveSelf(UhtResolvePhase phase)
 		{
-			bool bResults = base.ResolveSelf(Phase);
+			bool results = base.ResolveSelf(phase);
 
-			switch (Phase)
+			switch (phase)
 			{
 				case UhtResolvePhase.Bases:
-					if (this.FunctionType == UhtFunctionType.Function && this.Outer is UhtClass OuterClass)
+					if (this.FunctionType == UhtFunctionType.Function && this.Outer is UhtClass outerClass)
 					{
 						// non-static functions in a const class must be const themselves
-						if (OuterClass.ClassFlags.HasAnyFlags(EClassFlags.Const))
+						if (outerClass.ClassFlags.HasAnyFlags(EClassFlags.Const))
 						{
 							this.FunctionFlags |= EFunctionFlags.Const;
 						}
@@ -419,25 +413,25 @@ namespace EpicGames.UHT.Types
 					break;
 
 				case UhtResolvePhase.Final:
-					if (this.Outer is UhtClass Class)
+					if (this.Outer is UhtClass classObj)
 					{
 						if (this.FunctionFlags.HasAnyFlags(EFunctionFlags.Native) &&
 							!this.FunctionExportFlags.HasAnyFlags(UhtFunctionExportFlags.CustomThunk) &&
 							this.CppImplName != this.SourceName)
 						{
-							UhtFoundDeclaration Declaration;
-							if (Class.TryGetDeclaration(this.CppImplName, out Declaration))
+							UhtFoundDeclaration declaration;
+							if (classObj.TryGetDeclaration(this.CppImplName, out declaration))
 							{
 								this.FunctionExportFlags |= UhtFunctionExportFlags.ImplFound;
-								if (Declaration.bIsVirtual)
+								if (declaration.IsVirtual)
 								{
 									this.FunctionExportFlags |= UhtFunctionExportFlags.ImplVirtual;
 								}
 							}
-							if (Class.TryGetDeclaration(this.CppValidationImplName, out Declaration))
+							if (classObj.TryGetDeclaration(this.CppValidationImplName, out declaration))
 							{
 								this.FunctionExportFlags |= UhtFunctionExportFlags.ValidationImplFound;
-								if (Declaration.bIsVirtual)
+								if (declaration.IsVirtual)
 								{
 									this.FunctionExportFlags |= UhtFunctionExportFlags.ValidationImplVirtual;
 								}
@@ -452,7 +446,7 @@ namespace EpicGames.UHT.Types
 							// const identifier to determine purity is not desirable. We should remove the following logic:
 
 							// If its a const BlueprintCallable function with some sort of output and is not being marked as an BlueprintPure=false function, mark it as BlueprintPure as well
-							if (this.bHasAnyOutputs && this.FunctionFlags.HasAnyFlags(EFunctionFlags.BlueprintCallable) && 
+							if (this.HasAnyOutputs && this.FunctionFlags.HasAnyFlags(EFunctionFlags.BlueprintCallable) &&
 								!this.FunctionExportFlags.HasAnyFlags(UhtFunctionExportFlags.ForceBlueprintImpure))
 							{
 								this.FunctionFlags |= EFunctionFlags.BlueprintPure;
@@ -464,17 +458,17 @@ namespace EpicGames.UHT.Types
 					// The following code is only performed on functions in a class.
 					if (this.Outer is UhtClass)
 					{
-						foreach (UhtType Type in this.Children)
+						foreach (UhtType type in this.Children)
 						{
-							if (Type is UhtProperty Property)
+							if (type is UhtProperty property)
 							{
-								if (Property.PropertyFlags.HasExactFlags(EPropertyFlags.OutParm | EPropertyFlags.ReturnParm, EPropertyFlags.OutParm))
+								if (property.PropertyFlags.HasExactFlags(EPropertyFlags.OutParm | EPropertyFlags.ReturnParm, EPropertyFlags.OutParm))
 								{
 									this.FunctionFlags |= EFunctionFlags.HasOutParms;
 								}
-								if (Property is UhtStructProperty StructProperty)
+								if (property is UhtStructProperty structProperty)
 								{
-									if (StructProperty.ScriptStruct.bHasDefaults)
+									if (structProperty.ScriptStruct.HasDefaults)
 									{
 										this.FunctionFlags |= EFunctionFlags.HasDefaults;
 									}
@@ -484,29 +478,29 @@ namespace EpicGames.UHT.Types
 					}
 					break;
 			}
-			return bResults;
+			return results;
 		}
 		#endregion
 
 		#region Validation support
 		/// <inheritdoc/>
-		protected override UhtValidationOptions Validate(UhtValidationOptions Options)
+		protected override UhtValidationOptions Validate(UhtValidationOptions options)
 		{
-			Options = base.Validate(Options);
+			options = base.Validate(options);
 
 			ValidateSpecifierFlags();
 
 			// Some parameter test are optionally done
-			bool bCheckForBlueprint = false;
+			bool checkForBlueprint = false;
 
 			// Function type specific validation
 			switch (this.FunctionType)
 			{
 				case UhtFunctionType.Function:
-					UhtClass? OuterClass = (UhtClass?)this.Outer;
+					UhtClass? outerClass = (UhtClass?)this.Outer;
 
-					if (this.FunctionFlags.HasAnyFlags(EFunctionFlags.BlueprintPure) && OuterClass != null && 
-						OuterClass.ClassFlags.HasAnyFlags(EClassFlags.Interface) && 
+					if (this.FunctionFlags.HasAnyFlags(EFunctionFlags.BlueprintPure) && outerClass != null &&
+						outerClass.ClassFlags.HasAnyFlags(EClassFlags.Interface) &&
 						!this.FunctionExportFlags.HasAnyFlags(UhtFunctionExportFlags.AutoBlueprintPure))
 					{
 						this.LogError("BlueprintPure specifier is not allowed for interface functions");
@@ -515,14 +509,14 @@ namespace EpicGames.UHT.Types
 					// If this function is blueprint callable or blueprint pure, require a category 
 					if (this.FunctionFlags.HasAnyFlags(EFunctionFlags.BlueprintCallable | EFunctionFlags.BlueprintPure))
 					{
-						bool bBlueprintAccessor = this.MetaData.ContainsKey(UhtNames.BlueprintSetter) || this.MetaData.ContainsKey(UhtNames.BlueprintGetter);
-						bool bHasMenuCategory = this.MetaData.ContainsKey(UhtNames.Category);
-						bool bInternalOnly = this.MetaData.GetBoolean(UhtNames.BlueprintInternalUseOnly);
+						bool blueprintAccessor = this.MetaData.ContainsKey(UhtNames.BlueprintSetter) || this.MetaData.ContainsKey(UhtNames.BlueprintGetter);
+						bool hasMenuCategory = this.MetaData.ContainsKey(UhtNames.Category);
+						bool internalOnly = this.MetaData.GetBoolean(UhtNames.BlueprintInternalUseOnly);
 
-						if (!bHasMenuCategory && !bInternalOnly && !this.bDeprecated && !bBlueprintAccessor)
+						if (!hasMenuCategory && !internalOnly && !this.Deprecated && !blueprintAccessor)
 						{
 							// To allow for quick iteration, don't enforce the requirement that game functions have to be categorized
-							if (this.HeaderFile.Package.bIsPartOfEngine)
+							if (this.HeaderFile.Package.IsPartOfEngine)
 							{
 								this.LogError("An explicit Category specifier is required for Blueprint accessible functions in an Engine module.");
 							}
@@ -536,7 +530,7 @@ namespace EpicGames.UHT.Types
 						// if this is a BlueprintNativeEvent or BlueprintImplementableEvent in an interface, make sure it's not "virtual"
 						if (this.FunctionFlags.HasAnyFlags(EFunctionFlags.BlueprintEvent))
 						{
-							if (OuterClass != null && OuterClass.ClassFlags.HasAnyFlags(EClassFlags.Interface))
+							if (outerClass != null && outerClass.ClassFlags.HasAnyFlags(EClassFlags.Interface))
 							{
 								this.LogError("BlueprintImplementableEvents in Interfaces must not be declared 'virtual'");
 							}
@@ -556,7 +550,7 @@ namespace EpicGames.UHT.Types
 					else
 					{
 						// if this is a function in an Interface, it must be marked 'virtual' unless it's an event
-						if (OuterClass != null && OuterClass.ClassFlags.HasAnyFlags(EClassFlags.Interface) && !this.FunctionFlags.HasAnyFlags(EFunctionFlags.BlueprintEvent))
+						if (outerClass != null && outerClass.ClassFlags.HasAnyFlags(EClassFlags.Interface) && !this.FunctionFlags.HasAnyFlags(EFunctionFlags.BlueprintEvent))
 						{
 							this.LogError("Interface functions that are not BlueprintImplementableEvents must be declared 'virtual'");
 						}
@@ -565,7 +559,7 @@ namespace EpicGames.UHT.Types
 					// This is a final (prebinding, non-overridable) function
 					if (this.FunctionFlags.HasAnyFlags(EFunctionFlags.Final) || this.FunctionExportFlags.HasAnyFlags(UhtFunctionExportFlags.Final))
 					{
-						if (OuterClass != null && OuterClass.ClassFlags.HasAnyFlags(EClassFlags.Interface))
+						if (outerClass != null && outerClass.ClassFlags.HasAnyFlags(EClassFlags.Interface))
 						{
 							this.LogError("Interface functions cannot be declared 'final'");
 						}
@@ -576,13 +570,13 @@ namespace EpicGames.UHT.Types
 					}
 
 					if (this.FunctionFlags.HasAnyFlags(EFunctionFlags.RequiredAPI) &&
-						OuterClass != null && OuterClass.ClassFlags.HasAnyFlags(EClassFlags.RequiredAPI))
+						outerClass != null && outerClass.ClassFlags.HasAnyFlags(EClassFlags.RequiredAPI))
 					{
 						this.LogError($"API must not be used on methods of a class that is marked with an API itself.");
 					}
 
 					// Make sure that blueprint pure functions have some type of output
-					if (this.FunctionFlags.HasAnyFlags(EFunctionFlags.BlueprintPure) && !this.bHasAnyOutputs)
+					if (this.FunctionFlags.HasAnyFlags(EFunctionFlags.BlueprintPure) && !this.HasAnyOutputs)
 					{
 						this.LogError("BlueprintPure specifier is not allowed for functions with no return value and no output parameters.");
 					}
@@ -592,28 +586,28 @@ namespace EpicGames.UHT.Types
 						this.LogError("Replicated functions can't have return values");
 					}
 
-					bCheckForBlueprint = this.FunctionFlags.HasAnyFlags(EFunctionFlags.BlueprintCallable | EFunctionFlags.BlueprintEvent);
+					checkForBlueprint = this.FunctionFlags.HasAnyFlags(EFunctionFlags.BlueprintCallable | EFunctionFlags.BlueprintEvent);
 
-					if (this.Outer is UhtStruct OuterStruct)
+					if (this.Outer is UhtStruct outerStruct)
 					{
-						UhtStruct? OuterSuperStruct = OuterStruct.SuperStruct;
-						if (OuterSuperStruct != null)
+						UhtStruct? outerSuperStruct = outerStruct.SuperStruct;
+						if (outerSuperStruct != null)
 						{
-							UhtType? OverriddenFunction = OuterSuperStruct.FindType(UhtFindOptions.SourceName | UhtFindOptions.Function, this.SourceName);
-							if (OverriddenFunction != null)
+							UhtType? overriddenFunction = outerSuperStruct.FindType(UhtFindOptions.SourceName | UhtFindOptions.Function, this.SourceName);
+							if (overriddenFunction != null)
 							{
 								// Native function overrides should be done in CPP text, not in a UFUNCTION() declaration (you can't change flags, and it'd otherwise be a burden to keep them identical)
 								this.LogError(
-									$"Override of UFUNCTION '{this.SourceName}' in parent '{OverriddenFunction.Outer?.SourceName}' cannot have a UFUNCTION() declaration above it; it will use the same parameters as the original declaration.");
+									$"Override of UFUNCTION '{this.SourceName}' in parent '{overriddenFunction.Outer?.SourceName}' cannot have a UFUNCTION() declaration above it; it will use the same parameters as the original declaration.");
 							}
 						}
 					}
 
 					// Make sure that the replication flags set on an overridden function match the parent function
-					UhtFunction? SuperFunction = this.SuperFunction;
-					if (SuperFunction != null)
+					UhtFunction? superFunction = this.SuperFunction;
+					if (superFunction != null)
 					{
-						if ((this.FunctionFlags & EFunctionFlags.NetFuncFlags) != (SuperFunction.FunctionFlags & EFunctionFlags.NetFuncFlags))
+						if ((this.FunctionFlags & EFunctionFlags.NetFuncFlags) != (superFunction.FunctionFlags & EFunctionFlags.NetFuncFlags))
 						{
 							this.LogError($"Overridden function '{this.SourceName}' cannot specify different replication flags when overriding a function.");
 						}
@@ -621,50 +615,50 @@ namespace EpicGames.UHT.Types
 
 					// if this function is an RPC in state scope, verify that it is an override
 					// this is required because the networking code only checks the class for RPCs when initializing network data, not any states within it
-					if (this.FunctionFlags.HasAnyFlags(EFunctionFlags.Net) && SuperFunction == null && !(this.Outer is UhtClass))
+					if (this.FunctionFlags.HasAnyFlags(EFunctionFlags.Net) && superFunction == null && this.Outer is not UhtClass)
 					{
 						this.LogError($"Function '{this.SourceName}' base implementation of RPCs cannot be in a state. Add a stub outside state scope.");
 					}
 
 					// Check implemented RPC functions
-					if (this.CppImplName != this.EngineName && 
+					if (this.CppImplName != this.EngineName &&
 						!this.FunctionExportFlags.HasAnyFlags(UhtFunctionExportFlags.CustomThunk) &&
-						OuterClass!.GeneratedCodeVersion > EGeneratedCodeVersion.V1)
+						outerClass!.GeneratedCodeVersion > EGeneratedCodeVersion.V1)
 					{
-						bool bIsNative = this.FunctionFlags.HasAnyFlags(EFunctionFlags.Native);
-						bool bIsNet = this.FunctionFlags.HasAnyFlags(EFunctionFlags.Net);
-						bool bIsNetValidate = this.FunctionFlags.HasAnyFlags(EFunctionFlags.NetValidate);
-						bool bIsNetResponse = this.FunctionFlags.HasAnyFlags(EFunctionFlags.NetResponse);
-						bool bIsBlueprintEvent = this.FunctionFlags.HasAnyFlags(EFunctionFlags.BlueprintEvent);
+						bool isNative = this.FunctionFlags.HasAnyFlags(EFunctionFlags.Native);
+						bool isNet = this.FunctionFlags.HasAnyFlags(EFunctionFlags.Net);
+						bool isNetValidate = this.FunctionFlags.HasAnyFlags(EFunctionFlags.NetValidate);
+						bool isNetResponse = this.FunctionFlags.HasAnyFlags(EFunctionFlags.NetResponse);
+						bool isBlueprintEvent = this.FunctionFlags.HasAnyFlags(EFunctionFlags.BlueprintEvent);
 
-						bool bNeedsImplementation = (bIsNet && !bIsNetResponse) || bIsBlueprintEvent || bIsNative;
-						bool bNeedsValidate = (bIsNative || bIsNet) && !bIsNetResponse && bIsNetValidate;
+						bool needsImplementation = (isNet && !isNetResponse) || isBlueprintEvent || isNative;
+						bool needsValidate = (isNative || isNet) && !isNetResponse && isNetValidate;
 
-						bool bHasImplementation = this.FunctionExportFlags.HasAnyFlags(UhtFunctionExportFlags.ImplFound);
-						bool bHasValidate = this.FunctionExportFlags.HasAnyFlags(UhtFunctionExportFlags.ValidationImplFound);
-						bool bHasVirtualImplementation = this.FunctionExportFlags.HasAnyFlags(UhtFunctionExportFlags.ImplVirtual);
-						bool bHasVirtualValidate = this.FunctionExportFlags.HasAnyFlags(UhtFunctionExportFlags.ValidationImplVirtual);
+						bool hasImplementation = this.FunctionExportFlags.HasAnyFlags(UhtFunctionExportFlags.ImplFound);
+						bool hasValidate = this.FunctionExportFlags.HasAnyFlags(UhtFunctionExportFlags.ValidationImplFound);
+						bool hasVirtualImplementation = this.FunctionExportFlags.HasAnyFlags(UhtFunctionExportFlags.ImplVirtual);
+						bool hasVirtualValidate = this.FunctionExportFlags.HasAnyFlags(UhtFunctionExportFlags.ValidationImplVirtual);
 
-						if (bNeedsImplementation || bNeedsValidate)
+						if (needsImplementation || needsValidate)
 						{
 							// Check if functions are missing.
-							if (bNeedsImplementation && !bHasImplementation)
+							if (needsImplementation && !hasImplementation)
 							{
 								LogRpcFunctionError(false, this.CppImplName);
 							}
 
-							if (bNeedsValidate && !bHasValidate)
+							if (needsValidate && !hasValidate)
 							{
 								LogRpcFunctionError(false, this.CppValidationImplName);
 							}
 
 							// If all needed functions are declared, check if they have virtual specifiers.
-							if (bNeedsImplementation && bHasImplementation && !bHasVirtualImplementation)
+							if (needsImplementation && hasImplementation && !hasVirtualImplementation)
 							{
 								LogRpcFunctionError(true, this.CppImplName);
 							}
 
-							if (bNeedsValidate && bHasValidate && !bHasVirtualValidate)
+							if (needsValidate && hasValidate && !hasVirtualValidate)
 							{
 								LogRpcFunctionError(true, this.CppValidationImplName);
 							}
@@ -674,7 +668,7 @@ namespace EpicGames.UHT.Types
 					// Validate field notify
 					if (this.FunctionExportFlags.HasAnyFlags(UhtFunctionExportFlags.FieldNotify))
 					{
-						if (OuterClass == null)
+						if (outerClass == null)
 						{
 							this.LogError($"FieldNofity function '{this.SourceName}' are only valid as UClass member function.");
 						}
@@ -705,10 +699,10 @@ namespace EpicGames.UHT.Types
 				case UhtFunctionType.SparseDelegate:
 
 					// Check for shadowing if requested
-					if (Options.HasAnyFlags(UhtValidationOptions.Shadowing))
+					if (options.HasAnyFlags(UhtValidationOptions.Shadowing))
 					{
-						UhtType? Existing = this.Outer!.FindType(UhtFindOptions.DelegateFunction | UhtFindOptions.EngineName | UhtFindOptions.SelfOnly, this.EngineName);
-						if (Existing != null && Existing != this)
+						UhtType? existing = this.Outer!.FindType(UhtFindOptions.DelegateFunction | UhtFindOptions.EngineName | UhtFindOptions.SelfOnly, this.EngineName);
+						if (existing != null && existing != this)
 						{
 							this.LogError($"Can't override delegate signature function '{this.SourceName}'");
 						}
@@ -717,128 +711,128 @@ namespace EpicGames.UHT.Types
 			}
 
 			// Make sure this isn't a duplicate type.  In some ways this has already been check above
-			UhtType? ExistingType = this.Outer!.FindType(UhtFindOptions.PropertyOrFunction | UhtFindOptions.EngineName | UhtFindOptions.SelfOnly, this.EngineName);
-			if (ExistingType != null && ExistingType != this)
+			UhtType? existingType = this.Outer!.FindType(UhtFindOptions.PropertyOrFunction | UhtFindOptions.EngineName | UhtFindOptions.SelfOnly, this.EngineName);
+			if (existingType != null && existingType != this)
 			{
-				this.LogError($"'{this.EngineName}' conflicts with '{ExistingType.FullName}'");
+				this.LogError($"'{this.EngineName}' conflicts with '{existingType.FullName}'");
 			}
 
 			// Do a more generic validation of the arguments
-			foreach (UhtType Type in this.Children)
+			foreach (UhtType type in this.Children)
 			{
-				if (Type is UhtProperty Property)
+				if (type is UhtProperty property)
 				{
 					// This code is technically incorrect but here for compatibility reasons.  Container types (TObjectPtr, TArray, etc...) should have the 
 					// const on the inside of the template arguments but it places the const on the outside which isn't correct.  This needs to be addressed
 					// in any updates to the type parser.  See EGetVarTypeOptions::NoAutoConst for the parser side of the problem.
-					if (!Property.PropertyFlags.HasAnyFlags(EPropertyFlags.ConstParm) && Property.MustBeConstArgument(out UhtType? ErrorType))
+					if (!property.PropertyFlags.HasAnyFlags(EPropertyFlags.ConstParm) && property.MustBeConstArgument(out UhtType? errorType))
 					{
-						if (Property.PropertyFlags.HasAnyFlags(EPropertyFlags.ReturnParm))
+						if (property.PropertyFlags.HasAnyFlags(EPropertyFlags.ReturnParm))
 						{
-							Property.LogError($"Return value must be 'const' since '{ErrorType.SourceName}' is marked 'const'");
+							property.LogError($"Return value must be 'const' since '{errorType.SourceName}' is marked 'const'");
 						}
 						else
 						{
-							Property.LogError($"Argument '{Property.SourceName}' must be 'const' since '{ErrorType.SourceName}' is marked 'const'");
+							property.LogError($"Argument '{property.SourceName}' must be 'const' since '{errorType.SourceName}' is marked 'const'");
 						}
 					}
 
 					// Due to structures that might not be fully parsed at parse time, do blueprint validation here
-					if (bCheckForBlueprint)
+					if (checkForBlueprint)
 					{
-						if (Property.bIsStaticArray)
+						if (property.IsStaticArray)
 						{
-							this.LogError($"Static array cannot be exposed to blueprint. Function: {this.SourceName} Parameter {Property.SourceName}");
+							this.LogError($"Static array cannot be exposed to blueprint. Function: {this.SourceName} Parameter {property.SourceName}");
 						}
-						if (!Property.PropertyCaps.HasAnyFlags(UhtPropertyCaps.IsParameterSupportedByBlueprint))
+						if (!property.PropertyCaps.HasAnyFlags(UhtPropertyCaps.IsParameterSupportedByBlueprint))
 						{
-							this.LogError($"Type '{Property.GetUserFacingDecl()}' is not supported by blueprint. Function:  {this.SourceName} Parameter {Property.SourceName}");
+							this.LogError($"Type '{property.GetUserFacingDecl()}' is not supported by blueprint. Function:  {this.SourceName} Parameter {property.SourceName}");
 						}
 					}
 				}
 			}
 
-			if (!this.bDeprecated)
+			if (!this.Deprecated)
 			{
-				Options |= UhtValidationOptions.Deprecated;
+				options |= UhtValidationOptions.Deprecated;
 			}
 			else
 			{
-				Options &= ~UhtValidationOptions.Deprecated;
+				options &= ~UhtValidationOptions.Deprecated;
 			}
-			return Options;
+			return options;
 		}
 
-		private void LogRpcFunctionError(bool bVirtualError, string FunctionName)
+		private void LogRpcFunctionError(bool virtualError, string functionName)
 		{
-			StringBuilder Builder = new StringBuilder();
+			StringBuilder builder = new StringBuilder();
 
-			if (bVirtualError)
+			if (virtualError)
 			{
-				Builder.Append("Declared function \"");
+				builder.Append("Declared function \"");
 			}
 			else
 			{
-				Builder.Append($"Function {this.SourceName} was marked as ");
+				builder.Append($"Function {this.SourceName} was marked as ");
 				if (this.FunctionFlags.HasAnyFlags(EFunctionFlags.Native))
 				{
-					Builder.Append("Native, ");
+					builder.Append("Native, ");
 				}
 				if (this.FunctionFlags.HasAnyFlags(EFunctionFlags.Net))
 				{
-					Builder.Append("Net, ");
+					builder.Append("Net, ");
 				}
 				if (this.FunctionFlags.HasAnyFlags(EFunctionFlags.BlueprintEvent))
 				{
-					Builder.Append("BlueprintEvent, ");
+					builder.Append("BlueprintEvent, ");
 				}
 				if (this.FunctionFlags.HasAnyFlags(EFunctionFlags.NetValidate))
 				{
-					Builder.Append("NetValidate, ");
+					builder.Append("NetValidate, ");
 				}
-				Builder.Length = Builder.Length - 2;
-				Builder.Append(". Declare function \"virtual ");
+				builder.Length -= 2;
+				builder.Append(". Declare function \"virtual ");
 			}
 
-			UhtProperty? ReturnType = this.ReturnProperty;
-			if (ReturnType != null)
+			UhtProperty? returnType = this.ReturnProperty;
+			if (returnType != null)
 			{
-				Builder.AppendPropertyText(ReturnType, UhtPropertyTextType.EventFunctionArgOrRetVal);
-				Builder.Append(' ');
+				builder.AppendPropertyText(returnType, UhtPropertyTextType.EventFunctionArgOrRetVal);
+				builder.Append(' ');
 			}
 			else
 			{
-				Builder.Append("void ");
+				builder.Append("void ");
 			}
 
-			Builder.Append(this.Outer!.SourceName).Append("::").Append(FunctionName).Append('(');
+			builder.Append(this.Outer!.SourceName).Append("::").Append(functionName).Append('(');
 
-			bool bFirst = true;
-			foreach (UhtProperty Arg in this.ParameterProperties.Span)
+			bool first = true;
+			foreach (UhtProperty arg in this.ParameterProperties.Span)
 			{
-				if (!bFirst)
+				if (!first)
 				{
-					Builder.Append(", ");
+					builder.Append(", ");
 				}
-				bFirst = false;
-				Builder.AppendFullDecl(Arg, UhtPropertyTextType.EventFunctionArgOrRetVal, true);
+				first = false;
+				builder.AppendFullDecl(arg, UhtPropertyTextType.EventFunctionArgOrRetVal, true);
 			}
 
-			Builder.Append(')');
+			builder.Append(')');
 			if (this.FunctionFlags.HasAnyFlags(EFunctionFlags.Const))
 			{
-				Builder.Append(" const");
+				builder.Append(" const");
 			}
 
-			if (bVirtualError)
+			if (virtualError)
 			{
-				Builder.Append("\" is not marked as virtual.");
+				builder.Append("\" is not marked as virtual.");
 			}
 			else
 			{
-				Builder.Append('\"');
+				builder.Append('\"');
 			}
-			this.LogError(Builder.ToString());
+			this.LogError(builder.ToString());
 		}
 
 		private void ValidateSpecifierFlags()
@@ -850,20 +844,20 @@ namespace EpicGames.UHT.Types
 					this.LogError("Replicated functions can not be blueprint events or exec");
 				}
 
-				bool bIsNetService = this.FunctionFlags.HasAnyFlags(EFunctionFlags.NetRequest | EFunctionFlags.NetResponse);
-				bool bIsNetReliable = this.FunctionFlags.HasAnyFlags(EFunctionFlags.NetReliable);
+				bool isNetService = this.FunctionFlags.HasAnyFlags(EFunctionFlags.NetRequest | EFunctionFlags.NetResponse);
+				bool isNetReliable = this.FunctionFlags.HasAnyFlags(EFunctionFlags.NetReliable);
 
 				if (this.FunctionFlags.HasAnyFlags(EFunctionFlags.Static))
 				{
 					this.LogError("Static functions can't be replicated");
 				}
 
-				if (!bIsNetReliable && !this.FunctionExportFlags.HasAnyFlags(UhtFunctionExportFlags.Unreliable) && !bIsNetService)
+				if (!isNetReliable && !this.FunctionExportFlags.HasAnyFlags(UhtFunctionExportFlags.Unreliable) && !isNetService)
 				{
 					this.LogError("Replicated function: 'reliable' or 'unreliable' is required");
 				}
 
-				if (bIsNetReliable && this.FunctionExportFlags.HasAnyFlags(UhtFunctionExportFlags.Unreliable) && !bIsNetService)
+				if (isNetReliable && this.FunctionExportFlags.HasAnyFlags(UhtFunctionExportFlags.Unreliable) && !isNetService)
 				{
 					this.LogError("'reliable' and 'unreliable' are mutually exclusive");
 				}
@@ -894,78 +888,78 @@ namespace EpicGames.UHT.Types
 		}
 
 		/// <inheritdoc/>
-		protected override void ValidateDocumentationPolicy(UhtDocumentationPolicy Policy)
+		protected override void ValidateDocumentationPolicy(UhtDocumentationPolicy policy)
 		{
 			if (this.FunctionType != UhtFunctionType.Function)
 			{
 				return;
 			}
 
-			UhtClass? OuterClass = (UhtClass?)this.Outer;
+			UhtClass? outerClass = (UhtClass?)this.Outer;
 
-			if (Policy.bFunctionToolTipsRequired)
+			if (policy.FunctionToolTipsRequired)
 			{
 				if (!this.MetaData.ContainsKey(UhtNames.ToolTip))
 				{
-					this.LogError($"Function '{OuterClass?.SourceName}::{this.SourceName}' does not provide a tooltip / comment (DocumentationPolicy).");
+					this.LogError($"Function '{outerClass?.SourceName}::{this.SourceName}' does not provide a tooltip / comment (DocumentationPolicy).");
 				}
 			}
 
-			if (Policy.bParameterToolTipsRequired)
+			if (policy.ParameterToolTipsRequired)
 			{
 				if (!this.MetaData.ContainsKey(UhtNames.Comment))
 				{
-					this.LogError($"Function '{OuterClass?.SourceName}::{this.SourceName}' does not provide a comment (DocumentationPolicy).");
+					this.LogError($"Function '{outerClass?.SourceName}::{this.SourceName}' does not provide a comment (DocumentationPolicy).");
 
 				}
 
-				Dictionary<StringView, StringView> ParamToolTips = GetParameterToolTipsFromFunctionComment(this.MetaData.GetValueOrDefault(UhtNames.Comment));
-				bool bHasAnyParamToolTips = ParamToolTips.Count > 1 || (ParamToolTips.Count == 1 && !ParamToolTips.ContainsKey(UhtNames.ReturnValue));
+				Dictionary<StringView, StringView> paramToolTips = GetParameterToolTipsFromFunctionComment(this.MetaData.GetValueOrDefault(UhtNames.Comment));
+				bool hasAnyParamToolTips = paramToolTips.Count > 1 || (paramToolTips.Count == 1 && !paramToolTips.ContainsKey(UhtNames.ReturnValue));
 
 				// only apply the validation for parameter tooltips if a function has any @param statements at all.
-				if (bHasAnyParamToolTips)
+				if (hasAnyParamToolTips)
 				{
 					// ensure each parameter has a tooltip
-					HashSet<StringView> ExistingFields = new HashSet<StringView>(StringViewComparer.OrdinalIgnoreCase);
-					foreach (UhtType Parameter in this.ParameterProperties.Span)
+					HashSet<StringView> existingFields = new HashSet<StringView>(StringViewComparer.OrdinalIgnoreCase);
+					foreach (UhtType parameter in this.ParameterProperties.Span)
 					{
-						ExistingFields.Add(Parameter.SourceName);
-						if (!ParamToolTips.ContainsKey(Parameter.SourceName))
+						existingFields.Add(parameter.SourceName);
+						if (!paramToolTips.ContainsKey(parameter.SourceName))
 						{
-							this.LogError($"Function '{OuterClass?.SourceName}::{this.SourceName}' doesn't provide a tooltip for parameter '{Parameter.SourceName}' (DocumentationPolicy).");
+							this.LogError($"Function '{outerClass?.SourceName}::{this.SourceName}' doesn't provide a tooltip for parameter '{parameter.SourceName}' (DocumentationPolicy).");
 						}
 					}
 
 					// ensure we don't have parameter tooltips for parameters that don't exist
-					foreach (KeyValuePair<StringView, StringView> Kvp in ParamToolTips)
+					foreach (KeyValuePair<StringView, StringView> kvp in paramToolTips)
 					{
-						if (Kvp.Key == UhtNames.ReturnValue)
+						if (kvp.Key == UhtNames.ReturnValue)
 						{
 							continue;
 						}
 
-						if (!ExistingFields.Contains(Kvp.Key))
+						if (!existingFields.Contains(kvp.Key))
 						{
-							this.LogError($"Function '{OuterClass?.SourceName}::{this.SourceName}' provides a tooltip for an unknown parameter '{Kvp.Key}'");
+							this.LogError($"Function '{outerClass?.SourceName}::{this.SourceName}' provides a tooltip for an unknown parameter '{kvp.Key}'");
 						}
 					}
 
 					// check for duplicate tooltips
-					Dictionary<StringView, StringView> ToolTipToParam = new Dictionary<StringView, StringView>(StringViewComparer.OrdinalIgnoreCase);
-					foreach (KeyValuePair<StringView, StringView> Kvp in ParamToolTips)
+					Dictionary<StringView, StringView> toolTipToParam = new Dictionary<StringView, StringView>(StringViewComparer.OrdinalIgnoreCase);
+					foreach (KeyValuePair<StringView, StringView> kvp in paramToolTips)
 					{
-						if (Kvp.Key == UhtNames.ReturnValue)
+						if (kvp.Key == UhtNames.ReturnValue)
 						{
 							continue;
 						}
 
-						if (ToolTipToParam.TryGetValue(Kvp.Value, out StringView Existing))
+						if (toolTipToParam.TryGetValue(kvp.Value, out StringView existing))
 						{
-							this.LogError($"Function '{OuterClass?.SourceName}::{this.SourceName}' uses identical tooltips for parameters '{Kvp.Key}' and '{Existing}' (DocumentationPolicy).");
+							this.LogError($"Function '{outerClass?.SourceName}::{this.SourceName}' uses identical tooltips for parameters '{kvp.Key}' and '{existing}' (DocumentationPolicy).");
 						}
 						else
 						{
-							ToolTipToParam.Add(Kvp.Value, Kvp.Key);
+							toolTipToParam.Add(kvp.Value, kvp.Key);
 						}
 					}
 				}
@@ -975,88 +969,88 @@ namespace EpicGames.UHT.Types
 		/// <summary>
 		/// Add a parameter tooltip to the collection
 		/// </summary>
-		/// <param name="Name">Name of the parameter</param>
-		/// <param name="Text">Text of the parameter tooltip</param>
-		/// <param name="ParamMap">Parameter map to add the tooltip</param>
-		private static void AddParameterToolTip(StringView Name, StringView Text, Dictionary<StringView, StringView> ParamMap)
+		/// <param name="name">Name of the parameter</param>
+		/// <param name="text">Text of the parameter tooltip</param>
+		/// <param name="paramMap">Parameter map to add the tooltip</param>
+		private static void AddParameterToolTip(StringView name, StringView text, Dictionary<StringView, StringView> paramMap)
 		{
-			ParamMap.Add(Name, new StringView(UhtFCString.TabsToSpaces(Text.Memory.Trim(), 4, false)));
+			paramMap.Add(name, new StringView(UhtFCString.TabsToSpaces(text.Memory.Trim(), 4, false)));
 		}
 
 		/// <summary>
 		/// Add a parameter tooltip to the collection where the name of the parameter is the first word of the line
 		/// </summary>
-		/// <param name="Text">Text of the parameter tooltip</param>
-		/// <param name="ParamMap">Parameter map to add the tooltip</param>
-		private static void AddParameterToolTip(StringView Text, Dictionary<StringView, StringView> ParamMap)
+		/// <param name="text">Text of the parameter tooltip</param>
+		/// <param name="paramMap">Parameter map to add the tooltip</param>
+		private static void AddParameterToolTip(StringView text, Dictionary<StringView, StringView> paramMap)
 		{
-			ReadOnlyMemory<char> Trimmed = Text.Memory.Trim();
-			int NameEnd = Trimmed.Span.IndexOf(' ');
-			if (NameEnd >= 0)
+			ReadOnlyMemory<char> trimmed = text.Memory.Trim();
+			int nameEnd = trimmed.Span.IndexOf(' ');
+			if (nameEnd >= 0)
 			{
-				AddParameterToolTip(new StringView(Trimmed, 0, NameEnd), new StringView(Trimmed, NameEnd + 1), ParamMap);
+				AddParameterToolTip(new StringView(trimmed, 0, nameEnd), new StringView(trimmed, nameEnd + 1), paramMap);
 			}
 		}
 
 		/// <summary>
 		/// Parse the function comment looking for parameter documentation.
 		/// </summary>
-		/// <param name="Input">The function input comment</param>
+		/// <param name="input">The function input comment</param>
 		/// <returns>Dictionary of parameter names and the documentation.  The return value will have a name of "ReturnValue"</returns>
-		private static Dictionary<StringView, StringView> GetParameterToolTipsFromFunctionComment(StringView Input)
+		private static Dictionary<StringView, StringView> GetParameterToolTipsFromFunctionComment(StringView input)
 		{
 			const string ParamTag = "@param";
 			const string ReturnTag = "@return";
 
-			Dictionary<StringView, StringView> ParamMap = new Dictionary<StringView, StringView>(StringViewComparer.OrdinalIgnoreCase);
+			Dictionary<StringView, StringView> paramMap = new Dictionary<StringView, StringView>(StringViewComparer.OrdinalIgnoreCase);
 
 			// Loop until we are out of text
-			while (Input.Length > 0)
+			while (input.Length > 0)
 			{
 
 				// Extract the line
-				StringView Line = Input;
-				int LineEnd = Input.Span.IndexOf('\n');
-				if (LineEnd >= 0)
+				StringView line = input;
+				int lineEnd = input.Span.IndexOf('\n');
+				if (lineEnd >= 0)
 				{
-					Line = new StringView(Input, 0, LineEnd);
-					Input = new StringView(Input, LineEnd + 1);
+					line = new StringView(input, 0, lineEnd);
+					input = new StringView(input, lineEnd + 1);
 				}
 
 				// If this is a parameter, invoke the method to extract the name first
-				int ParamStart = Line.Span.IndexOf(ParamTag, StringComparison.OrdinalIgnoreCase);
-				if (ParamStart >= 0)
+				int paramStart = line.Span.IndexOf(ParamTag, StringComparison.OrdinalIgnoreCase);
+				if (paramStart >= 0)
 				{
-					AddParameterToolTip(new StringView(Line, ParamStart + ParamTag.Length), ParamMap);
+					AddParameterToolTip(new StringView(line, paramStart + ParamTag.Length), paramMap);
 					continue;
 				}
 
 				// If this is the return value, we have a known name.  Invoke the method where we already know the name
-				ParamStart = Line.Span.IndexOf(ReturnTag, StringComparison.OrdinalIgnoreCase);
-				if (ParamStart >= 0)
+				paramStart = line.Span.IndexOf(ReturnTag, StringComparison.OrdinalIgnoreCase);
+				if (paramStart >= 0)
 				{
-					AddParameterToolTip(new StringView(UhtNames.ReturnValue), new StringView(Line, ParamStart + ReturnTag.Length), ParamMap);
+					AddParameterToolTip(new StringView(UhtNames.ReturnValue), new StringView(line, paramStart + ReturnTag.Length), paramMap);
 				}
 			}
 
-			return ParamMap;
+			return paramMap;
 		}
 		#endregion
 
 		/// <inheritdoc/>
-		public override void CollectReferences(IUhtReferenceCollector Collector)
+		public override void CollectReferences(IUhtReferenceCollector collector)
 		{
 			if (this.FunctionType != UhtFunctionType.Function)
 			{
-				Collector.AddExportType(this);
-				Collector.AddSingleton(this);
-				Collector.AddDeclaration(this, true);
-				Collector.AddCrossModuleReference(this, true);
-				Collector.AddCrossModuleReference(this.Package, true);
+				collector.AddExportType(this);
+				collector.AddSingleton(this);
+				collector.AddDeclaration(this, true);
+				collector.AddCrossModuleReference(this, true);
+				collector.AddCrossModuleReference(this.Package, true);
 			}
-			foreach (UhtType Child in this.Children)
+			foreach (UhtType child in this.Children)
 			{
-				Child.CollectReferences(Collector);
+				child.CollectReferences(collector);
 			}
 		}
 	}
diff --git a/Engine/Source/Programs/Shared/EpicGames.UHT/Types/UhtHeaderFile.cs b/Engine/Source/Programs/Shared/EpicGames.UHT/Types/UhtHeaderFile.cs
--- a/Engine/Source/Programs/Shared/EpicGames.UHT/Types/UhtHeaderFile.cs
+++ b/Engine/Source/Programs/Shared/EpicGames.UHT/Types/UhtHeaderFile.cs
@@ -1,12 +1,12 @@
 // Copyright Epic Games, Inc. All Rights Reserved.
 
-using EpicGames.Core;
-using EpicGames.UHT.Utils;
 using System;
 using System.Collections.Generic;
 using System.IO;
 using System.Text;
 using System.Text.Json.Serialization;
+using EpicGames.Core;
+using EpicGames.UHT.Utils;
 
 namespace EpicGames.UHT.Types
 {
@@ -42,7 +42,7 @@ namespace EpicGames.UHT.Types
 	/// Series of flags not part of the engine's class flags that affect code generation or verification
 	/// </summary>
 	[Flags]
-	public enum UhtHeaderFileExportFlags : Int32
+	public enum UhtHeaderFileExportFlags : int
 	{
 		/// <summary>
 		/// No export flags
@@ -52,7 +52,7 @@ namespace EpicGames.UHT.Types
 		/// <summary>
 		/// This header is being included by another header
 		/// </summary>
-		Referenced = 0x00000001, 
+		Referenced = 0x00000001,
 	}
 
 	/// <summary>
@@ -65,35 +65,35 @@ namespace EpicGames.UHT.Types
 		/// <summary>
 		/// Test to see if any of the specified flags are set
 		/// </summary>
-		/// <param name="InFlags">Current flags</param>
-		/// <param name="TestFlags">Flags to test for</param>
+		/// <param name="inFlags">Current flags</param>
+		/// <param name="testFlags">Flags to test for</param>
 		/// <returns>True if any of the flags are set</returns>
-		public static bool HasAnyFlags(this UhtHeaderFileExportFlags InFlags, UhtHeaderFileExportFlags TestFlags)
+		public static bool HasAnyFlags(this UhtHeaderFileExportFlags inFlags, UhtHeaderFileExportFlags testFlags)
 		{
-			return (InFlags & TestFlags) != 0;
+			return (inFlags & testFlags) != 0;
 		}
 
 		/// <summary>
 		/// Test to see if all of the specified flags are set
 		/// </summary>
-		/// <param name="InFlags">Current flags</param>
-		/// <param name="TestFlags">Flags to test for</param>
+		/// <param name="inFlags">Current flags</param>
+		/// <param name="testFlags">Flags to test for</param>
 		/// <returns>True if all the flags are set</returns>
-		public static bool HasAllFlags(this UhtHeaderFileExportFlags InFlags, UhtHeaderFileExportFlags TestFlags)
+		public static bool HasAllFlags(this UhtHeaderFileExportFlags inFlags, UhtHeaderFileExportFlags testFlags)
 		{
-			return (InFlags & TestFlags) == TestFlags;
+			return (inFlags & testFlags) == testFlags;
 		}
 
 		/// <summary>
 		/// Test to see if a specific set of flags have a specific value.
 		/// </summary>
-		/// <param name="InFlags">Current flags</param>
-		/// <param name="TestFlags">Flags to test for</param>
-		/// <param name="MatchFlags">Expected value of the tested flags</param>
+		/// <param name="inFlags">Current flags</param>
+		/// <param name="testFlags">Flags to test for</param>
+		/// <param name="matchFlags">Expected value of the tested flags</param>
 		/// <returns>True if the given flags have a specific value.</returns>
-		public static bool HasExactFlags(this UhtHeaderFileExportFlags InFlags, UhtHeaderFileExportFlags TestFlags, UhtHeaderFileExportFlags MatchFlags)
+		public static bool HasExactFlags(this UhtHeaderFileExportFlags inFlags, UhtHeaderFileExportFlags testFlags, UhtHeaderFileExportFlags matchFlags)
 		{
-			return (InFlags & TestFlags) == MatchFlags;
+			return (inFlags & testFlags) == matchFlags;
 		}
 	}
 
@@ -103,21 +103,21 @@ namespace EpicGames.UHT.Types
 	/// </summary>
 	public class UhtHeaderFile : UhtType
 	{
-		private readonly UhtSimpleMessageSite MessageSite;
-		private readonly UhtSourceFile SourceFile;
-		private readonly List<UhtHeaderFile> ReferencedHeaders = new List<UhtHeaderFile>();
+		private readonly UhtSimpleMessageSite _messageSite;
+		private readonly UhtSourceFile _sourceFile;
+		private readonly List<UhtHeaderFile> _referencedHeaders = new List<UhtHeaderFile>();
 
 		/// <summary>
 		/// Contents of the header
 		/// </summary>
 		[JsonIgnore]
-		public StringView Data { get => this.SourceFile.Data; }
+		public StringView Data => this._sourceFile.Data;
 
 		/// <summary>
 		/// Path of the header
 		/// </summary>
 		[JsonIgnore]
-		public string FilePath { get => this.SourceFile.FilePath; }
+		public string FilePath => this._sourceFile.FilePath;
 
 		/// <summary>
 		/// File name without the extension
@@ -132,17 +132,17 @@ namespace EpicGames.UHT.Types
 		/// <summary>
 		/// True if this header is NoExportTypes.h
 		/// </summary>
-		public bool bIsNoExportTypes { get; }
+		public bool IsNoExportTypes { get; }
 
 		/// <summary>
 		/// The file path of the header relative to the module location
 		/// </summary>
-		public string ModuleRelativeFilePath { get; set; } = string.Empty;
+		public string ModuleRelativeFilePath { get; set; } = String.Empty;
 
 		/// <summary>
 		/// Include file path added as meta data to the types
 		/// </summary>
-		public string IncludeFilePath { get; set; } = string.Empty;
+		public string IncludeFilePath { get; set; } = String.Empty;
 
 		/// <summary>
 		/// Location where the header file was found
@@ -164,7 +164,7 @@ namespace EpicGames.UHT.Types
 		/// If true, the header file should be exported
 		/// </summary>
 		[JsonIgnore]
-		public bool bShouldExport => this.HeaderFileExportFlags.HasAnyFlags(UhtHeaderFileExportFlags.Referenced) || this.Children.Count > 0;
+		public bool ShouldExport => this.HeaderFileExportFlags.HasAnyFlags(UhtHeaderFileExportFlags.Referenced) || this.Children.Count > 0;
 
 		/// <summary>
 		/// Resource collector for the header file
@@ -188,14 +188,14 @@ namespace EpicGames.UHT.Types
 
 		/// <inheritdoc/>
 		[JsonIgnore]
-		public override UhtHeaderFile HeaderFile { get => this; }
+		public override UhtHeaderFile HeaderFile => this;
 
 		/// <inheritdoc/>
 		[JsonIgnore]
 		public override UhtEngineType EngineType => UhtEngineType.Header;
 
 		/// <inheritdoc/>
-		public override string EngineClassName { get => "UhtHeaderFile"; }
+		public override string EngineClassName => "UhtHeaderFile";
 
 		/// <summary>
 		/// Collection of headers directly included by this header
@@ -207,28 +207,28 @@ namespace EpicGames.UHT.Types
 
 		/// <inheritdoc/>
 		[JsonIgnore]
-		public override IUhtMessageSession MessageSession => this.MessageSite.MessageSession;
+		public override IUhtMessageSession MessageSession => this._messageSite.MessageSession;
 
 		/// <inheritdoc/>
 		[JsonIgnore]
-		public override IUhtMessageSource? MessageSource => this.MessageSite.MessageSource;
+		public override IUhtMessageSource? MessageSource => this._messageSite.MessageSource;
 		#endregion
 
 		/// <summary>
 		/// Construct a new header file
 		/// </summary>
-		/// <param name="Package">Owning package</param>
-		/// <param name="Path">Path to the header file</param>
-		public UhtHeaderFile(UhtPackage Package, string Path) : base(Package, 1)
+		/// <param name="package">Owning package</param>
+		/// <param name="path">Path to the header file</param>
+		public UhtHeaderFile(UhtPackage package, string path) : base(package, 1)
 		{
 			this.HeaderFileTypeIndex = this.Session.GetNextHeaderFileTypeIndex();
-			this.MessageSite = new UhtSimpleMessageSite(this.Session);
-			this.SourceFile = new UhtSourceFile(this.Session, Path);
-			this.MessageSite.MessageSource = this.SourceFile;
-			this.FileNameWithoutExtension = System.IO.Path.GetFileNameWithoutExtension(this.SourceFile.FilePath);
+			this._messageSite = new UhtSimpleMessageSite(this.Session);
+			this._sourceFile = new UhtSourceFile(this.Session, path);
+			this._messageSite.MessageSource = this._sourceFile;
+			this.FileNameWithoutExtension = System.IO.Path.GetFileNameWithoutExtension(this._sourceFile.FilePath);
 			this.GeneratedHeaderFileName = this.FileNameWithoutExtension + ".generated.h";
-			this.SourceName = System.IO.Path.GetFileName(this.SourceFile.FilePath);
-			this.bIsNoExportTypes = string.Equals(this.SourceFile.FileName, "NoExportTypes", StringComparison.OrdinalIgnoreCase);
+			this.SourceName = System.IO.Path.GetFileName(this._sourceFile.FilePath);
+			this.IsNoExportTypes = String.Equals(this._sourceFile.FileName, "NoExportTypes", StringComparison.OrdinalIgnoreCase);
 		}
 
 		/// <summary>
@@ -236,45 +236,45 @@ namespace EpicGames.UHT.Types
 		/// </summary>
 		public void Read()
 		{
-			this.SourceFile.Read();
+			this._sourceFile.Read();
 		}
 
 		/// <summary>
 		/// Add a reference to the given header
 		/// </summary>
-		/// <param name="Id">Path of the header</param>
-		/// <param name="bIsIncludedFile">True if this is a directly included file</param>
-		public void AddReferencedHeader(string Id, bool bIsIncludedFile)
+		/// <param name="id">Path of the header</param>
+		/// <param name="isIncludedFile">True if this is a directly included file</param>
+		public void AddReferencedHeader(string id, bool isIncludedFile)
 		{
-			UhtHeaderFile? HeaderFile = this.Session.FindHeaderFile(Path.GetFileName(Id));
-			if (HeaderFile != null)
+			UhtHeaderFile? headerFile = this.Session.FindHeaderFile(Path.GetFileName(id));
+			if (headerFile != null)
 			{
-				AddReferencedHeader(HeaderFile, bIsIncludedFile);
+				AddReferencedHeader(headerFile, isIncludedFile);
 			}
 		}
 
 		/// <summary>
 		/// Add a reference to the header that defines the given type
 		/// </summary>
-		/// <param name="Type">Type in question</param>
-		public void AddReferencedHeader(UhtType Type)
+		/// <param name="type">Type in question</param>
+		public void AddReferencedHeader(UhtType type)
 		{
-			AddReferencedHeader(Type.HeaderFile, false);
+			AddReferencedHeader(type.HeaderFile, false);
 		}
 
 		/// <summary>
 		/// Add a reference to the given header file
 		/// </summary>
-		/// <param name="HeaderFile">Header file in question</param>
-		/// <param name="bIsIncludedFile">True if this is a directly included file</param>
-		public void AddReferencedHeader(UhtHeaderFile HeaderFile, bool bIsIncludedFile)
+		/// <param name="headerFile">Header file in question</param>
+		/// <param name="isIncludedFile">True if this is a directly included file</param>
+		public void AddReferencedHeader(UhtHeaderFile headerFile, bool isIncludedFile)
 		{
-			lock (this.ReferencedHeaders)
+			lock (this._referencedHeaders)
 			{
 				// Check for a duplicate
-				foreach (UhtHeaderFile Ref in this.ReferencedHeaders)
+				foreach (UhtHeaderFile reference in this._referencedHeaders)
 				{
-					if (Ref == HeaderFile)
+					if (reference == headerFile)
 					{
 						return;
 					}
@@ -282,14 +282,14 @@ namespace EpicGames.UHT.Types
 
 				// There is questionable compatibility hack where a source file will always be exported
 				// regardless of having types when it is being included by the SAME package.
-				if (HeaderFile.Package == this.Package)
+				if (headerFile.Package == this.Package)
 				{
-					HeaderFile.HeaderFileExportFlags |= UhtHeaderFileExportFlags.Referenced;
+					headerFile.HeaderFileExportFlags |= UhtHeaderFileExportFlags.Referenced;
 				}
-				this.ReferencedHeaders.Add(HeaderFile);
-				if (bIsIncludedFile)
+				this._referencedHeaders.Add(headerFile);
+				if (isIncludedFile)
 				{
-					this.IncludedHeaders.Add(HeaderFile);
+					this.IncludedHeaders.Add(headerFile);
 				}
 			}
 		}
@@ -298,13 +298,7 @@ namespace EpicGames.UHT.Types
 		/// Return an enumerator without locking.  This method can only be utilized AFTER all header parsing is complete. 
 		/// </summary>
 		[JsonIgnore]
-		public IEnumerable<UhtHeaderFile> ReferencedHeadersNoLock
-		{
-			get
-			{
-				return this.ReferencedHeaders;
-			}
-		}
+		public IEnumerable<UhtHeaderFile> ReferencedHeadersNoLock => this._referencedHeaders;
 
 		/// <summary>
 		/// Return an enumerator of all current referenced headers under a lock.  This should be used during parsing.
@@ -314,71 +308,71 @@ namespace EpicGames.UHT.Types
 		{
 			get
 			{
-				lock (this.ReferencedHeaders)
+				lock (this._referencedHeaders)
 				{
-					foreach (UhtHeaderFile Ref in this.ReferencedHeaders)
+					foreach (UhtHeaderFile reference in this._referencedHeaders)
 					{
-						yield return Ref;
+						yield return reference;
 					}
 				}
 			}
 		}
 
 		/// <inheritdoc/>
-		public override void AppendPathName(StringBuilder Builder, UhtType? StopOuter = null)
+		public override void AppendPathName(StringBuilder builder, UhtType? stopOuter = null)
 		{
 			// Headers do not contribute to path names
-			if (this != StopOuter && this.Outer != null)
+			if (this != stopOuter && this.Outer != null)
 			{
-				this.Outer.AppendPathName(Builder, StopOuter);
+				this.Outer.AppendPathName(builder, stopOuter);
 			}
 		}
 
 		/// <inheritdoc/>
-		protected override UhtValidationOptions Validate(UhtValidationOptions Options)
+		protected override UhtValidationOptions Validate(UhtValidationOptions options)
 		{
-			Options = base.Validate(Options | UhtValidationOptions.Shadowing);
+			options = base.Validate(options | UhtValidationOptions.Shadowing);
 
-			Dictionary<int, UhtFunction> UsedRPCIds = new Dictionary<int, UhtFunction>();
-			Dictionary<int, UhtFunction> RPCsNeedingHookup = new Dictionary<int, UhtFunction>();
-			foreach (UhtType Type in this.Children)
+			Dictionary<int, UhtFunction> usedRPCIds = new Dictionary<int, UhtFunction>();
+			Dictionary<int, UhtFunction> rpcsNeedingHookup = new Dictionary<int, UhtFunction>();
+			foreach (UhtType type in this.Children)
 			{
-				if (Type is UhtClass Class)
+				if (type is UhtClass classObj)
 				{
-					foreach (UhtType Child in Class.Children)
+					foreach (UhtType child in classObj.Children)
 					{
-						if (Child is UhtFunction Function)
+						if (child is UhtFunction function)
 						{
-							if (Function.FunctionType != UhtFunctionType.Function || !Function.FunctionFlags.HasAnyFlags(EFunctionFlags.Net))
+							if (function.FunctionType != UhtFunctionType.Function || !function.FunctionFlags.HasAnyFlags(EFunctionFlags.Net))
 							{
 								continue;
 							}
 
-							if (Function.RPCId > 0)
+							if (function.RPCId > 0)
 							{
-								UhtFunction? ExistingFunc;
-								if (UsedRPCIds.TryGetValue(Function.RPCId, out ExistingFunc))
+								UhtFunction? existingFunc;
+								if (usedRPCIds.TryGetValue(function.RPCId, out existingFunc))
 								{
-									Function.LogError($"Function {ExistingFunc.SourceName} already uses identifier {Function.RPCId}");
+									function.LogError($"Function {existingFunc.SourceName} already uses identifier {function.RPCId}");
 								}
 								else
 								{
-									UsedRPCIds.Add(Function.RPCId, Function);
+									usedRPCIds.Add(function.RPCId, function);
 								}
 
-								if (Function.FunctionFlags.HasAnyFlags(EFunctionFlags.NetResponse))
+								if (function.FunctionFlags.HasAnyFlags(EFunctionFlags.NetResponse))
 								{
 									// Look for another function expecting this response
-									RPCsNeedingHookup.Remove(Function.RPCId);
+									rpcsNeedingHookup.Remove(function.RPCId);
 								}
 							}
 
-							if (Function.RPCResponseId > 0 && Function.EndpointName != "JSBridge")
+							if (function.RPCResponseId > 0 && function.EndpointName != "JSBridge")
 							{
 								// Look for an existing response function, if not found then add to the list of ids awaiting hookup
-								if (!UsedRPCIds.ContainsKey(Function.RPCResponseId))
+								if (!usedRPCIds.ContainsKey(function.RPCResponseId))
 								{
-									RPCsNeedingHookup.Add(Function.RPCResponseId, Function);
+									rpcsNeedingHookup.Add(function.RPCResponseId, function);
 								}
 							}
 						}
@@ -386,14 +380,14 @@ namespace EpicGames.UHT.Types
 				}
 			}
 
-			if (RPCsNeedingHookup.Count > 0)
+			if (rpcsNeedingHookup.Count > 0)
 			{
-				foreach (KeyValuePair<int, UhtFunction> KVP in RPCsNeedingHookup)
+				foreach (KeyValuePair<int, UhtFunction> kvp in rpcsNeedingHookup)
 				{
-					KVP.Value.LogError($"Request function '{KVP.Value.SourceName}' is missing a response function with the id of '{KVP.Key}'");
+					kvp.Value.LogError($"Request function '{kvp.Value.SourceName}' is missing a response function with the id of '{kvp.Key}'");
 				}
 			}
-			return Options;
+			return options;
 		}
 	}
 }
diff --git a/Engine/Source/Programs/Shared/EpicGames.UHT/Types/UhtManifestFile.cs b/Engine/Source/Programs/Shared/EpicGames.UHT/Types/UhtManifestFile.cs
--- a/Engine/Source/Programs/Shared/EpicGames.UHT/Types/UhtManifestFile.cs
+++ b/Engine/Source/Programs/Shared/EpicGames.UHT/Types/UhtManifestFile.cs
@@ -1,8 +1,8 @@
 // Copyright Epic Games, Inc. All Rights Reserved.
 
+using System.Text.Json;
 using EpicGames.Core;
 using EpicGames.UHT.Utils;
-using System.Text.Json;
 
 namespace EpicGames.UHT.Types
 {
@@ -18,15 +18,15 @@ namespace EpicGames.UHT.Types
 		/// </summary>
 		public UHTManifest? Manifest { get; set; } = null;
 
-		private readonly UhtSourceFile SourceFile;
+		private readonly UhtSourceFile _sourceFile;
 
 		#region IUHTMessageSite implementation
 
 		/// <inheritdoc/>
-		public IUhtMessageSession MessageSession => this.SourceFile.MessageSession;
+		public IUhtMessageSession MessageSession => this._sourceFile.MessageSession;
 
 		/// <inheritdoc/>
-		public IUhtMessageSource? MessageSource => this.SourceFile.MessageSource;
+		public IUhtMessageSource? MessageSource => this._sourceFile.MessageSource;
 
 		/// <inheritdoc/>
 		public IUhtMessageLineNumber? MessageLineNumber => null;
@@ -35,20 +35,20 @@ namespace EpicGames.UHT.Types
 		/// <summary>
 		/// Construct a new manifest file
 		/// </summary>
-		/// <param name="Session">Current session</param>
-		/// <param name="FilePath">Path of the file</param>
-		public UhtManifestFile(UhtSession Session, string FilePath)
+		/// <param name="session">Current session</param>
+		/// <param name="filePath">Path of the file</param>
+		public UhtManifestFile(UhtSession session, string filePath)
 		{
-			this.SourceFile = new UhtSourceFile(Session, FilePath);
+			this._sourceFile = new UhtSourceFile(session, filePath);
 		}
 
 		/// <summary>
 		/// Read the contents of the file
 		/// </summary>
-		public void Read() 
+		public void Read()
 		{
-			this.SourceFile.Read();
-			this.Manifest = JsonSerializer.Deserialize<UHTManifest>(this.SourceFile.Data.ToString());
+			this._sourceFile.Read();
+			this.Manifest = JsonSerializer.Deserialize<UHTManifest>(this._sourceFile.Data.ToString());
 		}
 	}
 }
diff --git a/Engine/Source/Programs/Shared/EpicGames.UHT/Types/UhtMetaData.cs b/Engine/Source/Programs/Shared/EpicGames.UHT/Types/UhtMetaData.cs
--- a/Engine/Source/Programs/Shared/EpicGames.UHT/Types/UhtMetaData.cs
+++ b/Engine/Source/Programs/Shared/EpicGames.UHT/Types/UhtMetaData.cs
@@ -1,11 +1,11 @@
 // Copyright Epic Games, Inc. All Rights Reserved.
 
-using EpicGames.Core;
-using EpicGames.UHT.Utils;
 using System;
 using System.Collections.Generic;
 using System.Diagnostics.CodeAnalysis;
 using System.Text.Json.Serialization;
+using EpicGames.Core;
+using EpicGames.UHT.Utils;
 
 namespace EpicGames.UHT.Types
 {
@@ -19,10 +19,10 @@ namespace EpicGames.UHT.Types
 		/// <summary>
 		/// Convert a name and index into a full meta data name
 		/// </summary>
-		/// <param name="Name">Base name of the meta data</param>
-		/// <param name="NameIndex">Index of the meta data.  -1 for the root object.</param>
+		/// <param name="name">Base name of the meta data</param>
+		/// <param name="nameIndex">Index of the meta data.  -1 for the root object.</param>
 		/// <returns>Complete meta data key name</returns>
-		string GetMetaDataKey(string Name, int NameIndex);
+		string GetMetaDataKey(string name, int nameIndex);
 	}
 
 	/// <summary>
@@ -44,12 +44,12 @@ namespace EpicGames.UHT.Types
 		/// <summary>
 		/// Construct a new meta data key
 		/// </summary>
-		/// <param name="Name">Meta data name</param>
-		/// <param name="Index">Meta data index</param>
-		public UhtMetaDataKey(string Name, int Index = UhtMetaData.IndexNone)
+		/// <param name="name">Meta data name</param>
+		/// <param name="index">Meta data index</param>
+		public UhtMetaDataKey(string name, int index = UhtMetaData.IndexNone)
 		{
-			this.Name = Name;
-			this.Index = Index;
+			this.Name = name;
+			this.Index = index;
 		}
 
 		/// <summary>
@@ -90,7 +90,7 @@ namespace EpicGames.UHT.Types
 			{
 				return 1;
 			}
-			return string.Compare(x.Name, y.Name, StringComparison.OrdinalIgnoreCase);
+			return String.Compare(x.Name, y.Name, StringComparison.OrdinalIgnoreCase);
 		}
 
 		/// <summary>
@@ -125,7 +125,7 @@ namespace EpicGames.UHT.Types
 		/// <summary>
 		/// Helper comparer for meta data keys
 		/// </summary>
-		private static readonly UhtMetaDataKeyComparer Comparer = new UhtMetaDataKeyComparer();
+		private static readonly UhtMetaDataKeyComparer s_comparer = new UhtMetaDataKeyComparer();
 
 		/// <summary>
 		/// Empty collection of meta data
@@ -176,12 +176,12 @@ namespace EpicGames.UHT.Types
 		/// <summary>
 		/// Construct new meta data
 		/// </summary>
-		/// <param name="MessageSite">Message site for generating errors</param>
-		/// <param name="Config">Configuration for redirects</param>
-		public UhtMetaData(IUhtMessageSite? MessageSite, IUhtConfig? Config)
+		/// <param name="messageSite">Message site for generating errors</param>
+		/// <param name="config">Configuration for redirects</param>
+		public UhtMetaData(IUhtMessageSite? messageSite, IUhtConfig? config)
 		{
-			this.MessageSite = MessageSite;
-			this.Config = Config;
+			this.MessageSite = messageSite;
+			this.Config = config;
 		}
 
 		/// <summary>
@@ -216,29 +216,29 @@ namespace EpicGames.UHT.Types
 		/// <summary>
 		/// Test to see if the meta data contains the given key
 		/// </summary>
-		/// <param name="Name">Name of the meta data</param>
-		/// <param name="NameIndex">Enumeration value index or -1 for the type's meta data</param>
+		/// <param name="name">Name of the meta data</param>
+		/// <param name="nameIndex">Enumeration value index or -1 for the type's meta data</param>
 		/// <returns>True if the key is found</returns>
-		public bool ContainsKey(string Name, int NameIndex = IndexNone)
+		public bool ContainsKey(string name, int nameIndex = IndexNone)
 		{
 			if (this.Dictionary == null)
 			{
 				return false;
 			}
-			return this.Dictionary.ContainsKey(new UhtMetaDataKey(Name, NameIndex));
+			return this.Dictionary.ContainsKey(new UhtMetaDataKey(name, nameIndex));
 		}
 
 		/// <summary>
 		/// Test to see if the meta data or parent meta data contains the given key
 		/// </summary>
-		/// <param name="Name">Name of the meta data</param>
-		/// <param name="NameIndex">Enumeration value index or -1 for the type's meta data</param>
+		/// <param name="name">Name of the meta data</param>
+		/// <param name="nameIndex">Enumeration value index or -1 for the type's meta data</param>
 		/// <returns>True if the key is found</returns>
-		public bool ContainsKeyHierarchical(string Name, int NameIndex = IndexNone)
+		public bool ContainsKeyHierarchical(string name, int nameIndex = IndexNone)
 		{
-			for (UhtMetaData? Current = this; Current != null; Current = Current.Parent)
+			for (UhtMetaData? current = this; current != null; current = current.Parent)
 			{
-				if (Current.ContainsKey(Name, NameIndex))
+				if (current.ContainsKey(name, nameIndex))
 				{
 					return true;
 				}
@@ -249,74 +249,74 @@ namespace EpicGames.UHT.Types
 		/// <summary>
 		/// Attempt to get the value associated with the key
 		/// </summary>
-		/// <param name="Name">Name of the meta data</param>
-		/// <param name="Value">Found value</param>
+		/// <param name="name">Name of the meta data</param>
+		/// <param name="value">Found value</param>
 		/// <returns>True if the key is found</returns>
-		public bool TryGetValue(string Name, [NotNullWhen(true)] out string? Value)
+		public bool TryGetValue(string name, [NotNullWhen(true)] out string? value)
 		{
-			return TryGetValue(Name, IndexNone, out Value);
+			return TryGetValue(name, IndexNone, out value);
 		}
 
 		/// <summary>
 		/// Attempt to get the value associated with the key including parent meta data
 		/// </summary>
-		/// <param name="Name">Name of the meta data</param>
-		/// <param name="Value">Found value</param>
+		/// <param name="name">Name of the meta data</param>
+		/// <param name="value">Found value</param>
 		/// <returns>True if the key is found</returns>
-		public bool TryGetValueHierarchical(string Name, [NotNullWhen(true)] out string? Value)
+		public bool TryGetValueHierarchical(string name, [NotNullWhen(true)] out string? value)
 		{
-			return TryGetValueHierarchical(Name, IndexNone, out Value);
+			return TryGetValueHierarchical(name, IndexNone, out value);
 		}
 
 		/// <summary>
 		/// Attempt to get the value associated with the key
 		/// </summary>
-		/// <param name="Name">Name of the meta data</param>
-		/// <param name="NameIndex">Enumeration value index or -1 for the type's meta data</param>
-		/// <param name="Value">Found value</param>
+		/// <param name="name">Name of the meta data</param>
+		/// <param name="nameIndex">Enumeration value index or -1 for the type's meta data</param>
+		/// <param name="value">Found value</param>
 		/// <returns>True if the key is found</returns>
-		public bool TryGetValue(string Name, int NameIndex, [NotNullWhen(true)] out string? Value)
+		public bool TryGetValue(string name, int nameIndex, [NotNullWhen(true)] out string? value)
 		{
 			if (this.Dictionary == null)
 			{
-				Value = String.Empty;
+				value = String.Empty;
 				return false;
 			}
-			return this.Dictionary.TryGetValue(new UhtMetaDataKey(Name, NameIndex), out Value);
+			return this.Dictionary.TryGetValue(new UhtMetaDataKey(name, nameIndex), out value);
 		}
 
 		/// <summary>
 		/// Attempt to get the value associated with the key including parent meta data
 		/// </summary>
-		/// <param name="Name">Name of the meta data</param>
-		/// <param name="NameIndex">Enumeration value index or -1 for the type's meta data</param>
-		/// <param name="Value">Found value</param>
+		/// <param name="name">Name of the meta data</param>
+		/// <param name="nameIndex">Enumeration value index or -1 for the type's meta data</param>
+		/// <param name="value">Found value</param>
 		/// <returns>True if the key is found</returns>
-		public bool TryGetValueHierarchical(string Name, int NameIndex, [NotNullWhen(true)] out string? Value)
+		public bool TryGetValueHierarchical(string name, int nameIndex, [NotNullWhen(true)] out string? value)
 		{
-			for (UhtMetaData? Current = this; Current != null; Current = Current.Parent)
+			for (UhtMetaData? current = this; current != null; current = current.Parent)
 			{
-				if (Current.TryGetValue(Name, NameIndex, out Value))
+				if (current.TryGetValue(name, nameIndex, out value))
 				{
 					return true;
 				}
 			}
-			Value = String.Empty;
+			value = String.Empty;
 			return false;
 		}
 
 		/// <summary>
 		/// Get the string value of the given meta data key or the default value.
 		/// </summary>
-		/// <param name="Name">Name of the meta data key</param>
-		/// <param name="NameIndex">Index of the meta data key</param>
+		/// <param name="name">Name of the meta data key</param>
+		/// <param name="nameIndex">Index of the meta data key</param>
 		/// <returns>Meta data value or empty string if not found.</returns>
-		public string GetValueOrDefault(string Name, int NameIndex = IndexNone)
+		public string GetValueOrDefault(string name, int nameIndex = IndexNone)
 		{
-			string? Out;
-			if (TryGetValue(Name, NameIndex, out Out))
+			string? output;
+			if (TryGetValue(name, nameIndex, out output))
 			{
-				return Out;
+				return output;
 			}
 			return String.Empty;
 		}
@@ -324,70 +324,70 @@ namespace EpicGames.UHT.Types
 		/// <summary>
 		/// Get the string value of the given meta data searching the whole meta data chain.
 		/// </summary>
-		/// <param name="Name">Name of the meta data key</param>
-		/// <param name="NameIndex">Index of the meta data key</param>
+		/// <param name="name">Name of the meta data key</param>
+		/// <param name="nameIndex">Index of the meta data key</param>
 		/// <returns>Meta data value or empty string if not found.</returns>
-		public string GetValueOrDefaultHierarchical(string Name, int NameIndex = IndexNone)
+		public string GetValueOrDefaultHierarchical(string name, int nameIndex = IndexNone)
 		{
-			string? Out;
-			if (TryGetValueHierarchical(Name, NameIndex, out Out))
+			string? output;
+			if (TryGetValueHierarchical(name, nameIndex, out output))
 			{
-				return Out;
+				return output;
 			}
-			return string.Empty;
+			return String.Empty;
 		}
 
 		/// <summary>
 		/// Get the boolean value of the given meta data.
 		/// </summary>
-		/// <param name="Name">Name of the meta data key</param>
-		/// <param name="NameIndex">Index of the meta data key</param>
+		/// <param name="name">Name of the meta data key</param>
+		/// <param name="nameIndex">Index of the meta data key</param>
 		/// <returns>Boolean value or false if not found</returns>
-		public bool GetBoolean(string Name, int NameIndex = IndexNone)
+		public bool GetBoolean(string name, int nameIndex = IndexNone)
 		{
-			return UhtFCString.ToBool(GetValueOrDefault(Name, NameIndex));
+			return UhtFCString.ToBool(GetValueOrDefault(name, nameIndex));
 		}
 
 		/// <summary>
 		/// Get the boolean value of the given meta data searching the whole meta data chain.
 		/// </summary>
-		/// <param name="Name">Name of the meta data key</param>
-		/// <param name="NameIndex">Index of the meta data key</param>
+		/// <param name="name">Name of the meta data key</param>
+		/// <param name="nameIndex">Index of the meta data key</param>
 		/// <returns>Boolean value or false if not found</returns>
-		public bool GetBooleanHierarchical(string Name, int NameIndex = IndexNone)
+		public bool GetBooleanHierarchical(string name, int nameIndex = IndexNone)
 		{
-			return UhtFCString.ToBool(GetValueOrDefaultHierarchical(Name, NameIndex));
+			return UhtFCString.ToBool(GetValueOrDefaultHierarchical(name, nameIndex));
 		}
 
 		/// <summary>
 		/// Get the double value of the given meta data.
 		/// </summary>
-		/// <param name="Name">Name of the meta data key</param>
-		/// <param name="NameIndex">Index of the meta data key</param>
+		/// <param name="name">Name of the meta data key</param>
+		/// <param name="nameIndex">Index of the meta data key</param>
 		/// <returns>Double value or zero if not found</returns>
-		public double GetDouble(string Name, int NameIndex = IndexNone)
+		public double GetDouble(string name, int nameIndex = IndexNone)
 		{
-			string Value = GetValueOrDefault(Name, NameIndex);
-			double Result;
-			if (!double.TryParse(Value, out Result))
+			string value = GetValueOrDefault(name, nameIndex);
+			double result;
+			if (!Double.TryParse(value, out result))
 			{
-				Result = 0;
+				result = 0;
 			}
-			return Result;
+			return result;
 		}
 
 		/// <summary>
 		/// Get the string array value of the given meta data.
 		/// </summary>
-		/// <param name="Name">Name of the meta data key</param>
-		/// <param name="NameIndex">Index of the meta data key</param>
+		/// <param name="name">Name of the meta data key</param>
+		/// <param name="nameIndex">Index of the meta data key</param>
 		/// <returns>String array or null if not found</returns>
-		public string[]? GetStringArray(string Name, int NameIndex = IndexNone)
+		public string[]? GetStringArray(string name, int nameIndex = IndexNone)
 		{
-			string? Temp;
-			if (TryGetValue(Name, NameIndex, out Temp))
+			string? temp;
+			if (TryGetValue(name, nameIndex, out temp))
 			{
-				return Temp.ToString().Split(' ', StringSplitOptions.RemoveEmptyEntries);
+				return temp.ToString().Split(' ', StringSplitOptions.RemoveEmptyEntries);
 			}
 			return null;
 		}
@@ -395,15 +395,15 @@ namespace EpicGames.UHT.Types
 		/// <summary>
 		/// Get the string array value of the given meta data searching the whole meta data chain.
 		/// </summary>
-		/// <param name="Name">Name of the meta data key</param>
-		/// <param name="NameIndex">Index of the meta data key</param>
+		/// <param name="name">Name of the meta data key</param>
+		/// <param name="nameIndex">Index of the meta data key</param>
 		/// <returns>String array or null if not found</returns>
-		public string[]? GetStringArrayHierarchical(string Name, int NameIndex = IndexNone)
+		public string[]? GetStringArrayHierarchical(string name, int nameIndex = IndexNone)
 		{
-			string? Temp;
-			if (TryGetValueHierarchical(Name, NameIndex, out Temp))
+			string? temp;
+			if (TryGetValueHierarchical(name, nameIndex, out temp))
 			{
-				return Temp.ToString().Split(' ', StringSplitOptions.RemoveEmptyEntries);
+				return temp.ToString().Split(' ', StringSplitOptions.RemoveEmptyEntries);
 			}
 			return null;
 		}
@@ -411,83 +411,83 @@ namespace EpicGames.UHT.Types
 		/// <summary>
 		/// Add new meta data
 		/// </summary>
-		/// <param name="Name">Name of the meta data key</param>
-		/// <param name="Value">Value of the meta data</param>
-		public void Add(string Name, string Value)
+		/// <param name="name">Name of the meta data key</param>
+		/// <param name="value">Value of the meta data</param>
+		public void Add(string name, string value)
 		{
-			Add(Name, IndexNone, Value);
+			Add(name, IndexNone, value);
 		}
 
 		/// <summary>
 		/// Add new meta data
 		/// </summary>
-		/// <param name="Name">Name of the meta data key</param>
-		/// <param name="NameIndex">Index of the meta data key</param>
-		/// <param name="Value">Value of the meta data</param>
-		public void Add(string Name, int NameIndex, string Value)
+		/// <param name="name">Name of the meta data key</param>
+		/// <param name="nameIndex">Index of the meta data key</param>
+		/// <param name="value">Value of the meta data</param>
+		public void Add(string name, int nameIndex, string value)
 		{
-			string RemappedName;
+			string remappedName;
 			if (this.Config != null)
 			{
-				if (this.Config.RedirectMetaDataKey(Name, out RemappedName))
+				if (this.Config.RedirectMetaDataKey(name, out remappedName))
 				{
 					if (this.MessageSite != null)
 					{
-						this.MessageSite.LogWarning(this.LineNumber, $"Remapping old metadata key '{Name}' to new key '{RemappedName}', please update the declaration.");
+						this.MessageSite.LogWarning(this.LineNumber, $"Remapping old metadata key '{name}' to new key '{remappedName}', please update the declaration.");
 					}
 				}
 			}
-			GetDictionary()[new UhtMetaDataKey(Name, NameIndex)] = Value;
+			GetDictionary()[new UhtMetaDataKey(name, nameIndex)] = value;
 		}
 
 		/// <summary>
 		/// Add new meta data
 		/// </summary>
-		/// <param name="Name">Name of the meta data key</param>
-		/// <param name="Value">Value of the meta data</param>
-		/// <param name="NameIndex">Index of the meta data key</param>
-		public void Add(string Name, bool Value, int NameIndex = IndexNone)
+		/// <param name="name">Name of the meta data key</param>
+		/// <param name="value">Value of the meta data</param>
+		/// <param name="nameIndex">Index of the meta data key</param>
+		public void Add(string name, bool value, int nameIndex = IndexNone)
 		{
-			Add(Name, NameIndex, Value ? "true" : "false");
+			Add(name, nameIndex, value ? "true" : "false");
 		}
 
 		/// <summary>
 		/// Add new meta data
 		/// </summary>
-		/// <param name="Name">Name of the meta data key</param>
-		/// <param name="Strings">Value of the meta data</param>
-		/// <param name="Separator">Separator to use to join the strings</param>
-		public void Add(string Name, List<string> Strings, char Separator = ' ')
+		/// <param name="name">Name of the meta data key</param>
+		/// <param name="strings">Value of the meta data</param>
+		/// <param name="separator">Separator to use to join the strings</param>
+		public void Add(string name, List<string> strings, char separator = ' ')
 		{
-			Add(Name, string.Join(Separator, Strings));
+			Add(name, String.Join(separator, strings));
 		}
 
 		/// <summary>
 		/// Add new meta data if there are strings in the value.
 		/// </summary>
-		/// <param name="Name">Name of the meta data key</param>
-		/// <param name="Strings">Value of the meta data</param>
-		/// <param name="Separator">Separator to use to join the strings</param>
-		public void AddIfNotEmpty(string Name, List<string> Strings, char Separator = ' ')
+		/// <param name="name">Name of the meta data key</param>
+		/// <param name="strings">Value of the meta data</param>
+		/// <param name="separator">Separator to use to join the strings</param>
+		public void AddIfNotEmpty(string name, List<string> strings, char separator = ' ')
 		{
-			if (Strings.Count != 0)
+			if (strings.Count != 0)
 			{
-				Add(Name, string.Join(Separator, Strings));
+				Add(name, String.Join(separator, strings));
 			}
 		}
 
 		/// <summary>
 		/// Add the meta data from another meta data block
 		/// </summary>
-		/// <param name="MetaData"></param>
-		public void Add(UhtMetaData MetaData)
+		/// <param name="metaData"></param>
+		public void Add(UhtMetaData metaData)
 		{
-			if (MetaData.Dictionary != null)
+			if (metaData.Dictionary != null)
 			{
-				SortedList<UhtMetaDataKey, string> Dictionary = GetDictionary();
-				foreach (KeyValuePair<UhtMetaDataKey, string> Kvp in MetaData.Dictionary)
+				SortedList<UhtMetaDataKey, string> dictionary = GetDictionary();
+				foreach (KeyValuePair<UhtMetaDataKey, string> kvp in metaData.Dictionary)
 				{
-					Dictionary[Kvp.Key] = Kvp.Value;
+					dictionary[kvp.Key] = kvp.Value;
 				}
 			}
 		}
@@ -495,42 +495,42 @@ namespace EpicGames.UHT.Types
 		/// <summary>
 		/// Remove the given meta data
 		/// </summary>
-		/// <param name="Name">Name of the meta data key</param>
-		/// <param name="NameIndex">Index of the meta data key</param>
-		public void Remove(string Name, int NameIndex = IndexNone)
+		/// <param name="name">Name of the meta data key</param>
+		/// <param name="nameIndex">Index of the meta data key</param>
+		public void Remove(string name, int nameIndex = IndexNone)
 		{
 			if (this.Dictionary != null)
 			{
-				this.Dictionary.Remove(new UhtMetaDataKey(Name, NameIndex));
+				this.Dictionary.Remove(new UhtMetaDataKey(name, nameIndex));
 			}
 		}
 
 		/// <summary>
 		/// Given a key, return the full meta data name
 		/// </summary>
-		/// <param name="Key">Meta data key</param>
+		/// <param name="key">Meta data key</param>
 		/// <returns>Full meta data name</returns>
-		public string GetKeyString(UhtMetaDataKey Key)
+		public string GetKeyString(UhtMetaDataKey key)
 		{
-			return GetKeyString(Key.Name, Key.Index);
+			return GetKeyString(key.Name, key.Index);
 		}
 
 		/// <summary>
 		/// Given a key, return the full meta data name
 		/// </summary>
-		/// <param name="Name">Name of the meta data key</param>
-		/// <param name="NameIndex">Index of the meta data key</param>
+		/// <param name="name">Name of the meta data key</param>
+		/// <param name="nameIndex">Index of the meta data key</param>
 		/// <returns>Full meta data name</returns>
 		/// <exception cref="UhtIceException">Thrown if an index is supplied (not -1) and no key conversion interface is set</exception>
-		public string GetKeyString(string Name, int NameIndex)
+		public string GetKeyString(string name, int nameIndex)
 		{
-			if (NameIndex == IndexNone)
+			if (nameIndex == IndexNone)
 			{
-				return Name;
+				return name;
 			}
 			else if (KeyConversion != null)
 			{
-				return KeyConversion.GetMetaDataKey(Name, NameIndex);
+				return KeyConversion.GetMetaDataKey(name, nameIndex);
 			}
 			else
 			{
@@ -544,27 +544,27 @@ namespace EpicGames.UHT.Types
 		/// <returns>List of meta data key and value pairs</returns>
 		public List<KeyValuePair<string, string>> GetSorted()
 		{
-			List<KeyValuePair<string, string>> Out = new List<KeyValuePair<string, string>>(this.Dictionary != null ? this.Dictionary.Count : 0);
+			List<KeyValuePair<string, string>> output = new List<KeyValuePair<string, string>>(this.Dictionary != null ? this.Dictionary.Count : 0);
 			if (this.Dictionary != null && this.Dictionary.Count > 0)
 			{
-				foreach (KeyValuePair<UhtMetaDataKey, string> Kvp in this.Dictionary)
+				foreach (KeyValuePair<UhtMetaDataKey, string> kvp in this.Dictionary)
 				{
-					Out.Add(new KeyValuePair<string, string>(this.GetKeyString(Kvp.Key), Kvp.Value));
+					output.Add(new KeyValuePair<string, string>(this.GetKeyString(kvp.Key), kvp.Value));
 				}
 
-				Out.Sort((KeyValuePair<string, string> Lhs, KeyValuePair<string, string> Rhs) =>
+				output.Sort((KeyValuePair<string, string> lhs, KeyValuePair<string, string> rhs) =>
 				{
-					return StringComparerUE.OrdinalIgnoreCase.Compare(Lhs.Key, Rhs.Key);
+					return StringComparerUE.OrdinalIgnoreCase.Compare(lhs.Key, rhs.Key);
 				});
 			}
-			return Out;
+			return output;
 		}
 
 		private SortedList<UhtMetaDataKey, string> GetDictionary()
 		{
 			if (this.Dictionary == null)
 			{
-				this.Dictionary = new SortedList<UhtMetaDataKey, string>(Comparer);
+				this.Dictionary = new SortedList<UhtMetaDataKey, string>(s_comparer);
 			}
 			return this.Dictionary;
 		}
diff --git a/Engine/Source/Programs/Shared/EpicGames.UHT/Types/UhtObject.cs b/Engine/Source/Programs/Shared/EpicGames.UHT/Types/UhtObject.cs
--- a/Engine/Source/Programs/Shared/EpicGames.UHT/Types/UhtObject.cs
+++ b/Engine/Source/Programs/Shared/EpicGames.UHT/Types/UhtObject.cs
@@ -30,13 +30,13 @@ namespace EpicGames.UHT.Types
 		public UhtObject? AlternateObject { get; set; } = null;
 
 		/// <inheritdoc/>
-		public override string EngineClassName { get => "Object"; }
+		public override string EngineClassName => "Object";
 
 		/// <summary>
 		/// Construct a new instance of the object
 		/// </summary>
-		/// <param name="Session">Session of the object</param>
-		protected UhtObject(UhtSession Session) : base(Session)
+		/// <param name="session">Session of the object</param>
+		protected UhtObject(UhtSession session) : base(session)
 		{
 			ObjectTypeIndex = this.Session.GetNextObjectTypeIndex();
 		}
@@ -44,9 +44,9 @@ namespace EpicGames.UHT.Types
 		/// <summary>
 		/// Construct a new instance of the object
 		/// </summary>
-		/// <param name="Outer">Outer object</param>
-		/// <param name="LineNumber">Line number where object is defined</param>
-		protected UhtObject(UhtType Outer, int LineNumber) : base(Outer, LineNumber)
+		/// <param name="outer">Outer object</param>
+		/// <param name="lineNumber">Line number where object is defined</param>
+		protected UhtObject(UhtType outer, int lineNumber) : base(outer, lineNumber)
 		{
 			ObjectTypeIndex = this.Session.GetNextObjectTypeIndex();
 		}
diff --git a/Engine/Source/Programs/Shared/EpicGames.UHT/Types/UhtPackage.cs b/Engine/Source/Programs/Shared/EpicGames.UHT/Types/UhtPackage.cs
--- a/Engine/Source/Programs/Shared/EpicGames.UHT/Types/UhtPackage.cs
+++ b/Engine/Source/Programs/Shared/EpicGames.UHT/Types/UhtPackage.cs
@@ -1,10 +1,10 @@
 // Copyright Epic Games, Inc. All Rights Reserved.
 
+using System;
+using System.Text.Json.Serialization;
 using EpicGames.Core;
 using EpicGames.UHT.Tables;
 using EpicGames.UHT.Utils;
-using System;
-using System.Text.Json.Serialization;
 using UnrealBuildBase;
 
 namespace EpicGames.UHT.Types
@@ -37,7 +37,7 @@ namespace EpicGames.UHT.Types
 		public override UhtEngineType EngineType => UhtEngineType.Package;
 
 		/// <inheritdoc/>
-		public override string EngineClassName { get => "Package"; }
+		public override string EngineClassName => "Package";
 
 		/// <inheritdoc/>
 		[JsonIgnore]
@@ -51,7 +51,7 @@ namespace EpicGames.UHT.Types
 		/// True if the package is part of the engine
 		/// </summary>
 		[JsonIgnore]
-		public bool bIsPartOfEngine
+		public bool IsPartOfEngine
 		{
 			get
 			{
@@ -81,7 +81,7 @@ namespace EpicGames.UHT.Types
 		/// True if the package is a plugin
 		/// </summary>
 		[JsonIgnore]
-		public bool bIsPlugin => this.Module.BaseDirectory.Replace('\\', '/').Contains("/Plugins/", StringComparison.Ordinal);
+		public bool IsPlugin => this.Module.BaseDirectory.Replace('\\', '/').Contains("/Plugins/", StringComparison.Ordinal);
 
 		/// <summary>
 		/// Short name of the package (without the /Script/)
@@ -91,25 +91,25 @@ namespace EpicGames.UHT.Types
 		/// <summary>
 		/// Construct a new instance of a package
 		/// </summary>
-		/// <param name="Session">Running session</param>
-		/// <param name="Module">Source module of the package</param>
-		/// <param name="PackageFlags">Assorted package flags</param>
-		public UhtPackage(UhtSession Session, UHTManifest.Module Module, EPackageFlags PackageFlags) : base(Session)
+		/// <param name="session">Running session</param>
+		/// <param name="module">Source module of the package</param>
+		/// <param name="packageFlags">Assorted package flags</param>
+		public UhtPackage(UhtSession session, UHTManifest.Module module, EPackageFlags packageFlags) : base(session)
 		{
-			this.Module = Module;
-			this.PackageFlags = PackageFlags;
+			this.Module = module;
+			this.PackageFlags = packageFlags;
 			this.PackageTypeIndex = this.Session.GetNextPackageTypeIndex();
 
-			int LastSlashIndex = Module.Name.LastIndexOf('/');
-			if (LastSlashIndex == -1)
+			int lastSlashIndex = module.Name.LastIndexOf('/');
+			if (lastSlashIndex == -1)
 			{
-				this.SourceName = $"/Script/{Module.Name}";
-				this.ShortName = Module.Name;
+				this.SourceName = $"/Script/{module.Name}";
+				this.ShortName = module.Name;
 			}
 			else
 			{
-				this.SourceName = Module.Name;
-				this.ShortName = this.SourceName.Substring(LastSlashIndex + 1);
+				this.SourceName = module.Name;
+				this.ShortName = this.SourceName.Substring(lastSlashIndex + 1);
 			}
 		}
 	}
