More coding standard fixes

diff --git a/Engine/Source/Programs/Shared/EpicGames.UHT/Types/UhtProperty.cs b/Engine/Source/Programs/Shared/EpicGames.UHT/Types/UhtProperty.cs
--- a/Engine/Source/Programs/Shared/EpicGames.UHT/Types/UhtProperty.cs
+++ b/Engine/Source/Programs/Shared/EpicGames.UHT/Types/UhtProperty.cs
@@ -1,16 +1,16 @@
 // Copyright Epic Games, Inc. All Rights Reserved.
 
-using EpicGames.Core;
-using EpicGames.UHT.Parsers;
-using EpicGames.UHT.Tables;
-using EpicGames.UHT.Tokenizer;
-using EpicGames.UHT.Utils;
 using System;
 using System.Collections.Generic;
 using System.Diagnostics.CodeAnalysis;
 using System.Linq;
 using System.Text;
 using System.Text.Json.Serialization;
+using EpicGames.Core;
+using EpicGames.UHT.Parsers;
+using EpicGames.UHT.Tables;
+using EpicGames.UHT.Tokenizer;
+using EpicGames.UHT.Utils;
 
 namespace EpicGames.UHT.Types
 {
@@ -35,7 +35,7 @@ namespace EpicGames.UHT.Types
 		/// Property should be exported as protected
 		/// </summary>
 		Protected = 0x00000004,
-		
+
 		/// <summary>
 		/// The BlueprintPure flag was set in software and should not be considered an error
 		/// </summary>
@@ -87,35 +87,35 @@ namespace EpicGames.UHT.Types
 		/// <summary>
 		/// Test to see if any of the specified flags are set
 		/// </summary>
-		/// <param name="InFlags">Current flags</param>
-		/// <param name="TestFlags">Flags to test for</param>
+		/// <param name="inFlags">Current flags</param>
+		/// <param name="testFlags">Flags to test for</param>
 		/// <returns>True if any of the flags are set</returns>
-		public static bool HasAnyFlags(this UhtPropertyExportFlags InFlags, UhtPropertyExportFlags TestFlags)
+		public static bool HasAnyFlags(this UhtPropertyExportFlags inFlags, UhtPropertyExportFlags testFlags)
 		{
-			return (InFlags & TestFlags) != 0;
+			return (inFlags & testFlags) != 0;
 		}
 
 		/// <summary>
 		/// Test to see if all of the specified flags are set
 		/// </summary>
-		/// <param name="InFlags">Current flags</param>
-		/// <param name="TestFlags">Flags to test for</param>
+		/// <param name="inFlags">Current flags</param>
+		/// <param name="testFlags">Flags to test for</param>
 		/// <returns>True if all the flags are set</returns>
-		public static bool HasAllFlags(this UhtPropertyExportFlags InFlags, UhtPropertyExportFlags TestFlags)
+		public static bool HasAllFlags(this UhtPropertyExportFlags inFlags, UhtPropertyExportFlags testFlags)
 		{
-			return (InFlags & TestFlags) == TestFlags;
+			return (inFlags & testFlags) == testFlags;
 		}
 
 		/// <summary>
 		/// Test to see if a specific set of flags have a specific value.
 		/// </summary>
-		/// <param name="InFlags">Current flags</param>
-		/// <param name="TestFlags">Flags to test for</param>
-		/// <param name="MatchFlags">Expected value of the tested flags</param>
+		/// <param name="inFlags">Current flags</param>
+		/// <param name="testFlags">Flags to test for</param>
+		/// <param name="matchFlags">Expected value of the tested flags</param>
 		/// <returns>True if the given flags have a specific value.</returns>
-		public static bool HasExactFlags(this UhtPropertyExportFlags InFlags, UhtPropertyExportFlags TestFlags, UhtPropertyExportFlags MatchFlags)
+		public static bool HasExactFlags(this UhtPropertyExportFlags inFlags, UhtPropertyExportFlags testFlags, UhtPropertyExportFlags matchFlags)
 		{
-			return (InFlags & TestFlags) == MatchFlags;
+			return (inFlags & testFlags) == matchFlags;
 		}
 	}
 
@@ -154,12 +154,12 @@ namespace EpicGames.UHT.Types
 		/// <summary>
 		/// Return the hint text for the property category
 		/// </summary>
-		/// <param name="PropertyCategory">Property category</param>
+		/// <param name="propertyCategory">Property category</param>
 		/// <returns>The user facing hint text</returns>
 		/// <exception cref="UhtIceException">Unexpected category</exception>
-		public static string GetHintText(this UhtPropertyCategory PropertyCategory)
+		public static string GetHintText(this UhtPropertyCategory propertyCategory)
 		{
-			switch (PropertyCategory)
+			switch (propertyCategory)
 			{
 				case UhtPropertyCategory.ReplicatedParameter:
 				case UhtPropertyCategory.RegularParameter:
@@ -282,35 +282,35 @@ namespace EpicGames.UHT.Types
 		/// <summary>
 		/// Test to see if any of the specified flags are set
 		/// </summary>
-		/// <param name="InFlags">Current flags</param>
-		/// <param name="TestFlags">Flags to test for</param>
+		/// <param name="inFlags">Current flags</param>
+		/// <param name="testFlags">Flags to test for</param>
 		/// <returns>True if any of the flags are set</returns>
-		public static bool HasAnyFlags(this UhtPropertyOptions InFlags, UhtPropertyOptions TestFlags)
+		public static bool HasAnyFlags(this UhtPropertyOptions inFlags, UhtPropertyOptions testFlags)
 		{
-			return (InFlags & TestFlags) != 0;
+			return (inFlags & testFlags) != 0;
 		}
 
 		/// <summary>
 		/// Test to see if all of the specified flags are set
 		/// </summary>
-		/// <param name="InFlags">Current flags</param>
-		/// <param name="TestFlags">Flags to test for</param>
+		/// <param name="inFlags">Current flags</param>
+		/// <param name="testFlags">Flags to test for</param>
 		/// <returns>True if all the flags are set</returns>
-		public static bool HasAllFlags(this UhtPropertyOptions InFlags, UhtPropertyOptions TestFlags)
+		public static bool HasAllFlags(this UhtPropertyOptions inFlags, UhtPropertyOptions testFlags)
 		{
-			return (InFlags & TestFlags) == TestFlags;
+			return (inFlags & testFlags) == testFlags;
 		}
 
 		/// <summary>
 		/// Test to see if a specific set of flags have a specific value.
 		/// </summary>
-		/// <param name="InFlags">Current flags</param>
-		/// <param name="TestFlags">Flags to test for</param>
-		/// <param name="MatchFlags">Expected value of the tested flags</param>
+		/// <param name="inFlags">Current flags</param>
+		/// <param name="testFlags">Flags to test for</param>
+		/// <param name="matchFlags">Expected value of the tested flags</param>
 		/// <returns>True if the given flags have a specific value.</returns>
-		public static bool HasExactFlags(this UhtPropertyOptions InFlags, UhtPropertyOptions TestFlags, UhtPropertyOptions MatchFlags)
+		public static bool HasExactFlags(this UhtPropertyOptions inFlags, UhtPropertyOptions testFlags, UhtPropertyOptions matchFlags)
 		{
-			return (InFlags & TestFlags) == MatchFlags;
+			return (inFlags & testFlags) == matchFlags;
 		}
 	}
 
@@ -411,35 +411,35 @@ namespace EpicGames.UHT.Types
 		/// <summary>
 		/// Test to see if any of the specified flags are set
 		/// </summary>
-		/// <param name="InFlags">Current flags</param>
-		/// <param name="TestFlags">Flags to test for</param>
+		/// <param name="inFlags">Current flags</param>
+		/// <param name="testFlags">Flags to test for</param>
 		/// <returns>True if any of the flags are set</returns>
-		public static bool HasAnyFlags(this UhtPropertyCaps InFlags, UhtPropertyCaps TestFlags)
+		public static bool HasAnyFlags(this UhtPropertyCaps inFlags, UhtPropertyCaps testFlags)
 		{
-			return (InFlags & TestFlags) != 0;
+			return (inFlags & testFlags) != 0;
 		}
 
 		/// <summary>
 		/// Test to see if all of the specified flags are set
 		/// </summary>
-		/// <param name="InFlags">Current flags</param>
-		/// <param name="TestFlags">Flags to test for</param>
+		/// <param name="inFlags">Current flags</param>
+		/// <param name="testFlags">Flags to test for</param>
 		/// <returns>True if all the flags are set</returns>
-		public static bool HasAllFlags(this UhtPropertyCaps InFlags, UhtPropertyCaps TestFlags)
+		public static bool HasAllFlags(this UhtPropertyCaps inFlags, UhtPropertyCaps testFlags)
 		{
-			return (InFlags & TestFlags) == TestFlags;
+			return (inFlags & testFlags) == testFlags;
 		}
 
 		/// <summary>
 		/// Test to see if a specific set of flags have a specific value.
 		/// </summary>
-		/// <param name="InFlags">Current flags</param>
-		/// <param name="TestFlags">Flags to test for</param>
-		/// <param name="MatchFlags">Expected value of the tested flags</param>
+		/// <param name="inFlags">Current flags</param>
+		/// <param name="testFlags">Flags to test for</param>
+		/// <param name="matchFlags">Expected value of the tested flags</param>
 		/// <returns>True if the given flags have a specific value.</returns>
-		public static bool HasExactFlags(this UhtPropertyCaps InFlags, UhtPropertyCaps TestFlags, UhtPropertyCaps MatchFlags)
+		public static bool HasExactFlags(this UhtPropertyCaps inFlags, UhtPropertyCaps testFlags, UhtPropertyCaps matchFlags)
 		{
-			return (InFlags & TestFlags) == MatchFlags;
+			return (inFlags & testFlags) == matchFlags;
 		}
 	}
 
@@ -543,16 +543,16 @@ namespace EpicGames.UHT.Types
 		/// <summary>
 		/// Test to see if the text type is for a function
 		/// </summary>
-		/// <param name="TextType">Type of text</param>
+		/// <param name="textType">Type of text</param>
 		/// <returns>True if the text type is a function</returns>
-		public static bool IsParameter(this UhtPropertyTextType TextType)
+		public static bool IsParameter(this UhtPropertyTextType textType)
 		{
 			return
-				TextType == UhtPropertyTextType.GenericFunctionArgOrRetVal ||
-				TextType == UhtPropertyTextType.GenericFunctionArgOrRetValImpl ||
-				TextType == UhtPropertyTextType.ClassFunctionArgOrRetVal ||
-				TextType == UhtPropertyTextType.EventFunctionArgOrRetVal ||
-				TextType == UhtPropertyTextType.InterfaceFunctionArgOrRetVal;
+				textType == UhtPropertyTextType.GenericFunctionArgOrRetVal ||
+				textType == UhtPropertyTextType.GenericFunctionArgOrRetValImpl ||
+				textType == UhtPropertyTextType.ClassFunctionArgOrRetVal ||
+				textType == UhtPropertyTextType.EventFunctionArgOrRetVal ||
+				textType == UhtPropertyTextType.InterfaceFunctionArgOrRetVal;
 		}
 	}
 
@@ -608,17 +608,17 @@ namespace EpicGames.UHT.Types
 		/// <summary>
 		/// Return the hash code for a given object
 		/// </summary>
-		/// <param name="Obj">Object in question</param>
+		/// <param name="obj">Object in question</param>
 		/// <returns>Hash code</returns>
-		public uint GetTypeHash(UhtObject Obj);
+		public uint GetTypeHash(UhtObject obj);
 
 		/// <summary>
 		/// Return the singleton name for the given object
 		/// </summary>
-		/// <param name="Obj">Object in question</param>
-		/// <param name="bRegistered">If true, the singleton that returns the registered object is returned.</param>
+		/// <param name="obj">Object in question</param>
+		/// <param name="registered">If true, the singleton that returns the registered object is returned.</param>
 		/// <returns>Singleton function name</returns>
-		public string GetSingletonName(UhtObject? Obj, bool bRegistered);
+		public string GetSingletonName(UhtObject? obj, bool registered);
 	}
 
 	/// <summary>
@@ -716,7 +716,7 @@ namespace EpicGames.UHT.Types
 		/// <summary>
 		/// If true, the property is a bit field
 		/// </summary>
-		public bool bIsBitfield { get; set; }
+		public bool IsBitfield { get; set; }
 
 		/// <summary>
 		/// Construct a new, uninitialized version of the property settings
@@ -730,48 +730,48 @@ namespace EpicGames.UHT.Types
 		/// <summary>
 		/// Construct property settings based on the property settings for a parent container
 		/// </summary>
-		/// <param name="ParentPropertySettings">Parent container property</param>
-		/// <param name="SourceName">Name of the property</param>
-		/// <param name="MessageSite">Message site used to construct meta data object</param>
-		public UhtPropertySettings(UhtPropertySettings ParentPropertySettings, string SourceName, IUhtMessageSite MessageSite)
-		{
-			this.SourceName = SourceName;
-			this.EngineName = SourceName;
-			this.MetaData = new UhtMetaData(MessageSite, ParentPropertySettings.Outer.Session.Config);
-			this.Outer = ParentPropertySettings.Outer;
-			this.LineNumber = ParentPropertySettings.LineNumber;
-			this.PropertyCategory = ParentPropertySettings.PropertyCategory;
-			this.PropertyFlags = ParentPropertySettings.PropertyFlags;
-			this.DisallowPropertyFlags = ParentPropertySettings.DisallowPropertyFlags;
+		/// <param name="parentPropertySettings">Parent container property</param>
+		/// <param name="sourceName">Name of the property</param>
+		/// <param name="messageSite">Message site used to construct meta data object</param>
+		public UhtPropertySettings(UhtPropertySettings parentPropertySettings, string sourceName, IUhtMessageSite messageSite)
+		{
+			this.SourceName = sourceName;
+			this.EngineName = sourceName;
+			this.MetaData = new UhtMetaData(messageSite, parentPropertySettings.Outer.Session.Config);
+			this.Outer = parentPropertySettings.Outer;
+			this.LineNumber = parentPropertySettings.LineNumber;
+			this.PropertyCategory = parentPropertySettings.PropertyCategory;
+			this.PropertyFlags = parentPropertySettings.PropertyFlags;
+			this.DisallowPropertyFlags = parentPropertySettings.DisallowPropertyFlags;
 			this.PropertyExportFlags = UhtPropertyExportFlags.Public;
 			this.RepNotifyName = null;
 			this.Allocator = UhtPropertyAllocator.Default;
-			this.Options = ParentPropertySettings.Options;
+			this.Options = parentPropertySettings.Options;
 			this.PointerType = UhtPointerType.None;
 			this.ArrayDimensions = null;
 			this.DefaultValueTokens = null;
 			this.Setter = null;
 			this.Getter = null;
-			this.bIsBitfield = false;
+			this.IsBitfield = false;
 		}
 
 		/// <summary>
 		/// Reset the property settings.  Used on a cached property settings object
 		/// </summary>
-		/// <param name="Outer">Outer/owning type</param>
-		/// <param name="LineNumber">Line number of property</param>
-		/// <param name="PropertyCategory">Category of property</param>
-		/// <param name="DisallowPropertyFlags">Property flags that are not allowed</param>
-		public void Reset(UhtType Outer, int LineNumber, UhtPropertyCategory PropertyCategory, EPropertyFlags DisallowPropertyFlags)
-		{
-			this.SourceName = string.Empty;
-			this.EngineName = string.Empty;
-			this.MetaData = new UhtMetaData(Outer, Outer.Session.Config);
-			this.Outer = Outer;
-			this.LineNumber = LineNumber;
-			this.PropertyCategory = PropertyCategory;
+		/// <param name="outer">Outer/owning type</param>
+		/// <param name="lineNumber">Line number of property</param>
+		/// <param name="propertyCategory">Category of property</param>
+		/// <param name="disallowPropertyFlags">Property flags that are not allowed</param>
+		public void Reset(UhtType outer, int lineNumber, UhtPropertyCategory propertyCategory, EPropertyFlags disallowPropertyFlags)
+		{
+			this.SourceName = String.Empty;
+			this.EngineName = String.Empty;
+			this.MetaData = new UhtMetaData(outer, outer.Session.Config);
+			this.Outer = outer;
+			this.LineNumber = lineNumber;
+			this.PropertyCategory = propertyCategory;
 			this.PropertyFlags = EPropertyFlags.None;
-			this.DisallowPropertyFlags = DisallowPropertyFlags;
+			this.DisallowPropertyFlags = disallowPropertyFlags;
 			this.PropertyExportFlags = UhtPropertyExportFlags.Public;
 			this.RepNotifyName = null;
 			this.Allocator = UhtPropertyAllocator.Default;
@@ -780,39 +780,39 @@ namespace EpicGames.UHT.Types
 			this.DefaultValueTokens = null;
 			this.Setter = null;
 			this.Getter = null;
-			this.bIsBitfield = false;
+			this.IsBitfield = false;
 		}
 
 		/// <summary>
 		/// Reset property settings based on the given property.  Used to prepare a cached property settings for parsing.
 		/// </summary>
-		/// <param name="Property">Source property</param>
-		/// <param name="Options">Property options</param>
+		/// <param name="property">Source property</param>
+		/// <param name="options">Property options</param>
 		/// <exception cref="UhtIceException">Thrown if the input property doesn't have an outer</exception>
-		public void Reset(UhtProperty Property, UhtPropertyOptions Options)
+		public void Reset(UhtProperty property, UhtPropertyOptions options)
 		{
-			if (Property.Outer == null)
+			if (property.Outer == null)
 			{
 				throw new UhtIceException("Property must have an outer specified");
 			}
-			this.SourceName = Property.SourceName;
-			this.EngineName = Property.EngineName;
-			this.MetaData = Property.MetaData;
-			this.Outer = Property.Outer;
-			this.LineNumber = Property.LineNumber;
-			this.PropertyCategory = Property.PropertyCategory;
-			this.PropertyFlags = Property.PropertyFlags;
-			this.DisallowPropertyFlags = Property.DisallowPropertyFlags;
-			this.PropertyExportFlags = Property.PropertyExportFlags;
-			this.Allocator = Property.Allocator;
-			this.Options = Options;
-			this.PointerType = Property.PointerType;
-			this.RepNotifyName = Property.RepNotifyName;
-			this.ArrayDimensions = Property.ArrayDimensions;
-			this.DefaultValueTokens = Property.DefaultValueTokens;
-			this.Setter = Property.Setter;
-			this.Getter = Property.Getter;
-			this.bIsBitfield = Property.bIsBitfield;
+			this.SourceName = property.SourceName;
+			this.EngineName = property.EngineName;
+			this.MetaData = property.MetaData;
+			this.Outer = property.Outer;
+			this.LineNumber = property.LineNumber;
+			this.PropertyCategory = property.PropertyCategory;
+			this.PropertyFlags = property.PropertyFlags;
+			this.DisallowPropertyFlags = property.DisallowPropertyFlags;
+			this.PropertyExportFlags = property.PropertyExportFlags;
+			this.Allocator = property.Allocator;
+			this.Options = options;
+			this.PointerType = property.PointerType;
+			this.RepNotifyName = property.RepNotifyName;
+			this.ArrayDimensions = property.ArrayDimensions;
+			this.DefaultValueTokens = property.DefaultValueTokens;
+			this.Setter = property.Setter;
+			this.Getter = property.Getter;
+			this.IsBitfield = property.IsBitfield;
 		}
 	}
 
@@ -827,7 +827,7 @@ namespace EpicGames.UHT.Types
 		/// <summary>
 		/// Collection of recognized casts when parsing array dimensions
 		/// </summary>
-		private static readonly string[] Casts = new string[]
+		private static readonly string[] s_casts = new string[]
 		{
 			"(uint32)",
 			"(int32)",
@@ -845,7 +845,7 @@ namespace EpicGames.UHT.Types
 		/// <summary>
 		/// Collection of invalid names for parameters 
 		/// </summary>
-		private static readonly HashSet<string> InvalidParamNames = new HashSet<string>(StringComparer.OrdinalIgnoreCase) { "self" };
+		private static readonly HashSet<string> s_invalidParamNames = new HashSet<string>(StringComparer.OrdinalIgnoreCase) { "self" };
 
 		/// <summary>
 		/// Standard object flags for properties
@@ -872,12 +872,12 @@ namespace EpicGames.UHT.Types
 		/// <summary>
 		/// If true, then references must be passed without a pointer
 		/// </summary>
-		protected virtual bool bPGetPassAsNoPtr { get => false; }
+		protected virtual bool PGetPassAsNoPtr => false;
 
 		/// <summary>
 		/// Type of the PGet argument if one is required
 		/// </summary>
-		protected virtual UhtPGetArgumentType PGetTypeArgument { get => UhtPGetArgumentType.None; }
+		protected virtual UhtPGetArgumentType PGetTypeArgument => UhtPGetArgumentType.None;
 		#endregion
 
 		/// <summary>
@@ -953,7 +953,7 @@ namespace EpicGames.UHT.Types
 		/// <summary>
 		/// If true, this property is a bit field
 		/// </summary>
-		public bool bIsBitfield { get; set; } = false;
+		public bool IsBitfield { get; set; } = false;
 
 		///<inheritdoc/>
 		[JsonIgnore]
@@ -961,7 +961,7 @@ namespace EpicGames.UHT.Types
 
 		///<inheritdoc/>
 		[JsonIgnore]
-		public override bool bDeprecated => this.PropertyFlags.HasAnyFlags(EPropertyFlags.Deprecated);
+		public override bool Deprecated => this.PropertyFlags.HasAnyFlags(EPropertyFlags.Deprecated);
 
 		///<inheritdoc/>
 		protected override UhtSpecifierValidatorTable? SpecifierValidatorTable
@@ -986,20 +986,20 @@ namespace EpicGames.UHT.Types
 		/// If true, the property is a fixed/static array
 		/// </summary>
 		[JsonIgnore]
-		public bool bIsStaticArray => !string.IsNullOrEmpty(this.ArrayDimensions);
+		public bool IsStaticArray => !String.IsNullOrEmpty(this.ArrayDimensions);
 
 		/// <summary>
 		/// If true, the property is editor only
 		/// </summary>
 		[JsonIgnore]
-		public bool bIsEditorOnlyProperty => this.PropertyFlags.HasAnyFlags(EPropertyFlags.DevelopmentAssets);
+		public bool IsEditorOnlyProperty => this.PropertyFlags.HasAnyFlags(EPropertyFlags.DevelopmentAssets);
 
 		/// <summary>
 		/// Construct a new property
 		/// </summary>
-		/// <param name="Outer">Outer type of the property</param>
-		/// <param name="LineNumber">Line number where property was declared</param>
-		protected UhtProperty(UhtType Outer, int LineNumber) : base(Outer, LineNumber)
+		/// <param name="outer">Outer type of the property</param>
+		/// <param name="lineNumber">Line number where property was declared</param>
+		protected UhtProperty(UhtType outer, int lineNumber) : base(outer, lineNumber)
 		{
 			this.PropertyFlags = EPropertyFlags.None;
 			this.PropertyCaps = UhtPropertyCaps.CanBeContainerValue | UhtPropertyCaps.CanBeContainerKey | UhtPropertyCaps.CanHaveConfig;
@@ -1008,27 +1008,27 @@ namespace EpicGames.UHT.Types
 		/// <summary>
 		/// Construct a new property
 		/// </summary>
-		/// <param name="PropertySettings">Property settings from parsing</param>
-		protected UhtProperty(UhtPropertySettings PropertySettings) : base(PropertySettings.Outer, PropertySettings.LineNumber, PropertySettings.MetaData)
+		/// <param name="propertySettings">Property settings from parsing</param>
+		protected UhtProperty(UhtPropertySettings propertySettings) : base(propertySettings.Outer, propertySettings.LineNumber, propertySettings.MetaData)
 		{
-			this.SourceName = PropertySettings.SourceName;
+			this.SourceName = propertySettings.SourceName;
 			// Engine name defaults to source name.  If it doesn't match what is coming in, then set it.
-			if (PropertySettings.EngineName.Length > 0 && PropertySettings.SourceName != PropertySettings.EngineName ) 
+			if (propertySettings.EngineName.Length > 0 && propertySettings.SourceName != propertySettings.EngineName)
 			{
-				this.EngineName = PropertySettings.EngineName;
+				this.EngineName = propertySettings.EngineName;
 			}
-			this.PropertyCategory = PropertySettings.PropertyCategory;
-			this.PropertyFlags = PropertySettings.PropertyFlags;
-			this.DisallowPropertyFlags = PropertySettings.DisallowPropertyFlags;
-			this.PropertyExportFlags = PropertySettings.PropertyExportFlags;
-			this.Allocator = PropertySettings.Allocator;
-			this.PointerType = PropertySettings.PointerType;
-			this.RepNotifyName = PropertySettings.RepNotifyName;
-			this.ArrayDimensions = PropertySettings.ArrayDimensions;
-			this.DefaultValueTokens = PropertySettings.DefaultValueTokens;
-			this.Getter = PropertySettings.Getter;
-			this.Setter = PropertySettings.Setter;
-			this.bIsBitfield = PropertySettings.bIsBitfield;
+			this.PropertyCategory = propertySettings.PropertyCategory;
+			this.PropertyFlags = propertySettings.PropertyFlags;
+			this.DisallowPropertyFlags = propertySettings.DisallowPropertyFlags;
+			this.PropertyExportFlags = propertySettings.PropertyExportFlags;
+			this.Allocator = propertySettings.Allocator;
+			this.PointerType = propertySettings.PointerType;
+			this.RepNotifyName = propertySettings.RepNotifyName;
+			this.ArrayDimensions = propertySettings.ArrayDimensions;
+			this.DefaultValueTokens = propertySettings.DefaultValueTokens;
+			this.Getter = propertySettings.Getter;
+			this.Setter = propertySettings.Setter;
+			this.IsBitfield = propertySettings.IsBitfield;
 			this.PropertyCaps = UhtPropertyCaps.CanBeContainerValue | UhtPropertyCaps.CanBeContainerKey | UhtPropertyCaps.CanHaveConfig;
 		}
 
@@ -1042,29 +1042,29 @@ namespace EpicGames.UHT.Types
 		/// <summary>
 		/// Internal version of AppendText.  Don't append any text to the builder to get the default behavior
 		/// </summary>
-		/// <param name="Builder">Destination builder</param>
-		/// <param name="TextType">Text type of where the property is being referenced</param>
-		/// <param name="bIsTemplateArgument">If true, this property is a template arguments</param>
+		/// <param name="builder">Destination builder</param>
+		/// <param name="textType">Text type of where the property is being referenced</param>
+		/// <param name="isTemplateArgument">If true, this property is a template arguments</param>
 		/// <returns>Destination builder</returns>
-		public virtual StringBuilder AppendText(StringBuilder Builder, UhtPropertyTextType TextType, bool bIsTemplateArgument = false)
+		public virtual StringBuilder AppendText(StringBuilder builder, UhtPropertyTextType textType, bool isTemplateArgument = false)
 		{
 			// By default, we assume it will be just the simple text
-			return Builder.Append(this.CppTypeText);
+			return builder.Append(this.CppTypeText);
 		}
 
 		/// <summary>
 		/// Append the full declaration including such things as property name and const<amp/> requirements
 		/// </summary>
-		/// <param name="Builder">Destination builder</param>
-		/// <param name="TextType">Text type of where the property is being referenced</param>
-		/// <param name="bSkipParameterName">If true, do not include the property name</param>
+		/// <param name="builder">Destination builder</param>
+		/// <param name="textType">Text type of where the property is being referenced</param>
+		/// <param name="skipParameterName">If true, do not include the property name</param>
 		/// <returns>Destination builder</returns>
-		public virtual StringBuilder AppendFullDecl(StringBuilder Builder, UhtPropertyTextType TextType, bool bSkipParameterName = false)
+		public virtual StringBuilder AppendFullDecl(StringBuilder builder, UhtPropertyTextType textType, bool skipParameterName = false)
 		{
-			UhtPropertyCaps Caps = this.PropertyCaps;
+			UhtPropertyCaps caps = this.PropertyCaps;
 
-			bool bIsParameter = TextType.IsParameter();
-			bool bIsInterfaceProp = this is UhtInterfaceProperty;
+			bool isParameter = textType.IsParameter();
+			bool isInterfaceProp = this is UhtInterfaceProperty;
 
 			// When do we need a leading const:
 			// 1) If this is a object or object ptr property and the referenced class is const
@@ -1072,250 +1072,250 @@ namespace EpicGames.UHT.Types
 			//		if this is a parameter
 			//		AND - if this is a const param OR (if this is an interface property and not an out param)
 			// 3) If this is a parameter without array dimensions, must be passed by reference, but not an out parameter or const param
-			bool bPassCppArgsByRef = Caps.HasAnyFlags(UhtPropertyCaps.PassCppArgsByRef);
-			bool bIsConstParam = bIsParameter && (this.PropertyFlags.HasAnyFlags(EPropertyFlags.ConstParm) || (bIsInterfaceProp && !this.PropertyFlags.HasAnyFlags(EPropertyFlags.OutParm)));
-			bool bIsConstArgsByRef = bIsParameter && this.ArrayDimensions == null && bPassCppArgsByRef && !this.PropertyFlags.HasAnyFlags(EPropertyFlags.ConstParm | EPropertyFlags.OutParm);
-			bool bIsOnConstClass = false;
-			if (this is UhtObjectProperty ObjectProperty)
+			bool passCppArgsByRef = caps.HasAnyFlags(UhtPropertyCaps.PassCppArgsByRef);
+			bool isConstParam = isParameter && (this.PropertyFlags.HasAnyFlags(EPropertyFlags.ConstParm) || (isInterfaceProp && !this.PropertyFlags.HasAnyFlags(EPropertyFlags.OutParm)));
+			bool isConstArgsByRef = isParameter && this.ArrayDimensions == null && passCppArgsByRef && !this.PropertyFlags.HasAnyFlags(EPropertyFlags.ConstParm | EPropertyFlags.OutParm);
+			bool isOnConstClass = false;
+			if (this is UhtObjectProperty objectProperty)
 			{
-				bIsOnConstClass = ObjectProperty.Class.ClassFlags.HasAnyFlags(EClassFlags.Const);
+				isOnConstClass = objectProperty.Class.ClassFlags.HasAnyFlags(EClassFlags.Const);
 			}
-			bool bShouldHaveRef = this.PropertyFlags.HasAnyFlags(EPropertyFlags.OutParm | EPropertyFlags.ReferenceParm);
+			bool shouldHaveRef = this.PropertyFlags.HasAnyFlags(EPropertyFlags.OutParm | EPropertyFlags.ReferenceParm);
 
-			bool bConstAtTheBeginning = bIsOnConstClass || bIsConstArgsByRef || (bIsConstParam && !bShouldHaveRef);
-			if (bConstAtTheBeginning)
+			bool constAtTheBeginning = isOnConstClass || isConstArgsByRef || (isConstParam && !shouldHaveRef);
+			if (constAtTheBeginning)
 			{
-				Builder.Append("const ");
+				builder.Append("const ");
 			}
 
-			this.AppendText(Builder, TextType);
+			this.AppendText(builder, textType);
 
-			bool bFromConstClass = false;
-			if (TextType == UhtPropertyTextType.ExportMember && this.Outer != null)
+			bool fromConstClass = false;
+			if (textType == UhtPropertyTextType.ExportMember && this.Outer != null)
 			{
-				if (this.Outer is UhtClass OuterClass)
+				if (this.Outer is UhtClass outerClass)
 				{
-					bFromConstClass = OuterClass.ClassFlags.HasAnyFlags(EClassFlags.Const);
+					fromConstClass = outerClass.ClassFlags.HasAnyFlags(EClassFlags.Const);
 				}
 			}
-			bool bConstAtTheEnd = bFromConstClass || (bIsConstParam && bShouldHaveRef);
-			if (bConstAtTheEnd)
+			bool constAtTheEnd = fromConstClass || (isConstParam && shouldHaveRef);
+			if (constAtTheEnd)
 			{
-				Builder.Append(" const");
+				builder.Append(" const");
 			}
 
-			if (bIsParameter && this.ArrayDimensions == null && (bPassCppArgsByRef || this.PropertyFlags.HasAnyFlags(EPropertyFlags.OutParm | EPropertyFlags.ReferenceParm)))
+			if (isParameter && this.ArrayDimensions == null && (passCppArgsByRef || this.PropertyFlags.HasAnyFlags(EPropertyFlags.OutParm | EPropertyFlags.ReferenceParm)))
 			{
-				Builder.Append('&');
+				builder.Append('&');
 			}
 
-			Builder.Append(' ');
-			if (!bSkipParameterName)
+			builder.Append(' ');
+			if (!skipParameterName)
 			{
-				Builder.Append(this.SourceName);
+				builder.Append(this.SourceName);
 			}
 
 			if (this.ArrayDimensions != null)
 			{
-				Builder.Append('[').Append(this.ArrayDimensions).Append(']');
+				builder.Append('[').Append(this.ArrayDimensions).Append(']');
 			}
-			return Builder;
+			return builder;
 		}
 
 		/// <summary>
 		/// Append the required code to declare the property as a member
 		/// </summary>
-		/// <param name="Builder">Output builder</param>
-		/// <param name="Context">Current context</param>
-		/// <param name="Name">Name of the property.  This is needed in some cases where the name in the declarations doesn't match the property name.</param>
-		/// <param name="NameSuffix">Suffix to the property name</param>
-		/// <param name="Tabs">Number of tabs prefix the line with</param>
+		/// <param name="builder">Output builder</param>
+		/// <param name="context">Current context</param>
+		/// <param name="name">Name of the property.  This is needed in some cases where the name in the declarations doesn't match the property name.</param>
+		/// <param name="nameSuffix">Suffix to the property name</param>
+		/// <param name="tabs">Number of tabs prefix the line with</param>
 		/// <returns>Output builder</returns>
-		public abstract StringBuilder AppendMemberDecl(StringBuilder Builder, IUhtPropertyMemberContext Context, string Name, string NameSuffix, int Tabs);
+		public abstract StringBuilder AppendMemberDecl(StringBuilder builder, IUhtPropertyMemberContext context, string name, string nameSuffix, int tabs);
 
 		/// <summary>
 		/// Append the required code to declare the property as a member
 		/// </summary>
-		/// <param name="Builder">Output builder</param>
-		/// <param name="Context">Current context</param>
-		/// <param name="Name">Name of the property.  This is needed in some cases where the name in the declarations doesn't match the property name.</param>
-		/// <param name="NameSuffix">Suffix to the property name</param>
-		/// <param name="Tabs">Number of tabs prefix the line with</param>
-		/// <param name="ParamsStructName">Structure name</param>
-		/// <param name="bAppendMetaDataDecl">If true, add the meta data decl prior to the member decl</param>
+		/// <param name="builder">Output builder</param>
+		/// <param name="context">Current context</param>
+		/// <param name="name">Name of the property.  This is needed in some cases where the name in the declarations doesn't match the property name.</param>
+		/// <param name="nameSuffix">Suffix to the property name</param>
+		/// <param name="tabs">Number of tabs prefix the line with</param>
+		/// <param name="paramsStructName">Structure name</param>
+		/// <param name="appendMetaDataDecl">If true, add the meta data decl prior to the member decl</param>
 		/// <returns>Output builder</returns>
-		public StringBuilder AppendMemberDecl(StringBuilder Builder, IUhtPropertyMemberContext Context, string Name, string NameSuffix, int Tabs, string ParamsStructName, bool bAppendMetaDataDecl = true)
+		public StringBuilder AppendMemberDecl(StringBuilder builder, IUhtPropertyMemberContext context, string name, string nameSuffix, int tabs, string paramsStructName, bool appendMetaDataDecl = true)
 		{
-			if (bAppendMetaDataDecl)
+			if (appendMetaDataDecl)
 			{
-				Builder.AppendMetaDataDecl(this, Context, Name, NameSuffix, Tabs);
+				builder.AppendMetaDataDecl(this, context, name, nameSuffix, tabs);
 			}
-			Builder.AppendTabs(Tabs).Append("static const UECodeGen_Private::").Append(ParamsStructName).Append(' ').AppendNameDecl(Context, Name, NameSuffix).Append(";\r\n");
-			return Builder;
+			builder.AppendTabs(tabs).Append("static const UECodeGen_Private::").Append(paramsStructName).Append(' ').AppendNameDecl(context, name, nameSuffix).Append(";\r\n");
+			return builder;
 		}
 
 		/// <summary>
 		/// Append the required code to define the property as a member
 		/// </summary>
-		/// <param name="Builder">Output builder</param>
-		/// <param name="Context">Context of the call</param>
-		/// <param name="Name">Name of the property.  This is needed in some cases where the name in the declarations doesn't match the property name.</param>
-		/// <param name="NameSuffix">Suffix to the property name</param>
-		/// <param name="Offset">Offset of the property</param>
-		/// <param name="Tabs">Number of tabs prefix the line with</param>
+		/// <param name="builder">Output builder</param>
+		/// <param name="context">Context of the call</param>
+		/// <param name="name">Name of the property.  This is needed in some cases where the name in the declarations doesn't match the property name.</param>
+		/// <param name="nameSuffix">Suffix to the property name</param>
+		/// <param name="offset">Offset of the property</param>
+		/// <param name="tabs">Number of tabs prefix the line with</param>
 		/// <returns>Output builder</returns>
-		public abstract StringBuilder AppendMemberDef(StringBuilder Builder, IUhtPropertyMemberContext Context, string Name, string NameSuffix, string? Offset, int Tabs);
+		public abstract StringBuilder AppendMemberDef(StringBuilder builder, IUhtPropertyMemberContext context, string name, string nameSuffix, string? offset, int tabs);
 
 		/// <summary>
 		/// Append the required start of code to define the property as a member
 		/// </summary>
-		/// <param name="Builder">Output builder</param>
-		/// <param name="Context">Context of the call</param>
-		/// <param name="Name">Name of the property.  This is needed in some cases where the name in the declarations doesn't match the property name.</param>
-		/// <param name="NameSuffix">Suffix to the property name</param>
-		/// <param name="Offset">Offset of the property</param>
-		/// <param name="Tabs">Number of tabs prefix the line with</param>
-		/// <param name="ParamsStructName">Structure name</param>
-		/// <param name="ParamsGenFlags">Structure flags</param>
-		/// <param name="bAppendMetaDataDef">If true, add the meta data def prior to the member def</param>
-		/// <param name="bAppendOffset">If true, add the offset parameter</param>
+		/// <param name="builder">Output builder</param>
+		/// <param name="context">Context of the call</param>
+		/// <param name="name">Name of the property.  This is needed in some cases where the name in the declarations doesn't match the property name.</param>
+		/// <param name="nameSuffix">Suffix to the property name</param>
+		/// <param name="offset">Offset of the property</param>
+		/// <param name="tabs">Number of tabs prefix the line with</param>
+		/// <param name="paramsStructName">Structure name</param>
+		/// <param name="paramsGenFlags">Structure flags</param>
+		/// <param name="appendMetaDataDef">If true, add the meta data def prior to the member def</param>
+		/// <param name="appendOffset">If true, add the offset parameter</param>
 		/// <returns>Output builder</returns>
-		public StringBuilder AppendMemberDefStart(StringBuilder Builder, IUhtPropertyMemberContext Context, string Name, string NameSuffix, string? Offset, int Tabs,
-			string ParamsStructName, string ParamsGenFlags, bool bAppendMetaDataDef = true, bool bAppendOffset = true)
+		public StringBuilder AppendMemberDefStart(StringBuilder builder, IUhtPropertyMemberContext context, string name, string nameSuffix, string? offset, int tabs,
+			string paramsStructName, string paramsGenFlags, bool appendMetaDataDef = true, bool appendOffset = true)
 		{
-			if (bAppendMetaDataDef)
+			if (appendMetaDataDef)
 			{
-				Builder.AppendMetaDataDef(this, Context, Name, NameSuffix, Tabs);
+				builder.AppendMetaDataDef(this, context, name, nameSuffix, tabs);
 			}
-			Builder
-				.AppendTabs(Tabs)
-				.Append("const UECodeGen_Private::").Append(ParamsStructName).Append(' ')
-				.AppendNameDef(Context, Name, NameSuffix).Append(" = { ")
+			builder
+				.AppendTabs(tabs)
+				.Append("const UECodeGen_Private::").Append(paramsStructName).Append(' ')
+				.AppendNameDef(context, name, nameSuffix).Append(" = { ")
 				.AppendUTF8LiteralString(this.EngineName).Append(", ")
 				.AppendNotifyFunc(this).Append(", ")
 				.AppendFlags(this.PropertyFlags).Append(", ")
-				.Append(ParamsGenFlags).Append(", ")
+				.Append(paramsGenFlags).Append(", ")
 				.Append(ObjectFlags).Append(", ")
-				.AppendArrayDim(this, Context).Append(", ");
+				.AppendArrayDim(this, context).Append(", ");
 
 			if (this.PropertyExportFlags.HasAnyFlags(UhtPropertyExportFlags.SetterFound))
 			{
-				Builder.Append('&').Append(this.Outer!.SourceName).Append("::").AppendPropertySetterWrapperName(this).Append(", ");
+				builder.Append('&').Append(this.Outer!.SourceName).Append("::").AppendPropertySetterWrapperName(this).Append(", ");
 			}
 			else
 			{
-				Builder.Append("nullptr, ");
+				builder.Append("nullptr, ");
 			}
 
 			if (this.PropertyExportFlags.HasAnyFlags(UhtPropertyExportFlags.GetterFound))
 			{
-				Builder.Append('&').Append(this.Outer!.SourceName).Append("::").AppendPropertyGetterWrapperName(this).Append(", ");
+				builder.Append('&').Append(this.Outer!.SourceName).Append("::").AppendPropertyGetterWrapperName(this).Append(", ");
 			}
 			else
 			{
-				Builder.Append("nullptr, ");
+				builder.Append("nullptr, ");
 			}
 
-			if (bAppendOffset)
+			if (appendOffset)
 			{
-				if (!string.IsNullOrEmpty(Offset))
+				if (!String.IsNullOrEmpty(offset))
 				{
-					Builder.Append(Offset).Append(", ");
+					builder.Append(offset).Append(", ");
 				}
 				else
 				{
-					Builder.Append("STRUCT_OFFSET(").Append(Context.OuterStructSourceName).Append(", ").Append(this.SourceName).Append("), ");
+					builder.Append("STRUCT_OFFSET(").Append(context.OuterStructSourceName).Append(", ").Append(this.SourceName).Append("), ");
 				}
 			}
-			return Builder;
+			return builder;
 		}
 
 		/// <summary>
 		/// Append the required end of code to define the property as a member
 		/// </summary>
-		/// <param name="Builder">Output builder</param>
-		/// <param name="Context">Context of the call</param>
-		/// <param name="Name">Name of the property.  This is needed in some cases where the name in the declarations doesn't match the property name.</param>
-		/// <param name="NameSuffix">Suffix to the property name</param>
+		/// <param name="builder">Output builder</param>
+		/// <param name="context">Context of the call</param>
+		/// <param name="name">Name of the property.  This is needed in some cases where the name in the declarations doesn't match the property name.</param>
+		/// <param name="nameSuffix">Suffix to the property name</param>
 		/// <returns>Output builder</returns>
-		protected StringBuilder AppendMemberDefEnd(StringBuilder Builder, IUhtPropertyMemberContext Context, string Name, string NameSuffix)
+		protected StringBuilder AppendMemberDefEnd(StringBuilder builder, IUhtPropertyMemberContext context, string name, string nameSuffix)
 		{
-			Builder
-				.AppendMetaDataParams(this, Context, Name, NameSuffix)
+			builder
+				.AppendMetaDataParams(this, context, name, nameSuffix)
 				.Append(" };")
-				.AppendObjectHashes(this, Context)
+				.AppendObjectHashes(this, context)
 				.Append("\r\n");
-			return Builder;
+			return builder;
 		}
 
 		/// <summary>
 		/// Append the a type reference to the member definition
 		/// </summary>
-		/// <param name="Builder">Output builder</param>
-		/// <param name="Context">Context of the call</param>
-		/// <param name="Object">Referenced object</param>
-		/// <param name="bRegistered">True if the registered singleton name is to be used</param>
-		/// <param name="bAppendNull">True if a "nullptr" is to be appended if the object is null</param>
+		/// <param name="builder">Output builder</param>
+		/// <param name="context">Context of the call</param>
+		/// <param name="obj">Referenced object</param>
+		/// <param name="registered">True if the registered singleton name is to be used</param>
+		/// <param name="appendNull">True if a "nullptr" is to be appended if the object is null</param>
 		/// <returns>Output builder</returns>
-		protected static StringBuilder AppendMemberDefRef(StringBuilder Builder, IUhtPropertyMemberContext Context, UhtObject? Object, bool bRegistered, bool bAppendNull = false)
+		protected static StringBuilder AppendMemberDefRef(StringBuilder builder, IUhtPropertyMemberContext context, UhtObject? obj, bool registered, bool appendNull = false)
 		{
-			if (Object != null)
+			if (obj != null)
 			{
-				Builder.AppendSingletonName(Context, Object, bRegistered).Append(", ");
+				builder.AppendSingletonName(context, obj, registered).Append(", ");
 			}
-			else if (bAppendNull)
+			else if (appendNull)
 			{
-				Builder.Append("nullptr, ");
+				builder.Append("nullptr, ");
 			}
-			return Builder;
+			return builder;
 		}
 
 		/// <summary>
 		/// Append the required code to add the properties to a pointer array
 		/// </summary>
-		/// <param name="Builder">Output builder</param>
-		/// <param name="Context">Context of the call</param>
-		/// <param name="Name">Name of the property.  This is needed in some cases where the name in the declarations doesn't match the property name.</param>
-		/// <param name="NameSuffix">Suffix to the property name</param>
-		/// <param name="Tabs">Number of tabs prefix the line with</param>
+		/// <param name="builder">Output builder</param>
+		/// <param name="context">Context of the call</param>
+		/// <param name="name">Name of the property.  This is needed in some cases where the name in the declarations doesn't match the property name.</param>
+		/// <param name="nameSuffix">Suffix to the property name</param>
+		/// <param name="tabs">Number of tabs prefix the line with</param>
 		/// <returns>Output builder</returns>
-		public virtual StringBuilder AppendMemberPtr(StringBuilder Builder, IUhtPropertyMemberContext Context, string Name, string NameSuffix, int Tabs)
+		public virtual StringBuilder AppendMemberPtr(StringBuilder builder, IUhtPropertyMemberContext context, string name, string nameSuffix, int tabs)
 		{
-			Builder.AppendTabs(Tabs).Append("(const UECodeGen_Private::FPropertyParamsBase*)&").AppendNameDef(Context, Name, NameSuffix).Append(",\r\n");
-			return Builder;
+			builder.AppendTabs(tabs).Append("(const UECodeGen_Private::FPropertyParamsBase*)&").AppendNameDef(context, name, nameSuffix).Append(",\r\n");
+			return builder;
 		}
 
 		/// <summary>
 		/// Append a P_GET macro
 		/// </summary>
-		/// <param name="Builder">Destination builder</param>
+		/// <param name="builder">Destination builder</param>
 		/// <returns>Destination builder</returns>
-		public virtual StringBuilder AppendFunctionThunkParameterGet(StringBuilder Builder)
+		public virtual StringBuilder AppendFunctionThunkParameterGet(StringBuilder builder)
 		{
-			Builder.Append("P_GET_");
+			builder.Append("P_GET_");
 			if (this.ArrayDimensions != null)
 			{
-				Builder.Append("ARRAY");
+				builder.Append("ARRAY");
 			}
 			else
 			{
-				Builder.Append(this.PGetMacroText);
+				builder.Append(this.PGetMacroText);
 			}
 			if (this.PropertyFlags.HasAnyFlags(EPropertyFlags.OutParm))
 			{
-				if (this.bPGetPassAsNoPtr)
+				if (this.PGetPassAsNoPtr)
 				{
-					Builder.Append("_REF_NO_PTR");
+					builder.Append("_REF_NO_PTR");
 				}
 				else
 				{
-					Builder.Append("_REF");
+					builder.Append("_REF");
 				}
 			}
-			Builder.Append('(');
+			builder.Append('(');
 			if (this.ArrayDimensions != null)
 			{
-				Builder.AppendFunctionThunkParameterArrayType(this).Append(',');
+				builder.AppendFunctionThunkParameterArrayType(this).Append(',');
 			}
 			else
 			{
@@ -1325,51 +1325,51 @@ namespace EpicGames.UHT.Types
 						break;
 
 					case UhtPGetArgumentType.EngineClass:
-						Builder.Append('F').Append(this.EngineClassName).Append(',');
+						builder.Append('F').Append(this.EngineClassName).Append(',');
 						break;
 
 					case UhtPGetArgumentType.TypeText:
-						Builder.AppendPropertyText(this, UhtPropertyTextType.FunctionThunkParameterArgType).Append(',');
+						builder.AppendPropertyText(this, UhtPropertyTextType.FunctionThunkParameterArgType).Append(',');
 						break;
 				}
 			}
-			Builder.AppendFunctionThunkParameterName(this).Append(')');
-			return Builder;
+			builder.AppendFunctionThunkParameterName(this).Append(')');
+			return builder;
 		}
 
 		/// <summary>
 		/// Append the text for a function thunk call argument
 		/// </summary>
-		/// <param name="Builder">Output builder</param>
+		/// <param name="builder">Output builder</param>
 		/// <returns>Output builder</returns>
-		public virtual StringBuilder AppendFunctionThunkParameterArg(StringBuilder Builder)
+		public virtual StringBuilder AppendFunctionThunkParameterArg(StringBuilder builder)
 		{
-			return Builder.AppendFunctionThunkParameterName(this);
+			return builder.AppendFunctionThunkParameterName(this);
 		}
 
 		/// <summary>
 		/// Apppend the name of a function thunk paramter
 		/// </summary>
-		/// <param name="Builder">Output builder</param>
+		/// <param name="builder">Output builder</param>
 		/// <returns>Output builder</returns>
-		public StringBuilder AppendFunctionThunkParameterName(StringBuilder Builder)
+		public StringBuilder AppendFunctionThunkParameterName(StringBuilder builder)
 		{
-			Builder.Append(FunctionParameterThunkPrefix);
+			builder.Append(FunctionParameterThunkPrefix);
 			if (this.PropertyFlags.HasAnyFlags(EPropertyFlags.OutParm))
 			{
-				Builder.Append("Out_");
+				builder.Append("Out_");
 			}
-			Builder.Append(this.EngineName);
-			return Builder;
+			builder.Append(this.EngineName);
+			return builder;
 		}
 
 		/// <summary>
 		/// Append the appropriate values to initialize the property to a "NULL" value;
 		/// </summary>
-		/// <param name="Builder"></param>
-		/// <param name="bIsInitializer"></param>
+		/// <param name="builder"></param>
+		/// <param name="isInitializer"></param>
 		/// <returns></returns>
-		public abstract StringBuilder AppendNullConstructorArg(StringBuilder Builder, bool bIsInitializer);
+		public abstract StringBuilder AppendNullConstructorArg(StringBuilder builder, bool isInitializer);
 
 		/// <summary>
 		/// Return the basic declaration type text for user facing messages
@@ -1377,9 +1377,9 @@ namespace EpicGames.UHT.Types
 		/// <returns></returns>
 		public string GetUserFacingDecl()
 		{
-			StringBuilder Builder = new StringBuilder();
-			AppendText(Builder, UhtPropertyTextType.Generic);
-			return Builder.ToString();
+			StringBuilder builder = new StringBuilder();
+			AppendText(builder, UhtPropertyTextType.Generic);
+			return builder.ToString();
 		}
 
 		/// <summary>
@@ -1388,20 +1388,20 @@ namespace EpicGames.UHT.Types
 		/// <returns></returns>
 		public string GetRigVMType()
 		{
-			using (BorrowStringBuilder Borrower = new BorrowStringBuilder(StringBuilderCache.Small))
+			using (BorrowStringBuilder borrower = new BorrowStringBuilder(StringBuilderCache.Small))
 			{
-				AppendText(Borrower.StringBuilder, UhtPropertyTextType.RigVMType);
-				return Borrower.StringBuilder.ToString();
+				AppendText(borrower.StringBuilder, UhtPropertyTextType.RigVMType);
+				return borrower.StringBuilder.ToString();
 			}
 		}
 
 		/// <summary>
 		/// Appends any applicable objects and child properties
 		/// </summary>
-		/// <param name="Builder">Output builder</param>
-		/// <param name="StartingLength">Initial length of the builder prior to appending the hashes</param>
-		/// <param name="Context">Context used to lookup the hashes</param>
-		public virtual void AppendObjectHashes(StringBuilder Builder, int StartingLength, IUhtPropertyMemberContext Context)
+		/// <param name="builder">Output builder</param>
+		/// <param name="startingLength">Initial length of the builder prior to appending the hashes</param>
+		/// <param name="context">Context used to lookup the hashes</param>
+		public virtual void AppendObjectHashes(StringBuilder builder, int startingLength, IUhtPropertyMemberContext context)
 		{
 		}
 		#endregion
@@ -1412,68 +1412,68 @@ namespace EpicGames.UHT.Types
 		/// 
 		/// All tokens in the token reader must be consumed.  Otherwise the default value will be considered to be invalid.
 		/// </summary>
-		/// <param name="DefaultValueReader">Reader containing the default value</param>
-		/// <param name="InnerDefaultValue">Sanitized representation of default value</param>
+		/// <param name="defaultValueReader">Reader containing the default value</param>
+		/// <param name="innerDefaultValue">Sanitized representation of default value</param>
 		/// <returns>True if a default value was parsed.</returns>
-		public abstract bool SanitizeDefaultValue(IUhtTokenReader DefaultValueReader, StringBuilder InnerDefaultValue);
+		public abstract bool SanitizeDefaultValue(IUhtTokenReader defaultValueReader, StringBuilder innerDefaultValue);
 		#endregion
 
 		#region Resolution support
 		/// <inheritdoc/>
-		protected override bool ResolveSelf(UhtResolvePhase Phase)
+		protected override bool ResolveSelf(UhtResolvePhase phase)
 		{
-			bool bResult = base.ResolveSelf(Phase);
+			bool result = base.ResolveSelf(phase);
 
-			switch (Phase)
+			switch (phase)
 			{
 				case UhtResolvePhase.Final:
 					if (this.ArrayDimensions != null)
 					{
-						ReadOnlySpan<char> Dim = this.ArrayDimensions.AsSpan();
+						ReadOnlySpan<char> dim = this.ArrayDimensions.AsSpan();
 
-						bool bAgain;
+						bool again;
 						do
 						{
-							bAgain = false;
+							again = false;
 
 							// Remove any outer brackets
-							if (Dim[0] == '(')
+							if (dim[0] == '(')
 							{
-								for (int Index = 1, Depth = 1; Index < Dim.Length; ++Index)
+								for (int index = 1, depth = 1; index < dim.Length; ++index)
 								{
-									if (Dim[Index] == ')')
+									if (dim[index] == ')')
 									{
-										if (--Depth == 0)
+										if (--depth == 0)
 										{
-											if (Index == Dim.Length - 1)
+											if (index == dim.Length - 1)
 											{
-												Dim = Dim.Slice(1, Index - 1);
-												bAgain = true;
+												dim = dim.Slice(1, index - 1);
+												again = true;
 											}
 											break;
 										}
 									}
-									else if (Dim[Index] == '(')
+									else if (dim[index] == '(')
 									{
-										++Depth;
+										++depth;
 									}
 								}
 							}
 
 							// Remove any well known casts
-							if (Dim.Length > 0)
+							if (dim.Length > 0)
 							{
-								foreach (string Cast in Casts)
+								foreach (string cast in s_casts)
 								{
-									if (Dim.StartsWith(Cast))
+									if (dim.StartsWith(cast))
 									{
-										Dim = Dim.Slice(Cast.Length);
-										bAgain = true;
+										dim = dim.Slice(cast.Length);
+										again = true;
 										break;
 									}
 								}
 							}
-						} while (bAgain && Dim.Length > 0);
+						} while (again && dim.Length > 0);
 
 						//COMPATIBILITY-TODO - This method is more robust, but causes differences.  See UhtSession for future
 						// plans on fix in this.
@@ -1502,31 +1502,31 @@ namespace EpicGames.UHT.Types
 						//	this.MetaData.Add(UhtNames.ArraySizeEnum, Enum.GetPathName());
 						//}
 
-						if (Dim.Length > 0 && !UhtFCString.IsDigit(Dim[0]))
+						if (dim.Length > 0 && !UhtFCString.IsDigit(dim[0]))
 						{
-							UhtEnum? Enum = this.Session.FindRegularEnumValue(Dim.ToString());
-							if (Enum == null)
+							UhtEnum? enumObj = this.Session.FindRegularEnumValue(dim.ToString());
+							if (enumObj == null)
 							{
-								Enum = this.Session.FindType(this.Outer, UhtFindOptions.Enum | UhtFindOptions.SourceName, this.ArrayDimensions) as UhtEnum;
+								enumObj = this.Session.FindType(this.Outer, UhtFindOptions.Enum | UhtFindOptions.SourceName, this.ArrayDimensions) as UhtEnum;
 							}
-							if (Enum != null)
+							if (enumObj != null)
 							{
-								this.MetaData.Add(UhtNames.ArraySizeEnum, Enum.PathName);
+								this.MetaData.Add(UhtNames.ArraySizeEnum, enumObj.PathName);
 							}
 						}
 					}
 					break;
 			}
-			return bResult;
+			return result;
 		}
 
 		/// <summary>
 		/// Check properties to see if any instances are referenced.
 		/// This method does NOT cache the result.
 		/// </summary>
-		/// <param name="bDeepScan">If true, the ScanForInstancedReferenced method on the properties will also be called.</param>
+		/// <param name="deepScan">If true, the ScanForInstancedReferenced method on the properties will also be called.</param>
 		/// <returns></returns>
-		public virtual bool ScanForInstancedReferenced(bool bDeepScan)
+		public virtual bool ScanForInstancedReferenced(bool deepScan)
 		{
 			return false;
 		}
@@ -1534,14 +1534,14 @@ namespace EpicGames.UHT.Types
 
 		#region Validation support
 		/// <inheritdoc/>
-		protected override UhtValidationOptions Validate(UhtValidationOptions Options)
+		protected override UhtValidationOptions Validate(UhtValidationOptions options)
 		{
-			Options = base.Validate(Options);
+			options = base.Validate(options);
 
 			// The outer should never be null...
 			if (this.Outer == null)
 			{
-				return Options;
+				return options;
 			}
 
 			// Shadowing checks are done at this level, not in the properties themselves
@@ -1549,59 +1549,59 @@ namespace EpicGames.UHT.Types
 			// In the old code if you defined the property after the function with argument with the same name, UHT would
 			// not issue an error.  However, if the property was defined PRIOR to the function with the matching argument name,
 			// UHT would generate an error.
-			if (Options.HasAnyFlags(UhtValidationOptions.Shadowing) && this.PropertyCategory == UhtPropertyCategory.Member)
+			if (options.HasAnyFlags(UhtValidationOptions.Shadowing) && this.PropertyCategory == UhtPropertyCategory.Member)
 			{
 
 				// First check for duplicate name in self and then duplicate name in parents
-				UhtType? Existing = this.Outer.FindType(UhtFindOptions.PropertyOrFunction | UhtFindOptions.EngineName | UhtFindOptions.SelfOnly, this.EngineName);
-				if (Existing == this)
+				UhtType? existing = this.Outer.FindType(UhtFindOptions.PropertyOrFunction | UhtFindOptions.EngineName | UhtFindOptions.SelfOnly, this.EngineName);
+				if (existing == this)
 				{
-					Existing = this.Outer.FindType(UhtFindOptions.PropertyOrFunction | UhtFindOptions.EngineName | UhtFindOptions.ParentsOnly | UhtFindOptions.NoGlobal | UhtFindOptions.NoIncludes, this.EngineName);
+					existing = this.Outer.FindType(UhtFindOptions.PropertyOrFunction | UhtFindOptions.EngineName | UhtFindOptions.ParentsOnly | UhtFindOptions.NoGlobal | UhtFindOptions.NoIncludes, this.EngineName);
 				}
 
-				if (Existing != null && Existing != this)
+				if (existing != null && existing != this)
 				{
-					if (Existing is UhtProperty ExistingProperty)
+					if (existing is UhtProperty existingProperty)
 					{
 						//@TODO: This exception does not seem sound either, but there is enough existing code that it will need to be
 						// fixed up first before the exception it is removed.
-						bool bExistingPropDeprecated = ExistingProperty.PropertyFlags.HasAnyFlags(EPropertyFlags.Deprecated);
-						bool bNewPropDeprecated = this.PropertyCategory == UhtPropertyCategory.Member && this.PropertyFlags.HasAnyFlags(EPropertyFlags.Deprecated);
-						if (!bNewPropDeprecated && !bExistingPropDeprecated)
+						bool existingPropDeprecated = existingProperty.PropertyFlags.HasAnyFlags(EPropertyFlags.Deprecated);
+						bool newPropDeprecated = this.PropertyCategory == UhtPropertyCategory.Member && this.PropertyFlags.HasAnyFlags(EPropertyFlags.Deprecated);
+						if (!newPropDeprecated && !existingPropDeprecated)
 						{
-							this.LogShadowingError(ExistingProperty);
+							this.LogShadowingError(existingProperty);
 						}
 					}
-					else if (Existing is UhtFunction ExistingFunction)
+					else if (existing is UhtFunction existingFunction)
 					{
 						if (this.PropertyCategory == UhtPropertyCategory.Member)
 						{
-							this.LogShadowingError(ExistingFunction);
+							this.LogShadowingError(existingFunction);
 						}
 					}
 				}
 			}
 
-			Validate((UhtStruct)this.Outer, this, Options);
-			return Options;
+			Validate((UhtStruct)this.Outer, this, options);
+			return options;
 		}
 
-		private void LogShadowingError(UhtType Shadows)
+		private void LogShadowingError(UhtType shadows)
 		{
-			this.LogError($"{this.PropertyCategory.GetHintText()}: '{this.SourceName}' cannot be defined in '{this.Outer?.SourceName}' as it is already defined in scope '{Shadows.Outer?.SourceName}' (shadowing is not allowed)");
+			this.LogError($"{this.PropertyCategory.GetHintText()}: '{this.SourceName}' cannot be defined in '{this.Outer?.SourceName}' as it is already defined in scope '{shadows.Outer?.SourceName}' (shadowing is not allowed)");
 		}
 
 		/// <summary>
 		/// Validate the property settings
 		/// </summary>
-		/// <param name="OuterStruct">The outer structure for the property.  For properties inside containers, this will be the owning struct of the container</param>
-		/// <param name="OutermostProperty">Outer most property being validated.  For properties in containers, 
+		/// <param name="outerStruct">The outer structure for the property.  For properties inside containers, this will be the owning struct of the container</param>
+		/// <param name="outermostProperty">Outer most property being validated.  For properties in containers, 
 		/// this will be the container property.  For properties outside of containers or the container itself, this will be the property.</param>
-		/// <param name="Options"></param>
-		public virtual void Validate(UhtStruct OuterStruct, UhtProperty OutermostProperty, UhtValidationOptions Options)
+		/// <param name="options"></param>
+		public virtual void Validate(UhtStruct outerStruct, UhtProperty outermostProperty, UhtValidationOptions options)
 		{
 			// Check for deprecation
-			if (Options.HasAnyFlags(UhtValidationOptions.Deprecated) && !this.bDeprecated)
+			if (options.HasAnyFlags(UhtValidationOptions.Deprecated) && !this.Deprecated)
 			{
 				ValidateDeprecated();
 			}
@@ -1632,7 +1632,7 @@ namespace EpicGames.UHT.Types
 				}
 			}
 
-			if (!Options.HasAnyFlags(UhtValidationOptions.IsKey) && this.PropertyFlags.HasAnyFlags(EPropertyFlags.PersistentInstance) && !this.PropertyCaps.HasAnyFlags(UhtPropertyCaps.CanBeInstanced))
+			if (!options.HasAnyFlags(UhtValidationOptions.IsKey) && this.PropertyFlags.HasAnyFlags(EPropertyFlags.PersistentInstance) && !this.PropertyCaps.HasAnyFlags(UhtPropertyCaps.CanBeInstanced))
 			{
 				this.LogError("'Instanced' is only allowed on an object property, an array of objects, a set of objects, or a map with an object value type.");
 			}
@@ -1670,23 +1670,23 @@ namespace EpicGames.UHT.Types
 			}
 
 			// Check for invalid transients
-			EPropertyFlags Transients = this.PropertyFlags & (EPropertyFlags.DuplicateTransient | EPropertyFlags.TextExportTransient | EPropertyFlags.NonPIEDuplicateTransient);
-			if (Transients != 0 && !(OuterStruct is UhtClass))
+			EPropertyFlags transients = this.PropertyFlags & (EPropertyFlags.DuplicateTransient | EPropertyFlags.TextExportTransient | EPropertyFlags.NonPIEDuplicateTransient);
+			if (transients != 0 && outerStruct is not UhtClass)
 			{
-				this.LogError($"'{string.Join(", ", Transients.ToStringList(false))}' specifier(s) are only allowed on class member variables");
+				this.LogError($"'{String.Join(", ", transients.ToStringList(false))}' specifier(s) are only allowed on class member variables");
 			}
 
-			if (!Options.HasAnyFlags(UhtValidationOptions.IsKeyOrValue))
+			if (!options.HasAnyFlags(UhtValidationOptions.IsKeyOrValue))
 			{
 				switch (this.PropertyCategory)
 				{
 					case UhtPropertyCategory.Member:
-						ValidateMember(OuterStruct, Options);
+						ValidateMember(outerStruct, options);
 						break;
 
 					case UhtPropertyCategory.RegularParameter:
 					case UhtPropertyCategory.ReplicatedParameter:
-						ValidateFunctionArgument((UhtFunction)OuterStruct, Options);
+						ValidateFunctionArgument((UhtFunction)outerStruct, options);
 						break;
 
 					case UhtPropertyCategory.Return:
@@ -1706,16 +1706,16 @@ namespace EpicGames.UHT.Types
 		/// <summary>
 		/// Verify function argument
 		/// </summary>
-		protected virtual void ValidateFunctionArgument(UhtFunction Func, UhtValidationOptions Options)
+		protected virtual void ValidateFunctionArgument(UhtFunction func, UhtValidationOptions options)
 		{
-			if (Func.FunctionFlags.HasAnyFlags(EFunctionFlags.Net))
+			if (func.FunctionFlags.HasAnyFlags(EFunctionFlags.Net))
 			{
 				if (this.PropertyFlags.HasExactFlags(EPropertyFlags.ReferenceParm | EPropertyFlags.ConstParm, EPropertyFlags.ReferenceParm))
 				{
 					this.LogError($"Replicated parameters cannot be passed by non-const reference");
 				}
 
-				if (Func.FunctionFlags.HasAnyFlags(EFunctionFlags.NetRequest))
+				if (func.FunctionFlags.HasAnyFlags(EFunctionFlags.NetRequest))
 				{
 					if (this.PropertyFlags.HasExactFlags(EPropertyFlags.OutParm | EPropertyFlags.RepSkip, EPropertyFlags.OutParm))
 					{
@@ -1739,10 +1739,10 @@ namespace EpicGames.UHT.Types
 			}
 
 			// The following checks are not performed on the value of a container
-			if (Func.FunctionFlags.HasAnyFlags(EFunctionFlags.BlueprintEvent | EFunctionFlags.BlueprintCallable))
+			if (func.FunctionFlags.HasAnyFlags(EFunctionFlags.BlueprintEvent | EFunctionFlags.BlueprintCallable))
 			{
 				// Check that the parameter name is valid and does not conflict with pre-defined types
-				if (InvalidParamNames.Contains(this.SourceName))
+				if (s_invalidParamNames.Contains(this.SourceName))
 				{
 					this.LogError($"Parameter name '{this.SourceName}' in function is invalid, '{this.SourceName}' is a reserved name.");
 				}
@@ -1752,19 +1752,19 @@ namespace EpicGames.UHT.Types
 		/// <summary>
 		/// Validate member settings
 		/// </summary>
-		/// <param name="Struct">Containing struct.  This is either a UhtScriptStruct or UhtClass</param>
-		/// <param name="Options">Validation options</param>
-		protected virtual void ValidateMember(UhtStruct Struct, UhtValidationOptions Options)
+		/// <param name="structObj">Containing struct.  This is either a UhtScriptStruct or UhtClass</param>
+		/// <param name="options">Validation options</param>
+		protected virtual void ValidateMember(UhtStruct structObj, UhtValidationOptions options)
 		{
-			if (!Options.HasAnyFlags(UhtValidationOptions.IsKeyOrValue))
+			if (!options.HasAnyFlags(UhtValidationOptions.IsKeyOrValue))
 			{
 				// First check if the category was specified at all and if the property was exposed to the editor.
-				string? Category;
-				if (!this.MetaData.TryGetValue(UhtNames.Category, out Category))
+				string? category;
+				if (!this.MetaData.TryGetValue(UhtNames.Category, out category))
 				{
 					if (this.PropertyFlags.HasAnyFlags(EPropertyFlags.Edit | EPropertyFlags.BlueprintVisible))
 					{
-						if (this.Package.bIsPartOfEngine)
+						if (this.Package.IsPartOfEngine)
 						{
 							this.LogError("An explicit Category specifier is required for any property exposed to the editor or Blueprints in an Engine module.");
 						}
@@ -1772,7 +1772,7 @@ namespace EpicGames.UHT.Types
 				}
 
 				// If the category was specified explicitly, it wins
-				if (!string.IsNullOrEmpty(Category) && !this.PropertyFlags.HasAnyFlags(EPropertyFlags.Edit | EPropertyFlags.BlueprintVisible |
+				if (!String.IsNullOrEmpty(category) && !this.PropertyFlags.HasAnyFlags(EPropertyFlags.Edit | EPropertyFlags.BlueprintVisible |
 					EPropertyFlags.BlueprintAssignable | EPropertyFlags.BlueprintCallable))
 				{
 					this.LogWarning("Property has a Category set but is not exposed to the editor or Blueprints with EditAnywhere, BlueprintReadWrite, " +
@@ -1794,21 +1794,21 @@ namespace EpicGames.UHT.Types
 				}
 			}
 
-			string ExposeOnSpawnValue = this.MetaData.GetValueOrDefault(UhtNames.ExposeOnSpawn);
-			if (ExposeOnSpawnValue.Equals("true", StringComparison.OrdinalIgnoreCase) && !this.PropertyCaps.HasAnyFlags(UhtPropertyCaps.CanExposeOnSpawn))
+			string exposeOnSpawnValue = this.MetaData.GetValueOrDefault(UhtNames.ExposeOnSpawn);
+			if (exposeOnSpawnValue.Equals("true", StringComparison.OrdinalIgnoreCase) && !this.PropertyCaps.HasAnyFlags(UhtPropertyCaps.CanExposeOnSpawn))
 			{
 				this.LogError("ExposeOnSpawn - Property cannot be exposed");
 			}
 
 			if (this.PropertyExportFlags.HasAnyFlags(UhtPropertyExportFlags.FieldNotify))
 			{
-				if (this.Outer is UhtClass Class)
+				if (this.Outer is UhtClass classObj)
 				{
-					if (Class.ClassType != UhtClassType.Class)
+					if (classObj.ClassType != UhtClassType.Class)
 					{
 						this.LogError("FieldNofity are not valid on UInterface.");
 					}
-				}	
+				}
 				else
 				{
 					this.LogError("FieldNofity property are only valid as UClass member variable.");
@@ -1819,35 +1819,35 @@ namespace EpicGames.UHT.Types
 		/// <summary>
 		/// Generate an error if the class has been deprecated
 		/// </summary>
-		/// <param name="Class">Class to check</param>
-		protected void ValidateDeprecatedClass(UhtClass? Class)
+		/// <param name="classObj">Class to check</param>
+		protected void ValidateDeprecatedClass(UhtClass? classObj)
 		{
-			if (Class == null)
+			if (classObj == null)
 			{
 				return;
 			}
 
-			if (!Class.bDeprecated)
+			if (!classObj.Deprecated)
 			{
 				return;
 			}
 
 			if (this.PropertyCategory == UhtPropertyCategory.Member)
 			{
-				this.LogError($"Property is using a deprecated class: '{Class.SourceName}'.  Property should be marked deprecated as well.");
+				this.LogError($"Property is using a deprecated class: '{classObj.SourceName}'.  Property should be marked deprecated as well.");
 			}
 			else
 			{
-				this.LogError($"Function is using a deprecated class: '{Class.SourceName}'.  Function should be marked deprecated as well.");
+				this.LogError($"Function is using a deprecated class: '{classObj.SourceName}'.  Function should be marked deprecated as well.");
 			}
 		}
 
 		/// <summary>
 		/// Check to see if the property is valid as a member of a networked structure
 		/// </summary>
-		/// <param name="ReferencingProperty">The property referencing the structure property.  All error should be logged on the referencing property.</param>
+		/// <param name="referencingProperty">The property referencing the structure property.  All error should be logged on the referencing property.</param>
 		/// <returns>True if the property is valid, false if not.  If the property is not valid, an error should also be generated.</returns>
-		public virtual bool ValidateStructPropertyOkForNet(UhtProperty ReferencingProperty)
+		public virtual bool ValidateStructPropertyOkForNet(UhtProperty referencingProperty)
 		{
 			return true;
 		}
@@ -1855,23 +1855,23 @@ namespace EpicGames.UHT.Types
 		/// <summary>
 		/// Test to see if this property references something that requires the argument to be marked as const
 		/// </summary>
-		/// <param name="ErrorType">If const is required, returns the type that is forcing the const</param>
+		/// <param name="errorType">If const is required, returns the type that is forcing the const</param>
 		/// <returns>True if the argument must be marked as const</returns>
-		public virtual bool MustBeConstArgument([NotNullWhen(true)] out UhtType? ErrorType)
+		public virtual bool MustBeConstArgument([NotNullWhen(true)] out UhtType? errorType)
 		{
-			ErrorType = null;
+			errorType = null;
 			return false;
 		}
 		#endregion
 
 		#region Reference support
 		/// <inheritdoc/>
-		public override void CollectReferences(IUhtReferenceCollector Collector)
+		public override void CollectReferences(IUhtReferenceCollector collector)
 		{
-			CollectReferencesInternal(Collector, false);
+			CollectReferencesInternal(collector, false);
 			if (this.PropertyFlags.HasAnyFlags(EPropertyFlags.ParmFlags))
 			{
-				Collector.AddForwardDeclaration(this.GetForwardDeclarations());
+				collector.AddForwardDeclaration(this.GetForwardDeclarations());
 			}
 		}
 
@@ -1879,9 +1879,9 @@ namespace EpicGames.UHT.Types
 		/// Collect the references for the property.  This method is used by container properties to
 		/// collect the contained property's references.
 		/// </summary>
-		/// <param name="Collector">Reference collector</param>
-		/// <param name="bIsTemplateProperty">If true, this is a property in a container</param>
-		public virtual void CollectReferencesInternal(IUhtReferenceCollector Collector, bool bIsTemplateProperty)
+		/// <param name="collector">Reference collector</param>
+		/// <param name="isTemplateProperty">If true, this is a property in a container</param>
+		public virtual void CollectReferencesInternal(IUhtReferenceCollector collector, bool isTemplateProperty)
 		{
 		}
 
@@ -1910,45 +1910,45 @@ namespace EpicGames.UHT.Types
 		/// <summary>
 		/// Generate a new name suffix based on the current suffix and the new suffix
 		/// </summary>
-		/// <param name="OuterSuffix">Current suffix</param>
-		/// <param name="NewSuffix">Suffix to be added</param>
+		/// <param name="outerSuffix">Current suffix</param>
+		/// <param name="newSuffix">Suffix to be added</param>
 		/// <returns>Combination of the two suffixes</returns>
-		protected static string GetNameSuffix(string OuterSuffix, string NewSuffix)
+		protected static string GetNameSuffix(string outerSuffix, string newSuffix)
 		{
-			return string.IsNullOrEmpty(OuterSuffix) ? NewSuffix : $"{OuterSuffix}{NewSuffix}";
+			return String.IsNullOrEmpty(outerSuffix) ? newSuffix : $"{outerSuffix}{newSuffix}";
 		}
 
 		/// <summary>
 		/// Test to see if the two properties are the same type
 		/// </summary>
-		/// <param name="Other">Other property to test</param>
+		/// <param name="other">Other property to test</param>
 		/// <returns>True if the properies are the same type</returns>
-		public abstract bool IsSameType(UhtProperty Other);
+		public abstract bool IsSameType(UhtProperty other);
 
 		/// <summary>
 		/// Test to see if the two properties are the same type and ConstParm/OutParm flags somewhat match
 		/// </summary>
-		/// <param name="Other">The other property to test</param>
+		/// <param name="other">The other property to test</param>
 		/// <returns></returns>
-		public bool MatchesType(UhtProperty Other)
+		public bool MatchesType(UhtProperty other)
 		{
 			if (this.PropertyFlags.HasAnyFlags(EPropertyFlags.OutParm))
 			{
-				if (!Other.PropertyFlags.HasAnyFlags(EPropertyFlags.OutParm))
+				if (!other.PropertyFlags.HasAnyFlags(EPropertyFlags.OutParm))
 				{
 					return false;
 				}
 
-				if (this.PropertyFlags.HasAnyFlags(EPropertyFlags.ConstParm) && !Other.PropertyFlags.HasAnyFlags(EPropertyFlags.ConstParm))
+				if (this.PropertyFlags.HasAnyFlags(EPropertyFlags.ConstParm) && !other.PropertyFlags.HasAnyFlags(EPropertyFlags.ConstParm))
 				{
 					return false;
 				}
 			}
-			if (this.bIsStaticArray != Other.bIsStaticArray)
+			if (this.IsStaticArray != other.IsStaticArray)
 			{
 				return false;
 			}
-			return IsSameType(Other);
+			return IsSameType(other);
 		}
 		#endregion
 	}
@@ -1962,305 +1962,305 @@ namespace EpicGames.UHT.Types
 		/// <summary>
 		/// Add the given property text
 		/// </summary>
-		/// <param name="Builder">Destination builder</param>
-		/// <param name="Property">Property in question</param>
-		/// <param name="TextType">Type of text to append</param>
-		/// <param name="bIsTemplateArgument">If true, this is a template argument property</param>
+		/// <param name="builder">Destination builder</param>
+		/// <param name="property">Property in question</param>
+		/// <param name="textType">Type of text to append</param>
+		/// <param name="isTemplateArgument">If true, this is a template argument property</param>
 		/// <returns>Destination builder</returns>
-		public static StringBuilder AppendPropertyText(this StringBuilder Builder, UhtProperty Property, UhtPropertyTextType TextType, bool bIsTemplateArgument = false)
+		public static StringBuilder AppendPropertyText(this StringBuilder builder, UhtProperty property, UhtPropertyTextType textType, bool isTemplateArgument = false)
 		{
-			return Property.AppendText(Builder, TextType, bIsTemplateArgument);
+			return property.AppendText(builder, textType, isTemplateArgument);
 		}
 
 		/// <summary>
 		/// Append the member declaration
 		/// </summary>
-		/// <param name="Builder">Destination builder</param>
-		/// <param name="Property">Property in question</param>
-		/// <param name="Context">Context of the property</param>
-		/// <param name="Name">Property name</param>
-		/// <param name="NameSuffix">Name suffix</param>
-		/// <param name="Tabs">Number of tabs in the formatting</param>
+		/// <param name="builder">Destination builder</param>
+		/// <param name="property">Property in question</param>
+		/// <param name="context">Context of the property</param>
+		/// <param name="name">Property name</param>
+		/// <param name="nameSuffix">Name suffix</param>
+		/// <param name="tabs">Number of tabs in the formatting</param>
 		/// <returns>Destination builder</returns>
-		public static StringBuilder AppendMemberDecl(this StringBuilder Builder, UhtProperty Property, IUhtPropertyMemberContext Context, string Name, string NameSuffix, int Tabs)
+		public static StringBuilder AppendMemberDecl(this StringBuilder builder, UhtProperty property, IUhtPropertyMemberContext context, string name, string nameSuffix, int tabs)
 		{
-			return Property.AppendMemberDecl(Builder, Context, Name, NameSuffix, Tabs);
+			return property.AppendMemberDecl(builder, context, name, nameSuffix, tabs);
 		}
 
 		/// <summary>
 		/// Append the member definition
 		/// </summary>
-		/// <param name="Builder">Destination builder</param>
-		/// <param name="Property">Property in question</param>
-		/// <param name="Context">Context of the property</param>
-		/// <param name="Name">Property name</param>
-		/// <param name="NameSuffix">Name suffix</param>
-		/// <param name="Offset">Offset of the property in the parent</param>
-		/// <param name="Tabs">Number of tabs in the formatting</param>
+		/// <param name="builder">Destination builder</param>
+		/// <param name="property">Property in question</param>
+		/// <param name="context">Context of the property</param>
+		/// <param name="name">Property name</param>
+		/// <param name="nameSuffix">Name suffix</param>
+		/// <param name="offset">Offset of the property in the parent</param>
+		/// <param name="tabs">Number of tabs in the formatting</param>
 		/// <returns>Destination builder</returns>
-		public static StringBuilder AppendMemberDef(this StringBuilder Builder, UhtProperty Property, IUhtPropertyMemberContext Context, string Name, string NameSuffix, string? Offset, int Tabs)
+		public static StringBuilder AppendMemberDef(this StringBuilder builder, UhtProperty property, IUhtPropertyMemberContext context, string name, string nameSuffix, string? offset, int tabs)
 		{
-			return Property.AppendMemberDef(Builder, Context, Name, NameSuffix, Offset, Tabs);
+			return property.AppendMemberDef(builder, context, name, nameSuffix, offset, tabs);
 		}
 
 		/// <summary>
 		/// Append the member pointer
 		/// </summary>
-		/// <param name="Builder">Destination builder</param>
-		/// <param name="Property">Property in question</param>
-		/// <param name="Context">Context of the property</param>
-		/// <param name="Name">Property name</param>
-		/// <param name="NameSuffix">Name suffix</param>
-		/// <param name="Tabs">Number of tabs in the formatting</param>
+		/// <param name="builder">Destination builder</param>
+		/// <param name="property">Property in question</param>
+		/// <param name="context">Context of the property</param>
+		/// <param name="name">Property name</param>
+		/// <param name="nameSuffix">Name suffix</param>
+		/// <param name="tabs">Number of tabs in the formatting</param>
 		/// <returns>Destination builder</returns>
-		public static StringBuilder AppendMemberPtr(this StringBuilder Builder, UhtProperty Property, IUhtPropertyMemberContext Context, string Name, string NameSuffix, int Tabs)
+		public static StringBuilder AppendMemberPtr(this StringBuilder builder, UhtProperty property, IUhtPropertyMemberContext context, string name, string nameSuffix, int tabs)
 		{
-			return Property.AppendMemberPtr(Builder, Context, Name, NameSuffix, Tabs);
+			return property.AppendMemberPtr(builder, context, name, nameSuffix, tabs);
 		}
 
 		/// <summary>
 		/// Append the full declaration including such things as const, *, and &amp;
 		/// </summary>
-		/// <param name="Builder">Destination builder</param>
-		/// <param name="Property">Property in question</param>
-		/// <param name="TextType">Type of text to append</param>
-		/// <param name="bSkipParameterName">If true, don't append the parameter name</param>
+		/// <param name="builder">Destination builder</param>
+		/// <param name="property">Property in question</param>
+		/// <param name="textType">Type of text to append</param>
+		/// <param name="skipParameterName">If true, don't append the parameter name</param>
 		/// <returns>Destination builder</returns>
-		public static StringBuilder AppendFullDecl(this StringBuilder Builder, UhtProperty Property, UhtPropertyTextType TextType, bool bSkipParameterName)
+		public static StringBuilder AppendFullDecl(this StringBuilder builder, UhtProperty property, UhtPropertyTextType textType, bool skipParameterName)
 		{
-			return Property.AppendFullDecl(Builder, TextType, bSkipParameterName);
+			return property.AppendFullDecl(builder, textType, skipParameterName);
 		}
 
 		/// <summary>
 		/// Append the function thunk parameter get
 		/// </summary>
-		/// <param name="Builder">Destination builder</param>
-		/// <param name="Property">Property in question</param>
+		/// <param name="builder">Destination builder</param>
+		/// <param name="property">Property in question</param>
 		/// <returns>Destination builder</returns>
-		public static StringBuilder AppendFunctionThunkParameterGet(this StringBuilder Builder, UhtProperty Property)
+		public static StringBuilder AppendFunctionThunkParameterGet(this StringBuilder builder, UhtProperty property)
 		{
-			return Property.AppendFunctionThunkParameterGet(Builder);
+			return property.AppendFunctionThunkParameterGet(builder);
 		}
 
 		/// <summary>
 		/// Append the function thunk parameter array type
 		/// </summary>
-		/// <param name="Builder">Destination builder</param>
-		/// <param name="Property">Property in question</param>
+		/// <param name="builder">Destination builder</param>
+		/// <param name="property">Property in question</param>
 		/// <returns>Destination builder</returns>
-		public static StringBuilder AppendFunctionThunkParameterArrayType(this StringBuilder Builder, UhtProperty Property)
+		public static StringBuilder AppendFunctionThunkParameterArrayType(this StringBuilder builder, UhtProperty property)
 		{
-			return Property.AppendText(Builder, UhtPropertyTextType.FunctionThunkParameterArrayType);
+			return property.AppendText(builder, UhtPropertyTextType.FunctionThunkParameterArrayType);
 		}
 
 		/// <summary>
 		/// Append the function thunk parameter argument
 		/// </summary>
-		/// <param name="Builder">Destination builder</param>
-		/// <param name="Property">Property in question</param>
+		/// <param name="builder">Destination builder</param>
+		/// <param name="property">Property in question</param>
 		/// <returns>Destination builder</returns>
-		public static StringBuilder AppendFunctionThunkParameterArg(this StringBuilder Builder, UhtProperty Property)
+		public static StringBuilder AppendFunctionThunkParameterArg(this StringBuilder builder, UhtProperty property)
 		{
-			return Property.AppendFunctionThunkParameterArg(Builder);
+			return property.AppendFunctionThunkParameterArg(builder);
 		}
 
 		/// <summary>
 		/// Append the function thunk return
 		/// </summary>
-		/// <param name="Builder">Destination builder</param>
-		/// <param name="Property">Property in question</param>
+		/// <param name="builder">Destination builder</param>
+		/// <param name="property">Property in question</param>
 		/// <returns>Destination builder</returns>
-		public static StringBuilder AppendFunctionThunkReturn(this StringBuilder Builder, UhtProperty Property)
+		public static StringBuilder AppendFunctionThunkReturn(this StringBuilder builder, UhtProperty property)
 		{
-			return Property.AppendText(Builder, UhtPropertyTextType.FunctionThunkRetVal);
+			return property.AppendText(builder, UhtPropertyTextType.FunctionThunkRetVal);
 		}
 
 		/// <summary>
 		/// Append the function thunk parameter name
 		/// </summary>
-		/// <param name="Builder">Destination builder</param>
-		/// <param name="Property">Property in question</param>
+		/// <param name="builder">Destination builder</param>
+		/// <param name="property">Property in question</param>
 		/// <returns>Destination builder</returns>
-		public static StringBuilder AppendFunctionThunkParameterName(this StringBuilder Builder, UhtProperty Property)
+		public static StringBuilder AppendFunctionThunkParameterName(this StringBuilder builder, UhtProperty property)
 		{
-			return Property.AppendFunctionThunkParameterName(Builder);
+			return property.AppendFunctionThunkParameterName(builder);
 		}
 
 		/// <summary>
 		/// Append the sparse type
 		/// </summary>
-		/// <param name="Builder">Destination builder</param>
-		/// <param name="Property">Property in question</param>
+		/// <param name="builder">Destination builder</param>
+		/// <param name="property">Property in question</param>
 		/// <returns>Destination builder</returns>
-		public static StringBuilder AppendSparse(this StringBuilder Builder, UhtProperty Property)
+		public static StringBuilder AppendSparse(this StringBuilder builder, UhtProperty property)
 		{
-			return Property.AppendText(Builder, UhtPropertyTextType.Sparse);
+			return property.AppendText(builder, UhtPropertyTextType.Sparse);
 		}
 
 		/// <summary>
 		/// Append the property's null constructor arg
 		/// </summary>
-		/// <param name="Builder">Destination builder</param>
-		/// <param name="Property">Property in question</param>
-		/// <param name="bIsInitializer">If true this is in an initializer context.</param>
+		/// <param name="builder">Destination builder</param>
+		/// <param name="property">Property in question</param>
+		/// <param name="isInitializer">If true this is in an initializer context.</param>
 		/// <returns>Destination builder</returns>
-		public static StringBuilder AppendNullConstructorArg(this StringBuilder Builder, UhtProperty Property, bool bIsInitializer)
+		public static StringBuilder AppendNullConstructorArg(this StringBuilder builder, UhtProperty property, bool isInitializer)
 		{
-			Property.AppendNullConstructorArg(Builder, bIsInitializer);
-			return Builder;
+			property.AppendNullConstructorArg(builder, isInitializer);
+			return builder;
 		}
 
 		/// <summary>
 		/// Append the replication notify function or a 'nullptr'
 		/// </summary>
-		/// <param name="Builder">Destination builder</param>
-		/// <param name="Property">Property in question</param>
+		/// <param name="builder">Destination builder</param>
+		/// <param name="property">Property in question</param>
 		/// <returns>Destination builder</returns>
-		public static StringBuilder AppendNotifyFunc(this StringBuilder Builder, UhtProperty Property)
+		public static StringBuilder AppendNotifyFunc(this StringBuilder builder, UhtProperty property)
 		{
-			if (Property.RepNotifyName != null)
+			if (property.RepNotifyName != null)
 			{
-				Builder.AppendUTF8LiteralString(Property.RepNotifyName);
+				builder.AppendUTF8LiteralString(property.RepNotifyName);
 			}
 			else
 			{
-				Builder.Append("nullptr");
+				builder.Append("nullptr");
 			}
-			return Builder;
+			return builder;
 		}
 
 		/// <summary>
 		/// Append the parameter flags
 		/// </summary>
-		/// <param name="Builder">Destination builder</param>
-		/// <param name="PropertyFlags">Property flags</param>
+		/// <param name="builder">Destination builder</param>
+		/// <param name="propertyFlags">Property flags</param>
 		/// <returns>Destination builder</returns>
-		public static StringBuilder AppendFlags(this StringBuilder Builder, EPropertyFlags PropertyFlags)
+		public static StringBuilder AppendFlags(this StringBuilder builder, EPropertyFlags propertyFlags)
 		{
-			PropertyFlags &= ~EPropertyFlags.ComputedFlags;
-			return Builder.Append("(EPropertyFlags)0x").AppendFormat("{0:x16}", (ulong)PropertyFlags);
+			propertyFlags &= ~EPropertyFlags.ComputedFlags;
+			return builder.Append("(EPropertyFlags)0x").AppendFormat("{0:x16}", (ulong)propertyFlags);
 		}
 
 		/// <summary>
 		/// Append the property array dimensions as a CPP_ARRAY_DIM macro or '1' if not a fixed array.
 		/// </summary>
-		/// <param name="Builder">Destination builder</param>
-		/// <param name="Property">Property in question</param>
-		/// <param name="Context">Context of the property</param>
+		/// <param name="builder">Destination builder</param>
+		/// <param name="property">Property in question</param>
+		/// <param name="context">Context of the property</param>
 		/// <returns>Destination builder</returns>
-		public static StringBuilder AppendArrayDim(this StringBuilder Builder, UhtProperty Property, IUhtPropertyMemberContext Context)
+		public static StringBuilder AppendArrayDim(this StringBuilder builder, UhtProperty property, IUhtPropertyMemberContext context)
 		{
-			if (Property.ArrayDimensions != null)
+			if (property.ArrayDimensions != null)
 			{
-				Builder.Append("CPP_ARRAY_DIM(").Append(Property.SourceName).Append(", ").Append(Context.OuterStruct.SourceName).Append(')');
+				builder.Append("CPP_ARRAY_DIM(").Append(property.SourceName).Append(", ").Append(context.OuterStruct.SourceName).Append(')');
 			}
 			else
 			{
-				Builder.Append('1');
+				builder.Append('1');
 			}
-			return Builder;
+			return builder;
 		}
 
 		/// <summary>
 		/// Given an object, append the hash (if applicable) to the builder
 		/// </summary>
-		/// <param name="Builder">Output builder</param>
-		/// <param name="ReferingType">Type asking for an object hash</param>
-		/// <param name="StartingLength">Initial length of the builder prior to appending the hashes</param>
-		/// <param name="Context">Context used to lookup the hashes</param>
-		/// <param name="Object">Object being appended</param>
+		/// <param name="builder">Output builder</param>
+		/// <param name="referingType">Type asking for an object hash</param>
+		/// <param name="startingLength">Initial length of the builder prior to appending the hashes</param>
+		/// <param name="context">Context used to lookup the hashes</param>
+		/// <param name="obj">Object being appended</param>
 		/// <returns>Destination builder</returns>
-		public static StringBuilder AppendObjectHash(this StringBuilder Builder, int StartingLength, UhtType ReferingType, IUhtPropertyMemberContext Context, UhtObject? Object)
+		public static StringBuilder AppendObjectHash(this StringBuilder builder, int startingLength, UhtType referingType, IUhtPropertyMemberContext context, UhtObject? obj)
 		{
-			if (Object == null)
+			if (obj == null)
 			{
-				return Builder;
+				return builder;
 			}
-			else if (Object is UhtClass Class)
+			else if (obj is UhtClass classObj)
 			{
-				if (Class.ClassExportFlags.HasAnyFlags(UhtClassExportFlags.NoExport))
+				if (classObj.ClassExportFlags.HasAnyFlags(UhtClassExportFlags.NoExport))
 				{
-					return Builder;
+					return builder;
 				}
 			}
-			else if (Object is UhtScriptStruct ScriptStruct)
+			else if (obj is UhtScriptStruct scriptStruct)
 			{
-				if (ScriptStruct.ScriptStructFlags.HasAnyFlags(EStructFlags.NoExport))
+				if (scriptStruct.ScriptStructFlags.HasAnyFlags(EStructFlags.NoExport))
 				{
-					return Builder;
+					return builder;
 				}
 			}
 
-			Builder.Append(Builder.Length == StartingLength ? " // " : " ");
-			uint Hash = Context.GetTypeHash(Object);
-			if (Hash == 0)
+			builder.Append(builder.Length == startingLength ? " // " : " ");
+			uint hash = context.GetTypeHash(obj);
+			if (hash == 0)
 			{
-				string Type = ReferingType is UhtProperty ? "property" : "object";
-				ReferingType.LogWarning($"The {Type} \"{ReferingType.SourceName}\" references type \"{Object.SourceName}\" but the code generation hash is zero");
+				string type = referingType is UhtProperty ? "property" : "object";
+				referingType.LogWarning($"The {type} \"{referingType.SourceName}\" references type \"{obj.SourceName}\" but the code generation hash is zero");
 			}
-			Builder.Append(Context.GetTypeHash(Object));
-			return Builder;
+			builder.Append(context.GetTypeHash(obj));
+			return builder;
 		}
 
 		/// <summary>
 		/// Given an object, append the hash (if applicable) to the builder
 		/// </summary>
-		/// <param name="Builder">Output builder</param>
-		/// <param name="ReferingType">Type asking for an object hash</param>
-		/// <param name="Context">Context used to lookup the hashes</param>
-		/// <param name="Object">Object being appended</param>
+		/// <param name="builder">Output builder</param>
+		/// <param name="referingType">Type asking for an object hash</param>
+		/// <param name="context">Context used to lookup the hashes</param>
+		/// <param name="obj">Object being appended</param>
 		/// <returns>Destination builder</returns>
-		public static StringBuilder AppendObjectHash(this StringBuilder Builder, UhtType ReferingType, IUhtPropertyMemberContext Context, UhtObject? Object)
+		public static StringBuilder AppendObjectHash(this StringBuilder builder, UhtType referingType, IUhtPropertyMemberContext context, UhtObject? obj)
 		{
-			return Builder.AppendObjectHash(Builder.Length, ReferingType, Context, Object);
+			return builder.AppendObjectHash(builder.Length, referingType, context, obj);
 		}
 
 		/// <summary>
 		/// Append the object hashes for all referenced objects
 		/// </summary>
-		/// <param name="Builder">Destination builder</param>
-		/// <param name="Property">Property in question</param>
-		/// <param name="Context">Context of the property</param>
+		/// <param name="builder">Destination builder</param>
+		/// <param name="property">Property in question</param>
+		/// <param name="context">Context of the property</param>
 		/// <returns>Destination builder</returns>
-		public static StringBuilder AppendObjectHashes(this StringBuilder Builder, UhtProperty Property, IUhtPropertyMemberContext Context)
+		public static StringBuilder AppendObjectHashes(this StringBuilder builder, UhtProperty property, IUhtPropertyMemberContext context)
 		{
-			Property.AppendObjectHashes(Builder, Builder.Length, Context);
-			return Builder;
+			property.AppendObjectHashes(builder, builder.Length, context);
+			return builder;
 		}
 
 		/// <summary>
 		/// Append the singleton name for the given type
 		/// </summary>
-		/// <param name="Builder">Destination builder</param>
-		/// <param name="Context">Context of the property</param>
-		/// <param name="Type">Type to append</param>
-		/// <param name="bRegistered">If true, append the registered type singleton.</param>
+		/// <param name="builder">Destination builder</param>
+		/// <param name="context">Context of the property</param>
+		/// <param name="type">Type to append</param>
+		/// <param name="registered">If true, append the registered type singleton.</param>
 		/// <returns>Destination builder</returns>
-		public static StringBuilder AppendSingletonName(this StringBuilder Builder, IUhtPropertyMemberContext Context, UhtObject? Type, bool bRegistered)
+		public static StringBuilder AppendSingletonName(this StringBuilder builder, IUhtPropertyMemberContext context, UhtObject? type, bool registered)
 		{
-			return Builder.Append(Context.GetSingletonName(Type, bRegistered));
+			return builder.Append(context.GetSingletonName(type, registered));
 		}
 
 		/// <summary>
 		/// Append the getter wrapper name
 		/// </summary>
-		/// <param name="Builder">Destination builder</param>
-		/// <param name="Property">Property in question</param>
+		/// <param name="builder">Destination builder</param>
+		/// <param name="property">Property in question</param>
 		/// <returns>Destination builder</returns>
-		public static StringBuilder AppendPropertyGetterWrapperName(this StringBuilder Builder, UhtProperty Property)
+		public static StringBuilder AppendPropertyGetterWrapperName(this StringBuilder builder, UhtProperty property)
 		{
-			return Builder.Append("Get").Append(Property.SourceName).Append("_WrapperImpl");
+			return builder.Append("Get").Append(property.SourceName).Append("_WrapperImpl");
 		}
 
 		/// <summary>
 		/// Append the setter wrapper name
 		/// </summary>
-		/// <param name="Builder">Destination builder</param>
-		/// <param name="Property">Property in question</param>
+		/// <param name="builder">Destination builder</param>
+		/// <param name="property">Property in question</param>
 		/// <returns>Destination builder</returns>
-		public static StringBuilder AppendPropertySetterWrapperName(this StringBuilder Builder, UhtProperty Property)
+		public static StringBuilder AppendPropertySetterWrapperName(this StringBuilder builder, UhtProperty property)
 		{
-			return Builder.Append("Set").Append(Property.SourceName).Append("_WrapperImpl");
+			return builder.Append("Set").Append(property.SourceName).Append("_WrapperImpl");
 		}
 	}
 }
