More coding standard fixes

diff --git a/Engine/Source/Programs/Shared/EpicGames.UHT/Types/UhtScriptStruct.cs b/Engine/Source/Programs/Shared/EpicGames.UHT/Types/UhtScriptStruct.cs
--- a/Engine/Source/Programs/Shared/EpicGames.UHT/Types/UhtScriptStruct.cs
+++ b/Engine/Source/Programs/Shared/EpicGames.UHT/Types/UhtScriptStruct.cs
@@ -1,11 +1,11 @@
 // Copyright Epic Games, Inc. All Rights Reserved.
 
-using EpicGames.Core;
-using EpicGames.UHT.Tables;
-using EpicGames.UHT.Utils;
 using System;
 using System.Collections.Generic;
 using System.Text.Json.Serialization;
+using EpicGames.Core;
+using EpicGames.UHT.Tables;
+using EpicGames.UHT.Utils;
 
 namespace EpicGames.UHT.Types
 {
@@ -13,7 +13,7 @@ namespace EpicGames.UHT.Types
 	/// Flags to represent information about a RigVM parameter
 	/// </summary>
 	[Flags]
-	public enum UhtRigVMParameterFlags : Int32
+	public enum UhtRigVMParameterFlags : int
 	{
 		/// <summary>
 		/// No RigVM flags
@@ -43,7 +43,7 @@ namespace EpicGames.UHT.Types
 		/// <summary>
 		/// Set if the property is editor only
 		/// </summary>
-		EditorOnly		= 0x00000010,
+		EditorOnly = 0x00000010,
 
 		/// <summary>
 		/// Set if the property is an enum
@@ -66,35 +66,35 @@ namespace EpicGames.UHT.Types
 		/// <summary>
 		/// Test to see if any of the specified flags are set
 		/// </summary>
-		/// <param name="InFlags">Current flags</param>
-		/// <param name="TestFlags">Flags to test for</param>
+		/// <param name="inFlags">Current flags</param>
+		/// <param name="testFlags">Flags to test for</param>
 		/// <returns>True if any of the flags are set</returns>
-		public static bool HasAnyFlags(this UhtRigVMParameterFlags InFlags, UhtRigVMParameterFlags TestFlags)
+		public static bool HasAnyFlags(this UhtRigVMParameterFlags inFlags, UhtRigVMParameterFlags testFlags)
 		{
-			return (InFlags & TestFlags) != 0;
+			return (inFlags & testFlags) != 0;
 		}
 
 		/// <summary>
 		/// Test to see if all of the specified flags are set
 		/// </summary>
-		/// <param name="InFlags">Current flags</param>
-		/// <param name="TestFlags">Flags to test for</param>
+		/// <param name="inFlags">Current flags</param>
+		/// <param name="testFlags">Flags to test for</param>
 		/// <returns>True if all the flags are set</returns>
-		public static bool HasAllFlags(this UhtRigVMParameterFlags InFlags, UhtRigVMParameterFlags TestFlags)
+		public static bool HasAllFlags(this UhtRigVMParameterFlags inFlags, UhtRigVMParameterFlags testFlags)
 		{
-			return (InFlags & TestFlags) == TestFlags;
+			return (inFlags & testFlags) == testFlags;
 		}
 
 		/// <summary>
 		/// Test to see if a specific set of flags have a specific value.
 		/// </summary>
-		/// <param name="InFlags">Current flags</param>
-		/// <param name="TestFlags">Flags to test for</param>
-		/// <param name="MatchFlags">Expected value of the tested flags</param>
+		/// <param name="inFlags">Current flags</param>
+		/// <param name="testFlags">Flags to test for</param>
+		/// <param name="matchFlags">Expected value of the tested flags</param>
 		/// <returns>True if the given flags have a specific value.</returns>
-		public static bool HasExactFlags(this UhtRigVMParameterFlags InFlags, UhtRigVMParameterFlags TestFlags, UhtRigVMParameterFlags MatchFlags)
+		public static bool HasExactFlags(this UhtRigVMParameterFlags inFlags, UhtRigVMParameterFlags testFlags, UhtRigVMParameterFlags matchFlags)
 		{
-			return (InFlags & TestFlags) == MatchFlags;
+			return (inFlags & testFlags) == matchFlags;
 		}
 	}
 
@@ -140,243 +140,243 @@ namespace EpicGames.UHT.Types
 		/// True if the parameter is marked as "Constant"
 		/// </summary>
 		[JsonIgnore]
-		public bool bConstant => this.ParameterFlags.HasAnyFlags(UhtRigVMParameterFlags.Constant);
+		public bool Constant => this.ParameterFlags.HasAnyFlags(UhtRigVMParameterFlags.Constant);
 
 		/// <summary>
 		/// True if the parameter is marked as "Input"
 		/// </summary>
 		[JsonIgnore]
-		public bool bInput => this.ParameterFlags.HasAnyFlags(UhtRigVMParameterFlags.Input);
+		public bool Input => this.ParameterFlags.HasAnyFlags(UhtRigVMParameterFlags.Input);
 
 		/// <summary>
 		/// True if the parameter is marked as "Output"
 		/// </summary>
 		[JsonIgnore]
-		public bool bOutput => this.ParameterFlags.HasAnyFlags(UhtRigVMParameterFlags.Output);
+		public bool Output => this.ParameterFlags.HasAnyFlags(UhtRigVMParameterFlags.Output);
 
 		/// <summary>
 		/// True if the parameter is marked as "Singleton"
 		/// </summary>
 		[JsonIgnore]
-		public bool bSingleton => this.ParameterFlags.HasAnyFlags(UhtRigVMParameterFlags.Singleton);
+		public bool Singleton => this.ParameterFlags.HasAnyFlags(UhtRigVMParameterFlags.Singleton);
 
 		/// <summary>
 		/// True if the parameter is editor only
 		/// </summary>
 		[JsonIgnore]
-		public bool bEditorOnly => this.ParameterFlags.HasAnyFlags(UhtRigVMParameterFlags.EditorOnly);
+		public bool EditorOnly => this.ParameterFlags.HasAnyFlags(UhtRigVMParameterFlags.EditorOnly);
 
 		/// <summary>
 		/// True if the parameter is an enum
 		/// </summary>
 		[JsonIgnore]
-		public bool bIsEnum => this.ParameterFlags.HasAnyFlags(UhtRigVMParameterFlags.IsEnum);
+		public bool IsEnum => this.ParameterFlags.HasAnyFlags(UhtRigVMParameterFlags.IsEnum);
 
 		/// <summary>
 		/// True if the parameter is an array
 		/// </summary>
 		[JsonIgnore]
-		public bool bIsArray => this.ParameterFlags.HasAnyFlags(UhtRigVMParameterFlags.IsArray);
+		public bool IsArray => this.ParameterFlags.HasAnyFlags(UhtRigVMParameterFlags.IsArray);
 
 		/// <summary>
 		/// Create a new RigVM parameter from a property
 		/// </summary>
-		/// <param name="Property">Source property</param>
-		/// <param name="Index">Parameter index.  Used to create a unique cast name.</param>
-		public UhtRigVMParameter(UhtProperty Property, int Index)
+		/// <param name="property">Source property</param>
+		/// <param name="index">Parameter index.  Used to create a unique cast name.</param>
+		public UhtRigVMParameter(UhtProperty property, int index)
 		{
-			this.Property = Property;
+			this.Property = property;
 
-			this.Name = Property.EngineName;
-			this.ParameterFlags |= Property.MetaData.ContainsKey("Constant") ? UhtRigVMParameterFlags.Constant : UhtRigVMParameterFlags.None;
-			this.ParameterFlags |= Property.MetaData.ContainsKey("Input") ? UhtRigVMParameterFlags.Input : UhtRigVMParameterFlags.None;
-			this.ParameterFlags |= Property.MetaData.ContainsKey("Output") ? UhtRigVMParameterFlags.Output : UhtRigVMParameterFlags.None;
-			this.ParameterFlags |= Property.bIsEditorOnlyProperty ? UhtRigVMParameterFlags.EditorOnly : UhtRigVMParameterFlags.None;
-			this.ParameterFlags |= Property.MetaData.ContainsKey("Singleton") ? UhtRigVMParameterFlags.Singleton : UhtRigVMParameterFlags.None;
+			this.Name = property.EngineName;
+			this.ParameterFlags |= property.MetaData.ContainsKey("Constant") ? UhtRigVMParameterFlags.Constant : UhtRigVMParameterFlags.None;
+			this.ParameterFlags |= property.MetaData.ContainsKey("Input") ? UhtRigVMParameterFlags.Input : UhtRigVMParameterFlags.None;
+			this.ParameterFlags |= property.MetaData.ContainsKey("Output") ? UhtRigVMParameterFlags.Output : UhtRigVMParameterFlags.None;
+			this.ParameterFlags |= property.IsEditorOnlyProperty ? UhtRigVMParameterFlags.EditorOnly : UhtRigVMParameterFlags.None;
+			this.ParameterFlags |= property.MetaData.ContainsKey("Singleton") ? UhtRigVMParameterFlags.Singleton : UhtRigVMParameterFlags.None;
 
-			if (Property.MetaData.ContainsKey("Visible"))
+			if (property.MetaData.ContainsKey("Visible"))
 			{
 				this.ParameterFlags |= UhtRigVMParameterFlags.Constant | UhtRigVMParameterFlags.Input;
 				this.ParameterFlags &= ~UhtRigVMParameterFlags.Output;
 			}
 
-			if (this.bEditorOnly)
+			if (this.EditorOnly)
 			{
-				Property.LogError($"RigVM Struct '{this.Property.Outer?.SourceName}' - Member '{this.Property.SourceName}' is editor only - WITH_EDITORONLY_DATA not allowed on structs with RIGVM_METHOD.");
+				property.LogError($"RigVM Struct '{this.Property.Outer?.SourceName}' - Member '{this.Property.SourceName}' is editor only - WITH_EDITORONLY_DATA not allowed on structs with RIGVM_METHOD.");
 			}
 
-			if (Property.PropertyCaps.HasAnyFlags(UhtPropertyCaps.SupportsRigVM))
+			if (property.PropertyCaps.HasAnyFlags(UhtPropertyCaps.SupportsRigVM))
 			{
-				this.Type = Property.GetRigVMType();
-				if (Property.PropertyCaps.HasAnyFlags(UhtPropertyCaps.IsRigVMEnum))
+				this.Type = property.GetRigVMType();
+				if (property.PropertyCaps.HasAnyFlags(UhtPropertyCaps.IsRigVMEnum))
 				{
 					this.ParameterFlags |= UhtRigVMParameterFlags.IsEnum;
 				}
-				if (Property.PropertyCaps.HasAnyFlags(UhtPropertyCaps.IsRigVMArray))
+				if (property.PropertyCaps.HasAnyFlags(UhtPropertyCaps.IsRigVMArray))
 				{
 					this.ParameterFlags |= UhtRigVMParameterFlags.IsArray;
 					if (this.IsConst())
 					{
-						string ExtendedType = this.ExtendedType(false);
-						this.CastName = $"{this.Name}_{Index}_Array";
-						this.CastType = $"TArrayView<const {ExtendedType.Substring(1, ExtendedType.Length - 2)}>";
+						string extendedType = this.ExtendedType(false);
+						this.CastName = $"{this.Name}_{index}_Array";
+						this.CastType = $"TArrayView<const {extendedType.Substring(1, extendedType.Length - 2)}>";
 					}
 				}
 			}
 			else
 			{
-				Property.LogError($"RigVM Struct '{this.Property.Outer?.SourceName}' - Member '{this.Property.SourceName}' type '{this.Property.GetUserFacingDecl()}' not supported by RigVM.");
+				property.LogError($"RigVM Struct '{this.Property.Outer?.SourceName}' - Member '{this.Property.SourceName}' type '{this.Property.GetUserFacingDecl()}' not supported by RigVM.");
 			}
 		}
 
 		/// <summary>
 		/// Create a new parameter
 		/// </summary>
-		/// <param name="Name">Name of the parameter</param>
-		/// <param name="Type">Type of the parameter</param>
-		public UhtRigVMParameter(string Name, string Type)
+		/// <param name="name">Name of the parameter</param>
+		/// <param name="type">Type of the parameter</param>
+		public UhtRigVMParameter(string name, string type)
 		{
-			this.Name = Name;
-			this.Type = Type;
+			this.Name = name;
+			this.Type = type;
 		}
 
 		/// <summary>
 		/// Get the name of the parameter
 		/// </summary>
-		/// <param name="bCastName">If true, return the cast name</param>
+		/// <param name="castName">If true, return the cast name</param>
 		/// <returns>Parameter name</returns>
-		public string NameOriginal(bool bCastName = false)
+		public string NameOriginal(bool castName = false)
 		{
-			return bCastName && this.CastName != null ? this.CastName : this.Name;
+			return castName && this.CastName != null ? this.CastName : this.Name;
 		}
 
 		/// <summary>
 		/// Get the type of the parameter
 		/// </summary>
-		/// <param name="bCastType">If true, return the cast type</param>
+		/// <param name="castType">If true, return the cast type</param>
 		/// <returns>Parameter type</returns>
-		public string TypeOriginal(bool bCastType = false)
+		public string TypeOriginal(bool castType = false)
 		{
-			return bCastType && this.CastType != null ? this.CastType : this.Type;
+			return castType && this.CastType != null ? this.CastType : this.Type;
 		}
 
 		/// <summary>
 		/// Get the full declaration (type and name)
 		/// </summary>
-		/// <param name="bCastType">If true, return the cast type</param>
-		/// <param name="bCastName">If true, return the cast name</param>
+		/// <param name="castType">If true, return the cast type</param>
+		/// <param name="castName">If true, return the cast name</param>
 		/// <returns>Parameter declaration</returns>
-		public string Declaration(bool bCastType = false, bool bCastName = false)
+		public string Declaration(bool castType = false, bool castName = false)
 		{
-			return $"{TypeOriginal(bCastType)} {NameOriginal(bCastName)}";
+			return $"{TypeOriginal(castType)} {NameOriginal(castName)}";
 		}
 
 		/// <summary>
 		/// Return the base type without any template arguments
 		/// </summary>
-		/// <param name="bCastType">If true, return the cast type</param>
+		/// <param name="castType">If true, return the cast type</param>
 		/// <returns>Base parameter type</returns>
-		public string BaseType(bool bCastType = false)
+		public string BaseType(bool castType = false)
 		{
-			string String = TypeOriginal(bCastType);
+			string typeOriginal = TypeOriginal(castType);
 
-			int LesserPos = String.IndexOf('<', StringComparison.Ordinal);
-			if (LesserPos >= 0)
+			int lesserPos = typeOriginal.IndexOf('<', StringComparison.Ordinal);
+			if (lesserPos >= 0)
 			{
-				return String.Substring(0, LesserPos);
+				return typeOriginal.Substring(0, lesserPos);
 			}
 			else
 			{
-				return String;
+				return typeOriginal;
 			}
 		}
 
 		/// <summary>
 		/// Template arguments of the type or type if not a template type.
 		/// </summary>
-		/// <param name="bCastType">If true, return the cast type</param>
+		/// <param name="castType">If true, return the cast type</param>
 		/// <returns>Template arguments of the type</returns>
-		public string ExtendedType(bool bCastType = false)
+		public string ExtendedType(bool castType = false)
 		{
-			string String = TypeOriginal(bCastType);
+			string typeOriginal = TypeOriginal(castType);
 
-			int LesserPos = String.IndexOf('<', StringComparison.Ordinal);
-			if (LesserPos >= 0)
+			int lesserPos = typeOriginal.IndexOf('<', StringComparison.Ordinal);
+			if (lesserPos >= 0)
 			{
-				return String.Substring(LesserPos);
+				return typeOriginal.Substring(lesserPos);
 			}
 			else
 			{
-				return String;
+				return typeOriginal;
 			}
 		}
 
 		/// <summary>
 		/// Return the type with a const reference
 		/// </summary>
-		/// <param name="bCastType">If true, return the cast type</param>
+		/// <param name="castType">If true, return the cast type</param>
 		/// <returns>Type with a const reference</returns>
-		public string TypeConstRef(bool bCastType = false)
+		public string TypeConstRef(bool castType = false)
 		{
-			string String = TypeNoRef(bCastType);
-			if (String.StartsWith("T", StringComparison.Ordinal) || String.StartsWith("F", StringComparison.Ordinal))
+			string typeNoRef = TypeNoRef(castType);
+			if (typeNoRef.StartsWith("T", StringComparison.Ordinal) || typeNoRef.StartsWith("F", StringComparison.Ordinal))
 			{
-				return $"const {String}&";
+				return $"const {typeNoRef}&";
 			}
 			else
 			{
-				return $"const {String}";
+				return $"const {typeNoRef}";
 			}
 		}
 
 		/// <summary>
 		/// Return the type with a reference
 		/// </summary>
-		/// <param name="bCastType">If true, return the cast type</param>
+		/// <param name="castType">If true, return the cast type</param>
 		/// <returns>Type with a reference</returns>
-		public string TypeRef(bool bCastType = false)
+		public string TypeRef(bool castType = false)
 		{
-			string String = TypeNoRef(bCastType);
-			return $"{String}&";
+			string typeNoRef = TypeNoRef(castType);
+			return $"{typeNoRef}&";
 		}
 
 		/// <summary>
 		/// Return the type without reference
 		/// </summary>
-		/// <param name="bCastType">If true, return the cast type</param>
+		/// <param name="castType">If true, return the cast type</param>
 		/// <returns>Type without the reference</returns>
-		public string TypeNoRef(bool bCastType = false)
+		public string TypeNoRef(bool castType = false)
 		{
-			string String = TypeOriginal(bCastType);
-			if (String.EndsWith("&", StringComparison.Ordinal))
+			string typeOriginal = TypeOriginal(castType);
+			if (typeOriginal.EndsWith("&", StringComparison.Ordinal))
 			{
-				return String.Substring(0, String.Length - 1);
+				return typeOriginal.Substring(0, typeOriginal.Length - 1);
 			}
 			else
 			{
-				return String;
+				return typeOriginal;
 			}
 		}
 
 		/// <summary>
 		/// Return the type as a reference
 		/// </summary>
-		/// <param name="bCastType">If true, return the cast type</param>
+		/// <param name="castType">If true, return the cast type</param>
 		/// <returns>Type as a reference</returns>
-		public string TypeVariableRef(bool bCastType = false)
+		public string TypeVariableRef(bool castType = false)
 		{
-			return IsConst() ? TypeConstRef(bCastType) : TypeRef(bCastType);
+			return IsConst() ? TypeConstRef(castType) : TypeRef(castType);
 		}
 
 		/// <summary>
 		/// Return a variable declaration for the parameter
 		/// </summary>
-		/// <param name="bCastType">If true, return the cast type</param>
-		/// <param name="bCastName">If true, return the cast name</param>
+		/// <param name="castType">If true, return the cast type</param>
+		/// <param name="castName">If true, return the cast name</param>
 		/// <returns>Parameter as a variable declaration</returns>
-		public string Variable(bool bCastType = false, bool bCastName = false)
+		public string Variable(bool castType = false, bool castName = false)
 		{
-			return $"{TypeVariableRef(bCastType)} {NameOriginal(bCastName)}";
+			return $"{TypeVariableRef(castType)} {NameOriginal(castName)}";
 		}
 
 		/// <summary>
@@ -385,7 +385,7 @@ namespace EpicGames.UHT.Types
 		/// <returns>True if the parameter is constant</returns>
 		public bool IsConst()
 		{
-			return this.bConstant || (this.bInput && !this.bOutput);
+			return this.Constant || (this.Input && !this.Output);
 		}
 
 		/// <summary>
@@ -404,18 +404,18 @@ namespace EpicGames.UHT.Types
 	/// </summary>
 	public class UhtRigVMMethodInfo
 	{
-		private static readonly string NoPrefixInternal = string.Empty;
+		private static readonly string s_noPrefixInternal = String.Empty;
 		private const string ReturnPrefixInternal = "return ";
 
 		/// <summary>
 		/// Return type of the method
 		/// </summary>
-		public string ReturnType { get; set; } = string.Empty;
+		public string ReturnType { get; set; } = String.Empty;
 
 		/// <summary>
 		/// Name of the method
 		/// </summary>
-		public string Name { get; set; } = string.Empty;
+		public string Name { get; set; } = String.Empty;
 
 		/// <summary>
 		/// Method parameters
@@ -428,7 +428,7 @@ namespace EpicGames.UHT.Types
 		/// <returns>Prefix required for the return value</returns>
 		public string ReturnPrefix()
 		{
-			return (ReturnType.Length == 0 || ReturnType == "void") ? NoPrefixInternal : ReturnPrefixInternal;
+			return (ReturnType.Length == 0 || ReturnType == "void") ? s_noPrefixInternal : ReturnPrefixInternal;
 		}
 	}
 
@@ -442,17 +442,17 @@ namespace EpicGames.UHT.Types
 		/// <summary>
 		/// True if the GetUpgradeInfoMethod was found. 
 		/// </summary>
-		public bool bHasGetUpgradeInfoMethod { get; set; } = false;
-		
+		public bool HasGetUpgradeInfoMethod { get; set; } = false;
+
 		/// <summary>
 		/// True if the GetNextAggregateNameMethod was found. 
 		/// </summary>
-		public bool bHasGetNextAggregateNameMethod { get; set; } = false;
+		public bool HasGetNextAggregateNameMethod { get; set; } = false;
 
 		/// <summary>
 		/// Engine name of the owning script struct
 		/// </summary>
-		public string Name { get; set; } = string.Empty;
+		public string Name { get; set; } = String.Empty;
 
 		/// <summary>
 		/// List of the members
@@ -469,7 +469,7 @@ namespace EpicGames.UHT.Types
 	/// Series of flags not part of the engine's script struct flags that affect code generation or verification
 	/// </summary>
 	[Flags]
-	public enum UhtScriptStructExportFlags : Int32
+	public enum UhtScriptStructExportFlags : int
 	{
 
 		/// <summary>
@@ -508,35 +508,35 @@ namespace EpicGames.UHT.Types
 		/// <summary>
 		/// Test to see if any of the specified flags are set
 		/// </summary>
-		/// <param name="InFlags">Current flags</param>
-		/// <param name="TestFlags">Flags to test for</param>
+		/// <param name="inFlags">Current flags</param>
+		/// <param name="testFlags">Flags to test for</param>
 		/// <returns>True if any of the flags are set</returns>
-		public static bool HasAnyFlags(this UhtScriptStructExportFlags InFlags, UhtScriptStructExportFlags TestFlags)
+		public static bool HasAnyFlags(this UhtScriptStructExportFlags inFlags, UhtScriptStructExportFlags testFlags)
 		{
-			return (InFlags & TestFlags) != 0;
+			return (inFlags & testFlags) != 0;
 		}
 
 		/// <summary>
 		/// Test to see if all of the specified flags are set
 		/// </summary>
-		/// <param name="InFlags">Current flags</param>
-		/// <param name="TestFlags">Flags to test for</param>
+		/// <param name="inFlags">Current flags</param>
+		/// <param name="testFlags">Flags to test for</param>
 		/// <returns>True if all the flags are set</returns>
-		public static bool HasAllFlags(this UhtScriptStructExportFlags InFlags, UhtScriptStructExportFlags TestFlags)
+		public static bool HasAllFlags(this UhtScriptStructExportFlags inFlags, UhtScriptStructExportFlags testFlags)
 		{
-			return (InFlags & TestFlags) == TestFlags;
+			return (inFlags & testFlags) == testFlags;
 		}
 
 		/// <summary>
 		/// Test to see if a specific set of flags have a specific value.
 		/// </summary>
-		/// <param name="InFlags">Current flags</param>
-		/// <param name="TestFlags">Flags to test for</param>
-		/// <param name="MatchFlags">Expected value of the tested flags</param>
+		/// <param name="inFlags">Current flags</param>
+		/// <param name="testFlags">Flags to test for</param>
+		/// <param name="matchFlags">Expected value of the tested flags</param>
 		/// <returns>True if the given flags have a specific value.</returns>
-		public static bool HasExactFlags(this UhtScriptStructExportFlags InFlags, UhtScriptStructExportFlags TestFlags, UhtScriptStructExportFlags MatchFlags)
+		public static bool HasExactFlags(this UhtScriptStructExportFlags inFlags, UhtScriptStructExportFlags testFlags, UhtScriptStructExportFlags matchFlags)
 		{
-			return (InFlags & TestFlags) == MatchFlags;
+			return (inFlags & testFlags) == matchFlags;
 		}
 	}
 
@@ -576,42 +576,36 @@ namespace EpicGames.UHT.Types
 		/// True if the struct has the "HasDefaults" specifier
 		/// </summary>
 		[JsonIgnore]
-		public bool bHasDefaults => this.ScriptStructExportFlags.HasAnyFlags(UhtScriptStructExportFlags.HasDefaults);
+		public bool HasDefaults => this.ScriptStructExportFlags.HasAnyFlags(UhtScriptStructExportFlags.HasDefaults);
 
 		/// <summary>
 		/// True if the struct has the "IsAlwaysAccessible" specifier
 		/// </summary>
 		[JsonIgnore]
-		public bool bIsAlwaysAccessible => this.ScriptStructExportFlags.HasAnyFlags(UhtScriptStructExportFlags.IsAlwaysAccessible);
+		public bool IsAlwaysAccessible => this.ScriptStructExportFlags.HasAnyFlags(UhtScriptStructExportFlags.IsAlwaysAccessible);
 
 		/// <summary>
 		/// True if the struct has the "IsCoreType" specifier
 		/// </summary>
 		[JsonIgnore]
-		public bool bIsCoreType => this.ScriptStructExportFlags.HasAnyFlags(UhtScriptStructExportFlags.IsCoreType);
+		public bool IsCoreType => this.ScriptStructExportFlags.HasAnyFlags(UhtScriptStructExportFlags.IsCoreType);
 
 		/// <summary>
 		/// True if the struct has the "HasNoOpConstructor" specifier
 		/// </summary>
 		[JsonIgnore]
-		public bool bHasNoOpConstructor => this.ScriptStructExportFlags.HasAnyFlags(UhtScriptStructExportFlags.HasNoOpConstructor);
+		public bool HasNoOpConstructor => this.ScriptStructExportFlags.HasAnyFlags(UhtScriptStructExportFlags.HasNoOpConstructor);
 
 		/// <inheritdoc/>
-		public override string EngineClassName { get => "ScriptStruct"; }
+		public override string EngineClassName => "ScriptStruct";
 
 		/// <inheritdoc/>
 		[JsonIgnore]
-		public override string EngineNamePrefix
-		{
-			get
-			{
-				return this.Session.Config!.IsStructWithTPrefix(this.EngineName) ? "T" : "F";
-			}
-		}
+		public override string EngineNamePrefix => this.Session.Config!.IsStructWithTPrefix(this.EngineName) ? "T" : "F";
 
 		///<inheritdoc/>
 		[JsonIgnore]
-		protected override UhtSpecifierValidatorTable? SpecifierValidatorTable { get => this.Session.GetSpecifierValidatorTable(UhtTableNames.ScriptStruct); }
+		protected override UhtSpecifierValidatorTable? SpecifierValidatorTable => this.Session.GetSpecifierValidatorTable(UhtTableNames.ScriptStruct);
 
 		/// <summary>
 		/// Super struct
@@ -622,14 +616,14 @@ namespace EpicGames.UHT.Types
 		/// <summary>
 		/// Construct a new script struct
 		/// </summary>
-		/// <param name="Outer">Outer type</param>
-		/// <param name="LineNumber">Line number of the definition</param>
-		public UhtScriptStruct(UhtType Outer, int LineNumber) : base(Outer, LineNumber)
+		/// <param name="outer">Outer type</param>
+		/// <param name="lineNumber">Line number of the definition</param>
+		public UhtScriptStruct(UhtType outer, int lineNumber) : base(outer, lineNumber)
 		{
 		}
 
 		/// <inheritdoc/>
-		protected override void ResolveChildren(UhtResolvePhase Phase)
+		protected override void ResolveChildren(UhtResolvePhase phase)
 		{
 
 			// Setup additional property as well as script struct def flags
@@ -638,41 +632,41 @@ namespace EpicGames.UHT.Types
 			// up an input / output pin of a RigVMNode. To allow authoring of those
 			// nodes we'll mark up the property as accessible in Blueprint / Python
 			// as well as make the struct a blueprint type.
-			switch (Phase)
+			switch (phase)
 			{
 				case UhtResolvePhase.Properties:
-					foreach (UhtType Child in this.Children)
+					foreach (UhtType child in this.Children)
 					{
-						if (Child is UhtProperty Property)
+						if (child is UhtProperty property)
 						{
-							EPropertyFlags OriginalFlags = Property.PropertyFlags;
-							if (Property.MetaData.ContainsKey(UhtNames.Constant))
+							EPropertyFlags originalFlags = property.PropertyFlags;
+							if (property.MetaData.ContainsKey(UhtNames.Constant))
 							{
-								Property.PropertyFlags |= EPropertyFlags.Edit | EPropertyFlags.EditConst | EPropertyFlags.BlueprintVisible;
+								property.PropertyFlags |= EPropertyFlags.Edit | EPropertyFlags.EditConst | EPropertyFlags.BlueprintVisible;
 							}
-							if (Property.MetaData.ContainsKey(UhtNames.Input) || Property.MetaData.ContainsKey(UhtNames.Visible))
+							if (property.MetaData.ContainsKey(UhtNames.Input) || property.MetaData.ContainsKey(UhtNames.Visible))
 							{
-								Property.PropertyFlags |= EPropertyFlags.Edit | EPropertyFlags.BlueprintVisible;
+								property.PropertyFlags |= EPropertyFlags.Edit | EPropertyFlags.BlueprintVisible;
 							}
-							if (Property.MetaData.ContainsKey(UhtNames.Output))
+							if (property.MetaData.ContainsKey(UhtNames.Output))
 							{
-								if (!Property.PropertyFlags.HasAnyFlags(EPropertyFlags.BlueprintVisible))
+								if (!property.PropertyFlags.HasAnyFlags(EPropertyFlags.BlueprintVisible))
 								{
-									Property.PropertyFlags |= EPropertyFlags.BlueprintVisible | EPropertyFlags.BlueprintReadOnly;
+									property.PropertyFlags |= EPropertyFlags.BlueprintVisible | EPropertyFlags.BlueprintReadOnly;
 								}
 							}
 
-							if (OriginalFlags != Property.PropertyFlags &&
-								Property.PropertyFlags.HasAnyFlags(EPropertyFlags.BlueprintVisible | EPropertyFlags.BlueprintReadOnly))
+							if (originalFlags != property.PropertyFlags &&
+								property.PropertyFlags.HasAnyFlags(EPropertyFlags.BlueprintVisible | EPropertyFlags.BlueprintReadOnly))
 							{
 								if (!this.MetaData.GetBooleanHierarchical(UhtNames.BlueprintType))
 								{
 									this.MetaData.Add(UhtNames.BlueprintType, true);
 								}
 
-								if (!Property.MetaData.ContainsKey(UhtNames.Category))
+								if (!property.MetaData.ContainsKey(UhtNames.Category))
 								{
-									Property.MetaData.Add(UhtNames.Category, UhtNames.Pins);
+									property.MetaData.Add(UhtNames.Category, UhtNames.Pins);
 								}
 							}
 						}
@@ -680,9 +674,9 @@ namespace EpicGames.UHT.Types
 					break;
 			}
 
-			base.ResolveChildren(Phase);
+			base.ResolveChildren(phase);
 
-			switch (Phase)
+			switch (phase)
 			{
 				case UhtResolvePhase.Final:
 					if (ScanForInstancedReferenced(false))
@@ -695,26 +689,26 @@ namespace EpicGames.UHT.Types
 		}
 
 		/// <inheritdoc/>
-		public override bool ScanForInstancedReferenced(bool bDeepScan)
+		public override bool ScanForInstancedReferenced(bool deepScan)
 		{
 			if (this.ScriptStructFlags.HasAnyFlags(EStructFlags.HasInstancedReference))
 			{
 				return true;
 			}
 
-			if (this.SuperScriptStruct != null && this.SuperScriptStruct.ScanForInstancedReferenced(bDeepScan))
+			if (this.SuperScriptStruct != null && this.SuperScriptStruct.ScanForInstancedReferenced(deepScan))
 			{
 				return true;
 			}
 
-			return base.ScanForInstancedReferenced(bDeepScan);
+			return base.ScanForInstancedReferenced(deepScan);
 		}
 
 		#region Validation support
 		/// <inheritdoc/>
-		protected override UhtValidationOptions Validate(UhtValidationOptions Options)
+		protected override UhtValidationOptions Validate(UhtValidationOptions options)
 		{
-			Options = base.Validate(Options);
+			options = base.Validate(options);
 
 			if (this.ScriptStructFlags.HasAnyFlags(EStructFlags.Immutable))
 			{
@@ -725,43 +719,43 @@ namespace EpicGames.UHT.Types
 			}
 
 			// Validate the engine name
-			string ExpectedName = $"{this.EngineNamePrefix}{this.EngineName}";
-			if (this.SourceName != ExpectedName)
+			string expectedName = $"{this.EngineNamePrefix}{this.EngineName}";
+			if (this.SourceName != expectedName)
 			{
-				this.LogError($"Struct '{this.SourceName}' has an invalid Unreal prefix, expecting '{ExpectedName}");
+				this.LogError($"Struct '{this.SourceName}' has an invalid Unreal prefix, expecting '{expectedName}");
 			}
 
 			// Validate RigVM
-			if (this.RigVMStructInfo != null && this.MetaData.ContainsKey("Deprecated") && !this.RigVMStructInfo.bHasGetUpgradeInfoMethod)
+			if (this.RigVMStructInfo != null && this.MetaData.ContainsKey("Deprecated") && !this.RigVMStructInfo.HasGetUpgradeInfoMethod)
 			{
 				this.LogError($"RigVMStruct '{this.SourceName}' is marked as deprecated but is missing 'GetUpgradeInfo method.");
 				this.LogError("Please implement a method like below:");
 				this.LogError("RIGVM_METHOD()");
-				this.LogError("virtual FRigVMStructUpgradeInfo GetUpgradeInfo() const override;");					
+				this.LogError("virtual FRigVMStructUpgradeInfo GetUpgradeInfo() const override;");
 			}
 
-			return Options |= UhtValidationOptions.Shadowing;
+			return options |= UhtValidationOptions.Shadowing;
 		}
 		#endregion
 
 		/// <inheritdoc/>
-		public override void CollectReferences(IUhtReferenceCollector Collector)
+		public override void CollectReferences(IUhtReferenceCollector collector)
 		{
 			if (this.ScriptStructFlags.HasAnyFlags(EStructFlags.NoExport))
 			{
-				Collector.AddSingleton(this);
+				collector.AddSingleton(this);
 			}
-			Collector.AddExportType(this);
-			Collector.AddDeclaration(this, true);
-			Collector.AddCrossModuleReference(this, true);
+			collector.AddExportType(this);
+			collector.AddDeclaration(this, true);
+			collector.AddCrossModuleReference(this, true);
 			if (this.SuperScriptStruct != null)
 			{
-				Collector.AddCrossModuleReference(this.SuperScriptStruct, true);
+				collector.AddCrossModuleReference(this.SuperScriptStruct, true);
 			}
-			Collector.AddCrossModuleReference(this.Package, true);
-			foreach (UhtType Child in this.Children)
+			collector.AddCrossModuleReference(this.Package, true);
+			foreach (UhtType child in this.Children)
 			{
-				Child.CollectReferences(Collector);
+				child.CollectReferences(collector);
 			}
 		}
 
@@ -769,11 +763,11 @@ namespace EpicGames.UHT.Types
 		{
 			if (this.RigVMStructInfo != null)
 			{
-				for (UhtStruct? Current = this; Current != null; Current = Current.SuperStruct)
+				for (UhtStruct? current = this; current != null; current = current.SuperStruct)
 				{
-					foreach (UhtProperty Property in Current.Properties)
+					foreach (UhtProperty property in current.Properties)
 					{
-						this.RigVMStructInfo.Members.Add(new UhtRigVMParameter(Property, this.RigVMStructInfo.Members.Count));
+						this.RigVMStructInfo.Members.Add(new UhtRigVMParameter(property, this.RigVMStructInfo.Members.Count));
 					}
 				}
 
diff --git a/Engine/Source/Programs/Shared/EpicGames.UHT/Types/UhtSourceFile.cs b/Engine/Source/Programs/Shared/EpicGames.UHT/Types/UhtSourceFile.cs
--- a/Engine/Source/Programs/Shared/EpicGames.UHT/Types/UhtSourceFile.cs
+++ b/Engine/Source/Programs/Shared/EpicGames.UHT/Types/UhtSourceFile.cs
@@ -1,8 +1,8 @@
 // Copyright Epic Games, Inc. All Rights Reserved.
 
+using System.IO;
 using EpicGames.Core;
 using EpicGames.UHT.Utils;
-using System.IO;
 
 namespace EpicGames.UHT.Types
 {
@@ -12,8 +12,8 @@ namespace EpicGames.UHT.Types
 	/// </summary>
 	public class UhtSourceFile : IUhtMessageSite, IUhtMessageSource
 	{
-		private readonly UhtSimpleMessageSite MessageSite;
-		private UhtSourceFragment SourceFragment;
+		private readonly UhtSimpleMessageSite _messageSite;
+		private UhtSourceFragment _sourceFragment;
 
 		/// <summary>
 		/// The session associated with the source
@@ -33,29 +33,29 @@ namespace EpicGames.UHT.Types
 		/// <summary>
 		/// The contents of the source file
 		/// </summary>
-		public StringView Data { get => SourceFragment.Data; }
+		public StringView Data => _sourceFragment.Data;
 
 		/// <summary>
 		/// If this source is from a fragment, the original source file
 		/// </summary>
-		public UhtSourceFile? FragmentSourceFile { get => SourceFragment.SourceFile; }
+		public UhtSourceFile? FragmentSourceFile => _sourceFragment.SourceFile;
 
 		/// <summary>
 		/// If this source is from a fragment, the starting line of the fragment
 		/// </summary>
-		public int FragmentLineNumber { get => SourceFragment.LineNumber; }
+		public int FragmentLineNumber => _sourceFragment.LineNumber;
 
 		/// <summary>
 		/// Full file path of the source file
 		/// </summary>
-		public string FullFilePath { get => SourceFragment.FilePath; }
+		public string FullFilePath => _sourceFragment.FilePath;
 
 		#region IUHTMessageSite implementation
 		/// <inheritdoc/>
-		public IUhtMessageSession MessageSession => this.MessageSite.MessageSession;
+		public IUhtMessageSession MessageSession => this._messageSite.MessageSession;
 
 		/// <inheritdoc/>
-		public IUhtMessageSource? MessageSource => this.MessageSite.MessageSource;
+		public IUhtMessageSource? MessageSource => this._messageSite.MessageSource;
 
 		/// <inheritdoc/>
 		public IUhtMessageLineNumber? MessageLineNumber => null;
@@ -63,35 +63,35 @@ namespace EpicGames.UHT.Types
 
 		#region IUHTMessageSource implementation
 		/// <inheritdoc/>
-		public string MessageFilePath { get => this.FilePath; }
+		public string MessageFilePath => this.FilePath;
 
 		/// <inheritdoc/>
-		public string MessageFullFilePath { get => this.FullFilePath; }
+		public string MessageFullFilePath => this.FullFilePath;
 
 		/// <inheritdoc/>
-		public bool bMessageIsFragment { get => this.FragmentSourceFile != null; }
+		public bool MessageIsFragment => this.FragmentSourceFile != null;
 
 		/// <inheritdoc/>
-		public string MessageFragmentFilePath { get => this.FragmentSourceFile != null ? this.FragmentSourceFile.MessageFilePath : ""; }
+		public string MessageFragmentFilePath => this.FragmentSourceFile != null ? this.FragmentSourceFile.MessageFilePath : "";
 
 		/// <inheritdoc/>
-		public string MessageFragmentFullFilePath { get => this.FragmentSourceFile != null ? this.FragmentSourceFile.MessageFullFilePath : ""; }
+		public string MessageFragmentFullFilePath => this.FragmentSourceFile != null ? this.FragmentSourceFile.MessageFullFilePath : "";
 
 		/// <inheritdoc/>
-		public int MessageFragmentLineNumber { get => this.FragmentLineNumber; }
+		public int MessageFragmentLineNumber => this.FragmentLineNumber;
 		#endregion
 
 		/// <summary>
 		/// Construct a new instance of a source file
 		/// </summary>
-		/// <param name="Session">The owning session</param>
-		/// <param name="FilePath">The full file path</param>
-		public UhtSourceFile(UhtSession Session, string FilePath)
+		/// <param name="session">The owning session</param>
+		/// <param name="filePath">The full file path</param>
+		public UhtSourceFile(UhtSession session, string filePath)
 		{
-			this.Session = Session;
-			this.FilePath = FilePath;
+			this.Session = session;
+			this.FilePath = filePath;
 			this.FileName = Path.GetFileNameWithoutExtension(this.FilePath);
-			this.MessageSite = new UhtSimpleMessageSite(this.Session, this);
+			this._messageSite = new UhtSimpleMessageSite(this.Session, this);
 		}
 
 		/// <summary>
@@ -99,7 +99,7 @@ namespace EpicGames.UHT.Types
 		/// </summary>
 		public virtual void Read()
 		{
-			this.SourceFragment = Session.ReadSource(this.FilePath);
+			this._sourceFragment = Session.ReadSource(this.FilePath);
 		}
 	}
 }
diff --git a/Engine/Source/Programs/Shared/EpicGames.UHT/Types/UhtStruct.cs b/Engine/Source/Programs/Shared/EpicGames.UHT/Types/UhtStruct.cs
--- a/Engine/Source/Programs/Shared/EpicGames.UHT/Types/UhtStruct.cs
+++ b/Engine/Source/Programs/Shared/EpicGames.UHT/Types/UhtStruct.cs
@@ -1,12 +1,12 @@
 // Copyright Epic Games, Inc. All Rights Reserved.
 
+using System;
+using System.Collections.Generic;
+using System.Text.Json.Serialization;
 using EpicGames.Core;
 using EpicGames.UHT.Tables;
 using EpicGames.UHT.Tokenizer;
 using EpicGames.UHT.Utils;
-using System;
-using System.Collections.Generic;
-using System.Text.Json.Serialization;
 
 namespace EpicGames.UHT.Types
 {
@@ -36,11 +36,11 @@ namespace EpicGames.UHT.Types
 		{
 			get
 			{
-				foreach (UhtType Type in this.Children)
+				foreach (UhtType type in this.Children)
 				{
-					if (Type is UhtProperty Property)
+					if (type is UhtProperty property)
 					{
-						yield return Property;
+						yield return property;
 					}
 				}
 			}
@@ -59,11 +59,11 @@ namespace EpicGames.UHT.Types
 		{
 			get
 			{
-				foreach (UhtType Type in this.Children)
+				foreach (UhtType type in this.Children)
 				{
-					if (Type is UhtFunction Function)
+					if (type is UhtFunction function)
 					{
-						yield return Function;
+						yield return function;
 					}
 				}
 			}
@@ -71,10 +71,10 @@ namespace EpicGames.UHT.Types
 
 		/// <inheritdoc/>
 		[JsonIgnore]
-		public virtual string EngineNamePrefix { get => "F"; }
+		public virtual string EngineNamePrefix => "F";
 
 		/// <inheritdoc/>
-		public override string EngineClassName { get => "Struct"; }
+		public override string EngineClassName => "Struct";
 
 		/// <summary>
 		/// Super type
@@ -86,10 +86,7 @@ namespace EpicGames.UHT.Types
 		/// </summary>
 		//[JsonConverter(typeof(UhtNullableTypeListJsonConverter<UhtStruct>))]
 		[JsonIgnore]
-		public IReadOnlyList<UhtStruct> Bases 
-		{ 
-			get => _basesInternal != null ? _basesInternal : _emptyBases;
-		}
+		public IReadOnlyList<UhtStruct> Bases => _basesInternal ?? s_emptyBases;
 
 		/// <summary>
 		/// Super struct type
@@ -100,70 +97,70 @@ namespace EpicGames.UHT.Types
 		/// <summary>
 		/// Construct a new instance
 		/// </summary>
-		/// <param name="Outer">Outer type</param>
-		/// <param name="LineNumber">Line number where definition begins</param>
-		protected UhtStruct(UhtType Outer, int LineNumber) : base(Outer, LineNumber)
+		/// <param name="outer">Outer type</param>
+		/// <param name="lineNumber">Line number where definition begins</param>
+		protected UhtStruct(UhtType outer, int lineNumber) : base(outer, lineNumber)
 		{
 		}
 
 		/// <summary>
 		/// Test to see if the given struct is derived from the base structure
 		/// </summary>
-		/// <param name="Base">Base structure.</param>
+		/// <param name="baseStruct">Base structure.</param>
 		/// <returns>True if the given structure is the specified base or derives from the base.  If the base is null, the false is returned.</returns>
-		public bool IsChildOf(UhtStruct? Base)
+		public bool IsChildOf(UhtStruct? baseStruct)
 		{
-			if (Base == null)
+			if (baseStruct == null)
 			{
 				return false;
 			}
-			for (UhtStruct? Current = this; Current != null; Current = Current.Super)
+			for (UhtStruct? current = this; current != null; current = current.Super)
 			{
-				if (Current == Base)
+				if (current == baseStruct)
 				{
 					return true;
 				}
 			}
 			return false;
 		}
-		
-		private static readonly List<UhtStruct> _emptyBases = new List<UhtStruct>();
+
+		private static readonly List<UhtStruct> s_emptyBases = new List<UhtStruct>();
 		private List<UhtStruct>? _basesInternal = null;
 
 		#region Resolution support
 		/// <inheritdoc/>
-		protected override void ResolveSuper(UhtResolvePhase ResolvePhase)
+		protected override void ResolveSuper(UhtResolvePhase resolvePhase)
 		{
 			if (this.Super != null)
 			{
-				this.Super.Resolve(ResolvePhase);
+				this.Super.Resolve(resolvePhase);
 			}
 
-			foreach (UhtStruct Base in this.Bases)
+			foreach (UhtStruct baseStruct in this.Bases)
 			{
-				Base.Resolve(ResolvePhase);
+				baseStruct.Resolve(resolvePhase);
 			}
 
-			base.ResolveSuper(ResolvePhase);
+			base.ResolveSuper(resolvePhase);
 		}
 
 		/// <summary>
 		/// Check properties to see if any instances are referenced.
 		/// This method does NOT cache the result.
 		/// </summary>
-		/// <param name="bDeepScan">If true, the ScanForInstancedReferenced method on the properties will also be called.</param>
+		/// <param name="deepScan">If true, the ScanForInstancedReferenced method on the properties will also be called.</param>
 		/// <returns></returns>
-		public virtual bool ScanForInstancedReferenced(bool bDeepScan)
+		public virtual bool ScanForInstancedReferenced(bool deepScan)
 		{
-			foreach (UhtType Type in this.Children)
+			foreach (UhtType type in this.Children)
 			{
-				if (Type is UhtProperty Property)
+				if (type is UhtProperty property)
 				{
-					if (Property.PropertyFlags.HasAnyFlags(EPropertyFlags.ContainsInstancedReference | EPropertyFlags.InstancedReference))
+					if (property.PropertyFlags.HasAnyFlags(EPropertyFlags.ContainsInstancedReference | EPropertyFlags.InstancedReference))
 					{
 						return true;
 					}
-					if (bDeepScan && Property.ScanForInstancedReferenced(bDeepScan))
+					if (deepScan && property.ScanForInstancedReferenced(deepScan))
 					{
 						return true;
 					}
@@ -177,17 +174,17 @@ namespace EpicGames.UHT.Types
 		/// <summary>
 		/// Resolve super identifier
 		/// </summary>
-		/// <param name="SuperIdentifier">Token that represent the super</param>
-		/// <param name="FindOptions">Find options to restrict types</param>
+		/// <param name="superIdentifier">Token that represent the super</param>
+		/// <param name="findOptions">Find options to restrict types</param>
 		/// <exception cref="UhtException">Thrown if super can not be found</exception>
-		public void BindAndResolveSuper(UhtToken SuperIdentifier, UhtFindOptions FindOptions)
+		public void BindAndResolveSuper(UhtToken superIdentifier, UhtFindOptions findOptions)
 		{
-			if (SuperIdentifier)
+			if (superIdentifier)
 			{
-				this.Super = (UhtStruct?)this.FindType(FindOptions | UhtFindOptions.SourceName | UhtFindOptions.NoSelf, ref SuperIdentifier);
+				this.Super = (UhtStruct?)this.FindType(findOptions | UhtFindOptions.SourceName | UhtFindOptions.NoSelf, ref superIdentifier);
 				if (this.Super == null)
 				{
-					throw new UhtException(this, $"Unable to find parent {this.EngineType.ShortLowercaseText()} type for '{this.SourceName}' named '{SuperIdentifier.Value}'");
+					throw new UhtException(this, $"Unable to find parent {this.EngineType.ShortLowercaseText()} type for '{this.SourceName}' named '{superIdentifier.Value}'");
 				}
 				this.HeaderFile.AddReferencedHeader(this.Super);
 				this.MetaData.Parent = this.Super.MetaData;
@@ -199,25 +196,25 @@ namespace EpicGames.UHT.Types
 		/// Resolve bases.  Unlike super, this routine will not generate an error if the type can not be found.
 		/// Having unrecognized types is expected.
 		/// </summary>
-		/// <param name="BaseIdentifiers">Collection of bases</param>
-		/// <param name="FindOptions">Options to restrict types being searched</param>
-		public void BindAndResolveBases(List<UhtToken[]>? BaseIdentifiers, UhtFindOptions FindOptions)
+		/// <param name="baseIdentifiers">Collection of bases</param>
+		/// <param name="findOptions">Options to restrict types being searched</param>
+		public void BindAndResolveBases(List<UhtToken[]>? baseIdentifiers, UhtFindOptions findOptions)
 		{
-			if (BaseIdentifiers != null)
+			if (baseIdentifiers != null)
 			{
-				foreach (UhtToken[] BaseIdentifier in BaseIdentifiers)
+				foreach (UhtToken[] baseIdentifier in baseIdentifiers)
 				{
 					// We really only case about interfaces, but we can also handle structs
-					UhtStruct? Base = (UhtStruct?)this.FindType(FindOptions | UhtFindOptions.Class | UhtFindOptions.ScriptStruct | UhtFindOptions.SourceName | UhtFindOptions.NoSelf, BaseIdentifier);
-					if (Base != null)
+					UhtStruct? baseStruct = (UhtStruct?)this.FindType(findOptions | UhtFindOptions.Class | UhtFindOptions.ScriptStruct | UhtFindOptions.SourceName | UhtFindOptions.NoSelf, baseIdentifier);
+					if (baseStruct != null)
 					{
 						if (this._basesInternal == null)
 						{
 							this._basesInternal = new List<UhtStruct>();
 						}
-						this._basesInternal.Add(Base);
-						this.HeaderFile.AddReferencedHeader(Base);
-						Base.Resolve(UhtResolvePhase.Bases);
+						this._basesInternal.Add(baseStruct);
+						this.HeaderFile.AddReferencedHeader(baseStruct);
+						baseStruct.Resolve(UhtResolvePhase.Bases);
 					}
 				}
 			}
@@ -226,115 +223,115 @@ namespace EpicGames.UHT.Types
 
 		#region Validation support
 		/// <inheritdoc/>
-		protected override UhtValidationOptions Validate(UhtValidationOptions Options)
+		protected override UhtValidationOptions Validate(UhtValidationOptions options)
 		{
-			Options = base.Validate(Options);
+			options = base.Validate(options);
 			ValidateSparseClassData();
-			return Options;
+			return options;
 		}
 
-		private static bool CheckUIMinMaxRangeFromMetaData(UhtType Child)
+		private static bool CheckUIMinMaxRangeFromMetaData(UhtType child)
 		{
-			string UIMin = Child.MetaData.GetValueOrDefault(UhtNames.UIMin);
-			string UIMax = Child.MetaData.GetValueOrDefault(UhtNames.UIMax);
-			if (UIMin.Length == 0 || UIMax.Length == 0)
+			string uiMin = child.MetaData.GetValueOrDefault(UhtNames.UIMin);
+			string uiMax = child.MetaData.GetValueOrDefault(UhtNames.UIMax);
+			if (uiMin.Length == 0 || uiMax.Length == 0)
 			{
 				return false;
 			}
 
 			// NOTE: Old UHT didn't handle parse errors
-			double MinValue;
-			if (!double.TryParse(UIMin, out MinValue))
+			double minValue;
+			if (!Double.TryParse(uiMin, out minValue))
 			{
-				MinValue = 0;
+				minValue = 0;
 			}
 
-			double MaxValue;
-			if (!double.TryParse(UIMax, out MaxValue))
+			double maxValue;
+			if (!Double.TryParse(uiMax, out maxValue))
 			{
-				MaxValue = 0;
+				maxValue = 0;
 			}
 
 			// NOTE: that we actually allow UIMin == UIMax to disable the range manually.
-			return MinValue <= MaxValue;
+			return minValue <= maxValue;
 		}
 
 		/// <inheritdoc/>
-		protected override void ValidateDocumentationPolicy(UhtDocumentationPolicy Policy)
+		protected override void ValidateDocumentationPolicy(UhtDocumentationPolicy policy)
 		{
-			if (Policy.bClassOrStructCommentRequired)
+			if (policy.ClassOrStructCommentRequired)
 			{
-				string ClassTooltip = this.MetaData.GetValueOrDefault(UhtNames.ToolTip);
-				if (ClassTooltip.Length == 0 || ClassTooltip.Equals(this.EngineName, StringComparison.OrdinalIgnoreCase))
+				string classTooltip = this.MetaData.GetValueOrDefault(UhtNames.ToolTip);
+				if (classTooltip.Length == 0 || classTooltip.Equals(this.EngineName, StringComparison.OrdinalIgnoreCase))
 				{
 					this.LogError($"{this.EngineType.CapitalizedText()} '{this.SourceName}' does not provide a tooltip / comment (DocumentationPolicy).");
 				}
 			}
 
-			if (Policy.bMemberToolTipsRequired)
+			if (policy.MemberToolTipsRequired)
 			{
-				Dictionary<string, UhtProperty> ToolTipToType = new Dictionary<string, UhtProperty>();
-				foreach (UhtProperty Property in this.Properties)
+				Dictionary<string, UhtProperty> toolTipToType = new Dictionary<string, UhtProperty>();
+				foreach (UhtProperty property in this.Properties)
 				{
-					string ToolTip = Property.GetToolTipText();
-					if (ToolTip.Length == 0 || ToolTip == Property.GetDisplayNameText())
+					string toolTip = property.GetToolTipText();
+					if (toolTip.Length == 0 || toolTip == property.GetDisplayNameText())
 					{
-						Property.LogError($"Property '{this.SourceName}::{Property.SourceName}' does not provide a tooltip / comment (DocumentationPolicy).");
+						property.LogError($"Property '{this.SourceName}::{property.SourceName}' does not provide a tooltip / comment (DocumentationPolicy).");
 						continue;
 					}
 
-					UhtProperty? Existing;
-					if (ToolTipToType.TryGetValue(ToolTip, out Existing))
+					UhtProperty? existing;
+					if (toolTipToType.TryGetValue(toolTip, out existing))
 					{
-						Property.LogError($"Property '{this.SourceName}::{Existing.SourceName}' and '{this.SourceName}::{Property.SourceName}' are using identical tooltips (DocumentationPolicy).");
+						property.LogError($"Property '{this.SourceName}::{existing.SourceName}' and '{this.SourceName}::{property.SourceName}' are using identical tooltips (DocumentationPolicy).");
 					}
 					else
 					{
-						ToolTipToType.Add(ToolTip, Property);
+						toolTipToType.Add(toolTip, property);
 					}
 				}
 			}
 
-			if (Policy.bFloatRangesRequired)
+			if (policy.FloatRangesRequired)
 			{
-				foreach (UhtType Child in this.Children)
+				foreach (UhtType child in this.Children)
 				{
-					if (Child is UhtDoubleProperty || Child is UhtFloatProperty)
+					if (child is UhtDoubleProperty || child is UhtFloatProperty)
 					{
-						if (!CheckUIMinMaxRangeFromMetaData(Child))
+						if (!CheckUIMinMaxRangeFromMetaData(child))
 						{
-							Child.LogError($"Property '{this.SourceName}::{Child.SourceName}' does not provide a valid UIMin / UIMax (DocumentationPolicy).");
+							child.LogError($"Property '{this.SourceName}::{child.SourceName}' does not provide a valid UIMin / UIMax (DocumentationPolicy).");
 						}
 					}
 				}
 			}
 
 			// also compare all tooltips to see if they are unique
-			if (Policy.bFunctionToolTipsRequired)
+			if (policy.FunctionToolTipsRequired)
 			{
 				if (this is UhtClass)
 				{
-					Dictionary<string, UhtType> ToolTipToType = new Dictionary<string, UhtType>();
-					foreach (UhtType Child in this.Children)
+					Dictionary<string, UhtType> toolTipToType = new Dictionary<string, UhtType>();
+					foreach (UhtType child in this.Children)
 					{
-						if (Child is UhtFunction Function)
+						if (child is UhtFunction function)
 						{
-							string ToolTip = Function.GetToolTipText();
-							if (ToolTip.Length == 0)
+							string toolTip = function.GetToolTipText();
+							if (toolTip.Length == 0)
 							{
 								// NOTE: This does not fire because it doesn't check to see if it matches the display name as above.
-								Child.LogError($"Function '{this.SourceName}::{Function.SourceName}' does not provide a tooltip / comment (DocumentationPolicy).");
+								child.LogError($"Function '{this.SourceName}::{function.SourceName}' does not provide a tooltip / comment (DocumentationPolicy).");
 								continue;
 							}
 
-							UhtType? Existing;
-							if (ToolTipToType.TryGetValue(ToolTip, out Existing))
+							UhtType? existing;
+							if (toolTipToType.TryGetValue(toolTip, out existing))
 							{
-								Child.LogError($"Functions '{this.SourceName}::{Existing.SourceName}' and '{this.SourceName}::{Function.SourceName}' are using identical tooltips / comments (DocumentationPolicy).");
+								child.LogError($"Functions '{this.SourceName}::{existing.SourceName}' and '{this.SourceName}::{function.SourceName}' are using identical tooltips / comments (DocumentationPolicy).");
 							}
 							else
 							{
-								ToolTipToType.Add(ToolTip, Function);
+								toolTipToType.Add(toolTip, function);
 							}
 						}
 					}
@@ -345,81 +342,81 @@ namespace EpicGames.UHT.Types
 		private void ValidateSparseClassData()
 		{
 			// Fetch the data types
-			string[]? SparseClassDataTypes = this.MetaData.GetStringArray(UhtNames.SparseClassDataTypes);
-			if (SparseClassDataTypes == null)
+			string[]? sparseClassDataTypes = this.MetaData.GetStringArray(UhtNames.SparseClassDataTypes);
+			if (sparseClassDataTypes == null)
 			{
 				return;
 			}
 
 			// Make sure we don't try to have sparse class data inside of a struct instead of a class
-			UhtClass? Class = this as UhtClass;
-			if (Class == null)
+			UhtClass? classObj = this as UhtClass;
+			if (classObj == null)
 			{
 				this.LogError($"{this.EngineType.CapitalizedText()} '{this.SourceName}' contains sparse class data but is not a class.");
 				return;
 			}
 
 			// for now we only support one sparse class data structure per class
-			if (SparseClassDataTypes.Length > 1)
+			if (sparseClassDataTypes.Length > 1)
 			{
 				this.LogError($"Class '{this.SourceName}' contains multiple sparse class data types");
 				return;
 			}
-			if (SparseClassDataTypes.Length == 0)
+			if (sparseClassDataTypes.Length == 0)
 			{
 				this.LogError($"Class '{this.SourceName}' has sparse class metadata but does not specify a type");
 				return;
 			}
 
-			foreach (string SparseClassDataTypeName in SparseClassDataTypes)
+			foreach (string sparseClassDataTypeName in sparseClassDataTypes)
 			{
-				UhtScriptStruct? SparseScriptStruct = this.Session.FindType(null, UhtFindOptions.EngineName | UhtFindOptions.ScriptStruct, SparseClassDataTypeName) as UhtScriptStruct;
+				UhtScriptStruct? sparseScriptStruct = this.Session.FindType(null, UhtFindOptions.EngineName | UhtFindOptions.ScriptStruct, sparseClassDataTypeName) as UhtScriptStruct;
 
 				// make sure the sparse class data struct actually exists
-				if (SparseScriptStruct == null)
+				if (sparseScriptStruct == null)
 				{
-					this.LogError($"Unable to find sparse data type '{SparseClassDataTypeName}' for class '{this.SourceName}'");
+					this.LogError($"Unable to find sparse data type '{sparseClassDataTypeName}' for class '{this.SourceName}'");
 					continue;
 				}
 
 				// check the data struct for invalid properties
-				foreach (UhtProperty Property in SparseScriptStruct.Properties)
+				foreach (UhtProperty property in sparseScriptStruct.Properties)
 				{
-					if (Property.PropertyFlags.HasAnyFlags(EPropertyFlags.BlueprintAssignable))
+					if (property.PropertyFlags.HasAnyFlags(EPropertyFlags.BlueprintAssignable))
 					{
-						Property.LogError($"Sparse class data types can not contain blueprint assignable delegates. Type '{SparseScriptStruct.EngineName}' Delegate '{Property.SourceName}'");
+						property.LogError($"Sparse class data types can not contain blueprint assignable delegates. Type '{sparseScriptStruct.EngineName}' Delegate '{property.SourceName}'");
 					}
 
 					// all sparse properties should have EditDefaultsOnly
-					if (!Property.PropertyFlags.HasAllFlags(EPropertyFlags.Edit | EPropertyFlags.DisableEditOnInstance))
+					if (!property.PropertyFlags.HasAllFlags(EPropertyFlags.Edit | EPropertyFlags.DisableEditOnInstance))
 					{
-						Property.LogError($"Sparse class data types must be VisibleDefaultsOnly or EditDefaultsOnly. Type '{SparseScriptStruct.EngineName}' Property '{Property.SourceName}'");
+						property.LogError($"Sparse class data types must be VisibleDefaultsOnly or EditDefaultsOnly. Type '{sparseScriptStruct.EngineName}' Property '{property.SourceName}'");
 					}
 
 					// no sparse properties should have BlueprintReadWrite
-					if (Property.PropertyFlags.HasAnyFlags(EPropertyFlags.BlueprintVisible) && !Property.PropertyFlags.HasAnyFlags(EPropertyFlags.BlueprintReadOnly))
+					if (property.PropertyFlags.HasAnyFlags(EPropertyFlags.BlueprintVisible) && !property.PropertyFlags.HasAnyFlags(EPropertyFlags.BlueprintReadOnly))
 					{
-						Property.LogError($"Sparse class data types must not be BlueprintReadWrite. Type '{SparseScriptStruct.EngineName}' Property '{Property.SourceName}'");
+						property.LogError($"Sparse class data types must not be BlueprintReadWrite. Type '{sparseScriptStruct.EngineName}' Property '{property.SourceName}'");
 					}
 				}
 
 				// if the class's parent has a sparse class data struct then the current class must also use the same struct or one that inherits from it
-				UhtClass? SuperClass = Class.SuperClass;
-				if (SuperClass != null)
+				UhtClass? superClass = classObj.SuperClass;
+				if (superClass != null)
 				{
-					string[]? SuperSparseClassDataTypes = SuperClass.MetaData.GetStringArray(UhtNames.SparseClassDataTypes);
-					if (SuperSparseClassDataTypes != null)
+					string[]? superSparseClassDataTypes = superClass.MetaData.GetStringArray(UhtNames.SparseClassDataTypes);
+					if (superSparseClassDataTypes != null)
 					{
-						foreach (string SuperSparseClassDataTypeName in SuperSparseClassDataTypes)
+						foreach (string superSparseClassDataTypeName in superSparseClassDataTypes)
 						{
-							UhtScriptStruct? SuperSparseScriptStruct = this.Session.FindType(null, UhtFindOptions.EngineName | UhtFindOptions.ScriptStruct, SuperSparseClassDataTypeName) as UhtScriptStruct;
-							if (SuperSparseScriptStruct != null)
+							UhtScriptStruct? superSparseScriptStruct = this.Session.FindType(null, UhtFindOptions.EngineName | UhtFindOptions.ScriptStruct, superSparseClassDataTypeName) as UhtScriptStruct;
+							if (superSparseScriptStruct != null)
 							{
-								if (!SparseScriptStruct.IsChildOf(SuperSparseScriptStruct))
+								if (!sparseScriptStruct.IsChildOf(superSparseScriptStruct))
 								{
 									this.LogError(
-										$"Class '{this.SourceName}' is a child of '{SuperClass.SourceName}' but its sparse class data struct " +
-										$"'{SparseScriptStruct.EngineName}', does not inherit from '{SuperSparseScriptStruct.EngineName}'.");
+										$"Class '{this.SourceName}' is a child of '{superClass.SourceName}' but its sparse class data struct " +
+										$"'{sparseScriptStruct.EngineName}', does not inherit from '{superSparseScriptStruct.EngineName}'.");
 								}
 							}
 						}
