More coding standard fixes

diff --git a/Engine/Source/Programs/Shared/EpicGames.UHT/Types/UhtType.cs b/Engine/Source/Programs/Shared/EpicGames.UHT/Types/UhtType.cs
--- a/Engine/Source/Programs/Shared/EpicGames.UHT/Types/UhtType.cs
+++ b/Engine/Source/Programs/Shared/EpicGames.UHT/Types/UhtType.cs
@@ -1,13 +1,13 @@
 // Copyright Epic Games, Inc. All Rights Reserved.
 
-using EpicGames.UHT.Tables;
-using EpicGames.UHT.Tokenizer;
-using EpicGames.UHT.Utils;
 using System;
 using System.Collections.Generic;
 using System.Text;
 using System.Text.Json.Serialization;
 using System.Threading;
+using EpicGames.UHT.Tables;
+using EpicGames.UHT.Tokenizer;
+using EpicGames.UHT.Utils;
 
 namespace EpicGames.UHT.Types
 {
@@ -78,16 +78,16 @@ namespace EpicGames.UHT.Types
 		/// <summary>
 		/// Get the name of the type where shared distinctions aren't made between different class and function types.
 		/// </summary>
-		/// <param name="EngineType">The engine type in question</param>
+		/// <param name="engineType">The engine type in question</param>
 		/// <returns>Lowercase name</returns>
 		/// <exception cref="UhtIceException">If the requested type isn't supported</exception>
-		public static string ShortLowercaseText(this UhtEngineType EngineType)
+		public static string ShortLowercaseText(this UhtEngineType engineType)
 		{
-			switch (EngineType)
+			switch (engineType)
 			{
 				case UhtEngineType.Header: return "header";
 				case UhtEngineType.Package: return "package";
-				case UhtEngineType.Class: 
+				case UhtEngineType.Class:
 				case UhtEngineType.Interface:
 				case UhtEngineType.NativeInterface: return "class";
 				case UhtEngineType.ScriptStruct: return "struct";
@@ -103,12 +103,12 @@ namespace EpicGames.UHT.Types
 		/// <summary>
 		/// Get the name of the type where shared distinctions aren't made between different class and function types.
 		/// </summary>
-		/// <param name="EngineType">The engine type in question</param>
+		/// <param name="engineType">The engine type in question</param>
 		/// <returns>Capitalized name</returns>
 		/// <exception cref="UhtIceException">If the requested type isn't supported</exception>
-		public static string ShortCapitalizedText(this UhtEngineType EngineType)
+		public static string ShortCapitalizedText(this UhtEngineType engineType)
 		{
-			switch (EngineType)
+			switch (engineType)
 			{
 				case UhtEngineType.Header: return "Header";
 				case UhtEngineType.Package: return "Package";
@@ -128,12 +128,12 @@ namespace EpicGames.UHT.Types
 		/// <summary>
 		/// Get the name of the type
 		/// </summary>
-		/// <param name="EngineType">The engine type in question</param>
+		/// <param name="engineType">The engine type in question</param>
 		/// <returns>Lowercase name</returns>
 		/// <exception cref="UhtIceException">If the requested type isn't supported</exception>
-		public static string LowercaseText(this UhtEngineType EngineType)
+		public static string LowercaseText(this UhtEngineType engineType)
 		{
-			switch (EngineType)
+			switch (engineType)
 			{
 				case UhtEngineType.Header: return "header";
 				case UhtEngineType.Package: return "package";
@@ -153,12 +153,12 @@ namespace EpicGames.UHT.Types
 		/// <summary>
 		/// Get the name of the type
 		/// </summary>
-		/// <param name="EngineType">The engine type in question</param>
+		/// <param name="engineType">The engine type in question</param>
 		/// <returns>Capitalized name</returns>
 		/// <exception cref="UhtIceException">If the requested type isn't supported</exception>
-		public static string CapitalizedText(this UhtEngineType EngineType)
+		public static string CapitalizedText(this UhtEngineType engineType)
 		{
-			switch (EngineType)
+			switch (engineType)
 			{
 				case UhtEngineType.Header: return "Header";
 				case UhtEngineType.Package: return "Package";
@@ -178,11 +178,11 @@ namespace EpicGames.UHT.Types
 		/// <summary>
 		/// Return true if the type must be unique in the symbol table
 		/// </summary>
-		/// <param name="EngineType">Type in question</param>
+		/// <param name="engineType">Type in question</param>
 		/// <returns>True if it must be unique</returns>
-		public static bool MustBeUnique(this UhtEngineType EngineType)
+		public static bool MustBeUnique(this UhtEngineType engineType)
 		{
-			switch (EngineType)
+			switch (engineType)
 			{
 				case UhtEngineType.Class:
 				case UhtEngineType.Interface:
@@ -199,11 +199,11 @@ namespace EpicGames.UHT.Types
 		/// <summary>
 		/// Find options to be used when attempting to find a duplicate name
 		/// </summary>
-		/// <param name="EngineType">Type in question</param>
+		/// <param name="engineType">Type in question</param>
 		/// <returns>Find options to be used for search</returns>
-		public static UhtFindOptions MustBeUniqueFindOptions(this UhtEngineType EngineType)
+		public static UhtFindOptions MustBeUniqueFindOptions(this UhtEngineType engineType)
 		{
-			switch (EngineType)
+			switch (engineType)
 			{
 				case UhtEngineType.Class:
 				case UhtEngineType.Interface:
@@ -222,11 +222,11 @@ namespace EpicGames.UHT.Types
 		/// <summary>
 		/// Return true if the type must not be a reserved name.
 		/// </summary>
-		/// <param name="EngineType">Type in question</param>
+		/// <param name="engineType">Type in question</param>
 		/// <returns>True if the type must not have a reserved name.</returns>
-		public static bool MustNotBeReserved(this UhtEngineType EngineType)
+		public static bool MustNotBeReserved(this UhtEngineType engineType)
 		{
-			switch (EngineType)
+			switch (engineType)
 			{
 				case UhtEngineType.Class:
 				case UhtEngineType.Interface:
@@ -243,21 +243,21 @@ namespace EpicGames.UHT.Types
 		/// <summary>
 		/// Return true if the type should be added to the engine symbol table
 		/// </summary>
-		/// <param name="EngineType">Type in question</param>
+		/// <param name="engineType">Type in question</param>
 		/// <returns>True if the type should be added to the symbol table</returns>
-		public static bool HasEngineName(this UhtEngineType EngineType)
+		public static bool HasEngineName(this UhtEngineType engineType)
 		{
-			return EngineType != UhtEngineType.NativeInterface;
+			return engineType != UhtEngineType.NativeInterface;
 		}
 
 		/// <summary>
 		/// Return true if the children of the type should be added to the symbol table
 		/// </summary>
-		/// <param name="EngineType">Type in question</param>
+		/// <param name="engineType">Type in question</param>
 		/// <returns>True if children should be added</returns>
-		public static bool AddChildrenToSymbolTable(this UhtEngineType EngineType)
+		public static bool AddChildrenToSymbolTable(this UhtEngineType engineType)
 		{
-			switch (EngineType)
+			switch (engineType)
 			{
 				case UhtEngineType.Property:
 					return false;
@@ -270,12 +270,12 @@ namespace EpicGames.UHT.Types
 		/// <summary>
 		/// Convert the engine type to find options
 		/// </summary>
-		/// <param name="EngineType">Engine type</param>
+		/// <param name="engineType">Engine type</param>
 		/// <returns>Find options</returns>
 		/// <exception cref="UhtIceException">Engine type is invalid</exception>
-		public static UhtFindOptions FindOptions(this UhtEngineType EngineType)
+		public static UhtFindOptions FindOptions(this UhtEngineType engineType)
 		{
-			switch (EngineType)
+			switch (engineType)
 			{
 				case UhtEngineType.Header: return 0;
 				case UhtEngineType.Package: return 0;
@@ -318,7 +318,7 @@ namespace EpicGames.UHT.Types
 		/// Search for UhtScriptStruct types
 		/// </summary>
 		ScriptStruct = 1 << 9,
-		
+
 		/// <summary>
 		/// Search for UhtClass types
 		/// </summary>
@@ -410,35 +410,35 @@ namespace EpicGames.UHT.Types
 		/// <summary>
 		/// Test to see if any of the specified flags are set
 		/// </summary>
-		/// <param name="InFlags">Current flags</param>
-		/// <param name="TestFlags">Flags to test for</param>
+		/// <param name="inFlags">Current flags</param>
+		/// <param name="testFlags">Flags to test for</param>
 		/// <returns>True if any of the flags are set</returns>
-		public static bool HasAnyFlags(this UhtFindOptions InFlags, UhtFindOptions TestFlags)
+		public static bool HasAnyFlags(this UhtFindOptions inFlags, UhtFindOptions testFlags)
 		{
-			return (InFlags & TestFlags) != 0;
+			return (inFlags & testFlags) != 0;
 		}
 
 		/// <summary>
 		/// Test to see if all of the specified flags are set
 		/// </summary>
-		/// <param name="InFlags">Current flags</param>
-		/// <param name="TestFlags">Flags to test for</param>
+		/// <param name="inFlags">Current flags</param>
+		/// <param name="testFlags">Flags to test for</param>
 		/// <returns>True if all the flags are set</returns>
-		public static bool HasAllFlags(this UhtFindOptions InFlags, UhtFindOptions TestFlags)
+		public static bool HasAllFlags(this UhtFindOptions inFlags, UhtFindOptions testFlags)
 		{
-			return (InFlags & TestFlags) == TestFlags;
+			return (inFlags & testFlags) == testFlags;
 		}
 
 		/// <summary>
 		/// Test to see if a specific set of flags have a specific value.
 		/// </summary>
-		/// <param name="InFlags">Current flags</param>
-		/// <param name="TestFlags">Flags to test for</param>
-		/// <param name="MatchFlags">Expected value of the tested flags</param>
+		/// <param name="inFlags">Current flags</param>
+		/// <param name="testFlags">Flags to test for</param>
+		/// <param name="matchFlags">Expected value of the tested flags</param>
 		/// <returns>True if the given flags have a specific value.</returns>
-		public static bool HasExactFlags(this UhtFindOptions InFlags, UhtFindOptions TestFlags, UhtFindOptions MatchFlags)
+		public static bool HasExactFlags(this UhtFindOptions inFlags, UhtFindOptions testFlags, UhtFindOptions matchFlags)
 		{
-			return (InFlags & TestFlags) == MatchFlags;
+			return (inFlags & testFlags) == matchFlags;
 		}
 	}
 
@@ -491,35 +491,35 @@ namespace EpicGames.UHT.Types
 		/// <summary>
 		/// Test to see if any of the specified flags are set
 		/// </summary>
-		/// <param name="InFlags">Current flags</param>
-		/// <param name="TestFlags">Flags to test for</param>
+		/// <param name="inFlags">Current flags</param>
+		/// <param name="testFlags">Flags to test for</param>
 		/// <returns>True if any of the flags are set</returns>
-		public static bool HasAnyFlags(this UhtValidationOptions InFlags, UhtValidationOptions TestFlags)
+		public static bool HasAnyFlags(this UhtValidationOptions inFlags, UhtValidationOptions testFlags)
 		{
-			return (InFlags & TestFlags) != 0;
+			return (inFlags & testFlags) != 0;
 		}
 
 		/// <summary>
 		/// Test to see if all of the specified flags are set
 		/// </summary>
-		/// <param name="InFlags">Current flags</param>
-		/// <param name="TestFlags">Flags to test for</param>
+		/// <param name="inFlags">Current flags</param>
+		/// <param name="testFlags">Flags to test for</param>
 		/// <returns>True if all the flags are set</returns>
-		public static bool HasAllFlags(this UhtValidationOptions InFlags, UhtValidationOptions TestFlags)
+		public static bool HasAllFlags(this UhtValidationOptions inFlags, UhtValidationOptions testFlags)
 		{
-			return (InFlags & TestFlags) == TestFlags;
+			return (inFlags & testFlags) == testFlags;
 		}
 
 		/// <summary>
 		/// Test to see if a specific set of flags have a specific value.
 		/// </summary>
-		/// <param name="InFlags">Current flags</param>
-		/// <param name="TestFlags">Flags to test for</param>
-		/// <param name="MatchFlags">Expected value of the tested flags</param>
+		/// <param name="inFlags">Current flags</param>
+		/// <param name="testFlags">Flags to test for</param>
+		/// <param name="matchFlags">Expected value of the tested flags</param>
 		/// <returns>True if the given flags have a specific value.</returns>
-		public static bool HasExactFlags(this UhtValidationOptions InFlags, UhtValidationOptions TestFlags, UhtValidationOptions MatchFlags)
+		public static bool HasExactFlags(this UhtValidationOptions inFlags, UhtValidationOptions testFlags, UhtValidationOptions matchFlags)
 		{
-			return (InFlags & TestFlags) == MatchFlags;
+			return (inFlags & testFlags) == matchFlags;
 		}
 	}
 
@@ -553,7 +553,7 @@ namespace EpicGames.UHT.Types
 	/// <summary>
 	/// Requested resolve phase
 	/// </summary>
-	public enum UhtResolvePhase : Int32
+	public enum UhtResolvePhase : int
 	{
 		/// <summary>
 		/// Initial resolve phase
@@ -587,72 +587,72 @@ namespace EpicGames.UHT.Types
 	/// </summary>
 	public class UhtDocumentationPolicy
 	{
-		private static readonly UhtDocumentationPolicy StrictDocumentationPolicy = new UhtDocumentationPolicy
+		private static readonly UhtDocumentationPolicy s_strictDocumentationPolicy = new UhtDocumentationPolicy
 		{
-			bPolicySet = true,
-			bClassOrStructCommentRequired = true,
-			bFunctionToolTipsRequired = true,
-			bMemberToolTipsRequired = true,
-			bParameterToolTipsRequired = true,
-			bFloatRangesRequired = true,
+			PolicySet = true,
+			ClassOrStructCommentRequired = true,
+			FunctionToolTipsRequired = true,
+			MemberToolTipsRequired = true,
+			ParameterToolTipsRequired = true,
+			FloatRangesRequired = true,
 		};
 
-		private static readonly UhtDocumentationPolicy NoDocumentationPolicy = new UhtDocumentationPolicy
+		private static readonly UhtDocumentationPolicy s_noDocumentationPolicy = new UhtDocumentationPolicy
 		{
-			bPolicySet = false,
-			bClassOrStructCommentRequired = false,
-			bFunctionToolTipsRequired = false,
-			bMemberToolTipsRequired = false,
-			bParameterToolTipsRequired = false,
-			bFloatRangesRequired = false,
+			PolicySet = false,
+			ClassOrStructCommentRequired = false,
+			FunctionToolTipsRequired = false,
+			MemberToolTipsRequired = false,
+			ParameterToolTipsRequired = false,
+			FloatRangesRequired = false,
 		};
 
 		/// <summary>
 		/// If true, this is a valid documentation policy that should be respected
 		/// </summary>
-		public bool bPolicySet { get; set; }
+		public bool PolicySet { get; set; }
 
 		/// <summary>
 		/// A tool tip must be defined
 		/// </summary>
-		public bool bClassOrStructCommentRequired { get; set; }
+		public bool ClassOrStructCommentRequired { get; set; }
 
 		/// <summary>
 		/// Functions must have a tool tip
 		/// </summary>
-		public bool bFunctionToolTipsRequired { get; set; }
+		public bool FunctionToolTipsRequired { get; set; }
 
 		/// <summary>
 		/// Property members must have a tool tip
 		/// </summary>
-		public bool bMemberToolTipsRequired { get; set; }
+		public bool MemberToolTipsRequired { get; set; }
 
 		/// <summary>
 		/// Function parameters must have a tool tip
 		/// </summary>
-		public bool bParameterToolTipsRequired { get; set; }
+		public bool ParameterToolTipsRequired { get; set; }
 
 		/// <summary>
 		/// Float properties must have a range specified
 		/// </summary>
-		public bool bFloatRangesRequired { get; set; }
+		public bool FloatRangesRequired { get; set; }
 
 		/// <summary>
 		/// Return the documentation policy given the name
 		/// </summary>
-		/// <param name="DocumentationPolicyName">Name of the policy</param>
-		/// <param name="Policy">Policy settings</param>
+		/// <param name="documentationPolicyName">Name of the policy</param>
+		/// <param name="policy">Policy settings</param>
 		/// <returns>True if the name is valid, false if not</returns>
-		public static bool TryGet(string DocumentationPolicyName, out UhtDocumentationPolicy Policy)
+		public static bool TryGet(string documentationPolicyName, out UhtDocumentationPolicy policy)
 		{
-			if (DocumentationPolicyName == "Strict")
+			if (documentationPolicyName == "Strict")
 			{
-				Policy = StrictDocumentationPolicy;
+				policy = s_strictDocumentationPolicy;
 				return true;
 			}
 
-			Policy = NoDocumentationPolicy;
-			return DocumentationPolicyName.Length == 0;
+			policy = s_noDocumentationPolicy;
+			return documentationPolicyName.Length == 0;
 		}
 	};
 
@@ -665,7 +665,7 @@ namespace EpicGames.UHT.Types
 		/// <summary>
 		/// Empty list of type used when children are requested but no children have been added.
 		/// </summary>
-		private static readonly List<UhtType> EmptyTypeList = new List<UhtType>();
+		private static readonly List<UhtType> s_emptyTypeList = new List<UhtType>();
 
 		/// <summary>
 		/// All UHT runs are associated with a given session.  The session holds all the global information for a run.
@@ -693,12 +693,12 @@ namespace EpicGames.UHT.Types
 		{
 			get
 			{
-				UhtType Type = this;
-				while (Type.Outer != null)
+				UhtType type = this;
+				while (type.Outer != null)
 				{
-					Type = Type.Outer;
+					type = type.Outer;
 				}
-				return (UhtPackage)Type;
+				return (UhtPackage)type;
 			}
 		}
 
@@ -710,12 +710,12 @@ namespace EpicGames.UHT.Types
 		{
 			get
 			{
-				UhtType Type = this;
-				while (Type.Outer != null && Type.Outer.Outer != null)
+				UhtType type = this;
+				while (type.Outer != null && type.Outer.Outer != null)
 				{
-					Type = Type.Outer;
+					type = type.Outer;
 				}
-				return (UhtHeaderFile)Type;
+				return (UhtHeaderFile)type;
 			}
 		}
 
@@ -727,7 +727,7 @@ namespace EpicGames.UHT.Types
 		/// <summary>
 		/// The name of the type as used by the engines type system.  If not set, will default to the source name
 		/// </summary>
-		public string EngineName { get => this.EngineNameInternal ?? this.SourceName; set => this.EngineNameInternal = value; }
+		public string EngineName { get => this._engineName ?? this.SourceName; set => this._engineName = value; }
 
 		/// <summary>
 		/// Simple enumeration that can be used to detect the different types of objects
@@ -743,13 +743,13 @@ namespace EpicGames.UHT.Types
 		/// <summary>
 		/// If true, then this type is visible to the symbol table.
 		/// </summary>
-		public bool bVisibleType { get; set; } = true;
+		public bool VisibleType { get; set; } = true;
 
 		/// <summary>
 		/// Return true if the type has been deprecated
 		/// </summary>
 		[JsonIgnore]
-		public virtual bool bDeprecated => false;
+		public virtual bool Deprecated => false;
 
 		/// <summary>
 		/// The line number of where the definition began
@@ -760,13 +760,13 @@ namespace EpicGames.UHT.Types
 		/// Return a combination of the engine type name followed by the path name of the type
 		/// </summary>
 		[JsonIgnore]
-		public virtual string FullName 
-		{ 
+		public virtual string FullName
+		{
 			get
 			{
-				StringBuilder Builder = new StringBuilder();
-				AppendFullName(Builder);
-				return Builder.ToString();
+				StringBuilder builder = new StringBuilder();
+				AppendFullName(builder);
+				return builder.ToString();
 			}
 		}
 
@@ -778,9 +778,9 @@ namespace EpicGames.UHT.Types
 		{
 			get
 			{
-				StringBuilder Builder = new StringBuilder();
-				AppendPathName(Builder, null);
-				return Builder.ToString();
+				StringBuilder builder = new StringBuilder();
+				AppendPathName(builder, null);
+				return builder.ToString();
 			}
 		}
 
@@ -788,12 +788,12 @@ namespace EpicGames.UHT.Types
 		/// Return the name of the documentation policy to be used
 		/// </summary>
 		[JsonIgnore]
-		public string DocumentationPolicyName { get => this.MetaData.GetValueOrDefaultHierarchical(UhtNames.DocumentationPolicy); }
+		public string DocumentationPolicyName => this.MetaData.GetValueOrDefaultHierarchical(UhtNames.DocumentationPolicy);
 
 		/// <summary>
 		/// Return the specifier validation table to be used to validate the meta data on this type
 		/// </summary>
-		protected virtual UhtSpecifierValidatorTable? SpecifierValidatorTable { get => null; }
+		protected virtual UhtSpecifierValidatorTable? SpecifierValidatorTable => null;
 
 		/// <summary>
 		/// Meta data associated with the type
@@ -806,7 +806,7 @@ namespace EpicGames.UHT.Types
 		/// </summary>
 		[JsonConverter(typeof(UhtNullableTypeListJsonConverter<UhtType>))]
 		//TODO - this should return a IReadOnlyList
-		public List<UhtType> Children { get => ChildrenInternal == null ? UhtType.EmptyTypeList : ChildrenInternal; }
+		public List<UhtType> Children => _children ?? UhtType.s_emptyTypeList;
 
 		/// <summary>
 		/// Helper to string method
@@ -814,9 +814,9 @@ namespace EpicGames.UHT.Types
 		/// <returns></returns>
 		public override string ToString() { return this.SourceName; }
 
-		private string? EngineNameInternal = null;
-		private List<UhtType>? ChildrenInternal = null;
-		private int ResolveState = (int)UhtResolvePhase.None * 2;
+		private string? _engineName = null;
+		private List<UhtType>? _children = null;
+		private int _resolveState = (int)UhtResolvePhase.None * 2;
 
 		#region IUHTMessageSite implementation
 		/// <inheritdoc/>
@@ -840,31 +840,31 @@ namespace EpicGames.UHT.Types
 		/// Construct the base type information.  This constructor is used exclusively by UHTPackage which is at 
 		/// the root of all type hierarchies.
 		/// </summary>
-		/// <param name="Session"></param>
-		protected UhtType(UhtSession Session)
+		/// <param name="session"></param>
+		protected UhtType(UhtSession session)
 		{
-			this.Session = Session;
+			this.Session = session;
 			this.Outer = null;
-			this.SourceName = string.Empty;
+			this.SourceName = String.Empty;
 			this.LineNumber = 1;
-			this.TypeIndex = Session.GetNextTypeIndex();
+			this.TypeIndex = session.GetNextTypeIndex();
 			this.MetaData = new UhtMetaData(this, this.Session.Config);
 		}
 
 		/// <summary>
 		/// Construct instance of the type given a parent type.
 		/// </summary>
-		/// <param name="Outer">The outer type of the type being constructed.  For example, a class defined in a given header will have the header as the outer.</param>
-		/// <param name="LineNumber">The line number in the source file where the type was defined.</param>
-		/// <param name="MetaData">Optional meta data to be associated with the type.</param>
-		protected UhtType(UhtType Outer, int LineNumber, UhtMetaData? MetaData = null)
+		/// <param name="outer">The outer type of the type being constructed.  For example, a class defined in a given header will have the header as the outer.</param>
+		/// <param name="lineNumber">The line number in the source file where the type was defined.</param>
+		/// <param name="metaData">Optional meta data to be associated with the type.</param>
+		protected UhtType(UhtType outer, int lineNumber, UhtMetaData? metaData = null)
 		{
-			this.Session = Outer.Session;
-			this.Outer = Outer;
-			this.SourceName = string.Empty;
-			this.LineNumber = LineNumber;
+			this.Session = outer.Session;
+			this.Outer = outer;
+			this.SourceName = String.Empty;
+			this.LineNumber = lineNumber;
 			this.TypeIndex = this.Session.GetNextTypeIndex();
-			this.MetaData = MetaData ?? new UhtMetaData(this, this.Session.Config);
+			this.MetaData = metaData ?? new UhtMetaData(this, this.Session.Config);
 			this.MetaData.MessageSite = this; // Make sure the site is correct
 			this.MetaData.Config = this.Session.Config;
 		}
@@ -872,14 +872,14 @@ namespace EpicGames.UHT.Types
 		/// <summary>
 		/// Add a type as a child
 		/// </summary>
-		/// <param name="Child">The child to be added.</param>
-		public virtual void AddChild(UhtType Child)
+		/// <param name="child">The child to be added.</param>
+		public virtual void AddChild(UhtType child)
 		{
-			if (this.ChildrenInternal == null)
+			if (this._children == null)
 			{
-				this.ChildrenInternal = new List<UhtType>();
+				this._children = new List<UhtType>();
 			}
-			this.ChildrenInternal.Add(Child);
+			this._children.Add(child);
 		}
 
 		/// <summary>
@@ -887,25 +887,25 @@ namespace EpicGames.UHT.Types
 		/// list associated with the type is cleared.
 		/// </summary>
 		/// <returns>The current list of children</returns>
-		public List<UhtType>? DetchChildren()
+		public List<UhtType>? DetachChildren()
 		{
-			List<UhtType>? Out = this.ChildrenInternal;
-			this.ChildrenInternal = null;
-			return Out;
+			List<UhtType>? children = this._children;
+			this._children = null;
+			return children;
 		}
 
 		/// <summary>
 		/// Merge the list children into the type.  The child's outer is set.
 		/// </summary>
-		/// <param name="Children">List of children to add.</param>
-		public void AddChildren(List<UhtType>? Children)
+		/// <param name="children">List of children to add.</param>
+		public void AddChildren(List<UhtType>? children)
 		{
-			if (Children != null)
+			if (children != null)
 			{
-				foreach (UhtType Child in Children)
+				foreach (UhtType child in children)
 				{
-					Child.Outer = this;
-					AddChild(Child);
+					child.Outer = this;
+					AddChild(child);
 				}
 			}
 		}
@@ -915,67 +915,66 @@ namespace EpicGames.UHT.Types
 		/// <summary>
 		/// Resolve the type.  To customize behavior, override ResolveSuper, ResolveSelf, or ResolveChildren.
 		/// </summary>
-		/// <param name="Phase">Resolution phase</param>
+		/// <param name="phase">Resolution phase</param>
 		/// <returns>False if the type fails any invalid checks during InvalidPhase check.</returns>
-		public bool Resolve(UhtResolvePhase Phase)
+		public bool Resolve(UhtResolvePhase phase)
 		{
-			int InitState = (int)Phase << 1;
+			int initState = (int)phase << 1;
+			int currentState = Interlocked.Add(ref this._resolveState, 0);
 
-			int CurrentState = Interlocked.Add(ref this.ResolveState, 0);
-			
 			// If we are already passed the init phase for this state, then we are good
-			if (CurrentState > InitState)
+			if (currentState > initState)
 			{
 				return true;
 			}
 
 			// Attempt to make myself the resolve thread for this type.  If we aren't the ones
 			// to get the init state, then we have to wait until the thread that did completes
-			CurrentState = Interlocked.CompareExchange(ref this.ResolveState, InitState, CurrentState);
-			if (CurrentState > InitState)
+			currentState = Interlocked.CompareExchange(ref this._resolveState, initState, currentState);
+			if (currentState > initState)
 			{
 				return true;
 			}
-			else if (CurrentState == InitState)
+			else if (currentState == initState)
 			{
 				do
 				{
 					Thread.Yield();
-					CurrentState = Interlocked.Add(ref this.ResolveState, 0);
-				} while (CurrentState == InitState);
+					currentState = Interlocked.Add(ref this._resolveState, 0);
+				} while (currentState == initState);
 				return true;
 			}
 			else
 			{
-				bool bResult = true;
+				bool result = true;
 				try
 				{
-					ResolveSuper(Phase);
-					bResult = ResolveSelf(Phase);
-					ResolveChildren(Phase);
+					ResolveSuper(phase);
+					result = ResolveSelf(phase);
+					ResolveChildren(phase);
 				}
 				finally
 				{
-					Interlocked.Increment(ref this.ResolveState);
+					Interlocked.Increment(ref this._resolveState);
 				}
-				return bResult;
+				return result;
 			}
 		}
 
 		/// <summary>
 		/// Resolve any super types
 		/// </summary>
-		/// <param name="Phase">Resolution phase</param>
-		protected virtual void ResolveSuper(UhtResolvePhase Phase)
+		/// <param name="phase">Resolution phase</param>
+		protected virtual void ResolveSuper(UhtResolvePhase phase)
 		{
 		}
 
 		/// <summary>
 		/// Resolve self
 		/// </summary>
-		/// <param name="Phase">Resolution phase</param>
+		/// <param name="phase">Resolution phase</param>
 		/// <returns>False if the type fails any invalid checks during InvalidPhase check.</returns>
-		protected virtual bool ResolveSelf(UhtResolvePhase Phase)
+		protected virtual bool ResolveSelf(UhtResolvePhase phase)
 		{
 			return true;
 		}
@@ -983,32 +982,32 @@ namespace EpicGames.UHT.Types
 		/// <summary>
 		/// Resolve children
 		/// </summary>
-		/// <param name="Phase">Resolution phase</param>
-		protected virtual void ResolveChildren(UhtResolvePhase Phase)
+		/// <param name="phase">Resolution phase</param>
+		protected virtual void ResolveChildren(UhtResolvePhase phase)
 		{
-			if (this.ChildrenInternal != null)
+			if (this._children != null)
 			{
-				if (Phase == UhtResolvePhase.InvalidCheck)
+				if (phase == UhtResolvePhase.InvalidCheck)
 				{
-					int OutIndex = 0;
-					for (int Index = 0; Index < this.ChildrenInternal.Count; ++Index)
+					int outIndex = 0;
+					for (int index = 0; index < this._children.Count; ++index)
 					{
-						UhtType Child = this.ChildrenInternal[Index];
-						if (Child.Resolve(Phase))
+						UhtType child = this._children[index];
+						if (child.Resolve(phase))
 						{
-							this.ChildrenInternal[OutIndex++] = Child;
+							this._children[outIndex++] = child;
 						}
 					}
-					if (OutIndex < this.ChildrenInternal.Count)
+					if (outIndex < this._children.Count)
 					{
-						this.ChildrenInternal.RemoveRange(OutIndex, this.ChildrenInternal.Count - OutIndex);
+						this._children.RemoveRange(outIndex, this._children.Count - outIndex);
 					}
 				}
 				else
 				{
-					foreach (UhtType Child in this.ChildrenInternal)
+					foreach (UhtType child in this._children)
 					{
-						Child.Resolve(Phase);
+						child.Resolve(phase);
 					}
 				}
 			}
@@ -1019,52 +1018,52 @@ namespace EpicGames.UHT.Types
 		/// <summary>
 		/// Find the given type in the type hierarchy
 		/// </summary>
-		/// <param name="Options">Options controlling what is searched</param>
-		/// <param name="Name">Name of the type.</param>
-		/// <param name="MessageSite">If supplied, then a error message will be generated if the type can not be found</param>
-		/// <param name="LineNumber">Source code line number requesting the lookup.</param>
+		/// <param name="options">Options controlling what is searched</param>
+		/// <param name="name">Name of the type.</param>
+		/// <param name="messageSite">If supplied, then a error message will be generated if the type can not be found</param>
+		/// <param name="lineNumber">Source code line number requesting the lookup.</param>
 		/// <returns>The located type of null if not found</returns>
-		public UhtType? FindType(UhtFindOptions Options, string Name, IUhtMessageSite? MessageSite = null, int LineNumber = -1)
+		public UhtType? FindType(UhtFindOptions options, string name, IUhtMessageSite? messageSite = null, int lineNumber = -1)
 		{
-			return this.Session.FindType(this, Options, Name, MessageSite, LineNumber);
+			return this.Session.FindType(this, options, name, messageSite, lineNumber);
 		}
 
 		/// <summary>
 		/// Find the given type in the type hierarchy
 		/// </summary>
-		/// <param name="Options">Options controlling what is searched</param>
-		/// <param name="Name">Name of the type.</param>
-		/// <param name="MessageSite">If supplied, then a error message will be generated if the type can not be found</param>
+		/// <param name="options">Options controlling what is searched</param>
+		/// <param name="name">Name of the type.</param>
+		/// <param name="messageSite">If supplied, then a error message will be generated if the type can not be found</param>
 		/// <returns>The located type of null if not found</returns>
-		public UhtType? FindType(UhtFindOptions Options, ref UhtToken Name, IUhtMessageSite? MessageSite = null)
+		public UhtType? FindType(UhtFindOptions options, ref UhtToken name, IUhtMessageSite? messageSite = null)
 		{
-			return this.Session.FindType(this, Options, ref Name, MessageSite);
+			return this.Session.FindType(this, options, ref name, messageSite);
 		}
 
 		/// <summary>
 		/// Find the given type in the type hierarchy
 		/// </summary>
-		/// <param name="Options">Options controlling what is searched</param>
-		/// <param name="Identifiers">Enumeration of identifiers.</param>
-		/// <param name="MessageSite">If supplied, then a error message will be generated if the type can not be found</param>
-		/// <param name="LineNumber">Source code line number requesting the lookup.</param>
+		/// <param name="options">Options controlling what is searched</param>
+		/// <param name="identifiers">Enumeration of identifiers.</param>
+		/// <param name="messageSite">If supplied, then a error message will be generated if the type can not be found</param>
+		/// <param name="lineNumber">Source code line number requesting the lookup.</param>
 		/// <returns>The located type of null if not found</returns>
-		public UhtType? FindType(UhtFindOptions Options, UhtTokenList Identifiers, IUhtMessageSite? MessageSite = null, int LineNumber = -1)
+		public UhtType? FindType(UhtFindOptions options, UhtTokenList identifiers, IUhtMessageSite? messageSite = null, int lineNumber = -1)
 		{
-			return this.Session.FindType(this, Options, Identifiers, MessageSite, LineNumber);
+			return this.Session.FindType(this, options, identifiers, messageSite, lineNumber);
 		}
 
 		/// <summary>
 		/// Find the given type in the type hierarchy
 		/// </summary>
-		/// <param name="Options">Options controlling what is searched</param>
-		/// <param name="Identifiers">Enumeration of identifiers.</param>
-		/// <param name="MessageSite">If supplied, then a error message will be generated if the type can not be found</param>
-		/// <param name="LineNumber">Source code line number requesting the lookup.</param>
+		/// <param name="options">Options controlling what is searched</param>
+		/// <param name="identifiers">Enumeration of identifiers.</param>
+		/// <param name="messageSite">If supplied, then a error message will be generated if the type can not be found</param>
+		/// <param name="lineNumber">Source code line number requesting the lookup.</param>
 		/// <returns>The located type of null if not found</returns>
-		public UhtType? FindType(UhtFindOptions Options, UhtToken[] Identifiers, IUhtMessageSite? MessageSite = null, int LineNumber = -1)
+		public UhtType? FindType(UhtFindOptions options, UhtToken[] identifiers, IUhtMessageSite? messageSite = null, int lineNumber = -1)
 		{
-			return this.Session.FindType(this, Options, Identifiers, MessageSite, LineNumber);
+			return this.Session.FindType(this, options, identifiers, messageSite, lineNumber);
 		}
 		#endregion
 
@@ -1072,28 +1071,28 @@ namespace EpicGames.UHT.Types
 		/// <summary>
 		/// Validate the type settings
 		/// </summary>
-		/// <param name="Options">Validation options</param>
+		/// <param name="options">Validation options</param>
 		/// <returns>Updated validation options.  This will be used to validate the children.</returns>
-		protected virtual UhtValidationOptions Validate(UhtValidationOptions Options)
+		protected virtual UhtValidationOptions Validate(UhtValidationOptions options)
 		{
 			ValidateMetaData();
 			ValidateDocumentationPolicy();
-			return Options;
+			return options;
 		}
 
 		private void ValidateMetaData()
 		{
-			UhtSpecifierValidatorTable? Table = this.SpecifierValidatorTable;
-			if (Table != null)
+			UhtSpecifierValidatorTable? table = this.SpecifierValidatorTable;
+			if (table != null)
 			{
 				if (this.MetaData.Dictionary != null && this.MetaData.Dictionary.Count > 0)
 				{
-					foreach (var KVP in this.MetaData.Dictionary)
+					foreach (KeyValuePair<UhtMetaDataKey, string> kvp in this.MetaData.Dictionary)
 					{
-						UhtSpecifierValidator? Specifier;
-						if (Table.TryGetValue(KVP.Key.Name, out Specifier))
+						UhtSpecifierValidator? specifier;
+						if (table.TryGetValue(kvp.Key.Name, out specifier))
 						{
-							Specifier.Delegate(this, MetaData, KVP.Key, KVP.Value);
+							specifier.Delegate(this, MetaData, kvp.Key, kvp.Value);
 						}
 					}
 				}
@@ -1102,11 +1101,11 @@ namespace EpicGames.UHT.Types
 
 		private void ValidateDocumentationPolicy()
 		{
-			if (UhtDocumentationPolicy.TryGet(this.DocumentationPolicyName, out UhtDocumentationPolicy Policy))
+			if (UhtDocumentationPolicy.TryGet(this.DocumentationPolicyName, out UhtDocumentationPolicy policy))
 			{
-				if (Policy.bPolicySet)
+				if (policy.PolicySet)
 				{
-					this.ValidateDocumentationPolicy(Policy);
+					this.ValidateDocumentationPolicy(policy);
 				}
 			}
 			else
@@ -1118,25 +1117,25 @@ namespace EpicGames.UHT.Types
 		/// <summary>
 		/// Validate the documentation policy
 		/// </summary>
-		/// <param name="Policy">The expected documentation policy</param>
-		protected virtual void ValidateDocumentationPolicy(UhtDocumentationPolicy Policy)
+		/// <param name="policy">The expected documentation policy</param>
+		protected virtual void ValidateDocumentationPolicy(UhtDocumentationPolicy policy)
 		{
 		}
 
 		/// <summary>
 		/// Validate the given type.  Validates itself and the children
 		/// </summary>
-		/// <param name="Type">Type to validate</param>
-		/// <param name="Options">Validation options</param>
-		public static void ValidateType(UhtType Type, UhtValidationOptions Options)
+		/// <param name="type">Type to validate</param>
+		/// <param name="options">Validation options</param>
+		public static void ValidateType(UhtType type, UhtValidationOptions options)
 		{
 			// Invoke the new method
-			Options = Type.Validate(Options);
+			options = type.Validate(options);
 
 			// Invoke validate on the children
-			foreach (UhtType Child in Type.Children)
+			foreach (UhtType child in type.Children)
 			{
-				ValidateType(Child, Options);
+				ValidateType(child, options);
 			}
 		}
 		#endregion
@@ -1145,56 +1144,56 @@ namespace EpicGames.UHT.Types
 		/// <summary>
 		/// Return the path name of the type which includes the outers
 		/// </summary>
-		/// <param name="Builder">Destination builder</param>
-		/// <param name="StopOuter">Type to stop at when generating the path</param>
-		public virtual void AppendPathName(StringBuilder Builder, UhtType? StopOuter = null)
+		/// <param name="builder">Destination builder</param>
+		/// <param name="stopOuter">Type to stop at when generating the path</param>
+		public virtual void AppendPathName(StringBuilder builder, UhtType? stopOuter = null)
 		{
-			AppendPathNameInternal(Builder, StopOuter);
-			if (Builder.Length == 0)
+			AppendPathNameInternal(builder, stopOuter);
+			if (builder.Length == 0)
 			{
-				Builder.Append("None");
+				builder.Append("None");
 			}
 		}
 
-		private void AppendPathNameInternal(StringBuilder Builder, UhtType? StopOuter = null)
+		private void AppendPathNameInternal(StringBuilder builder, UhtType? stopOuter = null)
 		{
-			if (this != StopOuter)
+			if (this != stopOuter)
 			{
-				UhtType? Outer = this.Outer;
-				if (Outer != null && Outer != StopOuter)
+				UhtType? outer = this.Outer;
+				if (outer != null && outer != stopOuter)
 				{
-					Outer.AppendPathName(Builder, StopOuter);
+					outer.AppendPathName(builder, stopOuter);
 
 					// SubObjectDelimiter is used to indicate that this object's outer is not a UPackage
 					// In the C# version of UHT, we key off the header file
-					if (!(Outer is UhtHeaderFile) && Outer.Outer is UhtHeaderFile)
+					if (outer is not UhtHeaderFile && outer.Outer is UhtHeaderFile)
 					{
-						Builder.Append(UhtFCString.SubObjectDelimiter);
+						builder.Append(UhtFCString.SubObjectDelimiter);
 					}
 
 					// SubObjectDelimiter is also used to bridge between the object:properties in the path
-					else if (!(Outer is UhtProperty) && this is UhtProperty)
+					else if (outer is not UhtProperty && this is UhtProperty)
 					{
-						Builder.Append(UhtFCString.SubObjectDelimiter);
+						builder.Append(UhtFCString.SubObjectDelimiter);
 					}
 					else
 					{
-						Builder.Append('.');
+						builder.Append('.');
 					}
 				}
-				Builder.Append(this.EngineName);
+				builder.Append(this.EngineName);
 			}
 		}
 
 		/// <summary>
 		/// Get the full type name which is the engine class name followed by the path name
 		/// </summary>
-		/// <param name="Builder">Destination builder</param>
-		public virtual void AppendFullName(StringBuilder Builder)
+		/// <param name="builder">Destination builder</param>
+		public virtual void AppendFullName(StringBuilder builder)
 		{
-			Builder.Append(this.EngineClassName);
-			Builder.Append(' ');
-			AppendPathName(Builder);
+			builder.Append(this.EngineClassName);
+			builder.Append(' ');
+			AppendPathName(builder);
 		}
 
 		/// <summary>
@@ -1203,10 +1202,10 @@ namespace EpicGames.UHT.Types
 		/// <returns>The display name for this object.</returns>
 		public string GetDisplayNameText()
 		{
-			string? NativeDisplayName;
-			if (this.MetaData.TryGetValue(UhtNames.DisplayName, out NativeDisplayName))
+			string? nativeDisplayName;
+			if (this.MetaData.TryGetValue(UhtNames.DisplayName, out nativeDisplayName))
 			{
-				return NativeDisplayName;
+				return nativeDisplayName;
 			}
 
 			return GetDisplayStringFromEngineName();
@@ -1215,28 +1214,28 @@ namespace EpicGames.UHT.Types
 		/// <summary>
 		/// Finds the localized tooltip or native tooltip as a fallback.
 		/// </summary>
-		/// <param name="bShortToolTip">Look for a shorter version of the tooltip (falls back to the long tooltip if none was specified)</param>
+		/// <param name="shortToolTip">Look for a shorter version of the tooltip (falls back to the long tooltip if none was specified)</param>
 		/// <returns>The tooltip for this object.</returns>
-		public virtual string GetToolTipText(bool bShortToolTip = false)
+		public virtual string GetToolTipText(bool shortToolTip = false)
 		{
-			string NativeToolTip = string.Empty;
-			if (bShortToolTip)
+			string nativeToolTip = String.Empty;
+			if (shortToolTip)
 			{
-				NativeToolTip = this.MetaData.GetValueOrDefault(UhtNames.ShortToolTip);
+				nativeToolTip = this.MetaData.GetValueOrDefault(UhtNames.ShortToolTip);
 			}
 
-			if (NativeToolTip.Length == 0)
+			if (nativeToolTip.Length == 0)
 			{
-				NativeToolTip = this.MetaData.GetValueOrDefault(UhtNames.ToolTip);
+				nativeToolTip = this.MetaData.GetValueOrDefault(UhtNames.ToolTip);
 			}
 
-			if (NativeToolTip.Length == 0)
+			if (nativeToolTip.Length == 0)
 			{
 				return GetDisplayStringFromEngineName();
 			}
 			else
 			{
-				return NativeToolTip.TrimEnd().Replace("@see", "See:", StringComparison.Ordinal);
+				return nativeToolTip.TrimEnd().Replace("@see", "See:", StringComparison.Ordinal);
 			}
 		}
 
@@ -1253,8 +1252,8 @@ namespace EpicGames.UHT.Types
 		/// <summary>
 		/// Collect all the references for the type
 		/// </summary>
-		/// <param name="Collector">Object collecting the references</param>
-		public virtual void CollectReferences(IUhtReferenceCollector Collector)
+		/// <param name="collector">Object collecting the references</param>
+		public virtual void CollectReferences(IUhtReferenceCollector collector)
 		{
 		}
 	}
@@ -1267,49 +1266,49 @@ namespace EpicGames.UHT.Types
 		/// <summary>
 		/// Append all of the outer names associated with the type
 		/// </summary>
-		/// <param name="Builder">Output string builder</param>
-		/// <param name="Outer">Outer to be appended</param>
+		/// <param name="builder">Output string builder</param>
+		/// <param name="outer">Outer to be appended</param>
 		/// <returns>Output string builder</returns>
-		public static StringBuilder AppendOuterNames(this StringBuilder Builder, UhtType? Outer)
+		public static StringBuilder AppendOuterNames(this StringBuilder builder, UhtType? outer)
 		{
-			if (Outer == null)
+			if (outer == null)
 			{
-				return Builder;
+				return builder;
 			}
 
-			if (Outer is UhtClass)
+			if (outer is UhtClass)
 			{
-				Builder.Append('_');
-				Builder.Append(Outer.SourceName);
+				builder.Append('_');
+				builder.Append(outer.SourceName);
 			}
-			else if (Outer is UhtScriptStruct)
+			else if (outer is UhtScriptStruct)
 			{
 				// Structs can also have UPackage outer.
-				if (!(Outer.Outer is UhtHeaderFile))
+				if (outer.Outer is not UhtHeaderFile)
 				{
-					AppendOuterNames(Builder, Outer.Outer);
+					AppendOuterNames(builder, outer.Outer);
 				}
-				Builder.Append('_');
-				Builder.Append(Outer.SourceName);
+				builder.Append('_');
+				builder.Append(outer.SourceName);
 			}
-			else if (Outer is UhtPackage Package)
+			else if (outer is UhtPackage package)
 			{
-				Builder.Append('_');
-				Builder.Append(Package.ShortName);
+				builder.Append('_');
+				builder.Append(package.ShortName);
 			}
-			else if (Outer is UhtHeaderFile)
+			else if (outer is UhtHeaderFile)
 			{
 				// Pickup the package
-				AppendOuterNames(Builder, Outer.Outer);
+				AppendOuterNames(builder, outer.Outer);
 			}
 			else
 			{
-				AppendOuterNames(Builder, Outer.Outer);
-				Builder.Append('_');
-				Builder.Append(Outer.EngineName);
+				AppendOuterNames(builder, outer.Outer);
+				builder.Append('_');
+				builder.Append(outer.EngineName);
 			}
 
-			return Builder;
+			return builder;
 		}
 	}
 }
diff --git a/Engine/Source/Programs/Shared/EpicGames.UHT/Utils/UhtBuffer.cs b/Engine/Source/Programs/Shared/EpicGames.UHT/Utils/UhtBuffer.cs
--- a/Engine/Source/Programs/Shared/EpicGames.UHT/Utils/UhtBuffer.cs
+++ b/Engine/Source/Programs/Shared/EpicGames.UHT/Utils/UhtBuffer.cs
@@ -23,17 +23,17 @@ namespace EpicGames.UHT.Utils
 		/// <summary>
 		/// Adjustment to the bucket index to account for the minimum bucket size
 		/// </summary>
-		private static readonly int BuckedAdjustment = BitOperations.Log2((uint)UhtBuffer.MinSize);
-		
+		private static readonly int s_buckedAdjustment = BitOperations.Log2((uint)UhtBuffer.MinSize);
+
 		/// <summary>
 		/// Total number of supported buckets
 		/// </summary>
-		private static readonly int BucketCount = 32 - UhtBuffer.BuckedAdjustment;
+		private static readonly int s_bucketCount = 32 - UhtBuffer.s_buckedAdjustment;
 
 		/// <summary>
 		/// Bucket lookaside list
 		/// </summary>
-		private static readonly UhtBuffer?[] LookAsideArray = new UhtBuffer?[UhtBuffer.BucketCount];
+		private static readonly UhtBuffer?[] s_lookAsideArray = new UhtBuffer?[UhtBuffer.s_bucketCount];
 
 		/// <summary>
 		/// The bucket index associated with the buffer
@@ -49,7 +49,8 @@ namespace EpicGames.UHT.Utils
 		/// The backing character block.  The size of the array will normally be larger than the 
 		/// requested size.
 		/// </summary>
-		[System.Diagnostics.CodeAnalysis.SuppressMessage("Design", "CA1051:Do not declare visible instance fields", Justification = "")]
+		[System.Diagnostics.CodeAnalysis.SuppressMessage("Design", "CA1051:Do not declare visible instance fields", Justification = "<Pending>")]
+		[System.Diagnostics.CodeAnalysis.SuppressMessage("Style", "IDE1006:Naming Styles", Justification = "<Pending>")]
 		public char[] Block;
 
 		/// <summary>
@@ -60,113 +61,113 @@ namespace EpicGames.UHT.Utils
 		/// <summary>
 		/// Construct a new buffer
 		/// </summary>
-		/// <param name="Size">The initial size of the buffer</param>
-		/// <param name="Bucket">The bucket associated with the buffer</param>
-		/// <param name="BucketSize">The size all blocks in this bucket</param>
-		private UhtBuffer(int Size, int Bucket, int BucketSize)
+		/// <param name="size">The initial size of the buffer</param>
+		/// <param name="bucket">The bucket associated with the buffer</param>
+		/// <param name="bucketSize">The size all blocks in this bucket</param>
+		private UhtBuffer(int size, int bucket, int bucketSize)
 		{
-			this.Block = new char[BucketSize];
-			this.Bucket = Bucket;
-			Reset(Size);
+			this.Block = new char[bucketSize];
+			this.Bucket = bucket;
+			Reset(size);
 		}
 
 		/// <summary>
 		/// Reset the memory region to the given size
 		/// </summary>
-		/// <param name="Size"></param>
-		public void Reset(int Size)
+		/// <param name="size"></param>
+		public void Reset(int size)
 		{
-			this.Memory = new Memory<char>(this.Block, 0, Size);
+			this.Memory = new Memory<char>(this.Block, 0, size);
 		}
 
 		/// <summary>
 		/// Borrow a new buffer of the given size
 		/// </summary>
-		/// <param name="Size">Size of the buffer</param>
+		/// <param name="size">Size of the buffer</param>
 		/// <returns>Buffer that should be returned with a call to Return</returns>
-		public static UhtBuffer Borrow(int Size)
+		public static UhtBuffer Borrow(int size)
 		{
-			if (Size <= UhtBuffer.MinSize)
+			if (size <= UhtBuffer.MinSize)
 			{
-				return BorrowInternal(Size, 0, UhtBuffer.MinSize);
+				return BorrowInternal(size, 0, UhtBuffer.MinSize);
 			}
 			else
 			{
 
 				// Round up the size to the next larger power of two if it isn't a power of two
-				uint USize = (uint)Size;
-				--USize;
-				USize |= USize >> 1;
-				USize |= USize >> 2;
-				USize |= USize >> 4;
-				USize |= USize >> 8;
-				USize |= USize >> 16;
-				++USize;
-				int Bucket = BitOperations.Log2(USize) - UhtBuffer.BuckedAdjustment;
-				return BorrowInternal(Size, Bucket, (int)USize);
+				uint usize = (uint)size;
+				--usize;
+				usize |= usize >> 1;
+				usize |= usize >> 2;
+				usize |= usize >> 4;
+				usize |= usize >> 8;
+				usize |= usize >> 16;
+				++usize;
+				int bucket = BitOperations.Log2(usize) - UhtBuffer.s_buckedAdjustment;
+				return BorrowInternal(size, bucket, (int)usize);
 			}
 		}
 
 		/// <summary>
 		/// Return a buffer initialized with the string builder.
 		/// </summary>
-		/// <param name="Builder">Source builder content</param>
+		/// <param name="builder">Source builder content</param>
 		/// <returns>Buffer that should be returned with a call to Return</returns>
-		public static UhtBuffer Borrow(StringBuilder Builder)
+		public static UhtBuffer Borrow(StringBuilder builder)
 		{
-			int Length = Builder.Length;
-			UhtBuffer Buffer = Borrow(Length);
-			Builder.CopyTo(0, Buffer.Memory.Span, Length);
-			return Buffer;
+			int length = builder.Length;
+			UhtBuffer buffer = Borrow(length);
+			builder.CopyTo(0, buffer.Memory.Span, length);
+			return buffer;
 		}
 
 		/// <summary>
 		/// Return a buffer initialized with the string builder sub string.
 		/// </summary>
-		/// <param name="Builder">Source builder content</param>
-		/// <param name="StartIndex">Starting index in the builder</param>
-		/// <param name="Length">Length of the content</param>
+		/// <param name="builder">Source builder content</param>
+		/// <param name="startIndex">Starting index in the builder</param>
+		/// <param name="length">Length of the content</param>
 		/// <returns>Buffer that should be returned with a call to Return</returns>
-		public static UhtBuffer Borrow(StringBuilder Builder, int StartIndex, int Length)
+		public static UhtBuffer Borrow(StringBuilder builder, int startIndex, int length)
 		{
-			UhtBuffer Buffer = Borrow(Length);
-			Builder.CopyTo(StartIndex, Buffer.Memory.Span, Length);
-			return Buffer;
+			UhtBuffer buffer = Borrow(length);
+			builder.CopyTo(startIndex, buffer.Memory.Span, length);
+			return buffer;
 		}
 
 		/// <summary>
 		/// Return the buffer to the cache.  The buffer should no longer be accessed.
 		/// </summary>
-		/// <param name="Buffer">The buffer to be returned.</param>
-		public static void Return(UhtBuffer Buffer)
+		/// <param name="buffer">The buffer to be returned.</param>
+		public static void Return(UhtBuffer buffer)
 		{
-			lock (UhtBuffer.LookAsideArray)
+			lock (UhtBuffer.s_lookAsideArray)
 			{
-				Buffer.NextBuffer = UhtBuffer.LookAsideArray[Buffer.Bucket];
-				UhtBuffer.LookAsideArray[Buffer.Bucket] = Buffer;
+				buffer.NextBuffer = UhtBuffer.s_lookAsideArray[buffer.Bucket];
+				UhtBuffer.s_lookAsideArray[buffer.Bucket] = buffer;
 			}
 		}
 
 		/// <summary>
 		/// Internal helper to allocate a buffer
 		/// </summary>
-		/// <param name="Size">The initial size of the buffer</param>
-		/// <param name="Bucket">The bucket associated with the buffer</param>
-		/// <param name="BucketSize">The size all blocks in this bucket</param>
+		/// <param name="size">The initial size of the buffer</param>
+		/// <param name="bucket">The bucket associated with the buffer</param>
+		/// <param name="bucketSize">The size all blocks in this bucket</param>
 		/// <returns>The allocated buffer</returns>
-		private static UhtBuffer BorrowInternal(int Size, int Bucket, int BucketSize)
+		private static UhtBuffer BorrowInternal(int size, int bucket, int bucketSize)
 		{
-			lock (UhtBuffer.LookAsideArray)
+			lock (UhtBuffer.s_lookAsideArray)
 			{
-				if (UhtBuffer.LookAsideArray[Bucket] != null)
+				if (UhtBuffer.s_lookAsideArray[bucket] != null)
 				{
-					UhtBuffer Buffer = UhtBuffer.LookAsideArray[Bucket]!;
-					UhtBuffer.LookAsideArray[Bucket] = Buffer.NextBuffer;
-					Buffer.Reset(Size);
-					return Buffer;
+					UhtBuffer buffer = UhtBuffer.s_lookAsideArray[bucket]!;
+					UhtBuffer.s_lookAsideArray[bucket] = buffer.NextBuffer;
+					buffer.Reset(size);
+					return buffer;
 				}
 			}
-			return new UhtBuffer(Size, Bucket, BucketSize);
+			return new UhtBuffer(size, bucket, bucketSize);
 		}
 	}
 
@@ -184,30 +185,30 @@ namespace EpicGames.UHT.Utils
 		/// <summary>
 		/// Borrow a buffer with the given size
 		/// </summary>
-		/// <param name="Size">The size to borrow</param>
-		public UhtBorrowBuffer(int Size)
+		/// <param name="size">The size to borrow</param>
+		public UhtBorrowBuffer(int size)
 		{
-			this.Buffer = UhtBuffer.Borrow(Size);
+			this.Buffer = UhtBuffer.Borrow(size);
 		}
 
 		/// <summary>
 		/// Borrow a buffer populated with the builder contents
 		/// </summary>
-		/// <param name="Builder">Initial contents of the buffer</param>
-		public UhtBorrowBuffer(StringBuilder Builder)
+		/// <param name="builder">Initial contents of the buffer</param>
+		public UhtBorrowBuffer(StringBuilder builder)
 		{
-			this.Buffer = UhtBuffer.Borrow(Builder);
+			this.Buffer = UhtBuffer.Borrow(builder);
 		}
 
 		/// <summary>
 		/// Borrow a buffer populated with the builder contents
 		/// </summary>
-		/// <param name="Builder">Initial contents of the buffer</param>
-		/// <param name="StartIndex">Starting index into the builder</param>
-		/// <param name="Length">Length of the data in the builder</param>
-		public UhtBorrowBuffer(StringBuilder Builder, int StartIndex, int Length)
+		/// <param name="builder">Initial contents of the buffer</param>
+		/// <param name="startIndex">Starting index into the builder</param>
+		/// <param name="length">Length of the data in the builder</param>
+		public UhtBorrowBuffer(StringBuilder builder, int startIndex, int length)
 		{
-			this.Buffer = UhtBuffer.Borrow(Builder, StartIndex, Length);
+			this.Buffer = UhtBuffer.Borrow(builder, startIndex, length);
 		}
 
 		/// <summary>
diff --git a/Engine/Source/Programs/Shared/EpicGames.UHT/Utils/UhtConfig.cs b/Engine/Source/Programs/Shared/EpicGames.UHT/Utils/UhtConfig.cs
--- a/Engine/Source/Programs/Shared/EpicGames.UHT/Utils/UhtConfig.cs
+++ b/Engine/Source/Programs/Shared/EpicGames.UHT/Utils/UhtConfig.cs
@@ -72,50 +72,50 @@ namespace EpicGames.UHT.Utils
 		/// <summary>
 		/// If the token references a remapped identifier, update the value in the token 
 		/// </summary>
-		/// <param name="Token">Token to be remapped</param>
-		public void RedirectTypeIdentifier(ref UhtToken Token);
+		/// <param name="token">Token to be remapped</param>
+		public void RedirectTypeIdentifier(ref UhtToken token);
 
 		/// <summary>
 		/// Return the remapped key or the existing key
 		/// </summary>
-		/// <param name="Key">Key to be remapped.</param>
-		/// <param name="NewKey">Resulting key name</param>
+		/// <param name="key">Key to be remapped.</param>
+		/// <param name="newKey">Resulting key name</param>
 		/// <returns>True if the key has been remapped and has changed.</returns>
-		public bool RedirectMetaDataKey(string Key, out string NewKey);
+		public bool RedirectMetaDataKey(string key, out string newKey);
 
 		/// <summary>
 		/// Test to see if the given units are valid.
 		/// </summary>
-		/// <param name="Units">Units to test</param>
+		/// <param name="units">Units to test</param>
 		/// <returns>True if the units are valid, false if not</returns>
-		public bool IsValidUnits(StringView Units);
+		public bool IsValidUnits(StringView units);
 
 		/// <summary>
 		/// Test to see if the structure name should be using a "T" prefix.
 		/// </summary>
-		/// <param name="Name">Name of the structure to test without any prefix.</param>
+		/// <param name="name">Name of the structure to test without any prefix.</param>
 		/// <returns>True if the structure should have a "T" prefix.</returns>
-		public bool IsStructWithTPrefix(StringView Name);
+		public bool IsStructWithTPrefix(StringView name);
 
 		/// <summary>
 		/// Test to see if the given macro has a parameter count as part of the name.
 		/// </summary>
-		/// <param name="DelegateMacro">Macro to test</param>
+		/// <param name="delegateMacro">Macro to test</param>
 		/// <returns>-1 if the macro does not contain a parameter count.  The number of parameters minus one.</returns>
-		public int FindDelegateParameterCount(StringView DelegateMacro);
+		public int FindDelegateParameterCount(StringView delegateMacro);
 
 		/// <summary>
 		/// Get the parameter count string associated with the given index
 		/// </summary>
-		/// <param name="Index">Index from a prior call to FindDelegateParameterCount or -1.</param>
+		/// <param name="index">Index from a prior call to FindDelegateParameterCount or -1.</param>
 		/// <returns>Parameter count string or an empty string if Index is -1.</returns>
-		public StringView GetDelegateParameterCountString(int Index);
+		public StringView GetDelegateParameterCountString(int index);
 
 		/// <summary>
 		/// Test to see if the exporter is enabled
 		/// </summary>
-		/// <param name="Name">Name of the exporter</param>
+		/// <param name="name">Name of the exporter</param>
 		/// <returns>True if the exporter is enabled, false if not</returns>
-		public bool IsExporterEnabled(string Name);
+		public bool IsExporterEnabled(string name);
 	}
 }
diff --git a/Engine/Source/Programs/Shared/EpicGames.UHT/Utils/UhtException.cs b/Engine/Source/Programs/Shared/EpicGames.UHT/Utils/UhtException.cs
--- a/Engine/Source/Programs/Shared/EpicGames.UHT/Utils/UhtException.cs
+++ b/Engine/Source/Programs/Shared/EpicGames.UHT/Utils/UhtException.cs
@@ -1,8 +1,8 @@
 // Copyright Epic Games, Inc. All Rights Reserved.
 
-using EpicGames.UHT.Tokenizer;
 using System;
 using System.Text;
+using EpicGames.UHT.Tokenizer;
 
 namespace EpicGames.UHT.Utils
 {
@@ -26,45 +26,45 @@ namespace EpicGames.UHT.Utils
 		/// <summary>
 		/// Exception with a simple message.  Context will be the current header file.
 		/// </summary>
-		/// <param name="Message">Text of the error</param>
-		public UhtException(string Message)
+		/// <param name="message">Text of the error</param>
+		public UhtException(string message)
 		{
-			this.UhtMessage = UhtMessage.MakeMessage(UhtMessageType.Error, null, null, 1, Message);
+			this.UhtMessage = UhtMessage.MakeMessage(UhtMessageType.Error, null, null, 1, message);
 		}
 
 		/// <summary>
 		/// Exception with a simple message.  Context from the given message site.
 		/// </summary>
-		/// <param name="MessageSite">Site generating the exception</param>
-		/// <param name="Message">Text of the error</param>
-		/// <param name="ExtraContext">Addition context to be appended to the error message</param>
-		public UhtException(IUhtMessageSite MessageSite, string Message, object? ExtraContext = null)
+		/// <param name="messageSite">Site generating the exception</param>
+		/// <param name="message">Text of the error</param>
+		/// <param name="extraContext">Addition context to be appended to the error message</param>
+		public UhtException(IUhtMessageSite messageSite, string message, object? extraContext = null)
 		{
-			if (ExtraContext != null)
+			if (extraContext != null)
 			{
-				Message = $"{Message} while parsing {UhtMessage.FormatContext(ExtraContext)}";
+				message = $"{message} while parsing {UhtMessage.FormatContext(extraContext)}";
 			}
-			this.UhtMessage = UhtMessage.MakeMessage(UhtMessageType.Error, MessageSite.MessageSource, null, MessageSite.GetLineNumber(), Message);
+			this.UhtMessage = UhtMessage.MakeMessage(UhtMessageType.Error, messageSite.MessageSource, null, messageSite.GetLineNumber(), message);
 		}
 
 		/// <summary>
 		/// Make an exception to be thrown
 		/// </summary>
-		/// <param name="MessageSite">Message site to be associated with the exception</param>
-		/// <param name="LineNumber">Line number of the error</param>
-		/// <param name="Message">Text of the error</param>
-		/// <param name="ExtraContext">Addition context to be appended to the error message</param>
-		public UhtException(IUhtMessageSite MessageSite, int LineNumber, string Message, object? ExtraContext = null)
+		/// <param name="messageSite">Message site to be associated with the exception</param>
+		/// <param name="lineNumber">Line number of the error</param>
+		/// <param name="message">Text of the error</param>
+		/// <param name="extraContext">Addition context to be appended to the error message</param>
+		public UhtException(IUhtMessageSite messageSite, int lineNumber, string message, object? extraContext = null)
 		{
-			if (ExtraContext != null)
+			if (extraContext != null)
 			{
-				Message = $"{Message} while parsing {UhtMessage.FormatContext(ExtraContext)}";
+				message = $"{message} while parsing {UhtMessage.FormatContext(extraContext)}";
 			}
-			this.UhtMessage = UhtMessage.MakeMessage(UhtMessageType.Error, MessageSite.MessageSource, null, MessageSite.GetLineNumber(LineNumber), Message);
+			this.UhtMessage = UhtMessage.MakeMessage(UhtMessageType.Error, messageSite.MessageSource, null, messageSite.GetLineNumber(lineNumber), message);
 		}
 
 		/// <inheritdoc/>
-		public override string Message { get => this.UhtMessage.Message; }
+		public override string Message => this.UhtMessage.Message;
 	}
 
 	/// <summary>
@@ -76,35 +76,35 @@ namespace EpicGames.UHT.Utils
 		/// <summary>
 		/// Make a parsing error for when there is a mismatch between the expected token and what was parsed.
 		/// </summary>
-		/// <param name="MessageSite">Message site to be associated with the exception</param>
-		/// <param name="Got">The parsed token.  Support for EOF also provided.</param>
-		/// <param name="Expected">What was expected.</param>
-		/// <param name="ExtraContext">Extra context to be appended to the error message</param>
+		/// <param name="messageSite">Message site to be associated with the exception</param>
+		/// <param name="got">The parsed token.  Support for EOF also provided.</param>
+		/// <param name="expected">What was expected.</param>
+		/// <param name="extraContext">Extra context to be appended to the error message</param>
 		/// <returns>The exception object to throw</returns>
-		public UhtTokenException(IUhtMessageSite MessageSite, UhtToken Got, object? Expected, object? ExtraContext = null)
+		public UhtTokenException(IUhtMessageSite messageSite, UhtToken got, object? expected, object? extraContext = null)
 		{
-			string Message = Expected != null 
-				? $"Found {UhtMessage.FormatContext(Got)} when expecting {UhtMessage.FormatContext(Expected)}{FormatExtraContext(ExtraContext)}"
-				: $"Found {UhtMessage.FormatContext(Got)}{FormatExtraContext(ExtraContext)}";
-			this.UhtMessage = UhtMessage.MakeMessage(UhtMessageType.Error, MessageSite.MessageSource, null, Got.InputLine, Message);
+			string message = expected != null
+				? $"Found {UhtMessage.FormatContext(got)} when expecting {UhtMessage.FormatContext(expected)}{FormatExtraContext(extraContext)}"
+				: $"Found {UhtMessage.FormatContext(got)}{FormatExtraContext(extraContext)}";
+			this.UhtMessage = UhtMessage.MakeMessage(UhtMessageType.Error, messageSite.MessageSource, null, got.InputLine, message);
 		}
 
 		/// <summary>
 		/// Format any extra context supplied by the caller or the message site
 		/// </summary>
-		/// <param name="ExtraContext">Additional caller supplied context</param>
+		/// <param name="extraContext">Additional caller supplied context</param>
 		/// <returns></returns>
-		private static string FormatExtraContext(object? ExtraContext = null)
+		private static string FormatExtraContext(object? extraContext = null)
 		{
-			StringBuilder Builder = new StringBuilder(" while parsing ");
-			int StartingLength = Builder.Length;
-			if (ExtraContext != null)
+			StringBuilder builder = new StringBuilder(" while parsing ");
+			int startingLength = builder.Length;
+			if (extraContext != null)
 			{
-				Builder.Append(UhtMessage.FormatContext(ExtraContext));
+				builder.Append(UhtMessage.FormatContext(extraContext));
 			}
-			UhtMessage.Append(Builder, UhtTlsMessageExtraContext.GetMessageExtraContext(), StartingLength != Builder.Length);
+			UhtMessage.Append(builder, UhtTlsMessageExtraContext.GetMessageExtraContext(), startingLength != builder.Length);
 
-			return Builder.Length != StartingLength ? Builder.ToString() : string.Empty;
+			return builder.Length != startingLength ? builder.ToString() : String.Empty;
 		}
 	}
 
@@ -116,10 +116,10 @@ namespace EpicGames.UHT.Utils
 		/// <summary>
 		/// Exception with a simple message.  Context will be the current header file.
 		/// </summary>
-		/// <param name="Message">Text of the error</param>
-		public UhtIceException(string Message)
+		/// <param name="message">Text of the error</param>
+		public UhtIceException(string message)
 		{
-			this.UhtMessage = UhtMessage.MakeMessage(UhtMessageType.Ice, null, null, 1, Message);
+			this.UhtMessage = UhtMessage.MakeMessage(UhtMessageType.Ice, null, null, 1, message);
 		}
 	}
 }
diff --git a/Engine/Source/Programs/Shared/EpicGames.UHT/Utils/UhtExport.cs b/Engine/Source/Programs/Shared/EpicGames.UHT/Utils/UhtExport.cs
--- a/Engine/Source/Programs/Shared/EpicGames.UHT/Utils/UhtExport.cs
+++ b/Engine/Source/Programs/Shared/EpicGames.UHT/Utils/UhtExport.cs
@@ -1,10 +1,10 @@
 // Copyright Epic Games, Inc. All Rights Reserved.
 
-using EpicGames.Core;
-using EpicGames.UHT.Types;
 using System.Collections.Generic;
 using System.Text;
 using System.Threading.Tasks;
+using EpicGames.Core;
+using EpicGames.UHT.Types;
 
 namespace EpicGames.UHT.Utils
 {
@@ -12,15 +12,15 @@ namespace EpicGames.UHT.Utils
 	/// <summary>
 	/// Delegate invoked by a factory
 	/// </summary>
-	/// <param name="Factory">Invoking factory</param>
-	public delegate void UhtExportTaskDelegate(IUhtExportFactory Factory);
+	/// <param name="factory">Invoking factory</param>
+	public delegate void UhtExportTaskDelegate(IUhtExportFactory factory);
 
 	/// <summary>
 	/// Factory object used to generate export tasks
 	/// </summary>
 	public interface IUhtExportFactory
 	{
-		
+
 		/// <summary>
 		/// Session being run
 		/// </summary>
@@ -34,62 +34,62 @@ namespace EpicGames.UHT.Utils
 		/// <summary>
 		/// Create a task
 		/// </summary>
-		/// <param name="Prereqs">Tasks that must be completed prior to this task running</param>
-		/// <param name="Action">Action to be invoked to generate the output(s)</param>
+		/// <param name="prereqs">Tasks that must be completed prior to this task running</param>
+		/// <param name="action">Action to be invoked to generate the output(s)</param>
 		/// <returns>Task object or null if the task was immediately executed.</returns>
-		public Task? CreateTask(List<Task?>? Prereqs, UhtExportTaskDelegate Action);
+		public Task? CreateTask(List<Task?>? prereqs, UhtExportTaskDelegate action);
 
 		/// <summary>
 		/// Create a task
 		/// </summary>
-		/// <param name="Action">Action to be invoked to generate the output(s)</param>
+		/// <param name="action">Action to be invoked to generate the output(s)</param>
 		/// <returns>Task object or null if the task was immediately executed.</returns>
-		public Task? CreateTask(UhtExportTaskDelegate Action);
+		public Task? CreateTask(UhtExportTaskDelegate action);
 
 		/// <summary>
 		/// Commit the contents of the string builder as the output.
 		/// If you have a string builder, use this method so that a 
 		/// temporary buffer can be used.
 		/// </summary>
-		/// <param name="FilePath">Destination file path</param>
-		/// <param name="Builder">Source for the content</param>
-		public void CommitOutput(string FilePath, StringBuilder Builder);
+		/// <param name="filePath">Destination file path</param>
+		/// <param name="builder">Source for the content</param>
+		public void CommitOutput(string filePath, StringBuilder builder);
 
 		/// <summary>
 		/// Commit the value of the string as the output
 		/// </summary>
-		/// <param name="FilePath">Destination file path</param>
-		/// <param name="Output">Output to commit</param>
-		public void CommitOutput(string FilePath, StringView Output);
+		/// <param name="filePath">Destination file path</param>
+		/// <param name="output">Output to commit</param>
+		public void CommitOutput(string filePath, StringView output);
 
 		/// <summary>
 		/// Make a path for an output based on the header file name.
 		/// </summary>
-		/// <param name="HeaderFile">Header file being exported.</param>
-		/// <param name="Suffix">Suffix to be added to the file name.</param>
+		/// <param name="headerFile">Header file being exported.</param>
+		/// <param name="suffix">Suffix to be added to the file name.</param>
 		/// <returns>Output file path</returns>
-		public string MakePath(UhtHeaderFile HeaderFile, string Suffix);
+		public string MakePath(UhtHeaderFile headerFile, string suffix);
 
 		/// <summary>
 		/// Make a path for an output based on the package name.
 		/// </summary>
-		/// <param name="Package">Package being exported</param>
-		/// <param name="Suffix">Suffix to be added to the file name.</param>
+		/// <param name="package">Package being exported</param>
+		/// <param name="suffix">Suffix to be added to the file name.</param>
 		/// <returns>Output file path</returns>
-		public string MakePath(UhtPackage Package, string Suffix);
+		public string MakePath(UhtPackage package, string suffix);
 
 		/// <summary>
 		/// Make a path for the given file name and extension.  This is only valid for plugins.
 		/// </summary>
-		/// <param name="FileName">Name of the file</param>
-		/// <param name="Extension">Extension to add to the file</param>
+		/// <param name="fileName">Name of the file</param>
+		/// <param name="extension">Extension to add to the file</param>
 		/// <returns>Output file path</returns>
-		public string MakePath(string FileName, string Extension);
+		public string MakePath(string fileName, string extension);
 
 		/// <summary>
 		/// Add an external dependency to the given file path
 		/// </summary>
-		/// <param name="FilePath">External dependency to add</param>
-		public void AddExternalDependency(string FilePath);
+		/// <param name="filePath">External dependency to add</param>
+		public void AddExternalDependency(string filePath);
 	}
 }
diff --git a/Engine/Source/Programs/Shared/EpicGames.UHT/Utils/UhtFCString.cs b/Engine/Source/Programs/Shared/EpicGames.UHT/Utils/UhtFCString.cs
--- a/Engine/Source/Programs/Shared/EpicGames.UHT/Utils/UhtFCString.cs
+++ b/Engine/Source/Programs/Shared/EpicGames.UHT/Utils/UhtFCString.cs
@@ -1,8 +1,8 @@
 // Copyright Epic Games, Inc. All Rights Reserved.
 
-using EpicGames.Core;
 using System;
 using System.Text;
+using EpicGames.Core;
 
 namespace EpicGames.UHT.Utils
 {
@@ -21,213 +21,213 @@ namespace EpicGames.UHT.Utils
 		/// <summary>
 		/// Test to see if the string is a boolean
 		/// </summary>
-		/// <param name="Value">Boolean to test</param>
+		/// <param name="value">Boolean to test</param>
 		/// <returns>True if the value is true</returns>
-		public static bool ToBool(StringView Value)
+		public static bool ToBool(StringView value)
 		{
-			ReadOnlySpan<char> Span = Value.Span;
-			if (Span.Length == 0)
+			ReadOnlySpan<char> span = value.Span;
+			if (span.Length == 0)
 			{
 				return false;
 			}
 			else if (
-				Span.Equals("true", StringComparison.OrdinalIgnoreCase) ||
-				Span.Equals("yes", StringComparison.OrdinalIgnoreCase) ||
-				Span.Equals("on", StringComparison.OrdinalIgnoreCase))
+				span.Equals("true", StringComparison.OrdinalIgnoreCase) ||
+				span.Equals("yes", StringComparison.OrdinalIgnoreCase) ||
+				span.Equals("on", StringComparison.OrdinalIgnoreCase))
 			{
 				return true;
 			}
 			else if (
-				Span.Equals("false", StringComparison.OrdinalIgnoreCase) ||
-				Span.Equals("no", StringComparison.OrdinalIgnoreCase) ||
-				Span.Equals("off", StringComparison.OrdinalIgnoreCase))
+				span.Equals("false", StringComparison.OrdinalIgnoreCase) ||
+				span.Equals("no", StringComparison.OrdinalIgnoreCase) ||
+				span.Equals("off", StringComparison.OrdinalIgnoreCase))
 			{
 				return false;
 			}
 			else
 			{
-				int IntValue;
-				return int.TryParse(Span, out IntValue) && IntValue != 0;
+				int intValue;
+				return Int32.TryParse(span, out intValue) && intValue != 0;
 			}
 		}
 
 		/// <summary>
 		/// Test to see if the string is a boolean
 		/// </summary>
-		/// <param name="Value">Boolean to test</param>
+		/// <param name="value">Boolean to test</param>
 		/// <returns>True if the value is true</returns>
-		public static bool ToBool(string? Value)
+		public static bool ToBool(string? value)
 		{
-			if (string.IsNullOrEmpty(Value))
+			if (String.IsNullOrEmpty(value))
 			{
 				return false;
 			}
-			else if (string.Equals(Value, "true", StringComparison.OrdinalIgnoreCase) ||
-				string.Equals(Value, "yes", StringComparison.OrdinalIgnoreCase) ||
-				string.Equals(Value, "on", StringComparison.OrdinalIgnoreCase))
+			else if (String.Equals(value, "true", StringComparison.OrdinalIgnoreCase) ||
+				String.Equals(value, "yes", StringComparison.OrdinalIgnoreCase) ||
+				String.Equals(value, "on", StringComparison.OrdinalIgnoreCase))
 			{
 				return true;
 			}
-			else if (string.Equals(Value, "false", StringComparison.OrdinalIgnoreCase) ||
-				string.Equals(Value, "no", StringComparison.OrdinalIgnoreCase) ||
-				string.Equals(Value, "off", StringComparison.OrdinalIgnoreCase))
+			else if (String.Equals(value, "false", StringComparison.OrdinalIgnoreCase) ||
+				String.Equals(value, "no", StringComparison.OrdinalIgnoreCase) ||
+				String.Equals(value, "off", StringComparison.OrdinalIgnoreCase))
 			{
 				return false;
 			}
 			else
 			{
-				int IntValue;
-				return int.TryParse(Value, out IntValue) && IntValue != 0;
+				int intValue;
+				return Int32.TryParse(value, out intValue) && intValue != 0;
 			}
 		}
 
 		/// <summary>
 		/// Test to see if the character is a digit
 		/// </summary>
-		/// <param name="C">Character to test</param>
+		/// <param name="c">Character to test</param>
 		/// <returns>True if it is</returns>
-		public static bool IsDigit(char C)
+		public static bool IsDigit(char c)
 		{
-			return C >= '0' && C <= '9';
+			return c >= '0' && c <= '9';
 		}
 
 		/// <summary>
 		/// Test to see if the character is an alphabet character
 		/// </summary>
-		/// <param name="C">Character to test</param>
+		/// <param name="c">Character to test</param>
 		/// <returns>True if it is</returns>
-		public static bool IsAlpha(char C)
+		public static bool IsAlpha(char c)
 		{
-			return (C >= 'A' && C <= 'Z') || (C >= 'a' && C <= 'z');
+			return (c >= 'A' && c <= 'Z') || (c >= 'a' && c <= 'z');
 		}
 
 		/// <summary>
 		/// Test to see if the character is a digit or alphabet character
 		/// </summary>
-		/// <param name="C">Character to test</param>
+		/// <param name="c">Character to test</param>
 		/// <returns>True if it is</returns>
-		public static bool IsAlnum(char C)
+		public static bool IsAlnum(char c)
 		{
-			return IsDigit(C) || IsAlpha(C);
+			return IsDigit(c) || IsAlpha(c);
 		}
 
 		/// <summary>
 		/// Test to see if the character is a hex digit (A-F)
 		/// </summary>
-		/// <param name="C">Character to test</param>
+		/// <param name="c">Character to test</param>
 		/// <returns>True if it is</returns>
-		public static bool IsHexAlphaDigit(char C)
+		public static bool IsHexAlphaDigit(char c)
 		{
-			return (C >= 'A' && C <= 'F') || (C >= 'a' && C <= 'f');
+			return (c >= 'A' && c <= 'F') || (c >= 'a' && c <= 'f');
 		}
 
 		/// <summary>
 		/// Test to see if the character is a hex digit (0-9a-f)
 		/// </summary>
-		/// <param name="C">Character to test</param>
+		/// <param name="c">Character to test</param>
 		/// <returns>True if it is</returns>
-		public static bool IsHexDigit(char C)
+		public static bool IsHexDigit(char c)
 		{
-			return IsDigit(C) || IsHexAlphaDigit(C);
+			return IsDigit(c) || IsHexAlphaDigit(c);
 		}
 
 		/// <summary>
 		/// Test to see if the character is whitespace
 		/// </summary>
-		/// <param name="C">Character to test</param>
+		/// <param name="c">Character to test</param>
 		/// <returns>True if it is</returns>
-		public static bool IsWhitespace(char C)
+		public static bool IsWhitespace(char c)
 		{
-			return Char.IsWhiteSpace(C);
+			return Char.IsWhiteSpace(c);
 		}
 
 		/// <summary>
 		/// Test to see if the character is a sign
 		/// </summary>
-		/// <param name="C">Character to test</param>
+		/// <param name="c">Character to test</param>
 		/// <returns>True if it is</returns>
-		public static bool IsSign(char C)
+		public static bool IsSign(char c)
 		{
-			return C == '+' || C == '-';
+			return c == '+' || c == '-';
 		}
 
 		/// <summary>
 		/// Test to see if the character is a hex marker (xX)
 		/// </summary>
-		/// <param name="C">Character to test</param>
+		/// <param name="c">Character to test</param>
 		/// <returns>True if it is</returns>
-		public static bool IsHexMarker(char C)
+		public static bool IsHexMarker(char c)
 		{
-			return C == 'x' || C == 'X';
+			return c == 'x' || c == 'X';
 		}
 
 		/// <summary>
 		/// Test to see if the character is a float marker (fF)
 		/// </summary>
-		/// <param name="C">Character to test</param>
+		/// <param name="c">Character to test</param>
 		/// <returns>True if it is</returns>
-		public static bool IsFloatMarker(char C)
+		public static bool IsFloatMarker(char c)
 		{
-			return C == 'f' || C == 'F';
+			return c == 'f' || c == 'F';
 		}
 
 		/// <summary>
 		/// Test to see if the character is an exponent marker (eE)
 		/// </summary>
-		/// <param name="C">Character to test</param>
+		/// <param name="c">Character to test</param>
 		/// <returns>True if it is</returns>
-		public static bool IsExponentMarker(char C)
+		public static bool IsExponentMarker(char c)
 		{
-			return C == 'e' || C == 'E';
+			return c == 'e' || c == 'E';
 		}
 
 		/// <summary>
 		/// Test to see if the character is an unsigned marker (uU)
 		/// </summary>
-		/// <param name="C">Character to test</param>
+		/// <param name="c">Character to test</param>
 		/// <returns>True if it is</returns>
-		public static bool IsUnsignedMarker(char C)
+		public static bool IsUnsignedMarker(char c)
 		{
-			return C == 'u' || C == 'U';
+			return c == 'u' || c == 'U';
 		}
 
 		/// <summary>
 		/// Test to see if the character is a long marker (lL)
 		/// </summary>
-		/// <param name="C">Character to test</param>
+		/// <param name="c">Character to test</param>
 		/// <returns>True if it is</returns>
-		public static bool IsLongMarker(char C)
+		public static bool IsLongMarker(char c)
 		{
-			return C == 'l' || C == 'L';
+			return c == 'l' || c == 'L';
 		}
 
 		/// <summary>
 		/// Test to see if the span is a numeric value
 		/// </summary>
-		/// <param name="Span">Span to test</param>
+		/// <param name="span">Span to test</param>
 		/// <returns>True if it is</returns>
-		public static bool IsNumeric(ReadOnlySpan<char> Span)
+		public static bool IsNumeric(ReadOnlySpan<char> span)
 		{
-			int Index = 0;
-			char C = Span[Index];
-			if (C == '-' || C == '+')
+			int index = 0;
+			char c = span[index];
+			if (c == '-' || c == '+')
 			{
-				++Index;
+				++index;
 			}
 
-			bool bHasDot = false;
-			for (; Index < Span.Length; ++Index)
+			bool hasDot = false;
+			for (; index < span.Length; ++index)
 			{
-				C = Span[Index];
-				if (C == '.')
+				c = span[index];
+				if (c == '.')
 				{
-					if (bHasDot)
+					if (hasDot)
 					{
 						return false;
 					}
-					bHasDot = true;
+					hasDot = true;
 				}
-				else if (!IsDigit(C))
+				else if (!IsDigit(c))
 				{
 					return false;
 				}
@@ -238,78 +238,78 @@ namespace EpicGames.UHT.Utils
 		/// <summary>
 		/// Test to see if the character is a linebreak character
 		/// </summary>
-		/// <param name="C">Character to test</param>
+		/// <param name="c">Character to test</param>
 		/// <returns>True if it is</returns>
-		public static bool IsLinebreak(char C)
+		public static bool IsLinebreak(char c)
 		{
-			return (C >= 0x0a && C <= 0x0d) || C == 0x85 || C == 0x2028 || C == 0x2029;
+			return (c >= 0x0a && c <= 0x0d) || c == 0x85 || c == 0x2028 || c == 0x2029;
 		}
 
 		/// <summary>
 		/// Return an unescaped string
 		/// </summary>
-		/// <param name="Text">Text to unescape</param>
+		/// <param name="text">Text to unescape</param>
 		/// <returns>Resulting string</returns>
-		public static string UnescapeText(string Text)
+		public static string UnescapeText(string text)
 		{
-			StringBuilder Result = new StringBuilder();
-			for (int Idx = 0; Idx < Text.Length; Idx++)
+			StringBuilder result = new StringBuilder();
+			for (int idx = 0; idx < text.Length; idx++)
 			{
-				if (Text[Idx] == '\\' && Idx + 1 < Text.Length)
+				if (text[idx] == '\\' && idx + 1 < text.Length)
 				{
-					switch (Text[++Idx])
+					switch (text[++idx])
 					{
 						case 't':
-							Result.Append('\t');
+							result.Append('\t');
 							break;
 						case 'r':
-							Result.Append('\r');
+							result.Append('\r');
 							break;
 						case 'n':
-							Result.Append('\n');
+							result.Append('\n');
 							break;
 						case '\'':
-							Result.Append('\'');
+							result.Append('\'');
 							break;
 						case '\"':
-							Result.Append('\"');
+							result.Append('\"');
 							break;
 						case 'u':
-							bool bParsed = false;
-							if (Idx + 4 < Text.Length)
+							bool parsed = false;
+							if (idx + 4 < text.Length)
 							{
-								int Value;
-								if (int.TryParse(Text.Substring(Idx + 1, 4), System.Globalization.NumberStyles.HexNumber, null, out Value))
+								int value;
+								if (Int32.TryParse(text.Substring(idx + 1, 4), System.Globalization.NumberStyles.HexNumber, null, out value))
 								{
-									bParsed = true;
-									Result.Append((char)Value);
-									Idx += 4;
+									parsed = true;
+									result.Append((char)value);
+									idx += 4;
 								}
 							}
-							if (!bParsed)
+							if (!parsed)
 							{
-								Result.Append('\\');
-								Result.Append('u');
+								result.Append('\\');
+								result.Append('u');
 							}
 							break;
 						default:
-							Result.Append(Text[Idx]);
+							result.Append(text[idx]);
 							break;
 					}
 				}
 				else
 				{
-					Result.Append(Text[Idx]);
+					result.Append(text[idx]);
 				}
 			}
-			return Result.ToString();
+			return result.ToString();
 		}
 
 		/// <summary>
 		/// Words that are considered articles when parsing comment strings
 		/// Some words are always forced lowercase
 		/// </summary>
-		private static readonly string[] Articles = new string[]
+		private static readonly string[] s_articles = new string[]
 		{
 			"In",
 			"As",
@@ -331,36 +331,36 @@ namespace EpicGames.UHT.Utils
 		/// <summary>
 		/// Convert a name to a display string
 		/// </summary>
-		/// <param name="Name">Name to convert</param>
-		/// <param name="bIsBool">True if the name is from a boolean property</param>
+		/// <param name="name">Name to convert</param>
+		/// <param name="isBool">True if the name is from a boolean property</param>
 		/// <returns>Display string</returns>
-		public static string NameToDisplayString(StringView Name, bool bIsBool)
+		public static string NameToDisplayString(StringView name, bool isBool)
 		{
-			ReadOnlySpan<char> Chars = Name.Span;
+			ReadOnlySpan<char> chars = name.Span;
 
 			// This is used to indicate that we are in a run of uppercase letter and/or digits.  The code attempts to keep
 			// these characters together as breaking them up often looks silly (i.e. "Draw Scale 3 D" as opposed to "Draw Scale 3D"
-			bool bInARun = false;
-			bool bWasSpace = false;
-			bool bWasOpenParen = false;
-			bool bWasNumber = false;
-			bool bWasMinusSign = false;
-
-			StringBuilder Builder = new StringBuilder();
-			for (int CharIndex = 0; CharIndex < Chars.Length; ++CharIndex)
+			bool inARun = false;
+			bool wasSpace = false;
+			bool wasOpenParen = false;
+			bool wasNumber = false;
+			bool wasMinusSign = false;
+
+			StringBuilder builder = new StringBuilder();
+			for (int charIndex = 0; charIndex < chars.Length; ++charIndex)
 			{
-				char ch = Chars[CharIndex];
+				char ch = chars[charIndex];
 
-				bool bLowerCase = Char.IsLower(ch);
-				bool bUpperCase = Char.IsUpper(ch);
-				bool bIsDigit = Char.IsDigit(ch);
-				bool bIsUnderscore = ch == '_';
+				bool lowerCase = Char.IsLower(ch);
+				bool upperCase = Char.IsUpper(ch);
+				bool isDigit = Char.IsDigit(ch);
+				bool isUnderscore = ch == '_';
 
 				// Skip the first character if the property is a bool (they should all start with a lowercase 'b', which we don't want to keep)
-				if (CharIndex == 0 && bIsBool && ch == 'b')
+				if (charIndex == 0 && isBool && ch == 'b')
 				{
 					// Check if next character is uppercase as it may be a user created string that doesn't follow the rules of Unreal variables
-					if (Chars.Length > 1 && Char.IsUpper(Chars[1]))
+					if (chars.Length > 1 && Char.IsUpper(chars[1]))
 					{
 						continue;
 					}
@@ -368,56 +368,56 @@ namespace EpicGames.UHT.Utils
 
 				// If the current character is upper case or a digit, and the previous character wasn't, then we need to insert a space if there wasn't one previously
 				// We don't do this for numerical expressions, for example "-1.2" should not be formatted as "- 1. 2"
-				if ((bUpperCase || (bIsDigit && !bWasMinusSign)) && !bInARun && !bWasOpenParen && !bWasNumber)
+				if ((upperCase || (isDigit && !wasMinusSign)) && !inARun && !wasOpenParen && !wasNumber)
 				{
-					if (!bWasSpace && Builder.Length > 0)
+					if (!wasSpace && builder.Length > 0)
 					{
-						Builder.Append(' ');
-						bWasSpace = true;
+						builder.Append(' ');
+						wasSpace = true;
 					}
-					bInARun = true;
+					inARun = true;
 				}
 
 				// A lower case character will break a run of upper case letters and/or digits
-				if (bLowerCase)
+				if (lowerCase)
 				{
-					bInARun = false;
+					inARun = false;
 				}
 
 				// An underscore denotes a space, so replace it and continue the run
-				if (bIsUnderscore)
+				if (isUnderscore)
 				{
 					ch = ' ';
-					bInARun = true;
+					inARun = true;
 				}
 
 				// If this is the first character in the string, then it will always be upper-case
-				if (Builder.Length == 0)
+				if (builder.Length == 0)
 				{
 					ch = Char.ToUpper(ch);
 				}
-				else if (!bIsDigit && (bWasSpace || bWasOpenParen)) // If this is first character after a space, then make sure it is case-correct
+				else if (!isDigit && (wasSpace || wasOpenParen)) // If this is first character after a space, then make sure it is case-correct
 				{
 
 					// Search for a word that needs case repaired
-					bool bIsArticle = false;
-					foreach (string Article in Articles)
+					bool isArticle = false;
+					foreach (string article in s_articles)
 					{
 						// Make sure the character following the string we're testing is not lowercase (we don't want to match "in" with "instance")
-						int ArticleLength = Article.Length;
-						if (CharIndex + ArticleLength < Chars.Length && !Char.IsLower(Chars[CharIndex + ArticleLength]))
+						int articleLength = article.Length;
+						if (charIndex + articleLength < chars.Length && !Char.IsLower(chars[charIndex + articleLength]))
 						{
 							// Does this match the current article?
-							if (Chars.Slice(CharIndex, ArticleLength).Equals(Article.AsSpan(), StringComparison.OrdinalIgnoreCase))
+							if (chars.Slice(charIndex, articleLength).Equals(article.AsSpan(), StringComparison.OrdinalIgnoreCase))
 							{
-								bIsArticle = true;
+								isArticle = true;
 								break;
 							}
 						}
 					}
 
 					// Start of a keyword, force to lowercase
-					if (bIsArticle)
+					if (isArticle)
 					{
 						ch = Char.ToLower(ch);
 					}
@@ -427,106 +427,106 @@ namespace EpicGames.UHT.Utils
 					}
 				}
 
-				bWasSpace = ch == ' ';
-				bWasOpenParen = ch == '(';
+				wasSpace = ch == ' ';
+				wasOpenParen = ch == '(';
 
 				// What could be included as part of a numerical representation.
 				// For example -1.2
-				bWasMinusSign = ch == '-';
-				bool bPotentialNumericalChar = bWasMinusSign || ch == '.';
-				bWasNumber = bIsDigit || (bWasNumber && bPotentialNumericalChar);
-				Builder.Append(ch);
+				wasMinusSign = ch == '-';
+				bool potentialNumericalChar = wasMinusSign || ch == '.';
+				wasNumber = isDigit || (wasNumber && potentialNumericalChar);
+				builder.Append(ch);
 			}
-			return Builder.ToString();
+			return builder.ToString();
 		}
 
 		/// <summary>
 		/// Replace tabs to spaces in a string containing only a single line.
 		/// </summary>
-		/// <param name="Input">Input string</param>
-		/// <param name="TabSpacing">Number of spaces to exchange for tabs</param>
-		/// <param name="bEmulateCrBug">Due to a bug in UE ConvertTabsToSpacesInline, any \n is considered part of the line length.</param>
+		/// <param name="input">Input string</param>
+		/// <param name="tabSpacing">Number of spaces to exchange for tabs</param>
+		/// <param name="emulateCrBug">Due to a bug in UE ConvertTabsToSpacesInline, any \n is considered part of the line length.</param>
 		/// <returns>Resulting string or the original string if the string didn't contain any spaces.</returns>
-		public static ReadOnlyMemory<char> TabsToSpaces(ReadOnlyMemory<char> Input, int TabSpacing, bool bEmulateCrBug)
+		public static ReadOnlyMemory<char> TabsToSpaces(ReadOnlyMemory<char> input, int tabSpacing, bool emulateCrBug)
 		{
 			// If we have any tab characters, then we need to convert them to spaces
-			int TabIndex = Input.Span.IndexOf('\t');
-			if (TabIndex == -1)
+			int tabIndex = input.Span.IndexOf('\t');
+			if (tabIndex == -1)
 			{
-				return Input;
+				return input;
 			}
-			return TabsToSpacesInternal(Input.Span, TabSpacing, bEmulateCrBug, TabIndex).AsMemory();
+			return TabsToSpacesInternal(input.Span, tabSpacing, emulateCrBug, tabIndex).AsMemory();
 		}
 
 		/// <summary>
 		/// Replace tabs to spaces in a string containing zero or more lines.
 		/// </summary>
-		/// <param name="Input">Input string</param>
-		/// <param name="TabSpacing">Number of spaces to exchange for tabs</param>
-		/// <param name="bEmulateCrBug">Due to a bug in UE ConvertTabsToSpacesInline, any \n is considered part of the line length.</param>
+		/// <param name="input">Input string</param>
+		/// <param name="tabSpacing">Number of spaces to exchange for tabs</param>
+		/// <param name="emulateCrBug">Due to a bug in UE ConvertTabsToSpacesInline, any \n is considered part of the line length.</param>
 		/// <returns>Resulting string or the original string if the string didn't contain any spaces.</returns>
-		public static string TabsToSpaces(String Input, int TabSpacing, bool bEmulateCrBug)
+		public static string TabsToSpaces(string input, int tabSpacing, bool emulateCrBug)
 		{
 			// If we have any tab characters, then we need to convert them to spaces
-			int TabIndex = Input.IndexOf('\t', StringComparison.Ordinal);
-			if (TabIndex == -1)
+			int tabIndex = input.IndexOf('\t', StringComparison.Ordinal);
+			if (tabIndex == -1)
 			{
-				return Input;
+				return input;
 			}
-			return TabsToSpacesInternal(Input.AsSpan(), TabSpacing, bEmulateCrBug, TabIndex);
+			return TabsToSpacesInternal(input.AsSpan(), tabSpacing, emulateCrBug, tabIndex);
 		}
 
 		/// <summary>
 		/// Replace tabs to spaces in a string containing zero or more lines.
 		/// </summary>
-		/// <param name="Span">Input string to convert</param>
-		/// <param name="TabSpacing">Number of spaces to exchange for tabs</param>
-		/// <param name="bEmulateCrBug">Due to a bug in UE ConvertTabsToSpacesInline, any \n is considered part of the line length.</param>
-		/// <param name="TabIndex">Initial tab index</param>
+		/// <param name="span">Input string to convert</param>
+		/// <param name="tabSpacing">Number of spaces to exchange for tabs</param>
+		/// <param name="emulateCrBug">Due to a bug in UE ConvertTabsToSpacesInline, any \n is considered part of the line length.</param>
+		/// <param name="tabIndex">Initial tab index</param>
 		/// <returns>Resulting string or the original string if the string didn't contain any spaces.</returns>
-		private static string TabsToSpacesInternal(ReadOnlySpan<char> Span, int TabSpacing, bool bEmulateCrBug, int TabIndex)
+		private static string TabsToSpacesInternal(ReadOnlySpan<char> span, int tabSpacing, bool emulateCrBug, int tabIndex)
 		{
 			// Locate the last \n since all tabs have to be computed relative to this.
-			int CrPos = Span.Slice(0, TabIndex).LastIndexOf('\n') + 1;
+			int crPos = span.Slice(0, tabIndex).LastIndexOf('\n') + 1;
 
-			using (BorrowStringBuilder Borrower = new BorrowStringBuilder(StringBuilderCache.Small))
+			using (BorrowStringBuilder borrower = new BorrowStringBuilder(StringBuilderCache.Small))
 			{
-				StringBuilder Builder = Borrower.StringBuilder;
-				int Committed = 0;
+				StringBuilder builder = borrower.StringBuilder;
+				int committed = 0;
 				do
 				{
 					// Commit everything prior to the tab to the builder
-					Builder.Append(Span.Slice(0, TabIndex));
-					Span = Span.Slice(TabIndex + 1);
-					Committed += TabIndex;
+					builder.Append(span.Slice(0, tabIndex));
+					span = span.Slice(tabIndex + 1);
+					committed += tabIndex;
 
 					// Add the appropriate number of spaces
-					int AdjustedCrPos = CrPos;
-					if (bEmulateCrBug && AdjustedCrPos > 0)
+					int adjustedCrPos = crPos;
+					if (emulateCrBug && adjustedCrPos > 0)
 					{
-						--AdjustedCrPos;
+						--adjustedCrPos;
 					}
-					int SpacesToInsert = TabSpacing - (Committed - AdjustedCrPos) % TabSpacing;
-					Builder.AppendSpaces(SpacesToInsert);
-					Committed += SpacesToInsert;
+					int spacesToInsert = tabSpacing - (committed - adjustedCrPos) % tabSpacing;
+					builder.AppendSpaces(spacesToInsert);
+					committed += spacesToInsert;
 
 					// Search for the next \t or \n
-					for (TabIndex = 0; TabIndex < Span.Length; ++TabIndex)
+					for (tabIndex = 0; tabIndex < span.Length; ++tabIndex)
 					{
-						if (Span[TabIndex] == '\n')
+						if (span[tabIndex] == '\n')
 						{
-							CrPos = Committed + TabIndex + 1;
+							crPos = committed + tabIndex + 1;
 						}
-						else if (Span[TabIndex] == '\t')
+						else if (span[tabIndex] == '\t')
 						{
 							break;
 						}
 					}
-				} while (TabIndex < Span.Length);
+				} while (tabIndex < span.Length);
 
 				// Commit the remaining data
-				Builder.Append(Span);
-				return Builder.ToString();
+				builder.Append(span);
+				return builder.ToString();
 			}
 		}
 	}
