More coding standard fixes

diff --git a/Engine/Source/Programs/Shared/EpicGames.UHT/Utils/UhtFileManager.cs b/Engine/Source/Programs/Shared/EpicGames.UHT/Utils/UhtFileManager.cs
--- a/Engine/Source/Programs/Shared/EpicGames.UHT/Utils/UhtFileManager.cs
+++ b/Engine/Source/Programs/Shared/EpicGames.UHT/Utils/UhtFileManager.cs
@@ -1,10 +1,10 @@
 // Copyright Epic Games, Inc. All Rights Reserved.
 
-using EpicGames.Core;
 using System;
 using System.IO;
 using System.Linq;
 using System.Text;
+using EpicGames.Core;
 
 namespace EpicGames.UHT.Utils
 {
@@ -17,40 +17,40 @@ namespace EpicGames.UHT.Utils
 		/// <summary>
 		/// Return the full file path for a partial path
 		/// </summary>
-		/// <param name="FilePath">The partial file path</param>
+		/// <param name="filePath">The partial file path</param>
 		/// <returns>The full file path</returns>
-		public string GetFullFilePath(string FilePath);
+		public string GetFullFilePath(string filePath);
 
 		/// <summary>
 		/// Read the given source file
 		/// </summary>
-		/// <param name="FilePath">File path</param>
-		/// <param name="Fragment">Read fragment information</param>
+		/// <param name="filePath">File path</param>
+		/// <param name="fragment">Read fragment information</param>
 		/// <returns>True if the file was read</returns>
-		public bool ReadSource(string FilePath, out UhtSourceFragment Fragment);
+		public bool ReadSource(string filePath, out UhtSourceFragment fragment);
 
 		/// <summary>
 		/// Read the given source file
 		/// </summary>
-		/// <param name="FilePath">File path</param>
+		/// <param name="filePath">File path</param>
 		/// <returns>Buffer containing the read data or null if not found.  The returned buffer must be returned to the cache via a call to UhtBuffer.Return</returns>
-		public UhtBuffer? ReadOutput(string FilePath);
+		public UhtBuffer? ReadOutput(string filePath);
 
 		/// <summary>
 		/// Write the given contents to the file
 		/// </summary>
-		/// <param name="FilePath">Path to write to</param>
-		/// <param name="Contents">Contents to write</param>
+		/// <param name="filePath">Path to write to</param>
+		/// <param name="contents">Contents to write</param>
 		/// <returns>True if the source was written</returns>
-		public bool WriteOutput(string FilePath, ReadOnlySpan<char> Contents);
+		public bool WriteOutput(string filePath, ReadOnlySpan<char> contents);
 
 		/// <summary>
 		/// Rename the given file
 		/// </summary>
-		/// <param name="OldFilePath">Old file path name</param>
-		/// <param name="NewFilePath">New file path name</param>
+		/// <param name="oldFilePath">Old file path name</param>
+		/// <param name="newFilePath">New file path name</param>
 		/// <returns>True if the file was renamed</returns>
-		public bool RenameOutput(string OldFilePath, string NewFilePath);
+		public bool RenameOutput(string oldFilePath, string newFilePath);
 	}
 
 	/// <summary>
@@ -67,58 +67,58 @@ namespace EpicGames.UHT.Utils
 		}
 
 		/// <inheritdoc/>
-		public string GetFullFilePath(string FilePath)
+		public string GetFullFilePath(string filePath)
 		{
-			return FilePath;
+			return filePath;
 		}
 
 		/// <inheritdoc/>
-		public bool ReadSource(string FilePath, out UhtSourceFragment Fragment)
+		public bool ReadSource(string filePath, out UhtSourceFragment fragment)
 		{
-			if (ReadFile(FilePath, out StringView Data))
+			if (ReadFile(filePath, out StringView data))
 			{
-				Fragment = new UhtSourceFragment { SourceFile = null, FilePath = FilePath, LineNumber = 0, Data = Data };
+				fragment = new UhtSourceFragment { SourceFile = null, FilePath = filePath, LineNumber = 0, Data = data };
 				return true;
 			}
-			Fragment = new UhtSourceFragment { SourceFile = null, FilePath = String.Empty, LineNumber = 0, Data = new StringView() };
+			fragment = new UhtSourceFragment { SourceFile = null, FilePath = String.Empty, LineNumber = 0, Data = new StringView() };
 			return false;
 		}
 
 		/// <inheritdoc/>
-		public UhtBuffer? ReadOutput(string FilePath)
+		public UhtBuffer? ReadOutput(string filePath)
 		{
 			// Exceptions are very expensive.  Don't bother trying to open the file if it doesn't exist
-			if (!File.Exists(FilePath))
+			if (!File.Exists(filePath))
 			{
 				return null;
 			}
 
 			try
 			{
-				using (FileStream fs = new FileStream(FilePath, FileMode.Open, FileAccess.Read, FileShare.Read, 4 * 1024, FileOptions.SequentialScan))
+				using (FileStream fs = new FileStream(filePath, FileMode.Open, FileAccess.Read, FileShare.Read, 4 * 1024, FileOptions.SequentialScan))
 				{
 					using (StreamReader sr = new StreamReader(fs, Encoding.UTF8, true))
 					{
 
 						// Try to read the whole file into a buffer created by hand.  This avoids a LOT of memory allocations which in turn reduces the
 						// GC stress on the system.  Removing the StreamReader would be nice in the future.
-						long RawFileLength = fs.Length;
-						UhtBuffer InitialBuffer = UhtBuffer.Borrow((int)RawFileLength);
-						int ReadLength = sr.Read(InitialBuffer.Memory.Span);
+						long rawFileLength = fs.Length;
+						UhtBuffer initialBuffer = UhtBuffer.Borrow((int)rawFileLength);
+						int readLength = sr.Read(initialBuffer.Memory.Span);
 						if (sr.EndOfStream)
 						{
-							InitialBuffer.Reset(ReadLength);
-							return InitialBuffer;
+							initialBuffer.Reset(readLength);
+							return initialBuffer;
 						}
 						else
 						{
-							string Remaining = sr.ReadToEnd();
-							long TotalSize = ReadLength + Remaining.Length;
-							UhtBuffer Combined = UhtBuffer.Borrow((int)TotalSize);
-							Buffer.BlockCopy(InitialBuffer.Block, 0, Combined.Block, 0, ReadLength * sizeof(char));
-							Buffer.BlockCopy(Remaining.ToArray(), 0, Combined.Block, ReadLength * sizeof(char), Remaining.Length * sizeof(char));
-							UhtBuffer.Return(InitialBuffer);
-							return Combined;
+							string remaining = sr.ReadToEnd();
+							long totalSize = readLength + remaining.Length;
+							UhtBuffer combined = UhtBuffer.Borrow((int)totalSize);
+							Buffer.BlockCopy(initialBuffer.Block, 0, combined.Block, 0, readLength * sizeof(char));
+							Buffer.BlockCopy(remaining.ToArray(), 0, combined.Block, readLength * sizeof(char), remaining.Length * sizeof(char));
+							UhtBuffer.Return(initialBuffer);
+							return combined;
 						}
 					}
 				}
@@ -130,18 +130,18 @@ namespace EpicGames.UHT.Utils
 		}
 
 		/// <inheritdoc/>
-		public bool WriteOutput(string FilePath, ReadOnlySpan<char> Contents)
+		public bool WriteOutput(string filePath, ReadOnlySpan<char> contents)
 		{
 			try
 			{
-				string? FileDirectory = Path.GetDirectoryName(FilePath);
-				if (!string.IsNullOrEmpty(FileDirectory))
+				string? fileDirectory = Path.GetDirectoryName(filePath);
+				if (!String.IsNullOrEmpty(fileDirectory))
 				{
-					Directory.CreateDirectory(FileDirectory);
+					Directory.CreateDirectory(fileDirectory);
 				}
-				using (StreamWriter Writer = new StreamWriter(FilePath, false, new UTF8Encoding(false, true), 16 * 1024))
+				using (StreamWriter writer = new StreamWriter(filePath, false, new UTF8Encoding(false, true), 16 * 1024))
 				{
-					Writer.Write(Contents);
+					writer.Write(contents);
 				}
 				return true;
 			}
@@ -152,11 +152,11 @@ namespace EpicGames.UHT.Utils
 		}
 
 		/// <inheritdoc/>
-		public bool RenameOutput(string OldFilePath, string NewFilePath)
+		public bool RenameOutput(string oldFilePath, string newFilePath)
 		{
 			try
 			{
-				File.Move(OldFilePath, NewFilePath, true);
+				File.Move(oldFilePath, newFilePath, true);
 				return true;
 			}
 			catch (Exception)
@@ -168,38 +168,38 @@ namespace EpicGames.UHT.Utils
 		/// <summary>
 		/// Read the given source file
 		/// </summary>
-		/// <param name="FilePath">Full file path</param>
-		/// <param name="Contents">Contents of the file</param>
+		/// <param name="filePath">Full file path</param>
+		/// <param name="contents">Contents of the file</param>
 		/// <returns>True if the file was read, false if not</returns>
-		private static bool ReadFile(string FilePath, out StringView Contents)
+		private static bool ReadFile(string filePath, out StringView contents)
 		{
 			// Exceptions are very expensive.  Don't bother trying to open the file if it doesn't exist
-			if (!File.Exists(FilePath))
+			if (!File.Exists(filePath))
 			{
-				Contents = new StringView();
+				contents = new StringView();
 				return false;
 			}
 
 			try
 			{
-				using (FileStream fs = new FileStream(FilePath, FileMode.Open, FileAccess.Read, FileShare.Read, 4 * 1024, FileOptions.SequentialScan))
+				using (FileStream fs = new FileStream(filePath, FileMode.Open, FileAccess.Read, FileShare.Read, 4 * 1024, FileOptions.SequentialScan))
 				{
 					using (StreamReader sr = new StreamReader(fs, Encoding.UTF8, true))
 					{
 
 						// Try to read the whole file into a buffer created by hand.  This avoids a LOT of memory allocations which in turn reduces the
 						// GC stress on the system.  Removing the StreamReader would be nice in the future.
-						long RawFileLength = fs.Length;
-						char[] InitialBuffer = new char[RawFileLength];
-						int ReadLength = sr.Read(InitialBuffer, 0, (int)RawFileLength);
+						long rawFileLength = fs.Length;
+						char[] initialBuffer = new char[rawFileLength];
+						int readLength = sr.Read(initialBuffer, 0, (int)rawFileLength);
 						if (sr.EndOfStream)
 						{
-							Contents = new StringView(new ReadOnlyMemory<char>(InitialBuffer, 0, ReadLength));
+							contents = new StringView(new ReadOnlyMemory<char>(initialBuffer, 0, readLength));
 						}
 						else
 						{
-							string Remaining = sr.ReadToEnd();
-							Contents = new StringView(String.Concat(new ReadOnlySpan<char>(InitialBuffer, 0, ReadLength), Remaining));
+							string remaining = sr.ReadToEnd();
+							contents = new StringView(String.Concat(new ReadOnlySpan<char>(initialBuffer, 0, readLength), remaining));
 						}
 						return true;
 					}
@@ -207,7 +207,7 @@ namespace EpicGames.UHT.Utils
 			}
 			catch (IOException)
 			{
-				Contents = new StringView();
+				contents = new StringView();
 				return false;
 			}
 		}
diff --git a/Engine/Source/Programs/Shared/EpicGames.UHT/Utils/UhtHash.cs b/Engine/Source/Programs/Shared/EpicGames.UHT/Utils/UhtHash.cs
--- a/Engine/Source/Programs/Shared/EpicGames.UHT/Utils/UhtHash.cs
+++ b/Engine/Source/Programs/Shared/EpicGames.UHT/Utils/UhtHash.cs
@@ -1,7 +1,7 @@
 // Copyright Epic Games, Inc. All Rights Reserved.
 
-using EpicGames.Core;
 using System;
+using EpicGames.Core;
 
 namespace EpicGames.UHT.Utils
 {
@@ -11,34 +11,34 @@ namespace EpicGames.UHT.Utils
 	/// </summary>
 	public struct UhtHash
 	{
-		private ulong HashInternal;
+		private ulong _hashInternal;
 
 		/// <summary>
 		/// Start the hash computation
 		/// </summary>
 		public void Begin()
 		{
-			this.HashInternal = 0;
+			this._hashInternal = 0;
 		}
 
 		/// <summary>
 		/// Add a span of text to the current hash value
 		/// </summary>
-		/// <param name="Text"></param>
-		public void Add(ReadOnlySpan<char> Text)
+		/// <param name="text"></param>
+		public void Add(ReadOnlySpan<char> text)
 		{
-			while (Text.Length != 0)
+			while (text.Length != 0)
 			{
-				int CrPos = Text.IndexOf('\r');
-				if (CrPos == -1)
+				int crPos = text.IndexOf('\r');
+				if (crPos == -1)
 				{
-					HashText(Text.Slice(0, Text.Length));
+					HashText(text.Slice(0, text.Length));
 					break;
 				}
 				else
 				{
-					HashText(Text.Slice(0, CrPos));
-					Text = Text.Slice(CrPos + 1);
+					HashText(text.Slice(0, crPos));
+					text = text.Slice(crPos + 1);
 				}
 			}
 		}
@@ -49,31 +49,31 @@ namespace EpicGames.UHT.Utils
 		/// <returns>Final hash value</returns>
 		public uint End()
 		{
-			return (uint)(this.HashInternal + (this.HashInternal >> 32));
+			return (uint)(this._hashInternal + (this._hashInternal >> 32));
 		}
 
 		/// <summary>
 		/// Generate the hash value for a block of text
 		/// </summary>
-		/// <param name="Text">Text to hash</param>
+		/// <param name="text">Text to hash</param>
 		/// <returns>Hash value</returns>
-		public static uint GenenerateTextHash(ReadOnlySpan<char> Text)
+		public static uint GenenerateTextHash(ReadOnlySpan<char> text)
 		{
-			UhtHash Hash = new UhtHash();
-			Hash.Begin();
-			Hash.Add(Text);
-			return Hash.End();
+			UhtHash hash = new UhtHash();
+			hash.Begin();
+			hash.Add(text);
+			return hash.End();
 		}
 
-		private void HashText(ReadOnlySpan<char> Text)
+		private void HashText(ReadOnlySpan<char> text)
 		{
-			if (Text.Length > 0)
+			if (text.Length > 0)
 			{
 				unsafe
 				{
-					fixed (char* TextPtr = Text)
+					fixed (char* textPtr = text)
 					{
-						this.HashInternal = CityHash.CityHash64WithSeed((byte*)TextPtr, (uint)(Text.Length * sizeof(char)), this.HashInternal);
+						this._hashInternal = CityHash.CityHash64WithSeed((byte*)textPtr, (uint)(text.Length * sizeof(char)), this._hashInternal);
 					}
 				}
 			}
diff --git a/Engine/Source/Programs/Shared/EpicGames.UHT/Utils/UhtJsonConverters.cs b/Engine/Source/Programs/Shared/EpicGames.UHT/Utils/UhtJsonConverters.cs
--- a/Engine/Source/Programs/Shared/EpicGames.UHT/Utils/UhtJsonConverters.cs
+++ b/Engine/Source/Programs/Shared/EpicGames.UHT/Utils/UhtJsonConverters.cs
@@ -1,10 +1,10 @@
 // Copyright Epic Games, Inc. All Rights Reserved.
 
-using EpicGames.UHT.Types;
 using System;
 using System.Collections.Generic;
 using System.Text.Json;
 using System.Text.Json.Serialization;
+using EpicGames.UHT.Types;
 
 namespace EpicGames.UHT.Utils
 {
@@ -18,12 +18,12 @@ namespace EpicGames.UHT.Utils
 		/// <summary>
 		/// Read the JSON value
 		/// </summary>
-		/// <param name="Reader">Source reader</param>
-		/// <param name="TypeToConvert">Type to convert</param>
-		/// <param name="Options">Serialization options</param>
+		/// <param name="reader">Source reader</param>
+		/// <param name="typeToConvert">Type to convert</param>
+		/// <param name="options">Serialization options</param>
 		/// <returns>Read value</returns>
 		/// <exception cref="NotImplementedException"></exception>
-		public override T Read(ref Utf8JsonReader Reader, Type TypeToConvert, JsonSerializerOptions Options)
+		public override T Read(ref Utf8JsonReader reader, Type typeToConvert, JsonSerializerOptions options)
 		{
 			throw new NotImplementedException();
 		}
@@ -31,12 +31,12 @@ namespace EpicGames.UHT.Utils
 		/// <summary>
 		/// Write the JSON value
 		/// </summary>
-		/// <param name="Writer">Destination writer</param>
-		/// <param name="Type">Value being written</param>
-		/// <param name="Options">Serialization options</param>
-		public override void Write(Utf8JsonWriter Writer, T Type, JsonSerializerOptions Options)
+		/// <param name="writer">Destination writer</param>
+		/// <param name="type">Value being written</param>
+		/// <param name="options">Serialization options</param>
+		public override void Write(Utf8JsonWriter writer, T type, JsonSerializerOptions options)
 		{
-			Writer.WriteStringValue(Type.SourceName);
+			writer.WriteStringValue(type.SourceName);
 		}
 	}
 
@@ -49,12 +49,12 @@ namespace EpicGames.UHT.Utils
 		/// <summary>
 		/// Read the JSON value
 		/// </summary>
-		/// <param name="Reader">Source reader</param>
-		/// <param name="TypeToConvert">Type to convert</param>
-		/// <param name="Options">Serialization options</param>
+		/// <param name="reader">Source reader</param>
+		/// <param name="typeToConvert">Type to convert</param>
+		/// <param name="options">Serialization options</param>
 		/// <returns>Read value</returns>
 		/// <exception cref="NotImplementedException"></exception>
-		public override T Read(ref Utf8JsonReader Reader, Type TypeToConvert, JsonSerializerOptions Options)
+		public override T Read(ref Utf8JsonReader reader, Type typeToConvert, JsonSerializerOptions options)
 		{
 			throw new NotImplementedException();
 		}
@@ -62,14 +62,14 @@ namespace EpicGames.UHT.Utils
 		/// <summary>
 		/// Write the JSON value
 		/// </summary>
-		/// <param name="Writer">Destination writer</param>
-		/// <param name="Type">Value being written</param>
-		/// <param name="Options">Serialization options</param>
-		public override void Write(Utf8JsonWriter Writer, T? Type, JsonSerializerOptions Options)
+		/// <param name="writer">Destination writer</param>
+		/// <param name="type">Value being written</param>
+		/// <param name="options">Serialization options</param>
+		public override void Write(Utf8JsonWriter writer, T? type, JsonSerializerOptions options)
 		{
-			if (Type != null)
+			if (type != null)
 			{
-				Writer.WriteStringValue(Type.SourceName);
+				writer.WriteStringValue(type.SourceName);
 			}
 		}
 	}
@@ -83,12 +83,12 @@ namespace EpicGames.UHT.Utils
 		/// <summary>
 		/// Read the JSON value
 		/// </summary>
-		/// <param name="Reader">Source reader</param>
-		/// <param name="TypeToConvert">Type to convert</param>
-		/// <param name="Options">Serialization options</param>
+		/// <param name="reader">Source reader</param>
+		/// <param name="typeToConvert">Type to convert</param>
+		/// <param name="options">Serialization options</param>
 		/// <returns>Read value</returns>
 		/// <exception cref="NotImplementedException"></exception>
-		public override List<T> Read(ref Utf8JsonReader Reader, Type TypeToConvert, JsonSerializerOptions Options)
+		public override List<T> Read(ref Utf8JsonReader reader, Type typeToConvert, JsonSerializerOptions options)
 		{
 			throw new NotImplementedException();
 		}
@@ -96,17 +96,17 @@ namespace EpicGames.UHT.Utils
 		/// <summary>
 		/// Write the JSON value
 		/// </summary>
-		/// <param name="Writer">Destination writer</param>
-		/// <param name="Collection">Value being written</param>
-		/// <param name="Options">Serialization options</param>
-		public override void Write(Utf8JsonWriter Writer, List<T> Collection, JsonSerializerOptions Options)
+		/// <param name="writer">Destination writer</param>
+		/// <param name="collection">Value being written</param>
+		/// <param name="options">Serialization options</param>
+		public override void Write(Utf8JsonWriter writer, List<T> collection, JsonSerializerOptions options)
 		{
-			Writer.WriteStartArray();
-			foreach (UhtType Type in Collection)
+			writer.WriteStartArray();
+			foreach (UhtType type in collection)
 			{
-				JsonSerializer.Serialize(Writer, Type, Type.GetType(), Options);
+				JsonSerializer.Serialize(writer, type, type.GetType(), options);
 			}
-			Writer.WriteEndArray();
+			writer.WriteEndArray();
 		}
 	}
 
@@ -119,12 +119,12 @@ namespace EpicGames.UHT.Utils
 		/// <summary>
 		/// Read the JSON value
 		/// </summary>
-		/// <param name="Reader">Source reader</param>
-		/// <param name="TypeToConvert">Type to convert</param>
-		/// <param name="Options">Serialization options</param>
+		/// <param name="reader">Source reader</param>
+		/// <param name="typeToConvert">Type to convert</param>
+		/// <param name="options">Serialization options</param>
 		/// <returns>Read value</returns>
 		/// <exception cref="NotImplementedException"></exception>
-		public override List<T> Read(ref Utf8JsonReader Reader, Type TypeToConvert, JsonSerializerOptions Options)
+		public override List<T> Read(ref Utf8JsonReader reader, Type typeToConvert, JsonSerializerOptions options)
 		{
 			throw new NotImplementedException();
 		}
@@ -132,20 +132,20 @@ namespace EpicGames.UHT.Utils
 		/// <summary>
 		/// Write the JSON value
 		/// </summary>
-		/// <param name="Writer">Destination writer</param>
-		/// <param name="Collection">Value being written</param>
-		/// <param name="Options">Serialization options</param>
-		public override void Write(Utf8JsonWriter Writer, List<T>? Collection, JsonSerializerOptions Options)
+		/// <param name="writer">Destination writer</param>
+		/// <param name="collection">Value being written</param>
+		/// <param name="options">Serialization options</param>
+		public override void Write(Utf8JsonWriter writer, List<T>? collection, JsonSerializerOptions options)
 		{
-			Writer.WriteStartArray();
-			if (Collection != null)
+			writer.WriteStartArray();
+			if (collection != null)
 			{
-				foreach (UhtType Type in Collection)
+				foreach (UhtType type in collection)
 				{
-					JsonSerializer.Serialize(Writer, Type, Type.GetType(), Options);
+					JsonSerializer.Serialize(writer, type, type.GetType(), options);
 				}
 			}
-			Writer.WriteEndArray();
+			writer.WriteEndArray();
 		}
 	}
 }
diff --git a/Engine/Source/Programs/Shared/EpicGames.UHT/Utils/UhtMessage.cs b/Engine/Source/Programs/Shared/EpicGames.UHT/Utils/UhtMessage.cs
--- a/Engine/Source/Programs/Shared/EpicGames.UHT/Utils/UhtMessage.cs
+++ b/Engine/Source/Programs/Shared/EpicGames.UHT/Utils/UhtMessage.cs
@@ -1,10 +1,10 @@
 // Copyright Epic Games, Inc. All Rights Reserved.
 
-using EpicGames.UHT.Tokenizer;
 using System;
 using System.Collections.Generic;
 using System.Text;
 using System.Threading;
+using EpicGames.UHT.Tokenizer;
 
 namespace EpicGames.UHT.Utils
 {
@@ -49,8 +49,8 @@ namespace EpicGames.UHT.Utils
 		/// <summary>
 		/// Add the given message
 		/// </summary>
-		/// <param name="Message">The message to be added</param>
-		void AddMessage(UhtMessage Message);
+		/// <param name="message">The message to be added</param>
+		void AddMessage(UhtMessage message);
 	}
 
 	/// <summary>
@@ -71,7 +71,7 @@ namespace EpicGames.UHT.Utils
 		/// <summary>
 		/// If true, the source is a source fragment from the testing harness
 		/// </summary>
-		bool bMessageIsFragment { get; }
+		bool MessageIsFragment { get; }
 
 		/// <summary>
 		/// If this is a fragment, this is the container file path of the fragment
@@ -126,7 +126,7 @@ namespace EpicGames.UHT.Utils
 		/// Destination message session for the messages
 		/// </summary>
 		public IUhtMessageSession MessageSession { get; }
-		
+
 		/// <summary>
 		/// Source file generating messages
 		/// </summary>
@@ -171,40 +171,40 @@ namespace EpicGames.UHT.Utils
 		/// <summary>
 		/// Make a new message with the given settings
 		/// </summary>
-		/// <param name="MessageType">Type of message</param>
-		/// <param name="MessageSource">Source of the message</param>
-		/// <param name="FilePath">File path of the message</param>
-		/// <param name="LineNumber">Line number where message occurred</param>
-		/// <param name="Message">Text of the message</param>
+		/// <param name="messageType">Type of message</param>
+		/// <param name="messageSource">Source of the message</param>
+		/// <param name="filePath">File path of the message</param>
+		/// <param name="lineNumber">Line number where message occurred</param>
+		/// <param name="message">Text of the message</param>
 		/// <returns>Created message</returns>
-		public static UhtMessage MakeMessage(UhtMessageType MessageType, IUhtMessageSource? MessageSource, string? FilePath, int LineNumber, string Message)
+		public static UhtMessage MakeMessage(UhtMessageType messageType, IUhtMessageSource? messageSource, string? filePath, int lineNumber, string message)
 		{
 			return new UhtMessage
 			{
-				MessageType = MessageType,
-				MessageSource = MessageSource,
-				FilePath = FilePath,
-				LineNumber = LineNumber,
-				Message = Message
+				MessageType = messageType,
+				MessageSource = messageSource,
+				FilePath = filePath,
+				LineNumber = lineNumber,
+				Message = message
 			};
 		}
 
 		/// <summary>
 		/// Format an object to be included in a message
 		/// </summary>
-		/// <param name="Context">Contextual object</param>
+		/// <param name="context">Contextual object</param>
 		/// <returns>Formatted context</returns>
-		public static string FormatContext(object Context)
+		public static string FormatContext(object context)
 		{
-			if (Context is IUhtMessageExtraContext ExtraContextInterface)
+			if (context is IUhtMessageExtraContext extraContextInterface)
 			{
-				StringBuilder Builder = new StringBuilder();
-				Append(Builder, ExtraContextInterface, false);
-				return Builder.ToString();
+				StringBuilder builder = new StringBuilder();
+				Append(builder, extraContextInterface, false);
+				return builder.ToString();
 			}
-			else if (Context is UhtToken Token)
+			else if (context is UhtToken token)
 			{
-				switch (Token.TokenType)
+				switch (token.TokenType)
 				{
 					case UhtTokenType.EndOfFile:
 						return "EOF";
@@ -217,49 +217,49 @@ namespace EpicGames.UHT.Utils
 					case UhtTokenType.StringConst:
 						return "string constant";
 					default:
-						return $"'{Token.Value}'";
+						return $"'{token.Value}'";
 				}
 			}
-			else if (Context is char C)
+			else if (context is char c)
 			{
-				return $"'{C}'";
+				return $"'{c}'";
 			}
-			else if (Context is string[] StringArray)
+			else if (context is string[] stringArray)
 			{
-				return UhtUtilities.MergeTypeNames(StringArray, "or", true);
+				return UhtUtilities.MergeTypeNames(stringArray, "or", true);
 			}
 			else
 			{
-				return Context.ToString() ?? string.Empty;
+				return context.ToString() ?? String.Empty;
 			}
 		}
 
 		/// <summary>
 		/// Append message context
 		/// </summary>
-		/// <param name="Builder">Destination builder</param>
-		/// <param name="MessageExtraContextInterface">Extra context to append</param>
-		/// <param name="bAlwaysIncludeSeparator">If true, always include the separator</param>
-		public static void Append(StringBuilder Builder, IUhtMessageExtraContext? MessageExtraContextInterface, bool bAlwaysIncludeSeparator)
+		/// <param name="builder">Destination builder</param>
+		/// <param name="messageExtraContextInterface">Extra context to append</param>
+		/// <param name="alwaysIncludeSeparator">If true, always include the separator</param>
+		public static void Append(StringBuilder builder, IUhtMessageExtraContext? messageExtraContextInterface, bool alwaysIncludeSeparator)
 		{
-			if (MessageExtraContextInterface == null)
+			if (messageExtraContextInterface == null)
 			{
 				return;
 			}
 
-			IEnumerable<object?>? ExtraContextList = MessageExtraContextInterface.MessageExtraContext;
-			if (ExtraContextList != null)
+			IEnumerable<object?>? extraContextList = messageExtraContextInterface.MessageExtraContext;
+			if (extraContextList != null)
 			{
-				int StartingLength = Builder.Length;
-				foreach (object? EC in ExtraContextList)
+				int startingLength = builder.Length;
+				foreach (object? ec in extraContextList)
 				{
-					if (EC != null)
+					if (ec != null)
 					{
-						if (Builder.Length != StartingLength || bAlwaysIncludeSeparator)
+						if (builder.Length != startingLength || alwaysIncludeSeparator)
 						{
-							Builder.Append(" in ");
+							builder.Append(" in ");
 						}
-						Builder.Append(UhtMessage.FormatContext(EC));
+						builder.Append(UhtMessage.FormatContext(ec));
 					}
 				}
 			}
@@ -286,14 +286,14 @@ namespace EpicGames.UHT.Utils
 	/// </summary>
 	public class UhtSimpleMessageSite : IUhtMessageSite
 	{
-		private readonly IUhtMessageSession MessageSessionInternal;
-		private IUhtMessageSource? MessageSourceInternal;
+		private readonly IUhtMessageSession _messageSessionInternal;
+		private IUhtMessageSource? _messageSourceInternal;
 
 		#region IUHTMessageSite implementation
 		/// <inheritdoc/>
-		public IUhtMessageSession MessageSession => this.MessageSessionInternal;
+		public IUhtMessageSession MessageSession => this._messageSessionInternal;
 		/// <inheritdoc/>
-		public IUhtMessageSource? MessageSource { get => this.MessageSourceInternal; set => this.MessageSourceInternal = value; }
+		public IUhtMessageSource? MessageSource { get => this._messageSourceInternal; set => this._messageSourceInternal = value; }
 		/// <inheritdoc/>
 		public IUhtMessageLineNumber? MessageLineNumber => null;
 		#endregion
@@ -301,12 +301,12 @@ namespace EpicGames.UHT.Utils
 		/// <summary>
 		/// Create a simple message site for the given session and source
 		/// </summary>
-		/// <param name="MessageSession">Associated message session</param>
-		/// <param name="MessageSource">Source for the messages</param>
-		public UhtSimpleMessageSite(IUhtMessageSession MessageSession, IUhtMessageSource? MessageSource = null)
+		/// <param name="messageSession">Associated message session</param>
+		/// <param name="messageSource">Source for the messages</param>
+		public UhtSimpleMessageSite(IUhtMessageSession messageSession, IUhtMessageSource? messageSource = null)
 		{
-			this.MessageSessionInternal = MessageSession;
-			this.MessageSourceInternal = MessageSource;
+			this._messageSessionInternal = messageSession;
+			this._messageSourceInternal = messageSource;
 		}
 	}
 
@@ -317,29 +317,29 @@ namespace EpicGames.UHT.Utils
 	{
 		#region IUHTMessageSource implementation
 		/// <inheritdoc/>
-		public string MessageFilePath { get => this.FilePath; }
+		public string MessageFilePath => this._filePath;
 		/// <inheritdoc/>
-		public string MessageFullFilePath { get => this.FilePath; }
+		public string MessageFullFilePath => this._filePath;
 		/// <inheritdoc/>
-		public bool bMessageIsFragment { get => false; }
+		public bool MessageIsFragment => false;
 		/// <inheritdoc/>
-		public string MessageFragmentFilePath { get => ""; }
+		public string MessageFragmentFilePath => "";
 		/// <inheritdoc/>
-		public string MessageFragmentFullFilePath { get => ""; }
+		public string MessageFragmentFullFilePath => "";
 		/// <inheritdoc/>
-		public int MessageFragmentLineNumber { get => -1; }
+		public int MessageFragmentLineNumber => -1;
 		#endregion
 
-		private readonly string FilePath;
+		private readonly string _filePath;
 
 		/// <summary>
 		/// Create a simple file site
 		/// </summary>
-		/// <param name="MessageSession">Associated message session</param>
-		/// <param name="FilePath">File associated with the site</param>
-		public UhtSimpleFileMessageSite(IUhtMessageSession MessageSession, string FilePath) : base(MessageSession, null)
+		/// <param name="messageSession">Associated message session</param>
+		/// <param name="filePath">File associated with the site</param>
+		public UhtSimpleFileMessageSite(IUhtMessageSession messageSession, string filePath) : base(messageSession, null)
 		{
-			this.FilePath = FilePath;
+			this._filePath = filePath;
 			this.MessageSource = this;
 		}
 	}
@@ -352,18 +352,18 @@ namespace EpicGames.UHT.Utils
 		/// <summary>
 		/// Get the line number generating the error
 		/// </summary>
-		/// <param name="MessageSite">The message site generating the error.</param>
-		/// <param name="LineNumber">An override line number</param>
+		/// <param name="messageSite">The message site generating the error.</param>
+		/// <param name="lineNumber">An override line number</param>
 		/// <returns>Either the overriding line number or the line number from the message site.</returns>
-		public static int GetLineNumber(this IUhtMessageSite MessageSite, int LineNumber = -1)
+		public static int GetLineNumber(this IUhtMessageSite messageSite, int lineNumber = -1)
 		{
-			if (LineNumber != -1)
+			if (lineNumber != -1)
 			{
-				return LineNumber;
+				return lineNumber;
 			}
-			else if (MessageSite.MessageLineNumber != null)
+			else if (messageSite.MessageLineNumber != null)
 			{
-				return MessageSite.MessageLineNumber.MessageLineNumber;
+				return messageSite.MessageLineNumber.MessageLineNumber;
 			}
 			else
 			{
@@ -374,110 +374,110 @@ namespace EpicGames.UHT.Utils
 		/// <summary>
 		/// Log an error
 		/// </summary>
-		/// <param name="MessageSite">Message site associated with the message</param>
-		/// <param name="LineNumber">Line number of the error</param>
-		/// <param name="Message">Text of the error</param>
-		/// <param name="ExtraContext">Addition context to be appended to the error message</param>
-		public static void LogError(this IUhtMessageSite MessageSite, int LineNumber, string Message, object? ExtraContext = null)
+		/// <param name="messageSite">Message site associated with the message</param>
+		/// <param name="lineNumber">Line number of the error</param>
+		/// <param name="message">Text of the error</param>
+		/// <param name="extraContext">Addition context to be appended to the error message</param>
+		public static void LogError(this IUhtMessageSite messageSite, int lineNumber, string message, object? extraContext = null)
 		{
-			LogMessage(UhtMessageType.Error, MessageSite, LineNumber, Message, ExtraContext);
+			LogMessage(UhtMessageType.Error, messageSite, lineNumber, message, extraContext);
 		}
 
 		/// <summary>
 		/// Log an error
 		/// </summary>
-		/// <param name="MessageSite">Message site associated with the message</param>
-		/// <param name="Message">Text of the error</param>
-		/// <param name="ExtraContext">Addition context to be appended to the error message</param>
-		public static void LogError(this IUhtMessageSite MessageSite, string Message, object? ExtraContext = null)
+		/// <param name="messageSite">Message site associated with the message</param>
+		/// <param name="message">Text of the error</param>
+		/// <param name="extraContext">Addition context to be appended to the error message</param>
+		public static void LogError(this IUhtMessageSite messageSite, string message, object? extraContext = null)
 		{
-			LogMessage(UhtMessageType.Error, MessageSite, -1, Message, ExtraContext);
+			LogMessage(UhtMessageType.Error, messageSite, -1, message, extraContext);
 		}
 
 		/// <summary>
 		/// Log a warning
 		/// </summary>
-		/// <param name="MessageSite">Message site associated with the message</param>
-		/// <param name="LineNumber">Line number of the warning</param>
-		/// <param name="Message">Text of the warning</param>
-		/// <param name="ExtraContext">Addition context to be appended to the error message</param>
-		public static void LogWarning(this IUhtMessageSite MessageSite, int LineNumber, string Message, object? ExtraContext = null)
+		/// <param name="messageSite">Message site associated with the message</param>
+		/// <param name="lineNumber">Line number of the warning</param>
+		/// <param name="message">Text of the warning</param>
+		/// <param name="extraContext">Addition context to be appended to the error message</param>
+		public static void LogWarning(this IUhtMessageSite messageSite, int lineNumber, string message, object? extraContext = null)
 		{
-			LogMessage(UhtMessageType.Warning, MessageSite, LineNumber, Message, ExtraContext);
+			LogMessage(UhtMessageType.Warning, messageSite, lineNumber, message, extraContext);
 		}
 
 		/// <summary>
 		/// Log a warning
 		/// </summary>
-		/// <param name="MessageSite">Message site associated with the message</param>
-		/// <param name="Message">Text of the warning</param>
-		/// <param name="ExtraContext">Addition context to be appended to the error message</param>
-		public static void LogWarning(this IUhtMessageSite MessageSite, string Message, object? ExtraContext = null)
+		/// <param name="messageSite">Message site associated with the message</param>
+		/// <param name="message">Text of the warning</param>
+		/// <param name="extraContext">Addition context to be appended to the error message</param>
+		public static void LogWarning(this IUhtMessageSite messageSite, string message, object? extraContext = null)
 		{
-			LogMessage(UhtMessageType.Warning, MessageSite, -1, Message, ExtraContext);
+			LogMessage(UhtMessageType.Warning, messageSite, -1, message, extraContext);
 		}
 
 		/// <summary>
 		/// Log information
 		/// </summary>
-		/// <param name="MessageSite">Message site associated with the message</param>
-		/// <param name="LineNumber">Line number of the information</param>
-		/// <param name="Message">Text of the information</param>
-		/// <param name="ExtraContext">Addition context to be appended to the error message</param>
-		public static void LogInfo(this IUhtMessageSite MessageSite, int LineNumber, string Message, object? ExtraContext = null)
+		/// <param name="messageSite">Message site associated with the message</param>
+		/// <param name="lineNumber">Line number of the information</param>
+		/// <param name="message">Text of the information</param>
+		/// <param name="extraContext">Addition context to be appended to the error message</param>
+		public static void LogInfo(this IUhtMessageSite messageSite, int lineNumber, string message, object? extraContext = null)
 		{
-			LogMessage(UhtMessageType.Info, MessageSite, LineNumber, Message, ExtraContext);
+			LogMessage(UhtMessageType.Info, messageSite, lineNumber, message, extraContext);
 		}
 
 		/// <summary>
 		/// Log a message directly to the log
 		/// </summary>
-		/// <param name="MessageSite">Message site associated with the message</param>
-		/// <param name="Message">Text of the information</param>
-		/// <param name="ExtraContext">Addition context to be appended to the error message</param>
-		public static void LogTrace(this IUhtMessageSite MessageSite, string Message, object? ExtraContext = null)
+		/// <param name="messageSite">Message site associated with the message</param>
+		/// <param name="message">Text of the information</param>
+		/// <param name="extraContext">Addition context to be appended to the error message</param>
+		public static void LogTrace(this IUhtMessageSite messageSite, string message, object? extraContext = null)
 		{
-			LogMessage(UhtMessageType.Trace, MessageSite, -1, Message, ExtraContext);
+			LogMessage(UhtMessageType.Trace, messageSite, -1, message, extraContext);
 		}
 
 		/// <summary>
 		/// Log a message directly to the log
 		/// </summary>
-		/// <param name="MessageSite">Message site associated with the message</param>
-		/// <param name="LineNumber">Line number of the information</param>
-		/// <param name="Message">Text of the information</param>
-		/// <param name="ExtraContext">Addition context to be appended to the error message</param>
-		public static void LogLog(this IUhtMessageSite MessageSite, int LineNumber, string Message, object? ExtraContext = null)
+		/// <param name="messageSite">Message site associated with the message</param>
+		/// <param name="lineNumber">Line number of the information</param>
+		/// <param name="message">Text of the information</param>
+		/// <param name="extraContext">Addition context to be appended to the error message</param>
+		public static void LogLog(this IUhtMessageSite messageSite, int lineNumber, string message, object? extraContext = null)
 		{
-			LogMessage(UhtMessageType.Trace, MessageSite, LineNumber, Message, ExtraContext);
+			LogMessage(UhtMessageType.Trace, messageSite, lineNumber, message, extraContext);
 		}
 
 		/// <summary>
 		/// Log a information
 		/// </summary>
-		/// <param name="MessageSite">Message site associated with the message</param>
-		/// <param name="Message">Text of the information</param>
-		/// <param name="ExtraContext">Addition context to be appended to the error message</param>
-		public static void LogInfo(this IUhtMessageSite MessageSite, string Message, object? ExtraContext = null)
+		/// <param name="messageSite">Message site associated with the message</param>
+		/// <param name="message">Text of the information</param>
+		/// <param name="extraContext">Addition context to be appended to the error message</param>
+		public static void LogInfo(this IUhtMessageSite messageSite, string message, object? extraContext = null)
 		{
-			LogMessage(UhtMessageType.Info, MessageSite, -1, Message, ExtraContext);
+			LogMessage(UhtMessageType.Info, messageSite, -1, message, extraContext);
 		}
 
 		/// <summary>
 		/// Log a message
 		/// </summary>
-		/// <param name="MessageType">Type of the message being generated</param>
-		/// <param name="MessageSite">Message site associated with the message</param>
-		/// <param name="LineNumber">Line number of the information</param>
-		/// <param name="Message">Text of the information</param>
-		/// <param name="ExtraContext">Addition context to be appended to the error message</param>
-		private static void LogMessage(UhtMessageType MessageType, IUhtMessageSite MessageSite, int LineNumber, string Message, object? ExtraContext)
+		/// <param name="messageType">Type of the message being generated</param>
+		/// <param name="messageSite">Message site associated with the message</param>
+		/// <param name="lineNumber">Line number of the information</param>
+		/// <param name="message">Text of the information</param>
+		/// <param name="extraContext">Addition context to be appended to the error message</param>
+		private static void LogMessage(UhtMessageType messageType, IUhtMessageSite messageSite, int lineNumber, string message, object? extraContext)
 		{
-			if (ExtraContext != null)
+			if (extraContext != null)
 			{
-				Message = $"{Message} in {UhtMessage.FormatContext(ExtraContext)}";
+				message = $"{message} in {UhtMessage.FormatContext(extraContext)}";
 			}
-			MessageSite.MessageSession.AddMessage(UhtMessage.MakeMessage(MessageType, MessageSite.MessageSource, null, MessageSite.GetLineNumber(LineNumber), Message));
+			messageSite.MessageSession.AddMessage(UhtMessage.MakeMessage(messageType, messageSite.MessageSource, null, messageSite.GetLineNumber(lineNumber), message));
 		}
 	}
 
@@ -486,24 +486,24 @@ namespace EpicGames.UHT.Utils
 	/// </summary>
 	public sealed class UhtTlsMessageExtraContext : IUhtMessageExtraContext
 	{
-		private Stack<object?>? ExtraContexts;
-		private static readonly ThreadLocal<UhtTlsMessageExtraContext> Tls = new ThreadLocal<UhtTlsMessageExtraContext>(() => new UhtTlsMessageExtraContext());
+		private Stack<object?>? _extraContexts;
+		private static readonly ThreadLocal<UhtTlsMessageExtraContext> s_tls = new ThreadLocal<UhtTlsMessageExtraContext>(() => new UhtTlsMessageExtraContext());
 
 		#region IUHTMessageExtraContext implementation
-		IEnumerable<object?>? IUhtMessageExtraContext.MessageExtraContext => this.ExtraContexts;
+		IEnumerable<object?>? IUhtMessageExtraContext.MessageExtraContext => this._extraContexts;
 		#endregion
 
 		/// <summary>
 		/// Add an extra context
 		/// </summary>
-		/// <param name="ExceptionContext"></param>
-		public void PushExtraContext(object? ExceptionContext)
+		/// <param name="exceptionContext"></param>
+		public void PushExtraContext(object? exceptionContext)
 		{
-			if (this.ExtraContexts == null)
+			if (this._extraContexts == null)
 			{
-				this.ExtraContexts = new Stack<object?>(8);
+				this._extraContexts = new Stack<object?>(8);
 			}
-			this.ExtraContexts.Push(ExceptionContext);
+			this._extraContexts.Push(exceptionContext);
 		}
 
 		/// <summary>
@@ -511,9 +511,9 @@ namespace EpicGames.UHT.Utils
 		/// </summary>
 		public void PopExtraContext()
 		{
-			if (this.ExtraContexts != null)
+			if (this._extraContexts != null)
 			{
-				this.ExtraContexts.Pop();
+				this._extraContexts.Pop();
 			}
 		}
 
@@ -521,13 +521,13 @@ namespace EpicGames.UHT.Utils
 		/// Get the extra context associated with this thread
 		/// </summary>
 		/// <returns>Extra context</returns>
-		public static UhtTlsMessageExtraContext? GetTls() { return UhtTlsMessageExtraContext.Tls.Value; }
+		public static UhtTlsMessageExtraContext? GetTls() { return UhtTlsMessageExtraContext.s_tls.Value; }
 
 		/// <summary>
 		/// Get the extra context interface
 		/// </summary>
 		/// <returns>Extra context interface</returns>
-		public static IUhtMessageExtraContext? GetMessageExtraContext() { return UhtTlsMessageExtraContext.Tls.Value; }
+		public static IUhtMessageExtraContext? GetMessageExtraContext() { return UhtTlsMessageExtraContext.s_tls.Value; }
 	}
 
 	/// <summary>
@@ -535,18 +535,18 @@ namespace EpicGames.UHT.Utils
 	/// </summary>
 	public struct UhtMessageContext : IDisposable
 	{
-		private readonly UhtTlsMessageExtraContext? Stack;
+		private readonly UhtTlsMessageExtraContext? _stack;
 
 		/// <summary>
 		/// Construct a new entry
 		/// </summary>
-		/// <param name="ExtraContext">Extra context to be added</param>
-		public UhtMessageContext(object? ExtraContext)
+		/// <param name="extraContext">Extra context to be added</param>
+		public UhtMessageContext(object? extraContext)
 		{
-			this.Stack = UhtTlsMessageExtraContext.GetTls();
-			if (this.Stack != null)
+			this._stack = UhtTlsMessageExtraContext.GetTls();
+			if (this._stack != null)
 			{
-				this.Stack.PushExtraContext(ExtraContext);
+				this._stack.PushExtraContext(extraContext);
 			}
 		}
 
@@ -554,13 +554,13 @@ namespace EpicGames.UHT.Utils
 		/// Replace the extra context.  This replaces the top level context and thus 
 		/// can have unexpected results if done when a more deeper context has been added
 		/// </summary>
-		/// <param name="ExtraContext">New extra context</param>
-		public void Reset(object? ExtraContext)
+		/// <param name="extraContext">New extra context</param>
+		public void Reset(object? extraContext)
 		{
-			if (this.Stack != null)
+			if (this._stack != null)
 			{
-				this.Stack.PopExtraContext();
-				this.Stack.PushExtraContext(ExtraContext);
+				this._stack.PopExtraContext();
+				this._stack.PushExtraContext(extraContext);
 			}
 		}
 
@@ -569,9 +569,9 @@ namespace EpicGames.UHT.Utils
 		/// </summary>
 		public void Dispose()
 		{
-			if (this.Stack != null)
+			if (this._stack != null)
 			{
-				this.Stack.PopExtraContext();
+				this._stack.PopExtraContext();
 			}
 		}
 	}
diff --git a/Engine/Source/Programs/Shared/EpicGames.UHT/Utils/UhtReferenceCollector.cs b/Engine/Source/Programs/Shared/EpicGames.UHT/Utils/UhtReferenceCollector.cs
--- a/Engine/Source/Programs/Shared/EpicGames.UHT/Utils/UhtReferenceCollector.cs
+++ b/Engine/Source/Programs/Shared/EpicGames.UHT/Utils/UhtReferenceCollector.cs
@@ -1,9 +1,9 @@
 // Copyright Epic Games, Inc. All Rights Reserved.
 
-using EpicGames.Core;
-using EpicGames.UHT.Types;
 using System;
 using System.Collections.Generic;
+using EpicGames.Core;
+using EpicGames.UHT.Types;
 
 namespace EpicGames.UHT.Utils
 {
@@ -18,43 +18,43 @@ namespace EpicGames.UHT.Utils
 		/// <summary>
 		/// Add a cross module reference to a given object type.
 		/// </summary>
-		/// <param name="Obj">Object type being referenced</param>
-		/// <param name="bRegistered">If true, the method being invoked must return the registered object.  This only applies to classes.</param>
-		void AddCrossModuleReference(UhtObject? Obj, bool bRegistered);
+		/// <param name="obj">Object type being referenced</param>
+		/// <param name="registered">If true, the method being invoked must return the registered object.  This only applies to classes.</param>
+		void AddCrossModuleReference(UhtObject? obj, bool registered);
 
 		/// <summary>
 		/// Add an object declaration
 		/// </summary>
-		/// <param name="Obj">Object in question</param>
-		/// <param name="bRegistered">If true, the method being invoked must return the registered object.  This only applies to classes.</param>
-		void AddDeclaration(UhtObject Obj, bool bRegistered);
+		/// <param name="obj">Object in question</param>
+		/// <param name="registered">If true, the method being invoked must return the registered object.  This only applies to classes.</param>
+		void AddDeclaration(UhtObject obj, bool registered);
 
 		/// <summary>
 		/// Add a field as a singleton for exporting
 		/// </summary>
-		/// <param name="Field">Field to be added</param>
-		void AddSingleton(UhtField Field);
+		/// <param name="field">Field to be added</param>
+		void AddSingleton(UhtField field);
 
 		/// <summary>
 		/// Add a field as a type being exported
 		/// </summary>
-		/// <param name="Field">Field to be added</param>
-		void AddExportType(UhtField Field);
+		/// <param name="field">Field to be added</param>
+		void AddExportType(UhtField field);
 
 		/// <summary>
 		/// Add a forward declaration.  The string can contain multiple declarations but must only exist on one line.
 		/// </summary>
-		/// <param name="Declaration">The declarations to add</param>
-		void AddForwardDeclaration(string? Declaration);
+		/// <param name="declaration">The declarations to add</param>
+		void AddForwardDeclaration(string? declaration);
 	}
 
 	/// <summary>
 	/// Delegate used to fetch the string associated with a reference
 	/// </summary>
-	/// <param name="ObjectIndex">Index of the referenced object</param>
-	/// <param name="bRegistered">If true return the registered string, otherwise the unregistered string.  Classes have an unregistered version.</param>
+	/// <param name="objectIndex">Index of the referenced object</param>
+	/// <param name="registered">If true return the registered string, otherwise the unregistered string.  Classes have an unregistered version.</param>
 	/// <returns>The requested string</returns>
-	public delegate string GetReferenceStringDelegate(int ObjectIndex, bool bRegistered);
+	public delegate string GetReferenceStringDelegate(int objectIndex, bool registered);
 
 	/// <summary>
 	/// Maintains a list of referenced object indices.
@@ -76,39 +76,39 @@ namespace EpicGames.UHT.Utils
 		/// If the object has the alternate object set (i.e. native interfaces), then
 		/// that object's index is used to generate the key.
 		/// </summary>
-		/// <param name="Object">Object being referenced</param>
-		/// <param name="bRegistered">If true, then the API that ensures the object is registered is returned.</param>
+		/// <param name="obj">Object being referenced</param>
+		/// <param name="registered">If true, then the API that ensures the object is registered is returned.</param>
 		/// <returns>Integer key value.</returns>
-		public static int GetKey(UhtObject Object, bool bRegistered)
+		public static int GetKey(UhtObject obj, bool registered)
 		{
-			return Object.AlternateObject != null ? GetKey(Object.AlternateObject, bRegistered) : (Object.ObjectTypeIndex << 1) + (bRegistered ? 1 : 0);
+			return obj.AlternateObject != null ? GetKey(obj.AlternateObject, registered) : (obj.ObjectTypeIndex << 1) + (registered ? 1 : 0);
 		}
 
 		/// <summary>
 		/// Given a key, return the object index and registered flag
 		/// </summary>
-		/// <param name="Key">The key in question</param>
-		/// <param name="ObjectIndex">Index of the referenced object</param>
-		/// <param name="bRegistered">True if referencing the registered API.</param>
-		public static void UngetKey(int Key, out int ObjectIndex, out bool bRegistered)
+		/// <param name="key">The key in question</param>
+		/// <param name="objectIndex">Index of the referenced object</param>
+		/// <param name="registered">True if referencing the registered API.</param>
+		public static void UngetKey(int key, out int objectIndex, out bool registered)
 		{
-			ObjectIndex = Key >> 1;
-			bRegistered = (Key & 1) != 0;
+			objectIndex = key >> 1;
+			registered = (key & 1) != 0;
 		}
 
 		/// <summary>
 		/// Add the given object to the references
 		/// </summary>
-		/// <param name="Object">Object to be added</param>
-		/// <param name="bRegistered">True if the registered API is being returned.</param>
-		public void Add(UhtObject? Object, bool bRegistered)
+		/// <param name="obj">Object to be added</param>
+		/// <param name="registered">True if the registered API is being returned.</param>
+		public void Add(UhtObject? obj, bool registered)
 		{
-			if (Object != null)
+			if (obj != null)
 			{
-				int Key = GetKey(Object, bRegistered);
-				if (this.Uniques.Add(Key))
+				int key = GetKey(obj, registered);
+				if (this.Uniques.Add(key))
 				{
-					this.References.Add(Key);
+					this.References.Add(key);
 				}
 			}
 		}
@@ -116,42 +116,42 @@ namespace EpicGames.UHT.Utils
 		/// <summary>
 		/// Return the collection of references sorted by the API string returned by the delegate.
 		/// </summary>
-		/// <param name="Delegate">Delegate to invoke to return the requested object API string</param>
+		/// <param name="referenceStringDelegate">Delegate to invoke to return the requested object API string</param>
 		/// <returns>Read only memory region of all the string.</returns>
-		public ReadOnlyMemory<string> GetSortedReferences(GetReferenceStringDelegate Delegate)
+		public ReadOnlyMemory<string> GetSortedReferences(GetReferenceStringDelegate referenceStringDelegate)
 		{
 			// Collect the unsorted array
-			string[] Sorted = new string[this.References.Count];
-			for (int Index = 0; Index < this.References.Count; ++Index)
+			string[] sorted = new string[this.References.Count];
+			for (int index = 0; index < this.References.Count; ++index)
 			{
-				int Key = this.References[Index];
-				int ObjectIndex;
-				bool bRegistered;
-				UngetKey(Key, out ObjectIndex, out bRegistered);
-				Sorted[Index] = Delegate(ObjectIndex, bRegistered);
+				int key = this.References[index];
+				int objectIndex;
+				bool registered;
+				UngetKey(key, out objectIndex, out registered);
+				sorted[index] = referenceStringDelegate(objectIndex, registered);
 			}
 
 			// Sort the array
-			Array.Sort(Sorted, StringComparerUE.OrdinalIgnoreCase);
+			Array.Sort(sorted, StringComparerUE.OrdinalIgnoreCase);
 
 			// Remove duplicates.  In some instances the different keys might return the same string.
 			// This removes those duplicates
 			if (this.References.Count > 1)
 			{
-				int PriorOut = 0;
-				for (int Index = 1; Index < Sorted.Length; ++Index)
+				int priorOut = 0;
+				for (int index = 1; index < sorted.Length; ++index)
 				{
-					if (Sorted[Index] != Sorted[PriorOut])
+					if (sorted[index] != sorted[priorOut])
 					{
-						++PriorOut;
-						Sorted[PriorOut] = Sorted[Index];
+						++priorOut;
+						sorted[priorOut] = sorted[index];
 					}
 				}
-				return Sorted.AsMemory(0, PriorOut + 1);
+				return sorted.AsMemory(0, priorOut + 1);
 			}
 			else
 			{
-				return Sorted.AsMemory();
+				return sorted.AsMemory();
 			}
 		}
 	}
@@ -195,54 +195,54 @@ namespace EpicGames.UHT.Utils
 		/// <summary>
 		/// Add a cross module reference
 		/// </summary>
-		/// <param name="Object">Object being referenced</param>
-		/// <param name="bRegistered">True if the object being referenced must be registered</param>
-		public void AddCrossModuleReference(UhtObject? Object, bool bRegistered)
+		/// <param name="obj">Object being referenced</param>
+		/// <param name="registered">True if the object being referenced must be registered</param>
+		public void AddCrossModuleReference(UhtObject? obj, bool registered)
 		{
-			CrossModule.Add(Object, bRegistered);
-			if (Object != null && !(Object is UhtPackage) && bRegistered)
+			CrossModule.Add(obj, registered);
+			if (obj != null && obj is not UhtPackage && registered)
 			{
-				this.ReferencedHeaders.Add(Object.HeaderFile);
+				this.ReferencedHeaders.Add(obj.HeaderFile);
 			}
 		}
 
 		/// <summary>
 		/// Add a declaration
 		/// </summary>
-		/// <param name="Object">Object being declared</param>
-		/// <param name="bRegistered">True if the object being declared must be registered</param>
-		public void AddDeclaration(UhtObject Object, bool bRegistered)
+		/// <param name="obj">Object being declared</param>
+		/// <param name="registered">True if the object being declared must be registered</param>
+		public void AddDeclaration(UhtObject obj, bool registered)
 		{
-			this.Declaration.Add(Object, bRegistered);
+			this.Declaration.Add(obj, registered);
 		}
 
 		/// <summary>
 		/// Add a singleton.  These are added as forward declared functions in the package file.
 		/// </summary>
-		/// <param name="Field">Field being added</param>
-		public void AddSingleton(UhtField Field)
+		/// <param name="field">Field being added</param>
+		public void AddSingleton(UhtField field)
 		{
-			this.Singletons.Add(Field);
+			this.Singletons.Add(field);
 		}
 
 		/// <summary>
 		/// Add a type to be exported.
 		/// </summary>
-		/// <param name="Field">Type to be exported</param>
-		public void AddExportType(UhtField Field)
+		/// <param name="field">Type to be exported</param>
+		public void AddExportType(UhtField field)
 		{
-			this.ExportTypes.Add(Field);
+			this.ExportTypes.Add(field);
 		}
 
 		/// <summary>
 		/// Add a symbol that must be forward declared
 		/// </summary>
-		/// <param name="Declaration">Symbol to be forward declared.</param>
-		public void AddForwardDeclaration(string? Declaration)
+		/// <param name="declaration">Symbol to be forward declared.</param>
+		public void AddForwardDeclaration(string? declaration)
 		{
-			if (!string.IsNullOrEmpty(Declaration))
+			if (!String.IsNullOrEmpty(declaration))
 			{
-				this.ForwardDeclarations.Add(Declaration);
+				this.ForwardDeclarations.Add(declaration);
 			}
 		}
 	}
