More coding standard fixes

diff --git a/Engine/Source/Programs/Shared/EpicGames.UHT/Utils/UhtSession.cs b/Engine/Source/Programs/Shared/EpicGames.UHT/Utils/UhtSession.cs
--- a/Engine/Source/Programs/Shared/EpicGames.UHT/Utils/UhtSession.cs
+++ b/Engine/Source/Programs/Shared/EpicGames.UHT/Utils/UhtSession.cs
@@ -1,20 +1,20 @@
 // Copyright Epic Games, Inc. All Rights Reserved.
 
-using EpicGames.Core;
-using EpicGames.UHT.Parsers;
-using EpicGames.UHT.Tokenizer;
-using EpicGames.UHT.Tables;
-using Microsoft.Extensions.Logging;
 using System;
 using System.Collections.Generic;
 using System.IO;
+using System.IO.Enumeration;
 using System.Linq;
 using System.Text;
 using System.Text.Json;
 using System.Threading;
 using System.Threading.Tasks;
+using EpicGames.Core;
+using EpicGames.UHT.Parsers;
+using EpicGames.UHT.Tables;
+using EpicGames.UHT.Tokenizer;
 using EpicGames.UHT.Types;
-using System.IO.Enumeration;
+using Microsoft.Extensions.Logging;
 
 namespace EpicGames.UHT.Utils
 {
@@ -54,26 +54,26 @@ namespace EpicGames.UHT.Utils
 	{
 		public struct Output
 		{
-			public string FilePath;
-			public string TempFilePath;
-			public bool bSaved;
+			public string _filePath;
+			public string _tempFilePath;
+			public bool _saved;
 		}
 
 		/// <summary>
 		/// UHT session
 		/// </summary>
-		private readonly UhtSession SessionInternal;
+		private readonly UhtSession _sessionInternal;
 
 		/// <summary>
 		/// Module associated with the plugin
 		/// </summary>
-		private readonly UHTManifest.Module? PluginModuleInternal;
+		private readonly UHTManifest.Module? _pluginModuleInternal;
 
 		/// <summary>
 		/// Limiter for the number of files being saved to the reference directory.
 		/// The OS can get swamped on high core systems
 		/// </summary>
-		private static readonly Semaphore WriteRefSemaphore = new Semaphore(32, 32);
+		private static readonly Semaphore s_writeRefSemaphore = new Semaphore(32, 32);
 
 		/// <summary>
 		/// Requesting exporter
@@ -83,49 +83,49 @@ namespace EpicGames.UHT.Utils
 		/// <summary>
 		/// UHT Session
 		/// </summary>
-		public UhtSession Session => this.SessionInternal;
+		public UhtSession Session => this._sessionInternal;
 
 		/// <summary>
 		/// Plugin module
 		/// </summary>
-		public UHTManifest.Module? PluginModule => this.PluginModuleInternal;
+		public UHTManifest.Module? PluginModule => this._pluginModuleInternal;
 
 		/// <summary>
 		/// Collection of error from mismatches with the reference files
 		/// </summary>
-		public Dictionary<string, bool> ReferenceErrorMessages = new Dictionary<string, bool>();
+		public Dictionary<string, bool> ReferenceErrorMessages { get; } = new Dictionary<string, bool>();
 
 		/// <summary>
 		/// List of export outputs
 		/// </summary>
-		public List<Output> Outputs = new List<Output>();
+		public List<Output> Outputs { get; } = new List<Output>();
 
 		/// <summary>
 		/// Directory for the reference output
 		/// </summary>
-		public string ReferenceDirectory = string.Empty;
+		public string ReferenceDirectory { get; set; } = String.Empty;
 
 		/// <summary>
 		/// Directory for the verify output
 		/// </summary>
-		public string VerifyDirectory = string.Empty;
+		public string VerifyDirectory { get; set; } = String.Empty;
 
 		/// <summary>
 		/// Collection of external dependencies
 		/// </summary>
-		public HashSet<string> ExternalDependencies = new HashSet<string>();
+		public HashSet<string> ExternalDependencies { get; } = new HashSet<string>();
 
 		/// <summary>
 		/// Create a new instance of the export factory
 		/// </summary>
-		/// <param name="Session">UHT session</param>
-		/// <param name="PluginModule">Plugin module</param>
-		/// <param name="Exporter">Exporter being run</param>
-		public UhtExportFactory(UhtSession Session, UHTManifest.Module? PluginModule, UhtExporter Exporter)
+		/// <param name="session">UHT session</param>
+		/// <param name="pluginModule">Plugin module</param>
+		/// <param name="exporter">Exporter being run</param>
+		public UhtExportFactory(UhtSession session, UHTManifest.Module? pluginModule, UhtExporter exporter)
 		{
-			this.SessionInternal = Session;
-			this.PluginModuleInternal = PluginModule;
-			this.Exporter = Exporter;
+			this._sessionInternal = session;
+			this._pluginModuleInternal = pluginModule;
+			this.Exporter = exporter;
 			if (this.Session.ReferenceMode != UhtReferenceMode.None)
 			{
 				this.ReferenceDirectory = Path.Combine(this.Session.ReferenceDirectory, this.Exporter.Name);
@@ -139,51 +139,51 @@ namespace EpicGames.UHT.Utils
 		/// If you have a string builder, use this method so that a 
 		/// temporary buffer can be used.
 		/// </summary>
-		/// <param name="FilePath">Destination file path</param>
-		/// <param name="Builder">Source for the content</param>
-		public void CommitOutput(string FilePath, StringBuilder Builder)
+		/// <param name="filePath">Destination file path</param>
+		/// <param name="builder">Source for the content</param>
+		public void CommitOutput(string filePath, StringBuilder builder)
 		{
-			using (UhtBorrowBuffer BorrowBuffer = new UhtBorrowBuffer(Builder))
+			using (UhtBorrowBuffer borrowBuffer = new UhtBorrowBuffer(builder))
 			{
-				string TempFilePath = FilePath + ".tmp";
-				SaveIfChanged(FilePath, TempFilePath, new StringView(BorrowBuffer.Buffer.Memory));
+				string tempFilePath = filePath + ".tmp";
+				SaveIfChanged(filePath, tempFilePath, new StringView(borrowBuffer.Buffer.Memory));
 			}
 		}
 
 		/// <summary>
 		/// Commit the value of the string as the output
 		/// </summary>
-		/// <param name="FilePath">Destination file path</param>
-		/// <param name="Output">Output to commit</param>
-		public void CommitOutput(string FilePath, StringView Output)
+		/// <param name="filePath">Destination file path</param>
+		/// <param name="output">Output to commit</param>
+		public void CommitOutput(string filePath, StringView output)
 		{
-			string TempFilePath = FilePath + ".tmp";
-			SaveIfChanged(FilePath, TempFilePath, Output);
+			string tempFilePath = filePath + ".tmp";
+			SaveIfChanged(filePath, tempFilePath, output);
 		}
 
 		/// <summary>
 		/// Create a task to export two files
 		/// </summary>
-		/// <param name="Prereqs">Tasks that must be completed prior to this task running</param>
-		/// <param name="Action">Action to be invoked to generate the output</param>
+		/// <param name="prereqs">Tasks that must be completed prior to this task running</param>
+		/// <param name="action">Action to be invoked to generate the output</param>
 		/// <returns>Task object or null if the task was immediately executed.</returns>
-		public Task? CreateTask(List<Task?>? Prereqs, UhtExportTaskDelegate Action)
+		public Task? CreateTask(List<Task?>? prereqs, UhtExportTaskDelegate action)
 		{
-			if (this.Session.bGoWide)
+			if (this.Session.GoWide)
 			{
-				Task[]? PrereqTasks = Prereqs != null ? Prereqs.Where(x => x != null).Cast<Task>().ToArray() : null;
-				if (PrereqTasks != null && PrereqTasks.Length > 0)
+				Task[]? prereqTasks = prereqs?.Where(x => x != null).Cast<Task>().ToArray();
+				if (prereqTasks != null && prereqTasks.Length > 0)
 				{
-					return Task.Factory.ContinueWhenAll(PrereqTasks, (Task[] Tasks) => { Action(this); });
+					return Task.Factory.ContinueWhenAll(prereqTasks, (Task[] tasks) => { action(this); });
 				}
 				else
 				{
-					return Task.Factory.StartNew(() => { Action(this); }, CancellationToken.None, TaskCreationOptions.None, TaskScheduler.Default);
+					return Task.Factory.StartNew(() => { action(this); }, CancellationToken.None, TaskCreationOptions.None, TaskScheduler.Default);
 				}
 			}
 			else
 			{
-				Action(this);
+				action(this);
 				return null;
 			}
 		}
@@ -191,172 +191,172 @@ namespace EpicGames.UHT.Utils
 		/// <summary>
 		/// Create a task to export two files
 		/// </summary>
-		/// <param name="Action">Action to be invoked to generate the output</param>
+		/// <param name="action">Action to be invoked to generate the output</param>
 		/// <returns>Task object or null if the task was immediately executed.</returns>
-		public Task? CreateTask(UhtExportTaskDelegate Action)
+		public Task? CreateTask(UhtExportTaskDelegate action)
 		{
-			return CreateTask(null, Action);
+			return CreateTask(null, action);
 		}
 
 		/// <summary>
 		/// Given a header file, generate the output file name.
 		/// </summary>
-		/// <param name="HeaderFile">Header file</param>
-		/// <param name="Suffix">Suffix/extension to be added to the file name.</param>
+		/// <param name="headerFile">Header file</param>
+		/// <param name="suffix">Suffix/extension to be added to the file name.</param>
 		/// <returns>Resulting file name</returns>
-		public string MakePath(UhtHeaderFile HeaderFile, string Suffix)
+		public string MakePath(UhtHeaderFile headerFile, string suffix)
 		{
-			return MakePath(HeaderFile.Package.Module, HeaderFile.FileNameWithoutExtension, Suffix);
+			return MakePath(headerFile.Package.Module, headerFile.FileNameWithoutExtension, suffix);
 		}
 
 		/// <summary>
 		/// Given a package file, generate the output file name
 		/// </summary>
-		/// <param name="Package">Package file</param>
-		/// <param name="Suffix">Suffix/extension to be added to the file name.</param>
+		/// <param name="package">Package file</param>
+		/// <param name="suffix">Suffix/extension to be added to the file name.</param>
 		/// <returns>Resulting file name</returns>
-		public string MakePath(UhtPackage Package, string Suffix)
+		public string MakePath(UhtPackage package, string suffix)
 		{
-			return MakePath(Package.Module, Package.ShortName, Suffix);
+			return MakePath(package.Module, package.ShortName, suffix);
 		}
 
 		/// <summary>
 		/// Make a path for an output based on the package output directory.
 		/// </summary>
-		/// <param name="FileName">Name of the file</param>
-		/// <param name="Extension">Extension to add to the file</param>
+		/// <param name="fileName">Name of the file</param>
+		/// <param name="extension">Extension to add to the file</param>
 		/// <returns>Output file path</returns>
-		public string MakePath(string FileName, string Extension)
+		public string MakePath(string fileName, string extension)
 		{
 			if (this.PluginModule == null)
 			{
 				throw new UhtIceException("MakePath with just a filename and extension can not be called from non-plugin exporters");
 			}
-			return MakePath(this.PluginModule, FileName, Extension);
+			return MakePath(this.PluginModule, fileName, extension);
 		}
 
 		/// <summary>
 		/// Add an external dependency to the given file path
 		/// </summary>
-		/// <param name="FilePath">External dependency to add</param>
-		public void AddExternalDependency(string FilePath)
+		/// <param name="filePath">External dependency to add</param>
+		public void AddExternalDependency(string filePath)
 		{
-			this.ExternalDependencies.Add(FilePath);
+			this.ExternalDependencies.Add(filePath);
 		}
 
-		private string MakePath(UHTManifest.Module Module, string FileName, string Suffix)
+		private string MakePath(UHTManifest.Module module, string fileName, string suffix)
 		{
 			if (PluginModule != null)
 			{
-				Module = PluginModule;
+				module = PluginModule;
 			}
-			return Path.Combine(Module.OutputDirectory, FileName) + Suffix;
+			return Path.Combine(module.OutputDirectory, fileName) + suffix;
 		}
 
 		/// <summary>
 		/// Helper method to test to see if the output has changed.
 		/// </summary>
-		/// <param name="FilePath">Name of the output file</param>
-		/// <param name="TempFilePath">Name of the temporary file</param>
-		/// <param name="Exported">Exported contents of the file</param>
-		internal void SaveIfChanged(string FilePath, string TempFilePath, StringView Exported)
+		/// <param name="filePath">Name of the output file</param>
+		/// <param name="tempFilePath">Name of the temporary file</param>
+		/// <param name="exported">Exported contents of the file</param>
+		internal void SaveIfChanged(string filePath, string tempFilePath, StringView exported)
 		{
 
-			ReadOnlySpan<char> ExportedSpan = Exported.Span;
+			ReadOnlySpan<char> exportedSpan = exported.Span;
 
 			if (this.Session.ReferenceMode != UhtReferenceMode.None)
 			{
-				string FileName = Path.GetFileName(FilePath);
+				string fileName = Path.GetFileName(filePath);
 
 				// Writing billions of files to the same directory causes issues.  Use ourselves to throttle reference writes
 				try
 				{
-					UhtExportFactory.WriteRefSemaphore.WaitOne();
+					UhtExportFactory.s_writeRefSemaphore.WaitOne();
 					{
-						string OutPath = Path.Combine(this.Session.ReferenceMode == UhtReferenceMode.Reference ? this.ReferenceDirectory : this.VerifyDirectory, FileName);
-						if (!this.Session.WriteSource(OutPath, Exported.Span))
+						string outPath = Path.Combine(this.Session.ReferenceMode == UhtReferenceMode.Reference ? this.ReferenceDirectory : this.VerifyDirectory, fileName);
+						if (!this.Session.WriteSource(outPath, exported.Span))
 						{
-							new UhtSimpleFileMessageSite(this.Session, OutPath).LogWarning($"Unable to write reference file {OutPath}");
+							new UhtSimpleFileMessageSite(this.Session, outPath).LogWarning($"Unable to write reference file {outPath}");
 						}
 					}
 				}
 				finally
 				{
-					UhtExportFactory.WriteRefSemaphore.Release();
+					UhtExportFactory.s_writeRefSemaphore.Release();
 				}
 
 				// If we are verifying, read the existing file and check the contents
 				if (this.Session.ReferenceMode == UhtReferenceMode.Verify)
 				{
-					string Message = String.Empty;
-					string RefPath = Path.Combine(this.ReferenceDirectory, FileName);
-					UhtBuffer? ExistingRef = this.Session.ReadSourceToBuffer(RefPath);
-					if (ExistingRef != null)
+					string message = String.Empty;
+					string refPath = Path.Combine(this.ReferenceDirectory, fileName);
+					UhtBuffer? existingRef = this.Session.ReadSourceToBuffer(refPath);
+					if (existingRef != null)
 					{
-						ReadOnlySpan<char> ExistingSpan = ExistingRef.Memory.Span;
-						if (ExistingSpan.CompareTo(ExportedSpan, StringComparison.Ordinal) != 0)
+						ReadOnlySpan<char> existingSpan = existingRef.Memory.Span;
+						if (existingSpan.CompareTo(exportedSpan, StringComparison.Ordinal) != 0)
 						{
-							Message = $"********************************* {FileName} has changed.";
+							message = $"********************************* {fileName} has changed.";
 						}
-						UhtBuffer.Return(ExistingRef);
+						UhtBuffer.Return(existingRef);
 					}
 					else
 					{
-						Message = $"********************************* {FileName} appears to be a new generated file.";
+						message = $"********************************* {fileName} appears to be a new generated file.";
 					}
 
-					if (!string.IsNullOrEmpty(Message))
+					if (!String.IsNullOrEmpty(message))
 					{
-						Log.TraceInformation(Message);
+						Log.TraceInformation(message);
 						lock (this.ReferenceErrorMessages)
 						{
-							this.ReferenceErrorMessages.Add(Message, true);
+							this.ReferenceErrorMessages.Add(message, true);
 						}
 					}
 				}
 			}
 
 			// Check to see if the contents have changed
-			UhtBuffer? Original = this.Session.ReadSourceToBuffer(FilePath);
-			bool bSave = Original == null;
-			if (Original != null)
+			UhtBuffer? original = this.Session.ReadSourceToBuffer(filePath);
+			bool save = original == null;
+			if (original != null)
 			{
-				ReadOnlySpan<char> OriginalSpan = Original.Memory.Span;
-				if (OriginalSpan.CompareTo(ExportedSpan, StringComparison.Ordinal) != 0)
+				ReadOnlySpan<char> originalSpan = original.Memory.Span;
+				if (originalSpan.CompareTo(exportedSpan, StringComparison.Ordinal) != 0)
 				{
-					if (this.Session.bFailIfGeneratedCodeChanges)
+					if (this.Session.FailIfGeneratedCodeChanges)
 					{
-						string ConflictPath = FilePath + ".conflict";
-						if (!this.Session.WriteSource(ConflictPath, Exported.Span))
+						string conflictPath = filePath + ".conflict";
+						if (!this.Session.WriteSource(conflictPath, exported.Span))
 						{
-							new UhtSimpleFileMessageSite(this.Session, FilePath).LogError($"Changes to generated code are not allowed - conflicts written to '{ConflictPath}'");
+							new UhtSimpleFileMessageSite(this.Session, filePath).LogError($"Changes to generated code are not allowed - conflicts written to '{conflictPath}'");
 						}
 					}
-					bSave = true;
+					save = true;
 				}
-				UhtBuffer.Return(Original);
+				UhtBuffer.Return(original);
 			}
 
 			// If changed of the original didn't exist, then save the new version
-			if (bSave && !this.Session.bNoOutput)
+			if (save && !this.Session.NoOutput)
 			{
-				if (!this.Session.WriteSource(TempFilePath, Exported.Span))
+				if (!this.Session.WriteSource(tempFilePath, exported.Span))
 				{
-					new UhtSimpleFileMessageSite(this.Session, FilePath).LogWarning($"Failed to save export file: '{TempFilePath}'");
+					new UhtSimpleFileMessageSite(this.Session, filePath).LogWarning($"Failed to save export file: '{tempFilePath}'");
 				}
 			}
 			else
 			{
-				bSave = false;
+				save = false;
 			}
 
 			// Add this to the list of outputs
 			lock (this.Outputs)
 			{
-				this.Outputs.Add(new Output { FilePath = FilePath, TempFilePath = TempFilePath, bSaved = bSave });
+				this.Outputs.Add(new Output { _filePath = filePath, _tempFilePath = tempFilePath, _saved = save });
 			}
-		}       
-		
+		}
+
 		/// <summary>
 		/// Run the output exporter
 		/// </summary>
@@ -371,65 +371,65 @@ namespace EpicGames.UHT.Utils
 			{
 
 				// These outputs are used to cull old outputs from the directories
-				Dictionary<string, HashSet<string>> OutputsByDirectory = new Dictionary<string, HashSet<string>>(StringComparer.OrdinalIgnoreCase);
-				List<UhtExportFactory.Output> Saves = new List<UhtExportFactory.Output>();
+				Dictionary<string, HashSet<string>> outputsByDirectory = new Dictionary<string, HashSet<string>>(StringComparer.OrdinalIgnoreCase);
+				List<UhtExportFactory.Output> saves = new List<UhtExportFactory.Output>();
 
 				// Collect information about the outputs
-				foreach (UhtExportFactory.Output Output in this.Outputs)
+				foreach (UhtExportFactory.Output output in this.Outputs)
 				{
 
 					// Add this output to the list of expected outputs by directory
-					string? FileDirectory = Path.GetDirectoryName(Output.FilePath);
-					if (FileDirectory != null)
+					string? fileDirectory = Path.GetDirectoryName(output._filePath);
+					if (fileDirectory != null)
 					{
-						HashSet<string>? Files;
-						if (!OutputsByDirectory.TryGetValue(FileDirectory, out Files))
+						HashSet<string>? files;
+						if (!outputsByDirectory.TryGetValue(fileDirectory, out files))
 						{
-							Files = new HashSet<string>(StringComparer.OrdinalIgnoreCase);
-							OutputsByDirectory.Add(FileDirectory, Files);
+							files = new HashSet<string>(StringComparer.OrdinalIgnoreCase);
+							outputsByDirectory.Add(fileDirectory, files);
 						}
-						Files.Add(Path.GetFileName(Output.FilePath));
+						files.Add(Path.GetFileName(output._filePath));
 					}
 
 					// Add the save task
-					if (Output.bSaved)
+					if (output._saved)
 					{
-						Saves.Add(Output);
+						saves.Add(output);
 					}
 				}
 
 				// Perform the renames
-				if (this.Session.bGoWide)
+				if (this.Session.GoWide)
 				{
-					Parallel.ForEach(Saves, (UhtExportFactory.Output Output) =>
+					Parallel.ForEach(saves, (UhtExportFactory.Output output) =>
 					{
-						RenameSource(Output);
+						RenameSource(output);
 					});
 				}
 				else
-				{ 
-					foreach (UhtExportFactory.Output Output in Saves)
+				{
+					foreach (UhtExportFactory.Output output in saves)
 					{
-						RenameSource(Output);
+						RenameSource(output);
 					}
 				}
 
 				// Perform the culling of the output directories
-				if (this.Session.bCullOutput && !this.Session.bNoOutput && 
+				if (this.Session.CullOutput && !this.Session.NoOutput &&
 					(this.Exporter.CppFilters != null || this.Exporter.HeaderFilters != null || this.Exporter.OtherFilters != null))
 				{
-					if (this.Session.bGoWide)
+					if (this.Session.GoWide)
 					{
-						Parallel.ForEach(OutputsByDirectory, (KeyValuePair<string, HashSet<string>> Kvp) =>
+						Parallel.ForEach(outputsByDirectory, (KeyValuePair<string, HashSet<string>> kvp) =>
 						{
-							CullOutputDirectory(Kvp.Key, Kvp.Value);
+							CullOutputDirectory(kvp.Key, kvp.Value);
 						});
 					}
 					else
 					{
-						foreach (KeyValuePair<string, HashSet<string>> Kvp in OutputsByDirectory)
+						foreach (KeyValuePair<string, HashSet<string>> kvp in outputsByDirectory)
 						{
-							CullOutputDirectory(Kvp.Key, Kvp.Value);
+							CullOutputDirectory(kvp.Key, kvp.Value);
 						}
 					}
 				}
@@ -440,34 +440,34 @@ namespace EpicGames.UHT.Utils
 		/// Given an output, rename the output file from the temporary file name to the final file name.
 		/// If there exists a current final file, it will be replaced.
 		/// </summary>
-		/// <param name="Output">The output file to rename</param>
-		private void RenameSource(UhtExportFactory.Output Output)
+		/// <param name="output">The output file to rename</param>
+		private void RenameSource(UhtExportFactory.Output output)
 		{
-			this.Session.RenameSource(Output.TempFilePath, Output.FilePath);
+			this.Session.RenameSource(output._tempFilePath, output._filePath);
 		}
 
 		/// <summary>
 		/// Given a directory and a list of known files, delete any unknown file that matches of the supplied filters
 		/// </summary>
-		/// <param name="OutputDirectory">Output directory to scan</param>
-		/// <param name="KnownOutputs">Collection of known output files not to be deleted</param>
-		private void CullOutputDirectory(string OutputDirectory, HashSet<string> KnownOutputs)
+		/// <param name="outputDirectory">Output directory to scan</param>
+		/// <param name="knownOutputs">Collection of known output files not to be deleted</param>
+		private void CullOutputDirectory(string outputDirectory, HashSet<string> knownOutputs)
 		{
-			foreach (string FilePath in Directory.EnumerateFiles(OutputDirectory))
+			foreach (string filePath in Directory.EnumerateFiles(outputDirectory))
 			{
-				string FileName = Path.GetFileName(FilePath);
-				if (KnownOutputs.Contains(FileName))
+				string fileName = Path.GetFileName(filePath);
+				if (knownOutputs.Contains(fileName))
 				{
 					continue;
 				}
 
-				if (IsFilterMatch(FileName, this.Exporter.CppFilters) ||
-					IsFilterMatch(FileName, this.Exporter.HeaderFilters) ||
-					IsFilterMatch(FileName, this.Exporter.OtherFilters))
+				if (IsFilterMatch(fileName, this.Exporter.CppFilters) ||
+					IsFilterMatch(fileName, this.Exporter.HeaderFilters) ||
+					IsFilterMatch(fileName, this.Exporter.OtherFilters))
 				{
 					try
 					{
-						File.Delete(Path.Combine(OutputDirectory, FilePath));
+						File.Delete(Path.Combine(outputDirectory, filePath));
 					}
 					catch (Exception)
 					{
@@ -479,14 +479,14 @@ namespace EpicGames.UHT.Utils
 		/// <summary>
 		/// Test to see if the given filename (without directory), matches one of the given filters
 		/// </summary>
-		/// <param name="FileName">File name to test</param>
-		/// <param name="Filters">List of wildcard filters</param>
+		/// <param name="fileName">File name to test</param>
+		/// <param name="filters">List of wildcard filters</param>
 		/// <returns>True if there is a match</returns>
-		private static bool IsFilterMatch(string FileName, IEnumerable<string> Filters)
+		private static bool IsFilterMatch(string fileName, IEnumerable<string> filters)
 		{
-			foreach (string Filter in Filters)
+			foreach (string filter in filters)
 			{
-				if (FileSystemName.MatchesSimpleExpression(Filter, FileName, true))
+				if (FileSystemName.MatchesSimpleExpression(filter, fileName, true))
 				{
 					return true;
 				}
@@ -532,7 +532,7 @@ namespace EpicGames.UHT.Utils
 			/// <summary>
 			/// Current number of types
 			/// </summary>
-			private int CountInternal = 0;
+			private int _count = 0;
 
 			/// <summary>
 			/// Get the next type index
@@ -540,13 +540,13 @@ namespace EpicGames.UHT.Utils
 			/// <returns>Index starting at zero</returns>
 			public int GetNext()
 			{
-				return Interlocked.Increment(ref this.CountInternal) - 1;
+				return Interlocked.Increment(ref this._count) - 1;
 			}
 
 			/// <summary>
 			/// The number of times GetNext was called.
 			/// </summary>
-			public int Count => Interlocked.Add(ref this.CountInternal, 0) + 1;
+			public int Count => Interlocked.Add(ref this._count, 0) + 1;
 		}
 
 		/// <summary>
@@ -554,14 +554,14 @@ namespace EpicGames.UHT.Utils
 		/// </summary>
 		private struct EnumAndValue
 		{
-			public UhtEnum Enum;
-			public long Value;
+			public UhtEnum Enum { get; set; }
+			public long Value { get; set; }
 		}
 
 		/// <summary>
 		/// Collection of reserved names
 		/// </summary>
-		private static readonly HashSet<string> ReservedNames = new HashSet<string> { "none" };
+		private static readonly HashSet<string> s_reservedNames = new HashSet<string> { "none" };
 
 		#region Configurable settings
 
@@ -593,12 +593,12 @@ namespace EpicGames.UHT.Utils
 		/// <summary>
 		/// Directory to store the reference output
 		/// </summary>
-		public string ReferenceDirectory { get; set; } = string.Empty;
+		public string ReferenceDirectory { get; set; } = String.Empty;
 
 		/// <summary>
 		/// Directory to store the verification output
 		/// </summary>
-		public string VerifyDirectory { get; set; } = string.Empty;
+		public string VerifyDirectory { get; set; } = String.Empty;
 
 		/// <summary>
 		/// Mode for generating and/or testing reference output
@@ -608,48 +608,48 @@ namespace EpicGames.UHT.Utils
 		/// <summary>
 		/// If true, warnings are considered to be errors
 		/// </summary>
-		public bool bWarningsAsErrors { get; set; } = false;
+		public bool WarningsAsErrors { get; set; } = false;
 
 		/// <summary>
 		/// If true, include relative file paths in the log file
 		/// </summary>
-		public bool bRelativePathInLog { get; set; } = false;
+		public bool RelativePathInLog { get; set; } = false;
 
 		/// <summary>
 		/// If true, use concurrent tasks to run UHT
 		/// </summary>
-		public bool bGoWide { get; set; } = true;
+		public bool GoWide { get; set; } = true;
 
 		/// <summary>
 		/// If any output file mismatches existing outputs, an error will be generated
 		/// </summary>
-		public bool bFailIfGeneratedCodeChanges { get; set; } = false;
+		public bool FailIfGeneratedCodeChanges { get; set; } = false;
 
 		/// <summary>
 		/// If true, no output files will be saved
 		/// </summary>
-		public bool bNoOutput { get; set; } = false;
+		public bool NoOutput { get; set; } = false;
 
 		/// <summary>
 		/// If true, cull the output for any extra files
 		/// </summary>
-		public bool bCullOutput { get; set; } = true;
+		public bool CullOutput { get; set; } = true;
 
 		/// <summary>
 		/// If true, include extra output in code generation
 		/// </summary>
-		public bool bIncludeDebugOutput { get; set; } = false;
+		public bool IncludeDebugOutput { get; set; } = false;
 
 		/// <summary>
 		/// If true, disable all exporters which would normally be run by default
 		/// </summary>
-		public bool bNoDefaultExporters { get; set; } = false;
+		public bool NoDefaultExporters { get; set; } = false;
 
 		/// <summary>
 		/// If true, cache any error messages until the end of processing.  This is used by the testing
 		/// harness to generate more stable console output.
 		/// </summary>
-		public bool bCacheMessages { get; set; } = false;
+		public bool CacheMessages { get; set; } = false;
 
 		/// <summary>
 		/// Collection of UHT tables
@@ -670,29 +670,29 @@ namespace EpicGames.UHT.Utils
 		/// <summary>
 		/// Manifest data from the manifest file
 		/// </summary>
-		public UHTManifest? Manifest { get => this.ManifestFile != null ? this.ManifestFile.Manifest : null; }
+		public UHTManifest? Manifest => this.ManifestFile?.Manifest;
 
 		/// <summary>
 		/// Collection of packages from the manifest
 		/// </summary>
-		public IReadOnlyList<UhtPackage> Packages { get => this.PackagesInternal; }
+		public IReadOnlyList<UhtPackage> Packages => this._packagesInternal;
 
 		/// <summary>
 		/// Collection of header files from the manifest.  The header files will also appear as the children 
 		/// of the packages
 		/// </summary>
-		public IReadOnlyList<UhtHeaderFile> HeaderFiles { get => this.HeaderFilesInternal; }
+		public IReadOnlyList<UhtHeaderFile> HeaderFiles => this._headerFilesInternal;
 
 		/// <summary>
 		/// Collection of header files topologically sorted.  This will not be populated until after header files
 		/// are parsed and resolved.
 		/// </summary>
-		public IReadOnlyList<UhtHeaderFile> SortedHeaderFiles { get => this.SortedHeaderFilesInternal; }
+		public IReadOnlyList<UhtHeaderFile> SortedHeaderFiles => this._sortedHeaderFilesInternal;
 
 		/// <summary>
 		/// Dictionary of stripped file name to the header file
 		/// </summary>
-		public IReadOnlyDictionary<string, UhtHeaderFile> HeaderFileDictionary { get => this.HeaderFileDictionaryInternal; }
+		public IReadOnlyDictionary<string, UhtHeaderFile> HeaderFileDictionary => this._headerFileDictionaryInternal;
 
 		/// <summary>
 		/// After headers are parsed, returns the UObject class.
@@ -701,11 +701,11 @@ namespace EpicGames.UHT.Utils
 		{
 			get
 			{
-				if (this.UObjectInternal == null)
+				if (this._uobject == null)
 				{
 					throw new UhtIceException("UObject was not defined.");
 				}
-				return this.UObjectInternal;
+				return this._uobject;
 			}
 		}
 
@@ -716,11 +716,11 @@ namespace EpicGames.UHT.Utils
 		{
 			get
 			{
-				if (this.UClassInternal == null)
+				if (this._uclass == null)
 				{
 					throw new UhtIceException("UClass was not defined.");
 				}
-				return this.UClassInternal;
+				return this._uclass;
 			}
 		}
 
@@ -731,11 +731,11 @@ namespace EpicGames.UHT.Utils
 		{
 			get
 			{
-				if (this.UInterfaceInternal == null)
+				if (this._uinterface == null)
 				{
 					throw new UhtIceException("UInterface was not defined.");
 				}
-				return this.UInterfaceInternal;
+				return this._uinterface;
 			}
 		}
 
@@ -746,11 +746,11 @@ namespace EpicGames.UHT.Utils
 		{
 			get
 			{
-				if (this.IInterfaceInternal == null)
+				if (this._iinterface == null)
 				{
 					throw new UhtIceException("IInterface was not defined.");
 				}
-				return this.IInterfaceInternal;
+				return this._iinterface;
 			}
 		}
 
@@ -766,54 +766,45 @@ namespace EpicGames.UHT.Utils
 		/// </summary>
 		public UhtClass? INotifyFieldValueChanged { get; set; } = null;
 
-		private readonly List<UhtPackage> PackagesInternal = new List<UhtPackage>();
-		private readonly List<UhtHeaderFile> HeaderFilesInternal = new List<UhtHeaderFile>();
-		private readonly List<UhtHeaderFile> SortedHeaderFilesInternal = new List<UhtHeaderFile>();
-		private readonly Dictionary<string, UhtHeaderFile> HeaderFileDictionaryInternal = new Dictionary<string, UhtHeaderFile>(StringComparer.OrdinalIgnoreCase);
-		private long ErrorCountInternal = 0;
-		private long WarningCountInternal = 0;
-		private readonly List<UhtMessage> Messages = new List<UhtMessage>();
-		private Task? MessageTask = null;
-		private UhtClass? UObjectInternal = null;
-		private UhtClass? UClassInternal = null;
-		private UhtClass? UInterfaceInternal = null;
-		private UhtClass? IInterfaceInternal = null;
-		private readonly TypeCounter TypeCounterInternal = new TypeCounter();
-		private readonly TypeCounter PackageTypeCountInternal = new TypeCounter();
-		private readonly TypeCounter HeaderFileTypeCountInternal = new TypeCounter();
-		private readonly TypeCounter ObjectTypeCountInternal = new TypeCounter();
-		private UhtSymbolTable SourceNameSymbolTable = new UhtSymbolTable(0);
-		private UhtSymbolTable EngineNameSymbolTable = new UhtSymbolTable(0);
-		private bool bSymbolTablePopulated = false;
-		private Task? ReferenceDeleteTask = null;
-		private readonly Dictionary<string, bool> ExporterStates = new Dictionary<string, bool>(StringComparer.OrdinalIgnoreCase);
-		private readonly Dictionary<string, EnumAndValue> FullEnumValueLookup = new Dictionary<string, EnumAndValue>();
-		private readonly Dictionary<string, UhtEnum> ShortEnumValueLookup = new Dictionary<string, UhtEnum>();
-		private JsonDocument? ProjectJson = null;
+		private readonly List<UhtPackage> _packagesInternal = new List<UhtPackage>();
+		private readonly List<UhtHeaderFile> _headerFilesInternal = new List<UhtHeaderFile>();
+		private readonly List<UhtHeaderFile> _sortedHeaderFilesInternal = new List<UhtHeaderFile>();
+		private readonly Dictionary<string, UhtHeaderFile> _headerFileDictionaryInternal = new Dictionary<string, UhtHeaderFile>(StringComparer.OrdinalIgnoreCase);
+		private long _errorCountInternal = 0;
+		private long _warningCountInternal = 0;
+		private readonly List<UhtMessage> _messages = new List<UhtMessage>();
+		private Task? _messageTask = null;
+		private UhtClass? _uobject = null;
+		private UhtClass? _uclass = null;
+		private UhtClass? _uinterface = null;
+		private UhtClass? _iinterface = null;
+		private readonly TypeCounter _typeCounterInternal = new TypeCounter();
+		private readonly TypeCounter _packageTypeCountInternal = new TypeCounter();
+		private readonly TypeCounter _headerFileTypeCountInternal = new TypeCounter();
+		private readonly TypeCounter _objectTypeCountInternal = new TypeCounter();
+		private UhtSymbolTable _sourceNameSymbolTable = new UhtSymbolTable(0);
+		private UhtSymbolTable _engineNameSymbolTable = new UhtSymbolTable(0);
+		private bool _symbolTablePopulated = false;
+		private Task? _referenceDeleteTask = null;
+		private readonly Dictionary<string, bool> _exporterStates = new Dictionary<string, bool>(StringComparer.OrdinalIgnoreCase);
+		private readonly Dictionary<string, EnumAndValue> _fullEnumValueLookup = new Dictionary<string, EnumAndValue>();
+		private readonly Dictionary<string, UhtEnum> _shortEnumValueLookup = new Dictionary<string, UhtEnum>();
+		private JsonDocument? _projectJson = null;
 
 		/// <summary>
 		/// The number of errors
 		/// </summary>
-		public long ErrorCount
-		{
-			get => Interlocked.Read(ref this.ErrorCountInternal);
-		}
+		public long ErrorCount => Interlocked.Read(ref this._errorCountInternal);
 
 		/// <summary>
 		/// The number of warnings
 		/// </summary>
-		public long WarningCount
-		{
-			get => Interlocked.Read(ref this.WarningCountInternal);
-		}
+		public long WarningCount => Interlocked.Read(ref this._warningCountInternal);
 
 		/// <summary>
 		/// True if any errors have occurred or warnings if warnings are to be treated as errors 
 		/// </summary>
-		public bool bHasErrors
-		{
-			get => this.ErrorCount > 0 || (this.bWarningsAsErrors && this.WarningCount > 0);
-		}
+		public bool HasErrors => this.ErrorCount > 0 || (this.WarningsAsErrors && this.WarningCount > 0);
 
 		#region IUHTMessageSession implementation
 		/// <inheritdoc/>
@@ -830,13 +821,13 @@ namespace EpicGames.UHT.Utils
 		/// <returns>New index</returns>
 		public int GetNextTypeIndex()
 		{
-			return this.TypeCounterInternal.GetNext();
+			return this._typeCounterInternal.GetNext();
 		}
 
 		/// <summary>
 		/// Return the number of types that have been defined.  This includes all types.
 		/// </summary>
-		public int TypeCount => this.TypeCounterInternal.Count;
+		public int TypeCount => this._typeCounterInternal.Count;
 
 		/// <summary>
 		/// Return the index for a newly defined packaging
@@ -844,13 +835,13 @@ namespace EpicGames.UHT.Utils
 		/// <returns>New index</returns>
 		public int GetNextHeaderFileTypeIndex()
 		{
-			return this.HeaderFileTypeCountInternal.GetNext();
+			return this._headerFileTypeCountInternal.GetNext();
 		}
 
 		/// <summary>
 		/// Return the number of headers that have been defined
 		/// </summary>
-		public int HeaderFileTypeCount => this.HeaderFileTypeCountInternal.Count;
+		public int HeaderFileTypeCount => this._headerFileTypeCountInternal.Count;
 
 		/// <summary>
 		/// Return the index for a newly defined package
@@ -858,13 +849,13 @@ namespace EpicGames.UHT.Utils
 		/// <returns>New index</returns>
 		public int GetNextPackageTypeIndex()
 		{
-			return this.PackageTypeCountInternal.GetNext();
+			return this._packageTypeCountInternal.GetNext();
 		}
 
 		/// <summary>
 		/// Return the number of UPackage types that have been defined
 		/// </summary>
-		public int PackageTypeCount => this.PackageTypeCountInternal.Count;
+		public int PackageTypeCount => this._packageTypeCountInternal.Count;
 
 		/// <summary>
 		/// Return the index for a newly defined object
@@ -872,13 +863,13 @@ namespace EpicGames.UHT.Utils
 		/// <returns>New index</returns>
 		public int GetNextObjectTypeIndex()
 		{
-			return this.ObjectTypeCountInternal.GetNext();
+			return this._objectTypeCountInternal.GetNext();
 		}
 
 		/// <summary>
 		/// Return the total number of UObject types that have been defined
 		/// </summary>
-		public int ObjectTypeCount => this.ObjectTypeCountInternal.Count;
+		public int ObjectTypeCount => this._objectTypeCountInternal.Count;
 
 		/// <summary>
 		/// Return the collection of exporters
@@ -888,62 +879,62 @@ namespace EpicGames.UHT.Utils
 		/// <summary>
 		/// Return the keyword table for the given table name
 		/// </summary>
-		/// <param name="TableName">Name of the table</param>
+		/// <param name="tableName">Name of the table</param>
 		/// <returns>The requested table</returns>
-		public UhtKeywordTable GetKeywordTable(string TableName)
+		public UhtKeywordTable GetKeywordTable(string tableName)
 		{
-			return this.Tables!.KeywordTables.Get(TableName);
+			return this.Tables!.KeywordTables.Get(tableName);
 		}
 
 		/// <summary>
 		/// Return the specifier table for the given table name
 		/// </summary>
-		/// <param name="TableName">Name of the table</param>
+		/// <param name="tableName">Name of the table</param>
 		/// <returns>The requested table</returns>
-		public UhtSpecifierTable GetSpecifierTable(string TableName)
+		public UhtSpecifierTable GetSpecifierTable(string tableName)
 		{
-			return this.Tables!.SpecifierTables.Get(TableName);
+			return this.Tables!.SpecifierTables.Get(tableName);
 		}
 
 		/// <summary>
 		/// Return the specifier validator table for the given table name
 		/// </summary>
-		/// <param name="TableName">Name of the table</param>
+		/// <param name="tableName">Name of the table</param>
 		/// <returns>The requested table</returns>
-		public UhtSpecifierValidatorTable GetSpecifierValidatorTable(string TableName)
+		public UhtSpecifierValidatorTable GetSpecifierValidatorTable(string tableName)
 		{
-			return this.Tables!.SpecifierValidatorTables.Get(TableName);
+			return this.Tables!.SpecifierValidatorTables.Get(tableName);
 		}
 
 		/// <summary>
 		/// Generate an error for the given unhandled keyword
 		/// </summary>
-		/// <param name="TokenReader">Token reader</param>
-		/// <param name="Token">Unhandled token</param>
-		public void LogUnhandledKeywordError(IUhtTokenReader TokenReader, UhtToken Token)
+		/// <param name="tokenReader">Token reader</param>
+		/// <param name="token">Unhandled token</param>
+		public void LogUnhandledKeywordError(IUhtTokenReader tokenReader, UhtToken token)
 		{
-			this.Tables!.KeywordTables.LogUnhandledError(TokenReader, Token);
+			this.Tables!.KeywordTables.LogUnhandledError(tokenReader, token);
 		}
 
 		/// <summary>
 		/// Test to see if the given class name is a property
 		/// </summary>
-		/// <param name="Name">Name of the class without the prefix</param>
+		/// <param name="name">Name of the class without the prefix</param>
 		/// <returns>True if the class name is a property.  False if the class name isn't a property or isn't an engine class.</returns>
-		public bool IsValidPropertyTypeName(StringView Name)
+		public bool IsValidPropertyTypeName(StringView name)
 		{
-			return this.Tables!.EngineClassTable.IsValidPropertyTypeName(Name);
+			return this.Tables!.EngineClassTable.IsValidPropertyTypeName(name);
 		}
 
 		/// <summary>
 		/// Return the loc text default value associated with the given name
 		/// </summary>
-		/// <param name="Name"></param>
-		/// <param name="LocTextDefaultValue">Loc text default value handler</param>
+		/// <param name="name"></param>
+		/// <param name="locTextDefaultValue">Loc text default value handler</param>
 		/// <returns></returns>
-		public bool TryGetLocTextDefaultValue(StringView Name, out UhtLocTextDefaultValue LocTextDefaultValue)
+		public bool TryGetLocTextDefaultValue(StringView name, out UhtLocTextDefaultValue locTextDefaultValue)
 		{
-			return this.Tables!.LocTextDefaultValueTable.TryGet(Name, out LocTextDefaultValue);
+			return this.Tables!.LocTextDefaultValueTable.TryGet(name, out locTextDefaultValue);
 		}
 
 		/// <summary>
@@ -954,12 +945,12 @@ namespace EpicGames.UHT.Utils
 		/// <summary>
 		/// Return the property type associated with the given name
 		/// </summary>
-		/// <param name="Name"></param>
-		/// <param name="PropertyType">Property type if matched</param>
+		/// <param name="name"></param>
+		/// <param name="propertyType">Property type if matched</param>
 		/// <returns></returns>
-		public bool TryGetPropertyType(StringView Name, out UhtPropertyType PropertyType)
+		public bool TryGetPropertyType(StringView name, out UhtPropertyType propertyType)
 		{
-			return this.Tables!.PropertyTypeTable.TryGet(Name, out PropertyType);
+			return this.Tables!.PropertyTypeTable.TryGet(name, out propertyType);
 		}
 
 		/// <summary>
@@ -970,37 +961,37 @@ namespace EpicGames.UHT.Utils
 		/// <summary>
 		/// Return the structure default value associated with the given name
 		/// </summary>
-		/// <param name="Name"></param>
-		/// <param name="StructDefaultValue">Structure default value handler</param>
+		/// <param name="name"></param>
+		/// <param name="structDefaultValue">Structure default value handler</param>
 		/// <returns></returns>
-		public bool TryGetStructDefaultValue(StringView Name, out UhtStructDefaultValue StructDefaultValue)
+		public bool TryGetStructDefaultValue(StringView name, out UhtStructDefaultValue structDefaultValue)
 		{
-			return this.Tables!.StructDefaultValueTable.TryGet(Name, out StructDefaultValue);
+			return this.Tables!.StructDefaultValueTable.TryGet(name, out structDefaultValue);
 		}
 
 		/// <summary>
 		/// Run UHT on the given manifest.  Use the bHasError property to see if process was successful.
 		/// </summary>
-		/// <param name="ManifestFilePath">Path to the manifest file</param>
-		public void Run(string ManifestFilePath)
+		/// <param name="manifestFilePath">Path to the manifest file</param>
+		public void Run(string manifestFilePath)
 		{
 			if (this.FileManager == null)
 			{
-				Interlocked.Increment(ref this.ErrorCountInternal);
+				Interlocked.Increment(ref this._errorCountInternal);
 				Log.Logger.LogError("No file manager supplied, aborting.");
 				return;
 			}
 
 			if (this.Config == null)
 			{
-				Interlocked.Increment(ref this.ErrorCountInternal);
+				Interlocked.Increment(ref this._errorCountInternal);
 				Log.Logger.LogError("No configuration supplied, aborting.");
 				return;
 			}
 
 			if (this.Tables == null)
 			{
-				Interlocked.Increment(ref this.ErrorCountInternal);
+				Interlocked.Increment(ref this._errorCountInternal);
 				Log.Logger.LogError("No parsing tables supplied, aborting.");
 				return;
 			}
@@ -1011,7 +1002,7 @@ namespace EpicGames.UHT.Utils
 					break;
 
 				case UhtReferenceMode.Reference:
-					if (string.IsNullOrEmpty(this.ReferenceDirectory))
+					if (String.IsNullOrEmpty(this.ReferenceDirectory))
 					{
 						Log.Logger.LogError("WRITEREF requested but directory not set, ignoring");
 						this.ReferenceMode = UhtReferenceMode.None;
@@ -1019,7 +1010,7 @@ namespace EpicGames.UHT.Utils
 					break;
 
 				case UhtReferenceMode.Verify:
-					if (string.IsNullOrEmpty(this.ReferenceDirectory) || string.IsNullOrEmpty(this.VerifyDirectory))
+					if (String.IsNullOrEmpty(this.ReferenceDirectory) || String.IsNullOrEmpty(this.VerifyDirectory))
 					{
 						Log.Logger.LogError("VERIFYREF requested but directories not set, ignoring");
 						this.ReferenceMode = UhtReferenceMode.None;
@@ -1029,19 +1020,19 @@ namespace EpicGames.UHT.Utils
 
 			if (this.ReferenceMode != UhtReferenceMode.None)
 			{
-				string DirectoryToDelete = this.ReferenceMode == UhtReferenceMode.Reference ? this.ReferenceDirectory : this.VerifyDirectory;
-				this.ReferenceDeleteTask = Task.Factory.StartNew(() =>
+				string directoryToDelete = this.ReferenceMode == UhtReferenceMode.Reference ? this.ReferenceDirectory : this.VerifyDirectory;
+				this._referenceDeleteTask = Task.Factory.StartNew(() =>
 				{
 					try
 					{
-						Directory.Delete(DirectoryToDelete, true);
+						Directory.Delete(directoryToDelete, true);
 					}
 					catch (Exception)
 					{ }
 				}, CancellationToken.None, TaskCreationOptions.None, TaskScheduler.Default);
 			}
 
-			StepReadManifestFile(ManifestFilePath);
+			StepReadManifestFile(manifestFilePath);
 			StepPrepareModules();
 			StepPrepareHeaders();
 			StepParseHeaders();
@@ -1055,10 +1046,10 @@ namespace EpicGames.UHT.Utils
 			TopologicalSortHeaderFiles();
 
 			// If we are deleting the reference directory, then wait for that task to complete
-			if (this.ReferenceDeleteTask != null)
+			if (this._referenceDeleteTask != null)
 			{
 				Log.Logger.LogTrace("Step - Waiting for reference output to be cleared.");
-				this.ReferenceDeleteTask.Wait();
+				this._referenceDeleteTask.Wait();
 			}
 
 			StepExport();
@@ -1068,17 +1059,17 @@ namespace EpicGames.UHT.Utils
 		/// Try the given action regardless of any prior errors.  If an exception occurs that doesn't have the required
 		/// context, use the supplied context to generate the message.
 		/// </summary>
-		/// <param name="MessageSource">Message context for when the exception doesn't contain a context.</param>
-		/// <param name="Action">The lambda to be invoked</param>
-		public void TryAlways(IUhtMessageSource? MessageSource, Action Action)
+		/// <param name="messageSource">Message context for when the exception doesn't contain a context.</param>
+		/// <param name="action">The lambda to be invoked</param>
+		public void TryAlways(IUhtMessageSource? messageSource, Action action)
 		{
 			try
 			{
-				Action();
+				action();
 			}
-			catch (Exception E)
+			catch (Exception e)
 			{
-				HandleException(MessageSource, E);
+				HandleException(messageSource, e);
 			}
 		}
 
@@ -1086,19 +1077,19 @@ namespace EpicGames.UHT.Utils
 		/// Try the given action.  If an exception occurs that doesn't have the required
 		/// context, use the supplied context to generate the message.
 		/// </summary>
-		/// <param name="MessageSource">Message context for when the exception doesn't contain a context.</param>
-		/// <param name="Action">The lambda to be invoked</param>
-		public void Try(IUhtMessageSource? MessageSource, Action Action)
+		/// <param name="messageSource">Message context for when the exception doesn't contain a context.</param>
+		/// <param name="action">The lambda to be invoked</param>
+		public void Try(IUhtMessageSource? messageSource, Action action)
 		{
-			if (!this.bHasErrors)
+			if (!this.HasErrors)
 			{
 				try
 				{
-					Action();
+					action();
 				}
-				catch (Exception E)
+				catch (Exception e)
 				{
-					HandleException(MessageSource, E);
+					HandleException(messageSource, e);
 				}
 			}
 		}
@@ -1106,74 +1097,74 @@ namespace EpicGames.UHT.Utils
 		/// <summary>
 		/// Read the given source file
 		/// </summary>
-		/// <param name="FilePath">Full or relative file path</param>
+		/// <param name="filePath">Full or relative file path</param>
 		/// <returns>Information about the read source</returns>
-		public UhtSourceFragment ReadSource(string FilePath)
+		public UhtSourceFragment ReadSource(string filePath)
 		{
-			if (this.FileManager!.ReadSource(FilePath, out UhtSourceFragment Fragment))
+			if (this.FileManager!.ReadSource(filePath, out UhtSourceFragment fragment))
 			{
-				return Fragment;
+				return fragment;
 			}
-			throw new UhtException($"File not found '{FilePath}'");
+			throw new UhtException($"File not found '{filePath}'");
 		}
 
 		/// <summary>
 		/// Read the given source file
 		/// </summary>
-		/// <param name="FilePath">Full or relative file path</param>
+		/// <param name="filePath">Full or relative file path</param>
 		/// <returns>Buffer containing the read data or null if not found.  The returned buffer must be returned to the cache via a call to UhtBuffer.Return</returns>
-		public UhtBuffer? ReadSourceToBuffer(string FilePath)
+		public UhtBuffer? ReadSourceToBuffer(string filePath)
 		{
-			return this.FileManager!.ReadOutput(FilePath);
+			return this.FileManager!.ReadOutput(filePath);
 		}
 
 		/// <summary>
 		/// Write the given contents to the file
 		/// </summary>
-		/// <param name="FilePath">Path to write to</param>
-		/// <param name="Contents">Contents to write</param>
+		/// <param name="filePath">Path to write to</param>
+		/// <param name="contents">Contents to write</param>
 		/// <returns>True if the source was written</returns>
-		internal bool WriteSource(string FilePath, ReadOnlySpan<char> Contents)
+		internal bool WriteSource(string filePath, ReadOnlySpan<char> contents)
 		{
-			return this.FileManager!.WriteOutput(FilePath, Contents);
+			return this.FileManager!.WriteOutput(filePath, contents);
 		}
 
 		/// <summary>
 		/// Rename the given file
 		/// </summary>
-		/// <param name="OldFilePath">Old file path name</param>
-		/// <param name="NewFilePath">New file path name</param>
-		public void RenameSource(string OldFilePath, string NewFilePath)
+		/// <param name="oldFilePath">Old file path name</param>
+		/// <param name="newFilePath">New file path name</param>
+		public void RenameSource(string oldFilePath, string newFilePath)
 		{
-			if (!this.FileManager!.RenameOutput(OldFilePath, NewFilePath))
+			if (!this.FileManager!.RenameOutput(oldFilePath, newFilePath))
 			{
-				new UhtSimpleFileMessageSite(this, NewFilePath).LogError($"Failed to rename exported file: '{OldFilePath}'");
+				new UhtSimpleFileMessageSite(this, newFilePath).LogError($"Failed to rename exported file: '{oldFilePath}'");
 			}
 		}
 
 		/// <summary>
 		/// Given the name of a regular enum value, return the enum type
 		/// </summary>
-		/// <param name="Name">Enum value</param>
+		/// <param name="name">Enum value</param>
 		/// <returns>Associated regular enum type or null if not found or enum isn't a regular enum.</returns>
-		public UhtEnum? FindRegularEnumValue(string Name)
+		public UhtEnum? FindRegularEnumValue(string name)
 		{
 			//COMPATIBILITY-TODO - See comment below on a more rebust version of the enum lookup
-			//if (this.RegularEnumValueLookup.TryGetValue(Name, out UhtEnum? Enum))
+			//if (this.RegularEnumValueLookup.TryGetValue(name, out UhtEnum? enumObj))
 			//{
-			//	return Enum;
+			//	return enumObj;
 			//}
-			if (this.FullEnumValueLookup.TryGetValue(Name, out EnumAndValue Value))
+			if (this._fullEnumValueLookup.TryGetValue(name, out EnumAndValue value))
 			{
-				if (Value.Value != -1)
+				if (value.Value != -1)
 				{
-					return Value.Enum;
+					return value.Enum;
 				}
 			}
 
-			if (!Name.Contains("::", StringComparison.Ordinal) && this.ShortEnumValueLookup.TryGetValue(Name, out UhtEnum? Enum))
+			if (!name.Contains("::", StringComparison.Ordinal) && this._shortEnumValueLookup.TryGetValue(name, out UhtEnum? enumObj))
 			{
-				return Enum;
+				return enumObj;
 			}
 
 			return null;
@@ -1182,160 +1173,160 @@ namespace EpicGames.UHT.Utils
 		/// <summary>
 		/// Find the given type in the type hierarchy
 		/// </summary>
-		/// <param name="StartingType">Starting point for searches</param>
-		/// <param name="Options">Options controlling what is searched</param>
-		/// <param name="Name">Name of the type.</param>
-		/// <param name="MessageSite">If supplied, then a error message will be generated if the type can not be found</param>
-		/// <param name="LineNumber">Source code line number requesting the lookup.</param>
+		/// <param name="startingType">Starting point for searches</param>
+		/// <param name="options">Options controlling what is searched</param>
+		/// <param name="name">Name of the type.</param>
+		/// <param name="messageSite">If supplied, then a error message will be generated if the type can not be found</param>
+		/// <param name="lineNumber">Source code line number requesting the lookup.</param>
 		/// <returns>The located type of null if not found</returns>
-		public UhtType? FindType(UhtType? StartingType, UhtFindOptions Options, string Name, IUhtMessageSite? MessageSite = null, int LineNumber = -1)
+		public UhtType? FindType(UhtType? startingType, UhtFindOptions options, string name, IUhtMessageSite? messageSite = null, int lineNumber = -1)
 		{
-			ValidateFindOptions(Options);
+			ValidateFindOptions(options);
 
-			UhtType? Type = FindTypeInternal(StartingType, Options, Name);
-			if (Type == null && MessageSite != null)
+			UhtType? type = FindTypeInternal(startingType, options, name);
+			if (type == null && messageSite != null)
 			{
-				FindTypeError(MessageSite, LineNumber, Options, Name);
+				FindTypeError(messageSite, lineNumber, options, name);
 			}
-			return Type;
+			return type;
 		}
 
 		/// <summary>
 		/// Find the given type in the type hierarchy
 		/// </summary>
-		/// <param name="StartingType">Starting point for searches</param>
-		/// <param name="Options">Options controlling what is searched</param>
-		/// <param name="Name">Name of the type.</param>
-		/// <param name="MessageSite">If supplied, then a error message will be generated if the type can not be found</param>
+		/// <param name="startingType">Starting point for searches</param>
+		/// <param name="options">Options controlling what is searched</param>
+		/// <param name="name">Name of the type.</param>
+		/// <param name="messageSite">If supplied, then a error message will be generated if the type can not be found</param>
 		/// <returns>The located type of null if not found</returns>
-		public UhtType? FindType(UhtType? StartingType, UhtFindOptions Options, ref UhtToken Name, IUhtMessageSite? MessageSite = null)
+		public UhtType? FindType(UhtType? startingType, UhtFindOptions options, ref UhtToken name, IUhtMessageSite? messageSite = null)
 		{
-			ValidateFindOptions(Options);
+			ValidateFindOptions(options);
 
-			UhtType? Type = FindTypeInternal(StartingType, Options, Name.Value.ToString());
-			if (Type == null && MessageSite != null)
+			UhtType? type = FindTypeInternal(startingType, options, name.Value.ToString());
+			if (type == null && messageSite != null)
 			{
-				FindTypeError(MessageSite, Name.InputLine, Options, Name.Value.ToString());
+				FindTypeError(messageSite, name.InputLine, options, name.Value.ToString());
 			}
-			return Type;
+			return type;
 		}
 
 		/// <summary>
 		/// Find the given type in the type hierarchy
 		/// </summary>
-		/// <param name="StartingType">If specified, this represents the starting type to use when searching base/owner chain for a match</param>
-		/// <param name="Options">Options controlling what is searched</param>
-		/// <param name="Identifiers">Enumeration of identifiers.</param>
-		/// <param name="MessageSite">If supplied, then a error message will be generated if the type can not be found</param>
-		/// <param name="LineNumber">Source code line number requesting the lookup.</param>
+		/// <param name="startingType">If specified, this represents the starting type to use when searching base/owner chain for a match</param>
+		/// <param name="options">Options controlling what is searched</param>
+		/// <param name="identifiers">Enumeration of identifiers.</param>
+		/// <param name="messageSite">If supplied, then a error message will be generated if the type can not be found</param>
+		/// <param name="lineNumber">Source code line number requesting the lookup.</param>
 		/// <returns>The located type of null if not found</returns>
-		public UhtType? FindType(UhtType? StartingType, UhtFindOptions Options, UhtTokenList Identifiers, IUhtMessageSite? MessageSite = null, int LineNumber = -1)
+		public UhtType? FindType(UhtType? startingType, UhtFindOptions options, UhtTokenList identifiers, IUhtMessageSite? messageSite = null, int lineNumber = -1)
 		{
-			ValidateFindOptions(Options);
+			ValidateFindOptions(options);
 
-			if (Identifiers.Next != null && Identifiers.Next.Next != null)
+			if (identifiers.Next != null && identifiers.Next.Next != null)
 			{
-				if (MessageSite != null)
+				if (messageSite != null)
 				{
-					MessageSite.LogError(LineNumber, "UnrealHeaderTool only supports C++ identifiers of two or less identifiers");
+					messageSite.LogError(lineNumber, "UnrealHeaderTool only supports C++ identifiers of two or less identifiers");
 					return null;
 				}
 			}
 
-			UhtType? Type;
-			if (Identifiers.Next != null)
+			UhtType? type;
+			if (identifiers.Next != null)
 			{
-				Type = FindTypeTwoNamesInternal(StartingType, Options, Identifiers.Token.Value.ToString(), Identifiers.Next.Token.Value.ToString());
+				type = FindTypeTwoNamesInternal(startingType, options, identifiers.Token.Value.ToString(), identifiers.Next.Token.Value.ToString());
 			}
 			else
 			{
-				Type = FindTypeInternal(StartingType, Options, Identifiers.Token.Value.ToString());
+				type = FindTypeInternal(startingType, options, identifiers.Token.Value.ToString());
 			}
 
-			if (Type == null && MessageSite != null)
+			if (type == null && messageSite != null)
 			{
-				string FullIdentifier = Identifiers.Join("::");
-				FindTypeError(MessageSite, LineNumber, Options, FullIdentifier);
+				string fullIdentifier = identifiers.Join("::");
+				FindTypeError(messageSite, lineNumber, options, fullIdentifier);
 			}
-			return Type;
+			return type;
 		}
 
 		/// <summary>
 		/// Find the given type in the type hierarchy
 		/// </summary>
-		/// <param name="StartingType">If specified, this represents the starting type to use when searching base/owner chain for a match</param>
-		/// <param name="Options">Options controlling what is searched</param>
-		/// <param name="Identifiers">Enumeration of identifiers.</param>
-		/// <param name="MessageSite">If supplied, then a error message will be generated if the type can not be found</param>
-		/// <param name="LineNumber">Source code line number requesting the lookup.</param>
+		/// <param name="startingType">If specified, this represents the starting type to use when searching base/owner chain for a match</param>
+		/// <param name="options">Options controlling what is searched</param>
+		/// <param name="identifiers">Enumeration of identifiers.</param>
+		/// <param name="messageSite">If supplied, then a error message will be generated if the type can not be found</param>
+		/// <param name="lineNumber">Source code line number requesting the lookup.</param>
 		/// <returns>The located type of null if not found</returns>
-		public UhtType? FindType(UhtType? StartingType, UhtFindOptions Options, UhtToken[] Identifiers, IUhtMessageSite? MessageSite = null, int LineNumber = -1)
+		public UhtType? FindType(UhtType? startingType, UhtFindOptions options, UhtToken[] identifiers, IUhtMessageSite? messageSite = null, int lineNumber = -1)
 		{
-			ValidateFindOptions(Options);
+			ValidateFindOptions(options);
 
-			if (Identifiers.Length == 0)
+			if (identifiers.Length == 0)
 			{
 				throw new UhtIceException("Empty identifier array");
 			}
-			if (Identifiers.Length > 2)
+			if (identifiers.Length > 2)
 			{
-				if (MessageSite != null)
+				if (messageSite != null)
 				{
-					MessageSite.LogError(LineNumber, "UnrealHeaderTool only supports C++ identifiers of two or less identifiers");
+					messageSite.LogError(lineNumber, "UnrealHeaderTool only supports C++ identifiers of two or less identifiers");
 					return null;
 				}
 			}
 
-			UhtType? Type;
-			if (Identifiers.Length == 0)
+			UhtType? type;
+			if (identifiers.Length == 0)
 			{
-				Type = FindTypeTwoNamesInternal(StartingType, Options, Identifiers[0].Value.ToString(), Identifiers[1].Value.ToString());
+				type = FindTypeTwoNamesInternal(startingType, options, identifiers[0].Value.ToString(), identifiers[1].Value.ToString());
 			}
 			else
 			{
-				Type = FindTypeInternal(StartingType, Options, Identifiers[0].Value.ToString());
+				type = FindTypeInternal(startingType, options, identifiers[0].Value.ToString());
 			}
 
-			if (Type == null && MessageSite != null)
+			if (type == null && messageSite != null)
 			{
-				string FullIdentifier = string.Join("::", Identifiers);
-				FindTypeError(MessageSite, LineNumber, Options, FullIdentifier);
+				string fullIdentifier = String.Join("::", identifiers);
+				FindTypeError(messageSite, lineNumber, options, fullIdentifier);
 			}
-			return Type;
+			return type;
 		}
 
 		/// <summary>
 		/// Find the given type in the type hierarchy
 		/// </summary>
-		/// <param name="StartingType">Starting point for searches</param>
-		/// <param name="Options">Options controlling what is searched</param>
-		/// <param name="FirstName">First name of the type.</param>
-		/// <param name="SecondName">Second name used by delegates in classes and namespace enumerations</param>
+		/// <param name="startingType">Starting point for searches</param>
+		/// <param name="options">Options controlling what is searched</param>
+		/// <param name="firstName">First name of the type.</param>
+		/// <param name="secondName">Second name used by delegates in classes and namespace enumerations</param>
 		/// <returns>The located type of null if not found</returns>
-		private UhtType? FindTypeTwoNamesInternal(UhtType? StartingType, UhtFindOptions Options, string FirstName, string SecondName)
+		private UhtType? FindTypeTwoNamesInternal(UhtType? startingType, UhtFindOptions options, string firstName, string secondName)
 		{
 			// If we have two names
-			if (SecondName.Length > 0)
+			if (secondName.Length > 0)
 			{
-				if (Options.HasAnyFlags(UhtFindOptions.DelegateFunction | UhtFindOptions.Enum))
+				if (options.HasAnyFlags(UhtFindOptions.DelegateFunction | UhtFindOptions.Enum))
 				{
-					UhtFindOptions SubOptions = UhtFindOptions.NoParents | (Options & ~UhtFindOptions.TypesMask) | (Options & UhtFindOptions.Enum);
-					if (Options.HasAnyFlags(UhtFindOptions.DelegateFunction))
+					UhtFindOptions subOptions = UhtFindOptions.NoParents | (options & ~UhtFindOptions.TypesMask) | (options & UhtFindOptions.Enum);
+					if (options.HasAnyFlags(UhtFindOptions.DelegateFunction))
 					{
-						SubOptions |= UhtFindOptions.Class;
+						subOptions |= UhtFindOptions.Class;
 					}
-					UhtType? Type = FindTypeInternal(StartingType, SubOptions, FirstName);
-					if (Type == null)
+					UhtType? type = FindTypeInternal(startingType, subOptions, firstName);
+					if (type == null)
 					{
 						return null;
 					}
-					if (Type is UhtEnum)
+					if (type is UhtEnum)
 					{
-						return Type;
+						return type;
 					}
-					if (Type is UhtClass)
+					if (type is UhtClass)
 					{
-						return FindTypeInternal(StartingType, UhtFindOptions.DelegateFunction | UhtFindOptions.NoParents | (Options & ~UhtFindOptions.TypesMask), SecondName);
+						return FindTypeInternal(startingType, UhtFindOptions.DelegateFunction | UhtFindOptions.NoParents | (options & ~UhtFindOptions.TypesMask), secondName);
 					}
 				}
 
@@ -1344,38 +1335,38 @@ namespace EpicGames.UHT.Utils
 			}
 
 			// Perform the lookup for just a single name
-			return FindTypeInternal(StartingType, Options, FirstName);
+			return FindTypeInternal(startingType, options, firstName);
 		}
 
 		/// <summary>
 		/// Find the given type in the type hierarchy
 		/// </summary>
-		/// <param name="StartingType">Starting point for searches</param>
-		/// <param name="Options">Options controlling what is searched</param>
-		/// <param name="Name">Name of the type.</param>
+		/// <param name="startingType">Starting point for searches</param>
+		/// <param name="options">Options controlling what is searched</param>
+		/// <param name="name">Name of the type.</param>
 		/// <returns>The located type of null if not found</returns>
-		public UhtType? FindTypeInternal(UhtType? StartingType, UhtFindOptions Options, string Name)
+		public UhtType? FindTypeInternal(UhtType? startingType, UhtFindOptions options, string name)
 		{
-			if (Options.HasAnyFlags(UhtFindOptions.EngineName))
+			if (options.HasAnyFlags(UhtFindOptions.EngineName))
 			{
-				if (Options.HasAnyFlags(UhtFindOptions.CaseCompare))
+				if (options.HasAnyFlags(UhtFindOptions.CaseCompare))
 				{
-					return this.EngineNameSymbolTable.FindCasedType(StartingType, Options, Name);
+					return this._engineNameSymbolTable.FindCasedType(startingType, options, name);
 				}
 				else
 				{
-					return this.EngineNameSymbolTable.FindCaselessType(StartingType, Options, Name);
+					return this._engineNameSymbolTable.FindCaselessType(startingType, options, name);
 				}
 			}
-			else if (Options.HasAnyFlags(UhtFindOptions.SourceName))
+			else if (options.HasAnyFlags(UhtFindOptions.SourceName))
 			{
-				if (Options.HasAnyFlags(UhtFindOptions.CaselessCompare))
+				if (options.HasAnyFlags(UhtFindOptions.CaselessCompare))
 				{
-					return this.SourceNameSymbolTable.FindCaselessType(StartingType, Options, Name);
+					return this._sourceNameSymbolTable.FindCaselessType(startingType, options, name);
 				}
 				else
 				{
-					return this.SourceNameSymbolTable.FindCasedType(StartingType, Options, Name);
+					return this._sourceNameSymbolTable.FindCasedType(startingType, options, name);
 				}
 			}
 			else
@@ -1387,26 +1378,26 @@ namespace EpicGames.UHT.Utils
 		/// <summary>
 		/// Verify that the options are valid.  Will also check to make sure the symbol table has been populated.
 		/// </summary>
-		/// <param name="Options">Find options</param>
-		private void ValidateFindOptions(UhtFindOptions Options)
+		/// <param name="options">Find options</param>
+		private void ValidateFindOptions(UhtFindOptions options)
 		{
-			if (!Options.HasAnyFlags(UhtFindOptions.EngineName | UhtFindOptions.SourceName))
+			if (!options.HasAnyFlags(UhtFindOptions.EngineName | UhtFindOptions.SourceName))
 			{
 				throw new UhtIceException("Either EngineName or SourceName must be specified in the options");
 			}
 
-			if (Options.HasAnyFlags(UhtFindOptions.CaseCompare) && Options.HasAnyFlags(UhtFindOptions.CaselessCompare))
+			if (options.HasAnyFlags(UhtFindOptions.CaseCompare) && options.HasAnyFlags(UhtFindOptions.CaselessCompare))
 			{
 				throw new UhtIceException("Both CaseCompare and CaselessCompare can't be specified as FindType options");
 			}
 
-			UhtFindOptions TypeOptions = Options & UhtFindOptions.TypesMask;
-			if (TypeOptions == 0)
+			UhtFindOptions typeOptions = options & UhtFindOptions.TypesMask;
+			if (typeOptions == 0)
 			{
 				throw new UhtIceException("No type options specified");
 			}
 
-			if (!this.bSymbolTablePopulated)
+			if (!this._symbolTablePopulated)
 			{
 				throw new UhtIceException("Symbol table has not been populated, don't call FindType until headers are parsed.");
 			}
@@ -1415,52 +1406,51 @@ namespace EpicGames.UHT.Utils
 		/// <summary>
 		/// Generate an error message for when a given symbol wasn't found.  The text will contain the list of types that the symbol must be
 		/// </summary>
-		/// <param name="MessageSite">Destination for the message</param>
-		/// <param name="LineNumber">Line number generating the error</param>
-		/// <param name="Options">Collection of required types</param>
-		/// <param name="Name">The name of the symbol</param>
-		private static void FindTypeError(IUhtMessageSite MessageSite, int LineNumber, UhtFindOptions Options, string Name)
+		/// <param name="messageSite">Destination for the message</param>
+		/// <param name="lineNumber">Line number generating the error</param>
+		/// <param name="options">Collection of required types</param>
+		/// <param name="name">The name of the symbol</param>
+		private static void FindTypeError(IUhtMessageSite messageSite, int lineNumber, UhtFindOptions options, string name)
 		{
-			List<string> Types = new List<string>();
-			if (Options.HasAnyFlags(UhtFindOptions.Enum))
+			List<string> types = new List<string>();
+			if (options.HasAnyFlags(UhtFindOptions.Enum))
 			{
-				Types.Add("'enum'");
+				types.Add("'enum'");
 			}
-			if (Options.HasAnyFlags(UhtFindOptions.ScriptStruct))
+			if (options.HasAnyFlags(UhtFindOptions.ScriptStruct))
 			{
-				Types.Add("'struct'");
+				types.Add("'struct'");
 			}
-			if (Options.HasAnyFlags(UhtFindOptions.Class))
+			if (options.HasAnyFlags(UhtFindOptions.Class))
 			{
-				Types.Add("'class'");
+				types.Add("'class'");
 			}
-			if (Options.HasAnyFlags(UhtFindOptions.DelegateFunction))
+			if (options.HasAnyFlags(UhtFindOptions.DelegateFunction))
 			{
-				Types.Add("'delegate'");
+				types.Add("'delegate'");
 			}
-			if (Options.HasAnyFlags(UhtFindOptions.Function))
+			if (options.HasAnyFlags(UhtFindOptions.Function))
 			{
-				Types.Add("'function'");
+				types.Add("'function'");
 			}
-			if (Options.HasAnyFlags(UhtFindOptions.Property))
+			if (options.HasAnyFlags(UhtFindOptions.Property))
 			{
-				Types.Add("'property'");
+				types.Add("'property'");
 			}
 
-			MessageSite.LogError(LineNumber, $"Unable to find {UhtUtilities.MergeTypeNames(Types, "or")} with name '{Name}'");
+			messageSite.LogError(lineNumber, $"Unable to find {UhtUtilities.MergeTypeNames(types, "or")} with name '{name}'");
 		}
 
 		/// <summary>
 		/// Search for the given header file by just the file name
 		/// </summary>
-		/// <param name="Name">Name to be found</param>
+		/// <param name="name">Name to be found</param>
 		/// <returns></returns>
-		public UhtHeaderFile? FindHeaderFile(string Name)
+		public UhtHeaderFile? FindHeaderFile(string name)
 		{
-			UhtHeaderFile? HeaderFile;
-			if (this.HeaderFileDictionaryInternal.TryGetValue(Name, out HeaderFile))
+			if (this._headerFileDictionaryInternal.TryGetValue(name, out UhtHeaderFile? headerFile))
 			{
-				return HeaderFile;
+				return headerFile;
 			}
 			return null;
 		}
@@ -1469,30 +1459,30 @@ namespace EpicGames.UHT.Utils
 		/// <summary>
 		/// Add a message to the collection of output messages
 		/// </summary>
-		/// <param name="Message">Message being added</param>
-		public void AddMessage(UhtMessage Message)
+		/// <param name="message">Message being added</param>
+		public void AddMessage(UhtMessage message)
 		{
-			lock (this.Messages)
+			lock (this._messages)
 			{
-				this.Messages.Add(Message);
-				
+				this._messages.Add(message);
+
 				// If we aren't caching messages and this is the first message,
 				// start a task to flush the messages.
-				if (!this.bCacheMessages && this.Messages.Count == 1)
+				if (!this.CacheMessages && this._messages.Count == 1)
 				{
-					this.MessageTask = Task.Factory.StartNew(() => FlushMessages(), CancellationToken.None, TaskCreationOptions.None, TaskScheduler.Default);
+					this._messageTask = Task.Factory.StartNew(() => FlushMessages(), CancellationToken.None, TaskCreationOptions.None, TaskScheduler.Default);
 				}
 			}
 
-			switch (Message.MessageType)
+			switch (message.MessageType)
 			{
 				case UhtMessageType.Error:
 				case UhtMessageType.Ice:
-					Interlocked.Increment(ref this.ErrorCountInternal);
+					Interlocked.Increment(ref this._errorCountInternal);
 					break;
 
 				case UhtMessageType.Warning:
-					Interlocked.Increment(ref this.WarningCountInternal);
+					Interlocked.Increment(ref this._warningCountInternal);
 					break;
 
 				case UhtMessageType.Info:
@@ -1510,14 +1500,14 @@ namespace EpicGames.UHT.Utils
 		/// </summary>
 		public void LogMessages()
 		{
-			if (this.MessageTask != null)
+			if (this._messageTask != null)
 			{
-				this.MessageTask.Wait();
+				this._messageTask.Wait();
 			}
 
-			foreach (UhtMessage Message in FetchOrderedMessages())
+			foreach (UhtMessage message in FetchOrderedMessages())
 			{
-				LogMessage(Message);
+				LogMessage(message);
 			}
 		}
 
@@ -1526,49 +1516,49 @@ namespace EpicGames.UHT.Utils
 		/// </summary>
 		private void FlushMessages()
 		{
-			UhtMessage[]? MessageArray = null;
-			lock (this.Messages)
+			UhtMessage[]? messageArray = null;
+			lock (this._messages)
 			{
-				MessageArray = this.Messages.ToArray();
-				this.Messages.Clear();
+				messageArray = this._messages.ToArray();
+				this._messages.Clear();
 			}
 
-			foreach (UhtMessage Message in MessageArray)
+			foreach (UhtMessage message in messageArray)
 			{
-				LogMessage(Message);
+				LogMessage(message);
 			}
 		}
 
 		/// <summary>
 		/// Log the given message
 		/// </summary>
-		/// <param name="Message">The message to be logged</param>
-		private void LogMessage(UhtMessage Message)
+		/// <param name="message">The message to be logged</param>
+		private void LogMessage(UhtMessage message)
 		{
-			string FormattedMessage = FormatMessage(Message);
-			LogLevel LogLevel;
-			switch (Message.MessageType)
+			string formattedMessage = FormatMessage(message);
+			LogLevel logLevel;
+			switch (message.MessageType)
 			{
 				default:
 				case UhtMessageType.Error:
 				case UhtMessageType.Ice:
-					LogLevel = LogLevel.Error;
+					logLevel = LogLevel.Error;
 					break;
 
 				case UhtMessageType.Warning:
-					LogLevel = LogLevel.Warning;
+					logLevel = LogLevel.Warning;
 					break;
 
 				case UhtMessageType.Info:
-					LogLevel = LogLevel.Information;
+					logLevel = LogLevel.Information;
 					break;
 
 				case UhtMessageType.Trace:
-					LogLevel = LogLevel.Trace;
+					logLevel = LogLevel.Trace;
 					break;
 			}
 
-			Log.Logger.Log(LogLevel, "{FormattedMessage}", FormattedMessage);
+			Log.Logger.Log(logLevel, "{FormattedMessage}", formattedMessage);
 		}
 
 		/// <summary>
@@ -1577,12 +1567,12 @@ namespace EpicGames.UHT.Utils
 		/// <returns>List of all the messages</returns>
 		public List<string> CollectMessages()
 		{
-			List<string> Out = new List<string>();
-			foreach (UhtMessage Message in FetchOrderedMessages())
+			List<string> messages = new List<string>();
+			foreach (UhtMessage message in FetchOrderedMessages())
 			{
-				Out.Add(FormatMessage(Message));
+				messages.Add(FormatMessage(message));
 			}
-			return Out;
+			return messages;
 		}
 
 		/// <summary>
@@ -1590,14 +1580,14 @@ namespace EpicGames.UHT.Utils
 		/// This is used by the property resolution system to replace properties created during
 		/// the parsing phase that couldn't be resoled until after all headers are parsed.
 		/// </summary>
-		/// <param name="OldType"></param>
-		/// <param name="NewType"></param>
-		public void ReplaceTypeInSymbolTable(UhtType OldType, UhtType NewType)
+		/// <param name="oldType"></param>
+		/// <param name="newType"></param>
+		public void ReplaceTypeInSymbolTable(UhtType oldType, UhtType newType)
 		{
-			this.SourceNameSymbolTable.Replace(OldType, NewType, OldType.SourceName);
-			if (OldType.EngineType.HasEngineName())
+			this._sourceNameSymbolTable.Replace(oldType, newType, oldType.SourceName);
+			if (oldType.EngineType.HasEngineName())
 			{
-				this.EngineNameSymbolTable.Replace(OldType, NewType, OldType.EngineName);
+				this._engineNameSymbolTable.Replace(oldType, newType, oldType.EngineName);
 			}
 		}
 
@@ -1607,102 +1597,102 @@ namespace EpicGames.UHT.Utils
 		/// <returns>Enumerator</returns>
 		private IOrderedEnumerable<UhtMessage> FetchOrderedMessages()
 		{
-			List<UhtMessage> Messages = new List<UhtMessage>();
-			lock (this.Messages)
+			List<UhtMessage> messages = new List<UhtMessage>();
+			lock (this._messages)
 			{
-				Messages.AddRange(this.Messages);
-				this.Messages.Clear();
+				messages.AddRange(this._messages);
+				this._messages.Clear();
 			}
-			return Messages.OrderBy(Context => Context.FilePath).ThenBy(Context => Context.LineNumber + Context.MessageSource?.MessageFragmentLineNumber);
+			return messages.OrderBy(context => context.FilePath).ThenBy(context => context.LineNumber + context.MessageSource?.MessageFragmentLineNumber);
 		}
 
 		/// <summary>
 		/// Format the given message
 		/// </summary>
-		/// <param name="Message">Message to be formatted</param>
+		/// <param name="message">Message to be formatted</param>
 		/// <returns>Text of the formatted message</returns>
-		private string FormatMessage(UhtMessage Message)
+		private string FormatMessage(UhtMessage message)
 		{
-			string FilePath;
-			string FragmentPath = "";
-			int LineNumber = Message.LineNumber;
-			if (Message.FilePath != null)
+			string filePath;
+			string fragmentPath = "";
+			int lineNumber = message.LineNumber;
+			if (message.FilePath != null)
 			{
-				FilePath = Message.FilePath;
+				filePath = message.FilePath;
 			}
-			else if (Message.MessageSource != null)
+			else if (message.MessageSource != null)
 			{
-				if (Message.MessageSource.bMessageIsFragment)
+				if (message.MessageSource.MessageIsFragment)
 				{
-					if (this.bRelativePathInLog)
+					if (this.RelativePathInLog)
 					{
-						FilePath = Message.MessageSource.MessageFragmentFilePath;
+						filePath = message.MessageSource.MessageFragmentFilePath;
 					}
 					else
 					{
-						FilePath = Message.MessageSource.MessageFragmentFullFilePath;
+						filePath = message.MessageSource.MessageFragmentFullFilePath;
 					}
-					FragmentPath = $"[{Message.MessageSource.MessageFilePath}]";
-					LineNumber += Message.MessageSource.MessageFragmentLineNumber;
+					fragmentPath = $"[{message.MessageSource.MessageFilePath}]";
+					lineNumber += message.MessageSource.MessageFragmentLineNumber;
 				}
 				else
 				{
-					if (this.bRelativePathInLog)
+					if (this.RelativePathInLog)
 					{
-						FilePath = Message.MessageSource.MessageFilePath;
+						filePath = message.MessageSource.MessageFilePath;
 					}
 					else
 					{
-						FilePath = Message.MessageSource.MessageFullFilePath;
+						filePath = message.MessageSource.MessageFullFilePath;
 					}
 				}
 			}
 			else
 			{
-				FilePath = "UnknownSource";
+				filePath = "UnknownSource";
 			}
 
-			switch (Message.MessageType)
+			switch (message.MessageType)
 			{
 				case UhtMessageType.Error:
-					return $"{FilePath}({LineNumber}){FragmentPath}: Error: {Message.Message}";
+					return $"{filePath}({lineNumber}){fragmentPath}: Error: {message.Message}";
 				case UhtMessageType.Warning:
-					return $"{FilePath}({LineNumber}){FragmentPath}: Warning: {Message.Message}";
+					return $"{filePath}({lineNumber}){fragmentPath}: Warning: {message.Message}";
 				case UhtMessageType.Info:
-					return $"{FilePath}({LineNumber}){FragmentPath}: Info: {Message.Message}";
+					return $"{filePath}({lineNumber}){fragmentPath}: Info: {message.Message}";
 				case UhtMessageType.Trace:
-					return $"{FilePath}({LineNumber}){FragmentPath}: Trace: {Message.Message}";
+					return $"{filePath}({lineNumber}){fragmentPath}: Trace: {message.Message}";
 				default:
 				case UhtMessageType.Ice:
-					return $"{FilePath}({LineNumber}){FragmentPath}:  Error: Internal Compiler Error - {Message.Message}";
+					return $"{filePath}({lineNumber}){fragmentPath}:  Error: Internal Compiler Error - {message.Message}";
 			}
 		}
 
 		/// <summary>
 		/// Handle the given exception with the provided message context
 		/// </summary>
-		/// <param name="MessageSource">Context for the exception.  Required to handled all exceptions other than UHTException</param>
-		/// <param name="E">Exception being handled</param>
-		private void HandleException(IUhtMessageSource? MessageSource, Exception E)
+		/// <param name="messageSource">Context for the exception.  Required to handled all exceptions other than UHTException</param>
+		/// <param name="e">Exception being handled</param>
+		private void HandleException(IUhtMessageSource? messageSource, Exception e)
 		{
-			switch (E)
+			switch (e)
 			{
-				case UhtException UHTException:
-					UhtMessage Message = UHTException.UhtMessage;
-					if (Message.MessageSource == null)
+				case UhtException uhtException:
+					UhtMessage message = uhtException.UhtMessage;
+					if (message.MessageSource == null)
 					{
-						Message.MessageSource = MessageSource;
+						message.MessageSource = messageSource;
 					}
-					AddMessage(Message);
+					AddMessage(message);
 					break;
 
-				case JsonException JsonException:
-					AddMessage(UhtMessage.MakeMessage(UhtMessageType.Error, MessageSource, null, (int)(JsonException.LineNumber + 1 ?? 1), JsonException.Message));
+				case JsonException jsonException:
+					AddMessage(UhtMessage.MakeMessage(UhtMessageType.Error, messageSource, null, (int)(jsonException.LineNumber + 1 ?? 1), jsonException.Message));
 					break;
 
 				default:
 					//Log.TraceInformation("{0}", E.StackTrace);
-					AddMessage(UhtMessage.MakeMessage(UhtMessageType.Ice, MessageSource, null, 1, $"{E.GetType().ToString()} - {E.Message}"));
+					AddMessage(UhtMessage.MakeMessage(UhtMessageType.Ice, messageSource, null, 1, $"{e.GetType()} - {e.Message}"));
 					break;
 			}
 		}
@@ -1714,28 +1704,28 @@ namespace EpicGames.UHT.Utils
 		/// In general, it is assumed that during normal UHT, all paths are already full paths.
 		/// Only the test harness deals in relative paths.
 		/// </summary>
-		/// <param name="FilePath">Path to normalize</param>
+		/// <param name="filePath">Path to normalize</param>
 		/// <returns>Normalized path possibly converted to a full path.</returns>
-		private string GetNormalizedFullFilePath(string FilePath)
+		private string GetNormalizedFullFilePath(string filePath)
 		{
-			return this.FileManager!.GetFullFilePath(FilePath).Replace('\\', '/');
+			return this.FileManager!.GetFullFilePath(filePath).Replace('\\', '/');
 		}
 
-		private UhtHeaderFileParser? ParseHeaderFile(UhtHeaderFile HeaderFile)
+		private UhtHeaderFileParser? ParseHeaderFile(UhtHeaderFile headerFile)
 		{
-			UhtHeaderFileParser? Parser = null;
-			TryAlways(HeaderFile.MessageSource, () =>
+			UhtHeaderFileParser? parser = null;
+			TryAlways(headerFile.MessageSource, () =>
 			{
-				HeaderFile.Read();
-				Parser = UhtHeaderFileParser.Parse(HeaderFile);
+				headerFile.Read();
+				parser = UhtHeaderFileParser.Parse(headerFile);
 			});
-			return Parser;
+			return parser;
 		}
 
 		#region Run steps
-		private void StepReadManifestFile(string ManifestFilePath)
+		private void StepReadManifestFile(string manifestFilePath)
 		{
-			this.ManifestFile = new UhtManifestFile(this, ManifestFilePath);
+			this.ManifestFile = new UhtManifestFile(this, manifestFilePath);
 
 			Try(this.ManifestFile.MessageSource, () =>
 			{
@@ -1761,104 +1751,104 @@ namespace EpicGames.UHT.Utils
 			{
 				Log.Logger.LogTrace("Step - Prepare Modules");
 
-				foreach (UHTManifest.Module Module in this.ManifestFile.Manifest.Modules)
+				foreach (UHTManifest.Module module in this.ManifestFile.Manifest.Modules)
 				{
-					EPackageFlags PackageFlags = EPackageFlags.ContainsScript | EPackageFlags.Compiling;
+					EPackageFlags packageFlags = EPackageFlags.ContainsScript | EPackageFlags.Compiling;
 
-					switch (Module.OverrideModuleType)
+					switch (module.OverrideModuleType)
 					{
 						case EPackageOverrideType.None:
-							switch (Module.ModuleType)
+							switch (module.ModuleType)
 							{
 								case UHTModuleType.GameEditor:
 								case UHTModuleType.EngineEditor:
-									PackageFlags |= EPackageFlags.EditorOnly;
+									packageFlags |= EPackageFlags.EditorOnly;
 									break;
 
 								case UHTModuleType.GameDeveloper:
 								case UHTModuleType.EngineDeveloper:
-									PackageFlags |= EPackageFlags.Developer;
+									packageFlags |= EPackageFlags.Developer;
 									break;
 
 								case UHTModuleType.GameUncooked:
 								case UHTModuleType.EngineUncooked:
-									PackageFlags |= EPackageFlags.UncookedOnly;
+									packageFlags |= EPackageFlags.UncookedOnly;
 									break;
 							}
 							break;
 
 						case EPackageOverrideType.EditorOnly:
-							PackageFlags |= EPackageFlags.EditorOnly;
+							packageFlags |= EPackageFlags.EditorOnly;
 							break;
 
 						case EPackageOverrideType.EngineDeveloper:
 						case EPackageOverrideType.GameDeveloper:
-							PackageFlags |= EPackageFlags.Developer;
+							packageFlags |= EPackageFlags.Developer;
 							break;
 
 						case EPackageOverrideType.EngineUncookedOnly:
 						case EPackageOverrideType.GameUncookedOnly:
-							PackageFlags |= EPackageFlags.UncookedOnly;
+							packageFlags |= EPackageFlags.UncookedOnly;
 							break;
 					}
 
-					UhtPackage Package = new UhtPackage(this, Module, PackageFlags);
-					this.PackagesInternal.Add(Package);
+					UhtPackage package = new UhtPackage(this, module, packageFlags);
+					this._packagesInternal.Add(package);
 				}
 			});
 		}
 
-		private void StepPrepareHeaders(UhtPackage Package, IEnumerable<string> HeaderFiles, UhtHeaderFileType HeaderFileType)
+		private void StepPrepareHeaders(UhtPackage package, IEnumerable<string> headerFiles, UhtHeaderFileType headerFileType)
 		{
-			if (Package.Module == null)
+			if (package.Module == null)
 			{
 				return;
 			}
 
-			string TypeDirectory = HeaderFileType.ToString() + '/';
-			string NormalizedModuleBaseFullFilePath = GetNormalizedFullFilePath(Package.Module.BaseDirectory);
-			foreach (string HeaderFilePath in HeaderFiles)
+			string typeDirectory = headerFileType.ToString() + '/';
+			string normalizedModuleBaseFullFilePath = GetNormalizedFullFilePath(package.Module.BaseDirectory);
+			foreach (string headerFilePath in headerFiles)
 			{
 
 				// Make sure this isn't a duplicate
-				string NormalizedFullFilePath = GetNormalizedFullFilePath(HeaderFilePath);
-				string FileName = Path.GetFileName(NormalizedFullFilePath);
-				UhtHeaderFile? ExistingHeaderFile;
-				if (HeaderFileDictionaryInternal.TryGetValue(FileName, out ExistingHeaderFile) && ExistingHeaderFile != null)
+				string normalizedFullFilePath = GetNormalizedFullFilePath(headerFilePath);
+				string fileName = Path.GetFileName(normalizedFullFilePath);
+				UhtHeaderFile? existingHeaderFile;
+				if (_headerFileDictionaryInternal.TryGetValue(fileName, out existingHeaderFile) && existingHeaderFile != null)
 				{
-					string NormalizedExistingFullFilePath = GetNormalizedFullFilePath(ExistingHeaderFile.FilePath);
-					if (!string.Equals(NormalizedFullFilePath, NormalizedExistingFullFilePath, StringComparison.OrdinalIgnoreCase))
+					string normalizedExistingFullFilePath = GetNormalizedFullFilePath(existingHeaderFile.FilePath);
+					if (!String.Equals(normalizedFullFilePath, normalizedExistingFullFilePath, StringComparison.OrdinalIgnoreCase))
 					{
-						IUhtMessageSite Site = (IUhtMessageSite?)this.ManifestFile ?? this;
-						Site.LogError($"Two headers with the same name is not allowed. '{HeaderFilePath}' conflicts with '{ExistingHeaderFile.FilePath}'");
+						IUhtMessageSite site = (IUhtMessageSite?)this.ManifestFile ?? this;
+						site.LogError($"Two headers with the same name is not allowed. '{headerFilePath}' conflicts with '{existingHeaderFile.FilePath}'");
 						continue;
 					}
 				}
 
 				// Create the header file and add to the collections
-				UhtHeaderFile HeaderFile = new UhtHeaderFile(Package, HeaderFilePath);
-				HeaderFile.HeaderFileType = HeaderFileType;
-				HeaderFilesInternal.Add(HeaderFile);
-				HeaderFileDictionaryInternal.Add(FileName, HeaderFile);
-				Package.AddChild(HeaderFile);
+				UhtHeaderFile headerFile = new UhtHeaderFile(package, headerFilePath);
+				headerFile.HeaderFileType = headerFileType;
+				_headerFilesInternal.Add(headerFile);
+				_headerFileDictionaryInternal.Add(fileName, headerFile);
+				package.AddChild(headerFile);
 
 				// Save metadata for the class path, both for it's include path and relative to the module base directory
-				if (NormalizedFullFilePath.StartsWith(NormalizedModuleBaseFullFilePath, true, null))
+				if (normalizedFullFilePath.StartsWith(normalizedModuleBaseFullFilePath, true, null))
 				{
-					int StripLength = NormalizedModuleBaseFullFilePath.Length;
-					if (StripLength < NormalizedFullFilePath.Length && NormalizedFullFilePath[StripLength] == '/')
+					int stripLength = normalizedModuleBaseFullFilePath.Length;
+					if (stripLength < normalizedFullFilePath.Length && normalizedFullFilePath[stripLength] == '/')
 					{
-						++StripLength;
+						++stripLength;
 					}
 
-					HeaderFile.ModuleRelativeFilePath = NormalizedFullFilePath.Substring(StripLength);
+					headerFile.ModuleRelativeFilePath = normalizedFullFilePath.Substring(stripLength);
 
-					if (NormalizedFullFilePath.Substring(StripLength).StartsWith(TypeDirectory, true, null))
+					if (normalizedFullFilePath.Substring(stripLength).StartsWith(typeDirectory, true, null))
 					{
-						StripLength += TypeDirectory.Length;
+						stripLength += typeDirectory.Length;
 					}
 
-					HeaderFile.IncludeFilePath = NormalizedFullFilePath.Substring(StripLength);
+					headerFile.IncludeFilePath = normalizedFullFilePath.Substring(stripLength);
 				}
 			}
 		}
@@ -1874,27 +1864,27 @@ namespace EpicGames.UHT.Utils
 			{
 				Log.Logger.LogTrace("Step - Prepare Headers");
 
-				foreach (UhtPackage Package in this.PackagesInternal)
+				foreach (UhtPackage package in this._packagesInternal)
 				{
-					if (Package.Module != null)
+					if (package.Module != null)
 					{
-						StepPrepareHeaders(Package, Package.Module.ClassesHeaders, UhtHeaderFileType.Classes);
-						StepPrepareHeaders(Package, Package.Module.PublicHeaders, UhtHeaderFileType.Public);
-						StepPrepareHeaders(Package, Package.Module.InternalHeaders, UhtHeaderFileType.Internal);
-						StepPrepareHeaders(Package, Package.Module.PrivateHeaders, UhtHeaderFileType.Private);
+						StepPrepareHeaders(package, package.Module.ClassesHeaders, UhtHeaderFileType.Classes);
+						StepPrepareHeaders(package, package.Module.PublicHeaders, UhtHeaderFileType.Public);
+						StepPrepareHeaders(package, package.Module.InternalHeaders, UhtHeaderFileType.Internal);
+						StepPrepareHeaders(package, package.Module.PrivateHeaders, UhtHeaderFileType.Private);
 					}
 				}
 
 				// Locate the NoExportTypes.h file and add it to every other header file
-				if (this.HeaderFileDictionaryInternal.TryGetValue("NoExportTypes.h", out UhtHeaderFile? NoExportTypes))
+				if (this._headerFileDictionaryInternal.TryGetValue("NoExportTypes.h", out UhtHeaderFile? noExportTypes))
 				{
-					foreach (UhtPackage Package in this.PackagesInternal)
+					foreach (UhtPackage package in this._packagesInternal)
 					{
-						foreach (UhtHeaderFile HeaderFile in Package.Children)
+						foreach (UhtHeaderFile headerFile in package.Children)
 						{
-							if (HeaderFile != NoExportTypes)
+							if (headerFile != noExportTypes)
 							{
-								HeaderFile.AddReferencedHeader(NoExportTypes);
+								headerFile.AddReferencedHeader(noExportTypes);
 							}
 						}
 					}
@@ -1904,25 +1894,25 @@ namespace EpicGames.UHT.Utils
 
 		private void StepParseHeaders()
 		{
-			if (this.bHasErrors)
+			if (this.HasErrors)
 			{
 				return;
 			}
 
 			Log.Logger.LogTrace("Step - Parse Headers");
 
-			if (this.bGoWide)
+			if (this.GoWide)
 			{
-				Parallel.ForEach(this.HeaderFilesInternal, HeaderFile =>
+				Parallel.ForEach(this._headerFilesInternal, headerFile =>
 				{
-					ParseHeaderFile(HeaderFile);
+					ParseHeaderFile(headerFile);
 				});
 			}
 			else
 			{
-				foreach (UhtHeaderFile HeaderFile in this.HeaderFilesInternal)
+				foreach (UhtHeaderFile headerFile in this._headerFilesInternal)
 				{
-					ParseHeaderFile(HeaderFile);
+					ParseHeaderFile(headerFile);
 				}
 			}
 		}
@@ -1933,15 +1923,15 @@ namespace EpicGames.UHT.Utils
 			{
 				Log.Logger.LogTrace("Step - Populate symbol table");
 
-				this.SourceNameSymbolTable = new UhtSymbolTable(this.TypeCount);
-				this.EngineNameSymbolTable = new UhtSymbolTable(this.TypeCount);
+				this._sourceNameSymbolTable = new UhtSymbolTable(this.TypeCount);
+				this._engineNameSymbolTable = new UhtSymbolTable(this.TypeCount);
 
 				PopulateSymbolTable();
 
-				this.UObjectInternal = (UhtClass?)FindType(null, UhtFindOptions.SourceName | UhtFindOptions.Class, "UObject");
-				this.UClassInternal = (UhtClass?)FindType(null, UhtFindOptions.SourceName | UhtFindOptions.Class, "UClass");
-				this.UInterfaceInternal = (UhtClass?)FindType(null, UhtFindOptions.SourceName | UhtFindOptions.Class, "UInterface");
-				this.IInterfaceInternal = (UhtClass?)FindType(null, UhtFindOptions.SourceName | UhtFindOptions.Class, "IInterface");
+				this._uobject = (UhtClass?)FindType(null, UhtFindOptions.SourceName | UhtFindOptions.Class, "UObject");
+				this._uclass = (UhtClass?)FindType(null, UhtFindOptions.SourceName | UhtFindOptions.Class, "UClass");
+				this._uinterface = (UhtClass?)FindType(null, UhtFindOptions.SourceName | UhtFindOptions.Class, "UInterface");
+				this._iinterface = (UhtClass?)FindType(null, UhtFindOptions.SourceName | UhtFindOptions.Class, "IInterface");
 				this.AActor = (UhtClass?)FindType(null, UhtFindOptions.SourceName | UhtFindOptions.Class, "AActor");
 				this.INotifyFieldValueChanged = (UhtClass?)FindType(null, UhtFindOptions.SourceName | UhtFindOptions.Class, "INotifyFieldValueChanged");
 			});
@@ -1949,73 +1939,73 @@ namespace EpicGames.UHT.Utils
 
 		private void StepResolveBases()
 		{
-			StepForAllHeaders(HeaderFile => Resolve(HeaderFile, UhtResolvePhase.Bases));
+			StepForAllHeaders(headerFile => Resolve(headerFile, UhtResolvePhase.Bases));
 		}
 
 		private void StepResolveInvalidCheck()
 		{
-			StepForAllHeaders(HeaderFile => Resolve(HeaderFile, UhtResolvePhase.InvalidCheck));
+			StepForAllHeaders(headerFile => Resolve(headerFile, UhtResolvePhase.InvalidCheck));
 		}
 
 		private void StepResolveProperties()
 		{
-			StepForAllHeaders(HeaderFile => Resolve(HeaderFile, UhtResolvePhase.Properties));
+			StepForAllHeaders(headerFile => Resolve(headerFile, UhtResolvePhase.Properties));
 		}
 
 		private void StepResolveFinal()
 		{
-			StepForAllHeaders(HeaderFile => Resolve(HeaderFile, UhtResolvePhase.Final));
+			StepForAllHeaders(headerFile => Resolve(headerFile, UhtResolvePhase.Final));
 		}
 
 		private void StepResolveValidate()
 		{
-			StepForAllHeaders(HeaderFile => UhtType.ValidateType(HeaderFile, UhtValidationOptions.None));
+			StepForAllHeaders(headerFile => UhtType.ValidateType(headerFile, UhtValidationOptions.None));
 		}
 
 		private void StepCollectReferences()
 		{
-			StepForAllHeaders(HeaderFile =>
+			StepForAllHeaders(headerFile =>
 			{
-				foreach (UhtType Child in HeaderFile.Children)
+				foreach (UhtType child in headerFile.Children)
 				{
-					Child.CollectReferences(HeaderFile.References);
+					child.CollectReferences(headerFile.References);
 				}
-				foreach (UhtHeaderFile RefHeaderFile in HeaderFile.References.ReferencedHeaders)
+				foreach (UhtHeaderFile refHeaderFile in headerFile.References.ReferencedHeaders)
 				{
-					HeaderFile.AddReferencedHeader(RefHeaderFile);
+					headerFile.AddReferencedHeader(refHeaderFile);
 				}
 			});
 		}
 
-		private void Resolve(UhtHeaderFile HeaderFile, UhtResolvePhase ResolvePhase)
+		private void Resolve(UhtHeaderFile headerFile, UhtResolvePhase resolvePhase)
 		{
-			TryAlways(HeaderFile.MessageSource, () =>
+			TryAlways(headerFile.MessageSource, () =>
 			{
-				HeaderFile.Resolve(ResolvePhase);
+				headerFile.Resolve(resolvePhase);
 			});
 		}
 
-		private delegate void StepDelegate(UhtHeaderFile HeaderFile);
+		private delegate void StepDelegate(UhtHeaderFile headerFile);
 
-		private void StepForAllHeaders(StepDelegate Delegate)
+		private void StepForAllHeaders(StepDelegate stepDelegate)
 		{
-			if (this.bHasErrors)
+			if (this.HasErrors)
 			{
 				return;
 			}
 
-			if (this.bGoWide)
+			if (this.GoWide)
 			{
-				Parallel.ForEach(this.HeaderFilesInternal, HeaderFile =>
+				Parallel.ForEach(this._headerFilesInternal, headerFile =>
 				{
-					Delegate(HeaderFile);
+					stepDelegate(headerFile);
 				});
 			}
 			else
 			{
-				foreach (UhtHeaderFile HeaderFile in this.HeaderFilesInternal)
+				foreach (UhtHeaderFile headerFile in this._headerFilesInternal)
 				{
-					Delegate(HeaderFile);
+					stepDelegate(headerFile);
 				}
 			}
 		}
@@ -2024,18 +2014,18 @@ namespace EpicGames.UHT.Utils
 		#region Symbol table initialization
 		private void PopulateSymbolTable()
 		{
-			foreach (UhtHeaderFile HeaderFile in this.HeaderFilesInternal)
+			foreach (UhtHeaderFile headerFile in this._headerFilesInternal)
 			{
-				AddTypeToSymbolTable(HeaderFile);
+				AddTypeToSymbolTable(headerFile);
 			}
-			this.bSymbolTablePopulated = true;
+			this._symbolTablePopulated = true;
 		}
 
-		private void AddTypeToSymbolTable(UhtType Type)
+		private void AddTypeToSymbolTable(UhtType type)
 		{
-			UhtEngineType EngineExtendedType = Type.EngineType;
+			UhtEngineType engineExtendedType = type.EngineType;
 
-			if (Type is UhtEnum Enum)
+			if (type is UhtEnum enumObj)
 			{
 				//COMPATIBILITY-TODO: We can get more reliable results by only adding regular enums to the table
 				// and then in the lookup code in the property system to look for the '::' and just lookup
@@ -2053,57 +2043,57 @@ namespace EpicGames.UHT.Utils
 				//		this.RegularEnumValueLookup.Add(Value.Name.ToString(), Enum);
 				//	}
 				//}
-				bool bAddShortNames = Enum.CppForm == UhtEnumCppForm.Namespaced || Enum.CppForm == UhtEnumCppForm.EnumClass;
-				string CheckName = $"{Enum.SourceName}::";
-				foreach (UhtEnumValue Value in Enum.EnumValues)
+				bool addShortNames = enumObj.CppForm == UhtEnumCppForm.Namespaced || enumObj.CppForm == UhtEnumCppForm.EnumClass;
+				string checkName = $"{enumObj.SourceName}::";
+				foreach (UhtEnumValue value in enumObj.EnumValues)
 				{
-					this.FullEnumValueLookup.Add(Value.Name, new EnumAndValue { Enum = Enum, Value = Value.Value });
-					if (bAddShortNames)
+					this._fullEnumValueLookup.Add(value.Name, new EnumAndValue { Enum = enumObj, Value = value.Value });
+					if (addShortNames)
 					{
-						if (Value.Name.StartsWith(CheckName, StringComparison.Ordinal))
+						if (value.Name.StartsWith(checkName, StringComparison.Ordinal))
 						{
-							this.ShortEnumValueLookup.TryAdd(Value.Name.Substring(CheckName.Length), Enum);
+							this._shortEnumValueLookup.TryAdd(value.Name.Substring(checkName.Length), enumObj);
 						}
 					}
 				}
 			}
 
-			if (EngineExtendedType.FindOptions() != 0)
+			if (engineExtendedType.FindOptions() != 0)
 			{
-				if (EngineExtendedType.MustNotBeReserved())
+				if (engineExtendedType.MustNotBeReserved())
 				{
-					if (ReservedNames.Contains(Type.EngineName))
+					if (s_reservedNames.Contains(type.EngineName))
 					{
-						Type.HeaderFile.LogError(Type.LineNumber, $"{EngineExtendedType.CapitalizedText()} '{Type.EngineName}' uses a reserved type name.");
+						type.HeaderFile.LogError(type.LineNumber, $"{engineExtendedType.CapitalizedText()} '{type.EngineName}' uses a reserved type name.");
 					}
 				}
 
-				if (EngineExtendedType.HasEngineName() && EngineExtendedType.MustBeUnique())
+				if (engineExtendedType.HasEngineName() && engineExtendedType.MustBeUnique())
 				{
-					UhtType? ExistingType = this.EngineNameSymbolTable.FindCaselessType(null, EngineExtendedType.MustBeUniqueFindOptions(), Type.EngineName);
-					if (ExistingType != null)
+					UhtType? existingType = this._engineNameSymbolTable.FindCaselessType(null, engineExtendedType.MustBeUniqueFindOptions(), type.EngineName);
+					if (existingType != null)
 					{
-						Type.HeaderFile.LogError(Type.LineNumber, 
-							$"{EngineExtendedType.CapitalizedText()} '{Type.SourceName}' shares engine name '{Type.EngineName}' with " +
-							$"{ExistingType.EngineType.LowercaseText()} '{ExistingType.SourceName}' in {ExistingType.HeaderFile.FilePath}({ExistingType.LineNumber})");
+						type.HeaderFile.LogError(type.LineNumber,
+							$"{engineExtendedType.CapitalizedText()} '{type.SourceName}' shares engine name '{type.EngineName}' with " +
+							$"{existingType.EngineType.LowercaseText()} '{existingType.SourceName}' in {existingType.HeaderFile.FilePath}({existingType.LineNumber})");
 					}
 				}
 
-				if (Type.bVisibleType)
+				if (type.VisibleType)
 				{
-					this.SourceNameSymbolTable.Add(Type, Type.SourceName);
-					if (EngineExtendedType.HasEngineName())
+					this._sourceNameSymbolTable.Add(type, type.SourceName);
+					if (engineExtendedType.HasEngineName())
 					{
-						this.EngineNameSymbolTable.Add(Type, Type.EngineName);
+						this._engineNameSymbolTable.Add(type, type.EngineName);
 					}
 				}
 			}
 
-			if (EngineExtendedType.AddChildrenToSymbolTable())
+			if (engineExtendedType.AddChildrenToSymbolTable())
 			{
-				foreach (UhtType Child in Type.Children)
+				foreach (UhtType child in type.Children)
 				{
-					AddTypeToSymbolTable(Child);
+					AddTypeToSymbolTable(child);
 				}
 			}
 		}
@@ -2117,45 +2107,45 @@ namespace EpicGames.UHT.Utils
 			Permanent,
 		}
 
-		private void TopologicalRecursion(List<TopologicalState> States, UhtHeaderFile First, UhtHeaderFile Visit)
+		private void TopologicalRecursion(List<TopologicalState> states, UhtHeaderFile first, UhtHeaderFile visit)
 		{
-			foreach (UhtHeaderFile Referenced in Visit.ReferencedHeadersNoLock)
+			foreach (UhtHeaderFile referenced in visit.ReferencedHeadersNoLock)
 			{
-				if (States[Referenced.HeaderFileTypeIndex] == TopologicalState.Temporary)
+				if (states[referenced.HeaderFileTypeIndex] == TopologicalState.Temporary)
 				{
-					First.LogError($"'{Visit.FilePath}' includes/requires '{Referenced.FilePath}'");
-					if (First != Referenced)
+					first.LogError($"'{visit.FilePath}' includes/requires '{referenced.FilePath}'");
+					if (first != referenced)
 					{
-						TopologicalRecursion(States, First, Referenced);
+						TopologicalRecursion(states, first, referenced);
 					}
 					break;
 				}
 			}
 		}
 
-		private UhtHeaderFile? TopologicalVisit(List<TopologicalState> States, UhtHeaderFile Visit)
+		private UhtHeaderFile? TopologicalVisit(List<TopologicalState> states, UhtHeaderFile visit)
 		{
-			switch (States[Visit.HeaderFileTypeIndex])
+			switch (states[visit.HeaderFileTypeIndex])
 			{
 				case TopologicalState.Unmarked:
-					States[Visit.HeaderFileTypeIndex] = TopologicalState.Temporary;
-					foreach (UhtHeaderFile Referenced in Visit.ReferencedHeadersNoLock)
+					states[visit.HeaderFileTypeIndex] = TopologicalState.Temporary;
+					foreach (UhtHeaderFile referenced in visit.ReferencedHeadersNoLock)
 					{
-						if (Visit != Referenced)
+						if (visit != referenced)
 						{
-							UhtHeaderFile? Out = TopologicalVisit(States, Referenced);
-							if (Out != null)
+							UhtHeaderFile? recursion = TopologicalVisit(states, referenced);
+							if (recursion != null)
 							{
-								return Out;
+								return recursion;
 							}
 						}
 					}
-					States[Visit.HeaderFileTypeIndex] = TopologicalState.Permanent;
-					this.SortedHeaderFilesInternal.Add(Visit);
+					states[visit.HeaderFileTypeIndex] = TopologicalState.Permanent;
+					this._sortedHeaderFilesInternal.Add(visit);
 					return null;
 
 				case TopologicalState.Temporary:
-					return Visit;
+					return visit;
 
 				case TopologicalState.Permanent:
 					return null;
@@ -2172,22 +2162,22 @@ namespace EpicGames.UHT.Utils
 				Log.Logger.LogTrace("Step - Topological Sort Header Files");
 
 				// Initialize a scratch table for topological states
-				this.SortedHeaderFilesInternal.Capacity = this.HeaderFileTypeCount;
-				List<TopologicalState> States = new List<TopologicalState>(this.HeaderFileTypeCount);
-				for (int Index = 0; Index < this.HeaderFileTypeCount; ++Index)
+				this._sortedHeaderFilesInternal.Capacity = this.HeaderFileTypeCount;
+				List<TopologicalState> states = new List<TopologicalState>(this.HeaderFileTypeCount);
+				for (int index = 0; index < this.HeaderFileTypeCount; ++index)
 				{
-					States.Add(TopologicalState.Unmarked);
+					states.Add(TopologicalState.Unmarked);
 				}
 
-				foreach (UhtHeaderFile HeaderFile in this.HeaderFiles)
+				foreach (UhtHeaderFile headerFile in this.HeaderFiles)
 				{
-					if (States[HeaderFile.HeaderFileTypeIndex] == TopologicalState.Unmarked)
+					if (states[headerFile.HeaderFileTypeIndex] == TopologicalState.Unmarked)
 					{
-						UhtHeaderFile? Recursion = TopologicalVisit(States, HeaderFile);
-						if (Recursion != null)
+						UhtHeaderFile? recursion = TopologicalVisit(states, headerFile);
+						if (recursion != null)
 						{
-							HeaderFile.LogError("Circular dependency detected:");
-							TopologicalRecursion(States, Recursion, Recursion);
+							headerFile.LogError("Circular dependency detected:");
+							TopologicalRecursion(states, recursion, recursion);
 							return;
 						}
 					}
@@ -2197,59 +2187,59 @@ namespace EpicGames.UHT.Utils
 		#endregion
 
 		#region Validation helpers
-		private readonly HashSet<UhtScriptStruct> ScriptStructsValidForNet = new HashSet<UhtScriptStruct>();
+		private readonly HashSet<UhtScriptStruct> _scriptStructsValidForNet = new HashSet<UhtScriptStruct>();
 
 		/// <summary>
 		/// Validate that the given referenced script structure is valid for network operations.  If the structure
 		/// is valid, then the result will be cached.  It not valid, errors will be generated each time the structure
 		/// is referenced.
 		/// </summary>
-		/// <param name="ReferencingProperty">The property referencing a structure</param>
-		/// <param name="ReferencedScriptStruct">The script structure being referenced</param>
+		/// <param name="referencingProperty">The property referencing a structure</param>
+		/// <param name="referencedScriptStruct">The script structure being referenced</param>
 		/// <returns></returns>
-		public bool ValidateScriptStructOkForNet(UhtProperty ReferencingProperty, UhtScriptStruct ReferencedScriptStruct)
+		public bool ValidateScriptStructOkForNet(UhtProperty referencingProperty, UhtScriptStruct referencedScriptStruct)
 		{
 
 			// Check for existing value
-			lock (this.ScriptStructsValidForNet)
+			lock (this._scriptStructsValidForNet)
 			{
-				if (this.ScriptStructsValidForNet.Contains(ReferencedScriptStruct))
+				if (this._scriptStructsValidForNet.Contains(referencedScriptStruct))
 				{
 					return true;
 				}
 			}
 
-			bool bIsStructValid = true;
+			bool isStructValid = true;
 
 			// Check the super chain structure
-			UhtScriptStruct? SuperScriptStruct = ReferencedScriptStruct.SuperScriptStruct;
-			if (SuperScriptStruct != null)
+			UhtScriptStruct? superScriptStruct = referencedScriptStruct.SuperScriptStruct;
+			if (superScriptStruct != null)
 			{
-				if (!ValidateScriptStructOkForNet(ReferencingProperty, SuperScriptStruct))
+				if (!ValidateScriptStructOkForNet(referencingProperty, superScriptStruct))
 				{
-					bIsStructValid = false;
+					isStructValid = false;
 				}
 			}
 
 			// Check the structure properties
-			foreach (UhtProperty Property in ReferencedScriptStruct.Properties)
+			foreach (UhtProperty property in referencedScriptStruct.Properties)
 			{
-				if (!Property.ValidateStructPropertyOkForNet(ReferencingProperty))
+				if (!property.ValidateStructPropertyOkForNet(referencingProperty))
 				{
-					bIsStructValid = false;
+					isStructValid = false;
 					break;
 				}
 			}
 
 			// Save the results
-			if (bIsStructValid)
+			if (isStructValid)
 			{
-				lock (this.ScriptStructsValidForNet)
+				lock (this._scriptStructsValidForNet)
 				{
-					this.ScriptStructsValidForNet.Add(ReferencedScriptStruct);
+					this._scriptStructsValidForNet.Add(referencedScriptStruct);
 				}
 			}
-			return bIsStructValid;
+			return isStructValid;
 		}
 		#endregion
 
@@ -2258,61 +2248,61 @@ namespace EpicGames.UHT.Utils
 		/// <summary>
 		/// Enable/Disable an exporter.  This overrides the default state of the exporter.
 		/// </summary>
-		/// <param name="Name">Name of the exporter</param>
-		/// <param name="Enabled">If true, the exporter is to be enabled</param>
-		public void SetExporterStatus(string Name, bool Enabled)
+		/// <param name="name">Name of the exporter</param>
+		/// <param name="enabled">If true, the exporter is to be enabled</param>
+		public void SetExporterStatus(string name, bool enabled)
 		{
-			this.ExporterStates[Name] = Enabled;
+			this._exporterStates[name] = enabled;
 		}
 
 		/// <summary>
 		/// Test to see if the given exporter plugin is enabled.
 		/// </summary>
-		/// <param name="PluginName">Name of the plugin</param>
-		/// <param name="IncludeTargetCheck">If true, include a target check</param>
+		/// <param name="pluginName">Name of the plugin</param>
+		/// <param name="includeTargetCheck">If true, include a target check</param>
 		/// <returns>True if enabled</returns>
-		public bool IsPluginEnabled(string PluginName, bool IncludeTargetCheck)
+		public bool IsPluginEnabled(string pluginName, bool includeTargetCheck)
 		{
-			if (this.ProjectJson == null && this.ProjectDirectory != null && this.ProjectFile != null)
+			if (this._projectJson == null && this.ProjectDirectory != null && this.ProjectFile != null)
 			{
-				UhtBuffer? Contents = this.ReadSourceToBuffer(this.ProjectFile);
-				if (Contents != null)
+				UhtBuffer? contents = this.ReadSourceToBuffer(this.ProjectFile);
+				if (contents != null)
 				{
-					this.ProjectJson = JsonDocument.Parse(Contents.Memory);
-					UhtBuffer.Return(Contents);
+					this._projectJson = JsonDocument.Parse(contents.Memory);
+					UhtBuffer.Return(contents);
 				}
 			}
 
-			if (this.ProjectJson == null)
+			if (this._projectJson == null)
 			{
 				return false;
 			}
 
-			JsonObject RootObject = new JsonObject(this.ProjectJson.RootElement);
-			if (RootObject.TryGetObjectArrayField("Plugins", out JsonObject[]? Plugins))
+			JsonObject rootObject = new JsonObject(this._projectJson.RootElement);
+			if (rootObject.TryGetObjectArrayField("Plugins", out JsonObject[]? plugins))
 			{
-				foreach (JsonObject Plugin in Plugins)
+				foreach (JsonObject plugin in plugins)
 				{
-					if (!Plugin.TryGetStringField("Name", out string? TestPluginName) || !string.Equals(PluginName, TestPluginName, StringComparison.OrdinalIgnoreCase))
+					if (!plugin.TryGetStringField("Name", out string? testPluginName) || !String.Equals(pluginName, testPluginName, StringComparison.OrdinalIgnoreCase))
 					{
 						continue;
 					}
-					if (!Plugin.TryGetBoolField("Enabled", out bool Enabled) || !Enabled)
+					if (!plugin.TryGetBoolField("Enabled", out bool enabled) || !enabled)
 					{
 						return false;
 					}
-					if (IncludeTargetCheck && this.Manifest != null)
+					if (includeTargetCheck && this.Manifest != null)
 					{
-						if (Plugin.TryGetStringArrayField("TargetAllowList", out string[]? AllowList))
+						if (plugin.TryGetStringArrayField("TargetAllowList", out string[]? allowList))
 						{
-							if (AllowList.Contains(this.Manifest.TargetName, StringComparer.OrdinalIgnoreCase))
+							if (allowList.Contains(this.Manifest.TargetName, StringComparer.OrdinalIgnoreCase))
 							{
 								return true;
 							}
 						}
-						if (Plugin.TryGetStringArrayField("TargetDenyList", out string[]? DenyList))
+						if (plugin.TryGetStringArrayField("TargetDenyList", out string[]? denyList))
 						{
-							if (DenyList.Contains(this.Manifest.TargetName, StringComparer.OrdinalIgnoreCase))
+							if (denyList.Contains(this.Manifest.TargetName, StringComparer.OrdinalIgnoreCase))
 							{
 								return false;
 							}
@@ -2326,76 +2316,76 @@ namespace EpicGames.UHT.Utils
 
 		private void StepExport()
 		{
-			HashSet<string> ExternalDependencies = new HashSet<string>();
-			long TotalWrittenFiles = 0;
+			HashSet<string> externalDependencies = new HashSet<string>();
+			long totalWrittenFiles = 0;
 			Try(null, () =>
 			{
 				Log.Logger.LogTrace("Step - Exports");
 
-				foreach (UhtExporter Exporter in this.ExporterTable)
+				foreach (UhtExporter exporter in this.ExporterTable)
 				{
-					bool Run = false;
-					if (!this.ExporterStates.TryGetValue(Exporter.Name, out Run))
+					bool run = false;
+					if (!this._exporterStates.TryGetValue(exporter.Name, out run))
 					{
-						Run = Config!.IsExporterEnabled(Exporter.Name) || 
-							(Exporter.Options.HasAnyFlags(UhtExporterOptions.Default) && !this.bNoDefaultExporters);
+						run = Config!.IsExporterEnabled(exporter.Name) ||
+							(exporter.Options.HasAnyFlags(UhtExporterOptions.Default) && !this.NoDefaultExporters);
 					}
 
-					UHTManifest.Module? PluginModule = null;
-					if (!string.IsNullOrEmpty(Exporter.ModuleName))
+					UHTManifest.Module? pluginModule = null;
+					if (!String.IsNullOrEmpty(exporter.ModuleName))
 					{
-						foreach (UHTManifest.Module Module in this.Manifest!.Modules)
+						foreach (UHTManifest.Module module in this.Manifest!.Modules)
 						{
-							if (string.Equals(Module.Name, Exporter.ModuleName, StringComparison.OrdinalIgnoreCase))
+							if (String.Equals(module.Name, exporter.ModuleName, StringComparison.OrdinalIgnoreCase))
 							{
-								PluginModule = Module;
+								pluginModule = module;
 								break;
 							}
 						}
-						if (PluginModule == null)
+						if (pluginModule == null)
 						{
-							Log.TraceWarning($"Exporter \"{Exporter.Name}\" skipped because module \"{Exporter.ModuleName}\" was not found in manifest");
+							Log.TraceWarning($"Exporter \"{exporter.Name}\" skipped because module \"{exporter.ModuleName}\" was not found in manifest");
 							continue;
 						}
 					}
 
-					if (Run)
+					if (run)
 					{
-						Log.TraceLog($"       Running exporter {Exporter.Name}");
-						UhtExportFactory Factory = new UhtExportFactory(this, PluginModule, Exporter);
-						Factory.Run();
-						foreach (var Output in Factory.Outputs)
+						Log.TraceLog($"       Running exporter {exporter.Name}");
+						UhtExportFactory factory = new UhtExportFactory(this, pluginModule, exporter);
+						factory.Run();
+						foreach (UhtExportFactory.Output output in factory.Outputs)
 						{
-							if (Output.bSaved)
+							if (output._saved)
 							{
-								TotalWrittenFiles++;
+								totalWrittenFiles++;
 							}
 						}
-						foreach (string Dep in Factory.ExternalDependencies)
+						foreach (string dep in factory.ExternalDependencies)
 						{
-							ExternalDependencies.Add(Dep);
+							externalDependencies.Add(dep);
 						}
 					}
 					else
 					{
-						Log.TraceLog($"       Exporter {Exporter.Name} skipped");
+						Log.TraceLog($"       Exporter {exporter.Name} skipped");
 					}
 				}
 
 				// Save the collected external dependencies
-				if (!string.IsNullOrEmpty(this.Manifest!.ExternalDependenciesFile))
+				if (!String.IsNullOrEmpty(this.Manifest!.ExternalDependenciesFile))
 				{
-					using (StreamWriter Out = new StreamWriter(this.Manifest!.ExternalDependenciesFile))
+					using (StreamWriter output = new StreamWriter(this.Manifest!.ExternalDependenciesFile))
 					{
-						foreach (string Dep in ExternalDependencies)
+						foreach (string dep in externalDependencies)
 						{
-							Out.WriteLine(Dep);
+							output.WriteLine(dep);
 						}
 					}
 				}
 			});
 
-			Log.TraceInformation($"Total of {TotalWrittenFiles} written");
+			Log.TraceInformation($"Total of {totalWrittenFiles} written");
 		}
 		#endregion
 	}
