More coding standard fixes

diff --git a/Engine/Source/Programs/Shared/EpicGames.UHT/Utils/UhtStringBuilderExtensions.cs b/Engine/Source/Programs/Shared/EpicGames.UHT/Utils/UhtStringBuilderExtensions.cs
--- a/Engine/Source/Programs/Shared/EpicGames.UHT/Utils/UhtStringBuilderExtensions.cs
+++ b/Engine/Source/Programs/Shared/EpicGames.UHT/Utils/UhtStringBuilderExtensions.cs
@@ -1,10 +1,10 @@
 ﻿// Copyright Epic Games, Inc. All Rights Reserved.
 
-using EpicGames.Core;
-using EpicGames.UHT.Types;
 using System;
 using System.Collections.Generic;
 using System.Text;
+using EpicGames.Core;
+using EpicGames.UHT.Types;
 
 namespace EpicGames.UHT.Utils
 {
@@ -23,396 +23,396 @@ namespace EpicGames.UHT.Utils
 		/// <summary>
 		/// Append tabs to the builder
 		/// </summary>
-		/// <param name="Builder">Destination builder</param>
-		/// <param name="Tabs">Number of tabs to insert</param>
+		/// <param name="builder">Destination builder</param>
+		/// <param name="tabs">Number of tabs to insert</param>
 		/// <returns>Destination builder</returns>
 		/// <exception cref="ArgumentOutOfRangeException">Thrown if the number of tabs is out of range</exception>
-		public static StringBuilder AppendTabs(this StringBuilder Builder, int Tabs)
+		public static StringBuilder AppendTabs(this StringBuilder builder, int tabs)
 		{
-			if (Tabs < 0 || Tabs > TabsString.Length)
+			if (tabs < 0 || tabs > TabsString.Length)
 			{
 				throw new ArgumentOutOfRangeException($"Number of tabs specified must be between 0 and {TabsString.Length - 1} inclusive");
 			}
-			else if (Tabs > 0)
+			else if (tabs > 0)
 			{
-				Builder.Append(TabsString.Span.Slice(0, Tabs));
+				builder.Append(TabsString.Span.Slice(0, tabs));
 			}
-			return Builder;
+			return builder;
 		}
 
 		/// <summary>
 		/// Append spaces to the builder
 		/// </summary>
-		/// <param name="Builder">Destination builder</param>
-		/// <param name="Spaces">Number of spaces to insert</param>
+		/// <param name="builder">Destination builder</param>
+		/// <param name="spaces">Number of spaces to insert</param>
 		/// <returns>Destination builder</returns>
 		/// <exception cref="ArgumentOutOfRangeException">Thrown if the number of spaces is out of range</exception>
-		public static StringBuilder AppendSpaces(this StringBuilder Builder, int Spaces)
+		public static StringBuilder AppendSpaces(this StringBuilder builder, int spaces)
 		{
-			if (Spaces < 0 || Spaces > SpacesString.Length)
+			if (spaces < 0 || spaces > SpacesString.Length)
 			{
 				throw new ArgumentOutOfRangeException($"Number of spaces specified must be between 0 and {SpacesString.Length - 1} inclusive");
 			}
-			else if (Spaces > 0)
+			else if (spaces > 0)
 			{
-				Builder.Append(SpacesString.Span.Slice(0, Spaces));
+				builder.Append(SpacesString.Span.Slice(0, spaces));
 			}
-			return Builder;
+			return builder;
 		}
 
 		/// <summary>
 		/// Append a name declaration to the builder
 		/// </summary>
-		/// <param name="Builder">Destination builder</param>
-		/// <param name="NamePrefix">Name prefix</param>
-		/// <param name="Name">Name</param>
-		/// <param name="NameSuffix">Optional name suffix</param>
+		/// <param name="builder">Destination builder</param>
+		/// <param name="namePrefix">Name prefix</param>
+		/// <param name="name">Name</param>
+		/// <param name="nameSuffix">Optional name suffix</param>
 		/// <returns>Destination builder</returns>
-		public static StringBuilder AppendNameDecl(this StringBuilder Builder, string? NamePrefix, string Name, string? NameSuffix)
+		public static StringBuilder AppendNameDecl(this StringBuilder builder, string? namePrefix, string name, string? nameSuffix)
 		{
-			return Builder.Append(NamePrefix).Append(Name).Append(NameSuffix);
+			return builder.Append(namePrefix).Append(name).Append(nameSuffix);
 		}
 
 		/// <summary>
 		/// Append a name declaration to the builder
 		/// </summary>
-		/// <param name="Builder">Destination builder</param>
-		/// <param name="Context">Property context used to get the name prefix</param>
-		/// <param name="Name">Name</param>
-		/// <param name="NameSuffix">Optional name suffix</param>
+		/// <param name="builder">Destination builder</param>
+		/// <param name="context">Property context used to get the name prefix</param>
+		/// <param name="name">Name</param>
+		/// <param name="nameSuffix">Optional name suffix</param>
 		/// <returns>Destination builder</returns>
-		public static StringBuilder AppendNameDecl(this StringBuilder Builder, IUhtPropertyMemberContext Context, string Name, string NameSuffix)
+		public static StringBuilder AppendNameDecl(this StringBuilder builder, IUhtPropertyMemberContext context, string name, string nameSuffix)
 		{
-			return Builder.AppendNameDecl(Context.NamePrefix, Name, NameSuffix);
+			return builder.AppendNameDecl(context.NamePrefix, name, nameSuffix);
 		}
 
 		/// <summary>
 		/// Append a name definition to the builder
 		/// </summary>
-		/// <param name="Builder">Destination builder</param>
-		/// <param name="StaticsName">Optional name of the statics block which will be output in the form of "StaticsName::"</param>
-		/// <param name="NamePrefix">Name prefix</param>
-		/// <param name="Name">Name</param>
-		/// <param name="NameSuffix">Optional name suffix</param>
+		/// <param name="builder">Destination builder</param>
+		/// <param name="staticsName">Optional name of the statics block which will be output in the form of "StaticsName::"</param>
+		/// <param name="namePrefix">Name prefix</param>
+		/// <param name="name">Name</param>
+		/// <param name="nameSuffix">Optional name suffix</param>
 		/// <returns>Destination builder</returns>
-		public static StringBuilder AppendNameDef(this StringBuilder Builder, string? StaticsName, string? NamePrefix, string Name, string? NameSuffix)
+		public static StringBuilder AppendNameDef(this StringBuilder builder, string? staticsName, string? namePrefix, string name, string? nameSuffix)
 		{
-			if (!string.IsNullOrEmpty(StaticsName))
+			if (!String.IsNullOrEmpty(staticsName))
 			{
-				Builder.Append(StaticsName).Append("::");
+				builder.Append(staticsName).Append("::");
 			}
-			return Builder.AppendNameDecl(NamePrefix, Name, NameSuffix);
+			return builder.AppendNameDecl(namePrefix, name, nameSuffix);
 		}
 
 		/// <summary>
 		/// Append a name definition to the builder
 		/// </summary>
-		/// <param name="Builder">Destination builder</param>
-		/// <param name="Context">Property context used to get the statics name and name prefix</param>
-		/// <param name="Name">Name</param>
-		/// <param name="NameSuffix">Optional name suffix</param>
+		/// <param name="builder">Destination builder</param>
+		/// <param name="context">Property context used to get the statics name and name prefix</param>
+		/// <param name="name">Name</param>
+		/// <param name="nameSuffix">Optional name suffix</param>
 		/// <returns>Destination builder</returns>
-		public static StringBuilder AppendNameDef(this StringBuilder Builder, IUhtPropertyMemberContext Context, string Name, string NameSuffix)
+		public static StringBuilder AppendNameDef(this StringBuilder builder, IUhtPropertyMemberContext context, string name, string nameSuffix)
 		{
-			return Builder.AppendNameDef(Context.StaticsName, Context.NamePrefix, Name, NameSuffix);
+			return builder.AppendNameDef(context.StaticsName, context.NamePrefix, name, nameSuffix);
 		}
 
 		/// <summary>
 		/// Append the meta data parameters.  This is intended to be used as arguments to a function call.
 		/// </summary>
-		/// <param name="Builder">Destination builder</param>
-		/// <param name="Type">Source type containing the meta data</param>
-		/// <param name="StaticsName">Optional name of the statics block which will be output in the form of "StaticsName::"</param>
-		/// <param name="NamePrefix">Name prefix</param>
-		/// <param name="Name">Name</param>
-		/// <param name="NameSuffix">Optional name suffix</param>
-		/// <param name="MetaNameSuffix">Suffix to be added to the meta data name</param>
+		/// <param name="builder">Destination builder</param>
+		/// <param name="type">Source type containing the meta data</param>
+		/// <param name="staticsName">Optional name of the statics block which will be output in the form of "StaticsName::"</param>
+		/// <param name="namePrefix">Name prefix</param>
+		/// <param name="name">Name</param>
+		/// <param name="nameSuffix">Optional name suffix</param>
+		/// <param name="metaNameSuffix">Suffix to be added to the meta data name</param>
 		/// <returns>Destination builder</returns>
-		private static StringBuilder AppendMetaDataParams(this StringBuilder Builder, UhtType Type, string? StaticsName, string? NamePrefix, string Name, string? NameSuffix, string? MetaNameSuffix)
+		private static StringBuilder AppendMetaDataParams(this StringBuilder builder, UhtType type, string? staticsName, string? namePrefix, string name, string? nameSuffix, string? metaNameSuffix)
 		{
-			if (!Type.MetaData.IsEmpty())
+			if (!type.MetaData.IsEmpty())
 			{
-				return Builder
+				return builder
 					.Append("METADATA_PARAMS(")
-					.AppendNameDef(StaticsName, NamePrefix, Name, NameSuffix).Append(MetaNameSuffix)
+					.AppendNameDef(staticsName, namePrefix, name, nameSuffix).Append(metaNameSuffix)
 					.Append(", UE_ARRAY_COUNT(")
-					.AppendNameDef(StaticsName, NamePrefix, Name, NameSuffix).Append(MetaNameSuffix)
+					.AppendNameDef(staticsName, namePrefix, name, nameSuffix).Append(metaNameSuffix)
 					.Append("))");
 			}
 			else
 			{
-				return Builder.Append("METADATA_PARAMS(nullptr, 0)");
+				return builder.Append("METADATA_PARAMS(nullptr, 0)");
 			}
 		}
 
 		/// <summary>
 		/// Append the meta data parameters.  This is intended to be used as arguments to a function call.
 		/// </summary>
-		/// <param name="Builder">Destination builder</param>
-		/// <param name="Type">Source type containing the meta data</param>
-		/// <param name="StaticsName">Optional name of the statics block which will be output in the form of "StaticsName::"</param>
-		/// <param name="Name">Name</param>
+		/// <param name="builder">Destination builder</param>
+		/// <param name="type">Source type containing the meta data</param>
+		/// <param name="staticsName">Optional name of the statics block which will be output in the form of "StaticsName::"</param>
+		/// <param name="name">Name</param>
 		/// <returns>Destination builder</returns>
-		public static StringBuilder AppendMetaDataParams(this StringBuilder Builder, UhtType Type, string? StaticsName, string Name)
+		public static StringBuilder AppendMetaDataParams(this StringBuilder builder, UhtType type, string? staticsName, string name)
 		{
-			return Builder.AppendMetaDataParams(Type, StaticsName, null, Name, null, null);
+			return builder.AppendMetaDataParams(type, staticsName, null, name, null, null);
 		}
 
 		/// <summary>
 		/// Append the meta data parameters.  This is intended to be used as arguments to a function call.
 		/// </summary>
-		/// <param name="Builder">Destination builder</param>
-		/// <param name="Property">Source type containing the meta data</param>
-		/// <param name="Context">Property context used to get the statics name and name prefix</param>
-		/// <param name="Name">Name</param>
-		/// <param name="NameSuffix">Optional name suffix</param>
+		/// <param name="builder">Destination builder</param>
+		/// <param name="property">Source type containing the meta data</param>
+		/// <param name="context">Property context used to get the statics name and name prefix</param>
+		/// <param name="name">Name</param>
+		/// <param name="nameSuffix">Optional name suffix</param>
 		/// <returns>Destination builder</returns>
-		public static StringBuilder AppendMetaDataParams(this StringBuilder Builder, UhtProperty Property, IUhtPropertyMemberContext Context, string Name, string NameSuffix)
+		public static StringBuilder AppendMetaDataParams(this StringBuilder builder, UhtProperty property, IUhtPropertyMemberContext context, string name, string nameSuffix)
 		{
-			return Builder.AppendMetaDataParams(Property, Context.StaticsName, Context.NamePrefix, Name, NameSuffix, Context.MetaDataSuffix);
+			return builder.AppendMetaDataParams(property, context.StaticsName, context.NamePrefix, name, nameSuffix, context.MetaDataSuffix);
 		}
 
 		/// <summary>
 		/// Append the meta data declaration.
 		/// </summary>
-		/// <param name="Builder">Destination builder</param>
-		/// <param name="Type">Source type containing the meta data</param>
-		/// <param name="NamePrefix">Name prefix</param>
-		/// <param name="Name">Name</param>
-		/// <param name="NameSuffix">Optional name suffix</param>
-		/// <param name="MetaNameSuffix">Optional meta data name suffix</param>
-		/// <param name="Tabs">Number of tabs to indent</param>
+		/// <param name="builder">Destination builder</param>
+		/// <param name="type">Source type containing the meta data</param>
+		/// <param name="namePrefix">Name prefix</param>
+		/// <param name="name">Name</param>
+		/// <param name="nameSuffix">Optional name suffix</param>
+		/// <param name="metaNameSuffix">Optional meta data name suffix</param>
+		/// <param name="tabs">Number of tabs to indent</param>
 		/// <returns>Destination builder</returns>
-		private static StringBuilder AppendMetaDataDecl(this StringBuilder Builder, UhtType Type, string? NamePrefix, string Name, string? NameSuffix, string? MetaNameSuffix, int Tabs)
+		private static StringBuilder AppendMetaDataDecl(this StringBuilder builder, UhtType type, string? namePrefix, string name, string? nameSuffix, string? metaNameSuffix, int tabs)
 		{
-			if (!Type.MetaData.IsEmpty())
+			if (!type.MetaData.IsEmpty())
 			{
-				Builder.Append("#if WITH_METADATA\r\n");
-				Builder.AppendTabs(Tabs).Append("static const UECodeGen_Private::FMetaDataPairParam ").AppendNameDecl(NamePrefix, Name, NameSuffix).Append(MetaNameSuffix).Append("[];\r\n");
-				Builder.Append("#endif\r\n");
+				builder.Append("#if WITH_METADATA\r\n");
+				builder.AppendTabs(tabs).Append("static const UECodeGen_Private::FMetaDataPairParam ").AppendNameDecl(namePrefix, name, nameSuffix).Append(metaNameSuffix).Append("[];\r\n");
+				builder.Append("#endif\r\n");
 			}
-			return Builder;
+			return builder;
 		}
 
 		/// <summary>
 		/// Append the meta data declaration
 		/// </summary>
-		/// <param name="Builder">Destination builder</param>
-		/// <param name="Type">Source type containing the meta data</param>
-		/// <param name="Name">Name</param>
-		/// <param name="Tabs">Number of tabs to indent</param>
+		/// <param name="builder">Destination builder</param>
+		/// <param name="type">Source type containing the meta data</param>
+		/// <param name="name">Name</param>
+		/// <param name="tabs">Number of tabs to indent</param>
 		/// <returns>Destination builder</returns>
-		public static StringBuilder AppendMetaDataDecl(this StringBuilder Builder, UhtType Type, string Name, int Tabs)
+		public static StringBuilder AppendMetaDataDecl(this StringBuilder builder, UhtType type, string name, int tabs)
 		{
-			return Builder.AppendMetaDataDecl(Type, null, Name, null, null, Tabs);
+			return builder.AppendMetaDataDecl(type, null, name, null, null, tabs);
 		}
 
 		/// <summary>
 		/// Append the meta data declaration
 		/// </summary>
-		/// <param name="Builder">Destination builder</param>
-		/// <param name="Property">Source type containing the meta data</param>
-		/// <param name="Context">Property context used to get the statics name and name prefix</param>
-		/// <param name="Name">Name</param>
-		/// <param name="NameSuffix">Optional name suffix</param>
-		/// <param name="Tabs">Number of tabs to indent</param>
+		/// <param name="builder">Destination builder</param>
+		/// <param name="property">Source type containing the meta data</param>
+		/// <param name="context">Property context used to get the statics name and name prefix</param>
+		/// <param name="name">Name</param>
+		/// <param name="nameSuffix">Optional name suffix</param>
+		/// <param name="tabs">Number of tabs to indent</param>
 		/// <returns>Destination builder</returns>
-		public static StringBuilder AppendMetaDataDecl(this StringBuilder Builder, UhtProperty Property, IUhtPropertyMemberContext Context, string Name, string NameSuffix, int Tabs)
+		public static StringBuilder AppendMetaDataDecl(this StringBuilder builder, UhtProperty property, IUhtPropertyMemberContext context, string name, string nameSuffix, int tabs)
 		{
-			return Builder.AppendMetaDataDecl(Property, Context.NamePrefix, Name, NameSuffix, Context.MetaDataSuffix, Tabs);
+			return builder.AppendMetaDataDecl(property, context.NamePrefix, name, nameSuffix, context.MetaDataSuffix, tabs);
 		}
 
 		/// <summary>
 		/// Append the meta data definition.
 		/// </summary>
-		/// <param name="Builder">Destination builder</param>
-		/// <param name="Type">Source type containing the meta data</param>
-		/// <param name="StaticsName">Optional name of the statics block which will be output in the form of "StaticsName::"</param>
-		/// <param name="NamePrefix">Name prefix</param>
-		/// <param name="Name">Name</param>
-		/// <param name="NameSuffix">Optional name suffix</param>
-		/// <param name="MetaNameSuffix">Optional meta data name suffix</param>
-		/// <param name="Tabs">Number of tabs to indent</param>
+		/// <param name="builder">Destination builder</param>
+		/// <param name="type">Source type containing the meta data</param>
+		/// <param name="staticsName">Optional name of the statics block which will be output in the form of "StaticsName::"</param>
+		/// <param name="namePrefix">Name prefix</param>
+		/// <param name="name">Name</param>
+		/// <param name="nameSuffix">Optional name suffix</param>
+		/// <param name="metaNameSuffix">Optional meta data name suffix</param>
+		/// <param name="tabs">Number of tabs to indent</param>
 		/// <returns>Destination builder</returns>
-		private static StringBuilder AppendMetaDataDef(this StringBuilder Builder, UhtType Type, string? StaticsName, string? NamePrefix, string Name, string? NameSuffix, string? MetaNameSuffix, int Tabs)
+		private static StringBuilder AppendMetaDataDef(this StringBuilder builder, UhtType type, string? staticsName, string? namePrefix, string name, string? nameSuffix, string? metaNameSuffix, int tabs)
 		{
-			if (!Type.MetaData.IsEmpty())
+			if (!type.MetaData.IsEmpty())
 			{
-				List<KeyValuePair<string, string>> SortedMetaData = Type.MetaData.GetSorted();
+				List<KeyValuePair<string, string>> sortedMetaData = type.MetaData.GetSorted();
 
-				Builder.Append("#if WITH_METADATA\r\n");
-				Builder.AppendTabs(Tabs).Append("const UECodeGen_Private::FMetaDataPairParam ").AppendNameDef(StaticsName, NamePrefix, Name, NameSuffix).Append(MetaNameSuffix).Append("[] = {\r\n");
+				builder.Append("#if WITH_METADATA\r\n");
+				builder.AppendTabs(tabs).Append("const UECodeGen_Private::FMetaDataPairParam ").AppendNameDef(staticsName, namePrefix, name, nameSuffix).Append(metaNameSuffix).Append("[] = {\r\n");
 
-				foreach (KeyValuePair<string, string> Kvp in SortedMetaData)
+				foreach (KeyValuePair<string, string> kvp in sortedMetaData)
 				{
-					Builder.AppendTabs(Tabs + 1).Append("{ ").AppendUTF8LiteralString(Kvp.Key).Append(", ").AppendUTF8LiteralString(Kvp.Value).Append(" },\r\n");
+					builder.AppendTabs(tabs + 1).Append("{ ").AppendUTF8LiteralString(kvp.Key).Append(", ").AppendUTF8LiteralString(kvp.Value).Append(" },\r\n");
 				}
 
-				Builder.AppendTabs(Tabs).Append("};\r\n");
-				Builder.Append("#endif\r\n");
+				builder.AppendTabs(tabs).Append("};\r\n");
+				builder.Append("#endif\r\n");
 			}
-			return Builder;
+			return builder;
 		}
 
 		/// <summary>
 		/// Append the meta data definition
 		/// </summary>
-		/// <param name="Builder">Destination builder</param>
-		/// <param name="Type">Source type containing the meta data</param>
-		/// <param name="StaticsName">Optional name of the statics block which will be output in the form of "StaticsName::"</param>
-		/// <param name="Name">Name</param>
-		/// <param name="Tabs">Number of tabs to indent</param>
+		/// <param name="builder">Destination builder</param>
+		/// <param name="type">Source type containing the meta data</param>
+		/// <param name="staticsName">Optional name of the statics block which will be output in the form of "StaticsName::"</param>
+		/// <param name="name">Name</param>
+		/// <param name="tabs">Number of tabs to indent</param>
 		/// <returns>Destination builder</returns>
-		public static StringBuilder AppendMetaDataDef(this StringBuilder Builder, UhtType Type, string? StaticsName, string Name, int Tabs)
+		public static StringBuilder AppendMetaDataDef(this StringBuilder builder, UhtType type, string? staticsName, string name, int tabs)
 		{
-			return Builder.AppendMetaDataDef(Type, StaticsName, null, Name, null, null, Tabs);
+			return builder.AppendMetaDataDef(type, staticsName, null, name, null, null, tabs);
 		}
 
 		/// <summary>
 		/// Append the meta data definition
 		/// </summary>
-		/// <param name="Builder">Destination builder</param>
-		/// <param name="Property">Source type containing the meta data</param>
-		/// <param name="Context">Property context used to get the statics name and name prefix</param>
-		/// <param name="Name">Name</param>
-		/// <param name="NameSuffix">Optional name suffix</param>
-		/// <param name="Tabs">Number of tabs to indent</param>
+		/// <param name="builder">Destination builder</param>
+		/// <param name="property">Source type containing the meta data</param>
+		/// <param name="context">Property context used to get the statics name and name prefix</param>
+		/// <param name="name">Name</param>
+		/// <param name="nameSuffix">Optional name suffix</param>
+		/// <param name="tabs">Number of tabs to indent</param>
 		/// <returns>Destination builder</returns>
-		public static StringBuilder AppendMetaDataDef(this StringBuilder Builder, UhtProperty Property, IUhtPropertyMemberContext Context, string Name, string NameSuffix, int Tabs)
+		public static StringBuilder AppendMetaDataDef(this StringBuilder builder, UhtProperty property, IUhtPropertyMemberContext context, string name, string nameSuffix, int tabs)
 		{
-			return Builder.AppendMetaDataDef(Property, Context.StaticsName, Context.NamePrefix, Name, NameSuffix, Context.MetaDataSuffix, Tabs);
+			return builder.AppendMetaDataDef(property, context.StaticsName, context.NamePrefix, name, nameSuffix, context.MetaDataSuffix, tabs);
 		}
 
 		/// <summary>
 		/// Append the given text as a UTF8 encoded string
 		/// </summary>
-		/// <param name="Builder">Destination builder</param>
-		/// <param name="bUseText">If false, don't encode the text but include a nullptr</param>
-		/// <param name="Text">Text to include or an empty string if null.</param>
+		/// <param name="builder">Destination builder</param>
+		/// <param name="useText">If false, don't encode the text but include a nullptr</param>
+		/// <param name="text">Text to include or an empty string if null.</param>
 		/// <returns>Destination builder</returns>
-		public static StringBuilder AppendUTF8LiteralString(this StringBuilder Builder, bool bUseText, string? Text)
+		public static StringBuilder AppendUTF8LiteralString(this StringBuilder builder, bool useText, string? text)
 		{
-			if (!bUseText)
+			if (!useText)
 			{
-				Builder.Append("nullptr");
+				builder.Append("nullptr");
 			}
-			else if (Text == null)
+			else if (text == null)
 			{
-				Builder.Append("");
+				builder.Append("");
 			}
 			else
 			{
-				Builder.AppendUTF8LiteralString(Text);
+				builder.AppendUTF8LiteralString(text);
 			}
-			return Builder;
+			return builder;
 		}
 
 		/// <summary>
 		/// Append the given text as a UTF8 encoded string
 		/// </summary>
-		/// <param name="Builder">Destination builder</param>
-		/// <param name="Text">Text to be encoded</param>
+		/// <param name="builder">Destination builder</param>
+		/// <param name="text">Text to be encoded</param>
 		/// <returns>Destination builder</returns>
-		public static StringBuilder AppendUTF8LiteralString(this StringBuilder Builder, StringView Text)
+		public static StringBuilder AppendUTF8LiteralString(this StringBuilder builder, StringView text)
 		{
-			Builder.Append('\"');
+			builder.Append('\"');
 
-			ReadOnlySpan<char> Span = Text.Span;
-			int Length = Span.Length;
+			ReadOnlySpan<char> span = text.Span;
+			int length = span.Length;
 
-			if (Length > 0)
+			if (length > 0)
 			{
 
-				bool bTrailingHex = false;
-				int Index = 0;
+				bool trailingHex = false;
+				int index = 0;
 				while (true)
 				{
 					// Scan forward looking for anything that can just be blindly copied
-					int StartIndex = Index;
-					while (Index < Length)
+					int startIndex = index;
+					while (index < length)
 					{
-						char cskip = Span[Index];
+						char cskip = span[index];
 						if (cskip < 31 || cskip > 127 || cskip == '"' || cskip == '\\')
 						{
 							break;
 						}
-						++Index;
+						++index;
 					}
 
 					// If we found anything
-					if (StartIndex < Index)
+					if (startIndex < index)
 					{
 						// We close and open the literal here in order to ensure that successive hex characters aren't appended to the hex sequence, causing a different number
-						if (bTrailingHex && UhtFCString.IsHexDigit(Span[StartIndex]))
+						if (trailingHex && UhtFCString.IsHexDigit(span[startIndex]))
 						{
-							Builder.Append("\"\"");
+							builder.Append("\"\"");
 						}
-						Builder.Append(Span.Slice(StartIndex, Index - StartIndex));
+						builder.Append(span.Slice(startIndex, index - startIndex));
 					}
 
 					// We have either reached the end of the string, break
-					if (Index == Length)
+					if (index == length)
 					{
 						break;
 					}
 
 					// This character requires special processing
-					char c = Span[Index++];
+					char c = span[index++];
 					switch (c)
 					{
 						case '\r':
-							bTrailingHex = false;
+							trailingHex = false;
 							break;
 						case '\n':
-							bTrailingHex = false;
-							Builder.Append("\\n");
+							trailingHex = false;
+							builder.Append("\\n");
 							break;
 						case '\\':
-							bTrailingHex = false;
-							Builder.Append("\\\\");
+							trailingHex = false;
+							builder.Append("\\\\");
 							break;
 						case '\"':
-							bTrailingHex = false;
-							Builder.Append("\\\"");
+							trailingHex = false;
+							builder.Append("\\\"");
 							break;
 						default:
 							if (c < 31)
 							{
-								bTrailingHex = true;
+								trailingHex = true;
 								//Builder.Append($"\\x{(uint)c:x2}");
-								Builder.Append("\\x").AppendFormat("{0:x2}", (uint)c);
+								builder.Append("\\x").AppendFormat("{0:x2}", (uint)c);
 							}
 							else
 							{
-								bTrailingHex = false;
-								if (char.IsHighSurrogate(c))
+								trailingHex = false;
+								if (Char.IsHighSurrogate(c))
 								{
-									if (Index == Length)
+									if (index == length)
 									{
-										Builder.Append('?');
+										builder.Append('?');
 										break;
 									}
 
-									char clow = Span[Index];
-									if (char.IsLowSurrogate(clow))
+									char clow = span[index];
+									if (Char.IsLowSurrogate(clow))
 									{
-										++Index;
-										Builder.AppendEscapedUtf32((ulong)char.ConvertToUtf32(c, clow));
-										bTrailingHex = true;
+										++index;
+										builder.AppendEscapedUtf32((ulong)Char.ConvertToUtf32(c, clow));
+										trailingHex = true;
 									}
 									else
 									{
-										Builder.Append('?');
+										builder.Append('?');
 									}
 								}
-								else if (char.IsLowSurrogate(c))
+								else if (Char.IsLowSurrogate(c))
 								{
-									Builder.Append('?');
+									builder.Append('?');
 								}
 								else
 								{
-									Builder.AppendEscapedUtf32(c);
-									bTrailingHex = true;
+									builder.AppendEscapedUtf32(c);
+									trailingHex = true;
 								}
 							}
 							break;
@@ -420,63 +420,63 @@ namespace EpicGames.UHT.Utils
 				}
 			}
 
-			Builder.Append('\"');
-			return Builder;
+			builder.Append('\"');
+			return builder;
 		}
 
 		/// <summary>
 		/// Encode a single UTF32 value as UTF8 characters
 		/// </summary>
-		/// <param name="Builder">Destination builder</param>
+		/// <param name="builder">Destination builder</param>
 		/// <param name="c">Character to encode</param>
 		/// <returns>Destination builder</returns>
-		public static StringBuilder AppendEscapedUtf32(this StringBuilder Builder, ulong c)
+		public static StringBuilder AppendEscapedUtf32(this StringBuilder builder, ulong c)
 		{
 			if (c < 0x80)
 			{
-				Builder
+				builder
 					.Append("\\x").AppendFormat("{0:x2}", c);
 			}
 			else if (c < 0x800)
 			{
-				Builder
+				builder
 					.Append("\\x").AppendFormat("{0:x2}", 0xC0 + (c >> 6))
 					.Append("\\x").AppendFormat("{0:x2}", 0x80 + (c & 0x3f));
 			}
 			else if (c < 0x10000)
 			{
-				Builder
+				builder
 					.Append("\\x").AppendFormat("{0:x2}", 0xE0 + (c >> 12))
 					.Append("\\x").AppendFormat("{0:x2}", 0x80 + ((c >> 6) & 0x3f))
 					.Append("\\x").AppendFormat("{0:x2}", 0x80 + (c & 0x3f));
 			}
 			else
 			{
-				Builder
+				builder
 					.Append("\\x").AppendFormat("{0:x2}", 0xF0 + (c >> 18))
 					.Append("\\x").AppendFormat("{0:x2}", 0x80 + ((c >> 12) & 0x3f))
 					.Append("\\x").AppendFormat("{0:x2}", 0x80 + ((c >> 6) & 0x3f))
 					.Append("\\x").AppendFormat("{0:x2}", 0x80 + (c & 0x3f));
 			}
-			return Builder;
+			return builder;
 		}
 
 		/// <summary>
 		/// Append the given name of the class but always encode interfaces as the native interface name (i.e. "I...")
 		/// </summary>
-		/// <param name="Builder">Destination builder</param>
-		/// <param name="Class">Class to append name</param>
+		/// <param name="builder">Destination builder</param>
+		/// <param name="classObj">Class to append name</param>
 		/// <returns>Destination builder</returns>
 		/// <exception cref="NotImplementedException">Class has an unexpected class type</exception>
-		public static StringBuilder AppendClassSourceNameOrInterfaceName(this StringBuilder Builder, UhtClass Class)
+		public static StringBuilder AppendClassSourceNameOrInterfaceName(this StringBuilder builder, UhtClass classObj)
 		{
-			switch (Class.ClassType)
+			switch (classObj.ClassType)
 			{
 				case UhtClassType.Class:
 				case UhtClassType.NativeInterface:
-					return Builder.Append(Class.SourceName);
+					return builder.Append(classObj.SourceName);
 				case UhtClassType.Interface:
-					return Builder.Append('I').Append(Class.EngineName);
+					return builder.Append('I').Append(classObj.EngineName);
 				default:
 					throw new NotImplementedException();
 			}
@@ -502,29 +502,29 @@ namespace EpicGames.UHT.Utils
 		/// <summary>
 		/// Capacity of the cache
 		/// </summary>
-		private readonly int Capacity;
+		private readonly int _capacity;
 
 		/// <summary>
 		/// Initial buffer size for new StringBuilders.  Resetting StringBuilders might result
 		/// in the initial chunk size being smaller.
 		/// </summary>
-		private readonly int InitialBufferSize;
+		private readonly int _initialBufferSize;
 
 		/// <summary>
 		/// Stack of cached StringBuilders
 		/// </summary>
-		private readonly Stack<StringBuilder> Stack;
+		private readonly Stack<StringBuilder> _stack;
 
 		/// <summary>
 		/// Create a new StringBuilder cache
 		/// </summary>
-		/// <param name="Capacity">Maximum number of StringBuilders to cache</param>
-		/// <param name="InitialBufferSize">Initial buffer size for newly created StringBuilders</param>
-		public StringBuilderCache(int Capacity, int InitialBufferSize)
+		/// <param name="capacity">Maximum number of StringBuilders to cache</param>
+		/// <param name="initialBufferSize">Initial buffer size for newly created StringBuilders</param>
+		public StringBuilderCache(int capacity, int initialBufferSize)
 		{
-			this.Capacity = Capacity;
-			this.InitialBufferSize = InitialBufferSize;
-			this.Stack = new Stack<StringBuilder>(this.Capacity);
+			this._capacity = capacity;
+			this._initialBufferSize = initialBufferSize;
+			this._stack = new Stack<StringBuilder>(this._capacity);
 		}
 
 		/// <summary>
@@ -533,30 +533,30 @@ namespace EpicGames.UHT.Utils
 		/// <returns></returns>
 		public StringBuilder Borrow()
 		{
-			lock (this.Stack)
+			lock (this._stack)
 			{
-				if (this.Stack.Count > 0)
+				if (this._stack.Count > 0)
 				{
-					return this.Stack.Pop();
+					return this._stack.Pop();
 				}
 			}
 
-			return new StringBuilder(this.InitialBufferSize);
+			return new StringBuilder(this._initialBufferSize);
 		}
 
 		/// <summary>
 		/// Return a StringBuilder to the cache
 		/// </summary>
-		/// <param name="Builder">The builder being returned</param>
-		public void Return(StringBuilder Builder)
+		/// <param name="builder">The builder being returned</param>
+		public void Return(StringBuilder builder)
 		{
 			// Sadly, clearing the builder (sets length to 0) will reallocate chunks.
-			Builder.Clear();
-			lock (this.Stack)
+			builder.Clear();
+			lock (this._stack)
 			{
-				if (this.Stack.Count < this.Capacity)
+				if (this._stack.Count < this._capacity)
 				{
-					this.Stack.Push(Builder);
+					this._stack.Push(builder);
 				}
 			}
 		}
@@ -582,10 +582,10 @@ namespace EpicGames.UHT.Utils
 		/// <summary>
 		/// Borrow a string builder from the given cache
 		/// </summary>
-		/// <param name="Cache">String builder cache</param>
-		public BorrowStringBuilder(StringBuilderCache Cache)
+		/// <param name="cache">String builder cache</param>
+		public BorrowStringBuilder(StringBuilderCache cache)
 		{
-			this.Cache = Cache;
+			this.Cache = cache;
 			this.StringBuilder = this.Cache.Borrow();
 		}
 
diff --git a/Engine/Source/Programs/Shared/EpicGames.UHT/Utils/UhtSymbolTable.cs b/Engine/Source/Programs/Shared/EpicGames.UHT/Utils/UhtSymbolTable.cs
--- a/Engine/Source/Programs/Shared/EpicGames.UHT/Utils/UhtSymbolTable.cs
+++ b/Engine/Source/Programs/Shared/EpicGames.UHT/Utils/UhtSymbolTable.cs
@@ -1,9 +1,9 @@
 // Copyright Epic Games, Inc. All Rights Reserved.
 
-using EpicGames.Core;
-using EpicGames.UHT.Types;
 using System;
 using System.Collections.Generic;
+using EpicGames.Core;
+using EpicGames.UHT.Types;
 
 namespace EpicGames.UHT.Utils
 {
@@ -24,13 +24,13 @@ namespace EpicGames.UHT.Utils
 			/// <summary>
 			/// The type index of the symbol
 			/// </summary>
-			public int SymbolIndex;
+			public int SymbolIndex { get; set; }
 
 			/// <summary>
 			/// When searching the caseless chain, the cased index is used to match the symbol based
 			/// on the case named.
 			/// </summary>
-			public int CasedIndex;
+			public int CasedIndex { get; set; }
 		}
 
 		/// <summary>
@@ -42,86 +42,86 @@ namespace EpicGames.UHT.Utils
 			/// <summary>
 			/// The type associated with the symbol
 			/// </summary>
-			public UhtType Type;
+			public UhtType Type { get; set; }
 
 			/// <summary>
 			/// Mask of different find options which will match this symbol
 			/// </summary>
-			public UhtFindOptions MatchOptions;
+			public UhtFindOptions MatchOptions { get; set; }
 
 			/// <summary>
 			/// The case lookup index for matching by case
 			/// </summary>
-			public int CasedIndex;
+			public int CasedIndex { get; set; }
 
 			/// <summary>
 			/// The next index in the symbol change based on cassless lookup
 			/// </summary>
-			public int NextIndex;
+			public int NextIndex { get; set; }
 
 			/// <summary>
 			/// The last index in the chain.  This index is only used when the symbol entry is also acting as the list
 			/// </summary>
-			public int LastIndex;
+			public int LastIndex { get; set; }
 		}
 
 		/// <summary>
 		/// Number of unique cased symbol names.  
 		/// </summary>
-		private int CasedCount = 0;
+		private int _casedCount = 0;
 
 		/// <summary>
 		/// Case name lookup table that returns the symbol index and the case index
 		/// </summary>
-		private readonly Dictionary<string, Lookup> CasedDictionary = new Dictionary<string, Lookup>(StringComparer.Ordinal);
+		private readonly Dictionary<string, Lookup> _casedDictionary = new Dictionary<string, Lookup>(StringComparer.Ordinal);
 
 		/// <summary>
 		/// Caseless name lookup table that returns the symbol index
 		/// </summary>
-		private readonly Dictionary<string, int> CaselessDictionary = new Dictionary<string, int>(StringComparer.OrdinalIgnoreCase);
+		private readonly Dictionary<string, int> _caselessDictionary = new Dictionary<string, int>(StringComparer.OrdinalIgnoreCase);
 
 		/// <summary>
 		/// Collection of symbols in the table
 		/// </summary>
-		private readonly Symbol[] Symbols;
+		private readonly Symbol[] _symbols;
 
 		/// <summary>
 		/// Constructs a new symbol table.
 		/// </summary>
-		/// <param name="TypeCount">Number of types in the table.</param>
-		public UhtSymbolTable(int TypeCount)
+		/// <param name="typeCount">Number of types in the table.</param>
+		public UhtSymbolTable(int typeCount)
 		{
-			this.Symbols = new Symbol[TypeCount];
+			this._symbols = new Symbol[typeCount];
 		}
 
 		/// <summary>
 		/// Add a new type to the symbol table
 		/// </summary>
-		/// <param name="Type">The type being added</param>
-		/// <param name="Name">The name of the type which could be the source name or the engine name</param>
-		public void Add(UhtType Type, string Name)
+		/// <param name="type">The type being added</param>
+		/// <param name="name">The name of the type which could be the source name or the engine name</param>
+		public void Add(UhtType type, string name)
 		{
-			Lookup Existing;
-			if (this.CasedDictionary.TryGetValue(Name, out Existing))
+			Lookup existing;
+			if (this._casedDictionary.TryGetValue(name, out existing))
 			{
-				AddExisting(Type, Existing.CasedIndex, Existing.SymbolIndex);
+				AddExisting(type, existing.CasedIndex, existing.SymbolIndex);
 				return;
 			}
 
-			int CasedIndex = ++this.CasedCount;
+			int casedIndex = ++this._casedCount;
 
-			int SymbolIndex;
-			if (this.CaselessDictionary.TryGetValue(Name, out SymbolIndex))
+			int symbolIndex;
+			if (this._caselessDictionary.TryGetValue(name, out symbolIndex))
 			{
-				this.CasedDictionary.Add(Name, new Lookup { SymbolIndex = SymbolIndex, CasedIndex = CasedIndex });
-				AddExisting(Type, CasedIndex, SymbolIndex);
+				this._casedDictionary.Add(name, new Lookup { SymbolIndex = symbolIndex, CasedIndex = casedIndex });
+				AddExisting(type, casedIndex, symbolIndex);
 				return;
 			}
 
-			SymbolIndex = Type.TypeIndex;
-			this.Symbols[SymbolIndex] = new Symbol { Type = Type, MatchOptions = Type.EngineType.FindOptions(), CasedIndex = CasedIndex, NextIndex = 0, LastIndex = SymbolIndex };
-			this.CaselessDictionary.Add(Name, SymbolIndex);
-			this.CasedDictionary.Add(Name, new Lookup { SymbolIndex = SymbolIndex, CasedIndex = CasedIndex });
+			symbolIndex = type.TypeIndex;
+			this._symbols[symbolIndex] = new Symbol { Type = type, MatchOptions = type.EngineType.FindOptions(), CasedIndex = casedIndex, NextIndex = 0, LastIndex = symbolIndex };
+			this._caselessDictionary.Add(name, symbolIndex);
+			this._casedDictionary.Add(name, new Lookup { SymbolIndex = symbolIndex, CasedIndex = casedIndex });
 		}
 
 		/// <summary>
@@ -129,20 +129,20 @@ namespace EpicGames.UHT.Utils
 		/// parser property (which could not resolve the property prior to the symbol table being created)
 		/// with the fully resolved property.
 		/// </summary>
-		/// <param name="OldType">The old type being replaced.</param>
-		/// <param name="NewType">The new type.</param>
-		/// <param name="Name">The name of the type.</param>
+		/// <param name="oldType">The old type being replaced.</param>
+		/// <param name="newType">The new type.</param>
+		/// <param name="name">The name of the type.</param>
 		/// <exception cref="UhtIceException">Thrown if the symbol wasn't found.</exception>
-		public void Replace(UhtType OldType, UhtType NewType, string Name)
+		public void Replace(UhtType oldType, UhtType newType, string name)
 		{
-			int SymbolIndex;
-			if (this.CaselessDictionary.TryGetValue(Name, out SymbolIndex))
+			int symbolIndex;
+			if (this._caselessDictionary.TryGetValue(name, out symbolIndex))
 			{
-				for (; SymbolIndex != 0; SymbolIndex = this.Symbols[SymbolIndex].NextIndex)
+				for (; symbolIndex != 0; symbolIndex = this._symbols[symbolIndex].NextIndex)
 				{
-					if (this.Symbols[SymbolIndex].Type == OldType)
+					if (this._symbols[symbolIndex].Type == oldType)
 					{
-						this.Symbols[SymbolIndex].Type = NewType;
+						this._symbols[symbolIndex].Type = newType;
 						return;
 					}
 				}
@@ -153,16 +153,16 @@ namespace EpicGames.UHT.Utils
 		/// <summary>
 		/// Lookup the given name using cased string compare.
 		/// </summary>
-		/// <param name="StartingType">Starting type used to limit the scope of the search.</param>
-		/// <param name="Options">Options controlling what is search and what is returned.</param>
-		/// <param name="Name">Name to locate.</param>
+		/// <param name="startingType">Starting type used to limit the scope of the search.</param>
+		/// <param name="options">Options controlling what is search and what is returned.</param>
+		/// <param name="name">Name to locate.</param>
 		/// <returns>Found type or null if not found.</returns>
-		public UhtType? FindCasedType(UhtType? StartingType, UhtFindOptions Options, string Name)
+		public UhtType? FindCasedType(UhtType? startingType, UhtFindOptions options, string name)
 		{
-			Lookup Existing;
-			if (this.CasedDictionary.TryGetValue(Name, out Existing))
+			Lookup existing;
+			if (this._casedDictionary.TryGetValue(name, out existing))
 			{
-				return FindType(StartingType, Options, Existing);
+				return FindType(startingType, options, existing);
 			}
 			return null;
 		}
@@ -170,16 +170,16 @@ namespace EpicGames.UHT.Utils
 		/// <summary>
 		/// Lookup the given name using caseless string compare.
 		/// </summary>
-		/// <param name="StartingType">Starting type used to limit the scope of the search.</param>
-		/// <param name="Options">Options controlling what is search and what is returned.</param>
-		/// <param name="Name">Name to locate.</param>
+		/// <param name="startingType">Starting type used to limit the scope of the search.</param>
+		/// <param name="options">Options controlling what is search and what is returned.</param>
+		/// <param name="name">Name to locate.</param>
 		/// <returns>Found type or null if not found.</returns>
-		public UhtType? FindCaselessType(UhtType? StartingType, UhtFindOptions Options, string Name)
+		public UhtType? FindCaselessType(UhtType? startingType, UhtFindOptions options, string name)
 		{
-			int SymbolIndex;
-			if (this.CaselessDictionary.TryGetValue(Name, out SymbolIndex))
+			int symbolIndex;
+			if (this._caselessDictionary.TryGetValue(name, out symbolIndex))
 			{
-				return FindType(StartingType, Options, new Lookup { SymbolIndex = SymbolIndex, CasedIndex = 0 });
+				return FindType(startingType, options, new Lookup { SymbolIndex = symbolIndex, CasedIndex = 0 });
 			}
 			return null;
 		}
@@ -187,52 +187,52 @@ namespace EpicGames.UHT.Utils
 		/// <summary>
 		/// Lookup the given name.
 		/// </summary>
-		/// <param name="StartingType">Starting type used to limit the scope of the search.</param>
-		/// <param name="Options">Options controlling what is search and what is returned.</param>
-		/// <param name="Lookup">Starting lookup location.</param>
+		/// <param name="startingType">Starting type used to limit the scope of the search.</param>
+		/// <param name="options">Options controlling what is search and what is returned.</param>
+		/// <param name="lookup">Starting lookup location.</param>
 		/// <returns>Found type or null if not found.</returns>
-		private UhtType? FindType(UhtType? StartingType, UhtFindOptions Options, Lookup Lookup)
+		private UhtType? FindType(UhtType? startingType, UhtFindOptions options, Lookup lookup)
 		{
-			if (StartingType != null)
+			if (startingType != null)
 			{
-				UhtType? Found = FindTypeSuperChain(StartingType, Options, ref Lookup);
-				if (Found != null)
+				UhtType? found = FindTypeSuperChain(startingType, options, ref lookup);
+				if (found != null)
 				{
-					return Found;
+					return found;
 				}
 
-				Found = FindTypeOuterChain(StartingType, Options, ref Lookup);
-				if (Found != null)
+				found = FindTypeOuterChain(startingType, options, ref lookup);
+				if (found != null)
 				{
-					return Found;
+					return found;
 				}
 
-				if (!Options.HasAnyFlags(UhtFindOptions.NoIncludes))
+				if (!options.HasAnyFlags(UhtFindOptions.NoIncludes))
 				{
-					UhtHeaderFile HeaderFile = StartingType.HeaderFile;
-					foreach (UhtHeaderFile IncludedFile in HeaderFile.IncludedHeaders)
+					UhtHeaderFile headerFile = startingType.HeaderFile;
+					foreach (UhtHeaderFile includedFile in headerFile.IncludedHeaders)
 					{
-						Found = this.FindSymbolChain(IncludedFile, Options, ref Lookup);
-						if (Found != null)
+						found = this.FindSymbolChain(includedFile, options, ref lookup);
+						if (found != null)
 						{
 							break;
 						}
 					}
-					if (Found != null)
+					if (found != null)
 					{
-						return Found;
+						return found;
 					}
 				}
 			}
 
 			// Global search.  Match anything that has an owner of parent
-			if (!Options.HasAnyFlags(UhtFindOptions.NoGlobal))
+			if (!options.HasAnyFlags(UhtFindOptions.NoGlobal))
 			{
-				for (int Index = Lookup.SymbolIndex; Index != 0; Index = this.Symbols[Index].NextIndex)
+				for (int index = lookup.SymbolIndex; index != 0; index = this._symbols[index].NextIndex)
 				{
-					if (IsMatch(Options, Index, Lookup.CasedIndex) && this.Symbols[Index].Type.Outer is UhtHeaderFile)
+					if (IsMatch(options, index, lookup.CasedIndex) && this._symbols[index].Type.Outer is UhtHeaderFile)
 					{
-						return this.Symbols[Index].Type;
+						return this._symbols[index].Type;
 					}
 				}
 			}
@@ -244,44 +244,44 @@ namespace EpicGames.UHT.Utils
 		/// <summary>
 		/// Lookup the given name using the super class/struct chain.
 		/// </summary>
-		/// <param name="StartingType">Starting type used to limit the scope of the search.</param>
-		/// <param name="Options">Options controlling what is search and what is returned.</param>
-		/// <param name="Lookup">Starting lookup location.</param>
+		/// <param name="startingType">Starting type used to limit the scope of the search.</param>
+		/// <param name="options">Options controlling what is search and what is returned.</param>
+		/// <param name="lookup">Starting lookup location.</param>
 		/// <returns>Found type or null if not found.</returns>
-		private UhtType? FindTypeSuperChain(UhtType StartingType, UhtFindOptions Options, ref Lookup Lookup)
+		private UhtType? FindTypeSuperChain(UhtType startingType, UhtFindOptions options, ref Lookup lookup)
 		{
-			UhtType? CurrentType = StartingType;
+			UhtType? currentType = startingType;
 
 			// In a super chain search, we have to start at a UHTStruct that contributes to the symbol table
-			for (; CurrentType != null; CurrentType = CurrentType.Outer)
+			for (; currentType != null; currentType = currentType.Outer)
 			{
-				if (CurrentType is UhtStruct Struct && Struct.EngineType.AddChildrenToSymbolTable())
+				if (currentType is UhtStruct structObj && structObj.EngineType.AddChildrenToSymbolTable())
 				{
 					break;
 				}
 			}
 
 			// Not symbol that supports a super chain
-			if (CurrentType == null)
+			if (currentType == null)
 			{
 				return null;
 			}
 
 			// If requested, skip self
-			if (Options.HasAnyFlags(UhtFindOptions.NoSelf) && CurrentType == StartingType)
+			if (options.HasAnyFlags(UhtFindOptions.NoSelf) && currentType == startingType)
 			{
-				CurrentType = ((UhtStruct)CurrentType).Super;
+				currentType = ((UhtStruct)currentType).Super;
 			}
 
 			// Search the chain
-			for (; CurrentType != null; CurrentType = ((UhtStruct)CurrentType).Super)
+			for (; currentType != null; currentType = ((UhtStruct)currentType).Super)
 			{
-				UhtType? FoundType = FindSymbolChain(CurrentType, Options, ref Lookup);
-				if (FoundType != null)
+				UhtType? foundType = FindSymbolChain(currentType, options, ref lookup);
+				if (foundType != null)
 				{
-					return FoundType;
+					return foundType;
 				}
-				if (Options.HasAnyFlags(UhtFindOptions.NoParents))
+				if (options.HasAnyFlags(UhtFindOptions.NoParents))
 				{
 					return null;
 				}
@@ -292,33 +292,33 @@ namespace EpicGames.UHT.Utils
 		/// <summary>
 		/// Lookup the given name using the outer chain
 		/// </summary>
-		/// <param name="StartingType">Starting type used to limit the scope of the search.</param>
-		/// <param name="Options">Options controlling what is search and what is returned.</param>
-		/// <param name="Lookup">Starting lookup location.</param>
+		/// <param name="startingType">Starting type used to limit the scope of the search.</param>
+		/// <param name="options">Options controlling what is search and what is returned.</param>
+		/// <param name="lookup">Starting lookup location.</param>
 		/// <returns>Found type or null if not found.</returns>
-		private UhtType? FindTypeOuterChain(UhtType StartingType, UhtFindOptions Options, ref Lookup Lookup)
+		private UhtType? FindTypeOuterChain(UhtType startingType, UhtFindOptions options, ref Lookup lookup)
 		{
-			UhtType? CurrentType = StartingType;
+			UhtType? currentType = startingType;
 
 			// If requested, skip self
-			if (Options.HasAnyFlags(UhtFindOptions.NoSelf) && CurrentType == StartingType)
+			if (options.HasAnyFlags(UhtFindOptions.NoSelf) && currentType == startingType)
 			{
-				CurrentType = CurrentType.Outer;
+				currentType = currentType.Outer;
 			}
 
 			// Search the chain
-			for (; CurrentType != null; CurrentType = CurrentType.Outer)
+			for (; currentType != null; currentType = currentType.Outer)
 			{
-				if (CurrentType is UhtPackage)
+				if (currentType is UhtPackage)
 				{
 					return null;
 				}
-				UhtType? FoundType = FindSymbolChain(CurrentType, Options, ref Lookup);
-				if (FoundType != null)
+				UhtType? foundType = FindSymbolChain(currentType, options, ref lookup);
+				if (foundType != null)
 				{
-					return FoundType;
+					return foundType;
 				}
-				if (Options.HasAnyFlags(UhtFindOptions.NoOuter))
+				if (options.HasAnyFlags(UhtFindOptions.NoOuter))
 				{
 					return null;
 				}
@@ -329,17 +329,17 @@ namespace EpicGames.UHT.Utils
 		/// <summary>
 		/// Lookup the given name 
 		/// </summary>
-		/// <param name="Owner">Matching owner.</param>
-		/// <param name="Options">Options controlling what is search and what is returned.</param>
-		/// <param name="Lookup">Starting lookup location.</param>
+		/// <param name="owner">Matching owner.</param>
+		/// <param name="options">Options controlling what is search and what is returned.</param>
+		/// <param name="lookup">Starting lookup location.</param>
 		/// <returns>Found type or null if not found.</returns>
-		private UhtType? FindSymbolChain(UhtType Owner, UhtFindOptions Options, ref Lookup Lookup)
+		private UhtType? FindSymbolChain(UhtType owner, UhtFindOptions options, ref Lookup lookup)
 		{
-			for (int Index = Lookup.SymbolIndex; Index != 0; Index = this.Symbols[Index].NextIndex)
+			for (int index = lookup.SymbolIndex; index != 0; index = this._symbols[index].NextIndex)
 			{
-				if (IsMatch(Options, Index, Lookup.CasedIndex) && this.Symbols[Index].Type.Outer == Owner)
+				if (IsMatch(options, index, lookup.CasedIndex) && this._symbols[index].Type.Outer == owner)
 				{
-					return this.Symbols[Index].Type;
+					return this._symbols[index].Type;
 				}
 			}
 			return null;
@@ -348,29 +348,29 @@ namespace EpicGames.UHT.Utils
 		/// <summary>
 		/// Add a new type to the given symbol chain
 		/// </summary>
-		/// <param name="Type">Type being added</param>
-		/// <param name="CasedIndex">Cased index</param>
-		/// <param name="SymbolIndex">Symbol index</param>
-		private void AddExisting(UhtType Type, int CasedIndex, int SymbolIndex)
+		/// <param name="type">Type being added</param>
+		/// <param name="casedIndex">Cased index</param>
+		/// <param name="symbolIndex">Symbol index</param>
+		private void AddExisting(UhtType type, int casedIndex, int symbolIndex)
 		{
-			int TypeIndex = Type.TypeIndex;
-			this.Symbols[TypeIndex] = new Symbol { Type = Type, MatchOptions = Type.EngineType.FindOptions(), CasedIndex = CasedIndex, NextIndex = 0, LastIndex = 0 };
-			this.Symbols[this.Symbols[SymbolIndex].LastIndex].NextIndex = TypeIndex;
-			this.Symbols[SymbolIndex].LastIndex = TypeIndex;
+			int typeIndex = type.TypeIndex;
+			this._symbols[typeIndex] = new Symbol { Type = type, MatchOptions = type.EngineType.FindOptions(), CasedIndex = casedIndex, NextIndex = 0, LastIndex = 0 };
+			this._symbols[this._symbols[symbolIndex].LastIndex].NextIndex = typeIndex;
+			this._symbols[symbolIndex].LastIndex = typeIndex;
 		}
 
 		/// <summary>
 		/// Test to see if the given symbol matches the options
 		/// </summary>
-		/// <param name="Options">Options to match</param>
-		/// <param name="SymbolIndex">Symbol index</param>
-		/// <param name="CasedIndex">Case index</param>
+		/// <param name="options">Options to match</param>
+		/// <param name="symbolIndex">Symbol index</param>
+		/// <param name="casedIndex">Case index</param>
 		/// <returns>True if the symbol is a match</returns>
-		private bool IsMatch(UhtFindOptions Options, int SymbolIndex, int CasedIndex)
+		private bool IsMatch(UhtFindOptions options, int symbolIndex, int casedIndex)
 		{
-			return (CasedIndex == 0 || CasedIndex == this.Symbols[SymbolIndex].CasedIndex) && 
-				this.Symbols[SymbolIndex].Type.bVisibleType && 
-				(this.Symbols[SymbolIndex].MatchOptions & Options) != 0;
+			return (casedIndex == 0 || casedIndex == this._symbols[symbolIndex].CasedIndex) &&
+				this._symbols[symbolIndex].Type.VisibleType &&
+				(this._symbols[symbolIndex].MatchOptions & options) != 0;
 		}
 	}
 }
diff --git a/Engine/Source/Programs/Shared/EpicGames.UHT/Utils/UhtTables.cs b/Engine/Source/Programs/Shared/EpicGames.UHT/Utils/UhtTables.cs
--- a/Engine/Source/Programs/Shared/EpicGames.UHT/Utils/UhtTables.cs
+++ b/Engine/Source/Programs/Shared/EpicGames.UHT/Utils/UhtTables.cs
@@ -1,12 +1,12 @@
 // Copyright Epic Games, Inc. All Rights Reserved.
 
-using EpicGames.Core;
-using EpicGames.UHT.Tables;
 using System;
 using System.Collections.Generic;
 using System.Diagnostics.CodeAnalysis;
 using System.IO;
 using System.Reflection;
+using EpicGames.Core;
+using EpicGames.UHT.Tables;
 using UnrealBuildBase;
 
 namespace EpicGames.UHT.Utils
@@ -101,14 +101,14 @@ namespace EpicGames.UHT.Utils
 		/// <summary>
 		/// Name of the table
 		/// </summary>
-		public string TableName { get; set; } = string.Empty;
+		public string TableName { get; set; } = String.Empty;
 
 		/// <summary>
 		/// User facing name of the table
 		/// </summary>
 		public string UserName
 		{
-			get => string.IsNullOrEmpty(this.UserNameInternal) ? this.TableName : this.UserNameInternal;
+			get => String.IsNullOrEmpty(this.UserNameInternal) ? this.TableName : this.UserNameInternal;
 			set => this.UserNameInternal = value;
 		}
 
@@ -125,39 +125,39 @@ namespace EpicGames.UHT.Utils
 		/// <summary>
 		/// Internal version of the user name.  If it hasn't been set, then the table name will be used
 		/// </summary>
-		private string UserNameInternal { get; set; } = string.Empty;
+		private string UserNameInternal { get; set; } = String.Empty;
 
 		/// <summary>
 		/// Merge the lookup table.  Duplicates will be ignored.
 		/// </summary>
-		/// <param name="BaseTable">Base table being merged</param>
-		public virtual void Merge(UhtLookupTableBase BaseTable)
+		/// <param name="baseTable">Base table being merged</param>
+		public virtual void Merge(UhtLookupTableBase baseTable)
 		{
 		}
 
 		/// <summary>
 		/// Given a method name, try to extract the entry name for a table
 		/// </summary>
-		/// <param name="ClassType">Class containing the method</param>
-		/// <param name="MethodInfo">Method information</param>
-		/// <param name="InName">Optional name supplied by the attributes.  If specified, this name will be returned instead of extracted from the method name</param>
-		/// <param name="Suffix">Required suffix</param>
+		/// <param name="classType">Class containing the method</param>
+		/// <param name="methodInfo">Method information</param>
+		/// <param name="inName">Optional name supplied by the attributes.  If specified, this name will be returned instead of extracted from the method name</param>
+		/// <param name="suffix">Required suffix</param>
 		/// <returns>The extracted name or the supplied name</returns>
-		public static string GetSuffixedName(Type ClassType, MethodInfo MethodInfo, string? InName, string Suffix)
+		public static string GetSuffixedName(Type classType, MethodInfo methodInfo, string? inName, string suffix)
 		{
-			string Name = InName ?? string.Empty;
-			if (string.IsNullOrEmpty(Name))
+			string name = inName ?? String.Empty;
+			if (String.IsNullOrEmpty(name))
 			{
-				if (MethodInfo.Name.EndsWith(Suffix, StringComparison.Ordinal))
+				if (methodInfo.Name.EndsWith(suffix, StringComparison.Ordinal))
 				{
-					Name = MethodInfo.Name.Substring(0, MethodInfo.Name.Length - Suffix.Length);
+					name = methodInfo.Name.Substring(0, methodInfo.Name.Length - suffix.Length);
 				}
 				else
 				{
-					throw new UhtIceException($"The '{Suffix}' attribute on the {ClassType.Name}.{MethodInfo.Name} method doesn't have a name specified or the method name doesn't end in '{Suffix}'.");
+					throw new UhtIceException($"The '{suffix}' attribute on the {classType.Name}.{methodInfo.Name} method doesn't have a name specified or the method name doesn't end in '{suffix}'.");
 				}
 			}
-			return Name;
+			return name;
 		}
 	}
 
@@ -171,47 +171,47 @@ namespace EpicGames.UHT.Utils
 		/// <summary>
 		/// Lookup dictionary for the specifiers
 		/// </summary>
-		private readonly Dictionary<StringView, TValue> Lookup;
+		private readonly Dictionary<StringView, TValue> _lookup;
 
 		/// <summary>
 		/// Construct a new table
 		/// </summary>
-		public UhtLookupTable(StringViewComparer Comparer)
+		public UhtLookupTable(StringViewComparer comparer)
 		{
-			this.Lookup = new Dictionary<StringView, TValue>(Comparer);
+			this._lookup = new Dictionary<StringView, TValue>(comparer);
 		}
 
 		/// <summary>
 		/// Add the given value to the lookup table.  It will throw an exception if it is a duplicate.
 		/// </summary>
-		/// <param name="Key">Key to be added</param>
-		/// <param name="Value">Value to be added</param>
-		public UhtLookupTable<TValue> Add(string Key, TValue Value)
+		/// <param name="key">Key to be added</param>
+		/// <param name="value">Value to be added</param>
+		public UhtLookupTable<TValue> Add(string key, TValue value)
 		{
-			this.Lookup.Add(Key, Value);
+			this._lookup.Add(key, value);
 			return this;
 		}
 
 		/// <summary>
 		/// Attempt to fetch the value associated with the key
 		/// </summary>
-		/// <param name="Key">Lookup key</param>
-		/// <param name="Value">Value associated with the key</param>
+		/// <param name="key">Lookup key</param>
+		/// <param name="value">Value associated with the key</param>
 		/// <returns>True if the key was found, false if not</returns>
-		public bool TryGetValue(StringView Key, [MaybeNullWhen(false)] out TValue Value)
+		public bool TryGetValue(StringView key, [MaybeNullWhen(false)] out TValue value)
 		{
-			return this.Lookup.TryGetValue(Key, out Value);
+			return this._lookup.TryGetValue(key, out value);
 		}
 
 		/// <summary>
 		/// Merge the lookup table.  Duplicates will be ignored.
 		/// </summary>
-		/// <param name="BaseTable">Base table being merged</param>
-		public override void Merge(UhtLookupTableBase BaseTable)
+		/// <param name="baseTable">Base table being merged</param>
+		public override void Merge(UhtLookupTableBase baseTable)
 		{
-			foreach (var KVP in ((UhtLookupTable<TValue>)BaseTable).Lookup)
+			foreach (KeyValuePair<StringView, TValue> kvp in ((UhtLookupTable<TValue>)baseTable)._lookup)
 			{
-				this.Lookup.TryAdd(KVP.Key, KVP.Value);
+				this._lookup.TryAdd(kvp.Key, kvp.Value);
 			}
 		}
 	}
@@ -236,48 +236,48 @@ namespace EpicGames.UHT.Utils
 		/// <summary>
 		/// Create a new group of tables
 		/// </summary>
-		/// <param name="Name">The name of the group</param>
-		public UhtLookupTables(string Name)
+		/// <param name="name">The name of the group</param>
+		public UhtLookupTables(string name)
 		{
-			this.Name = Name;
+			this.Name = name;
 		}
 
 		/// <summary>
 		/// Given a table name, return the table.  If not found, a new one will be added with the given name.
 		/// </summary>
-		/// <param name="TableName">The name of the table to return</param>
+		/// <param name="tableName">The name of the table to return</param>
 		/// <returns>The table associated with the name.</returns>
-		public TTable Get(string TableName)
+		public TTable Get(string tableName)
 		{
-			TTable? Table;
-			if (!this.Tables.TryGetValue(TableName, out Table))
+			TTable? table;
+			if (!this.Tables.TryGetValue(tableName, out table))
 			{
-				Table = new TTable();
-				Table.TableName = TableName;
-				this.Tables.Add(TableName, Table);
+				table = new TTable();
+				table.TableName = tableName;
+				this.Tables.Add(tableName, table);
 			}
-			return Table;
+			return table;
 		}
 
 		/// <summary>
 		/// Create a table with the given information.  If the table already exists, it will be initialized with the given data.
 		/// </summary>
-		/// <param name="TableName">The name of the table</param>
-		/// <param name="UserName">The user facing name of the name</param>
-		/// <param name="ParentTableName">The parent table name used to merge table entries</param>
-		/// <param name="bInternal">If true, the table is internal and won't be visible to the user.</param>
+		/// <param name="tableName">The name of the table</param>
+		/// <param name="userName">The user facing name of the name</param>
+		/// <param name="parentTableName">The parent table name used to merge table entries</param>
+		/// <param name="tableIsInternal">If true, the table is internal and won't be visible to the user.</param>
 		/// <returns>The created table</returns>
-		public TTable Create(string TableName, string UserName, string? ParentTableName, bool bInternal = false)
+		public TTable Create(string tableName, string userName, string? parentTableName, bool tableIsInternal = false)
 		{
-			TTable Table = Get(TableName);
-			Table.UserName = UserName;
-			Table.Internal = bInternal;
-			Table.Implemented = true;
-			if (!string.IsNullOrEmpty(ParentTableName))
+			TTable table = Get(tableName);
+			table.UserName = userName;
+			table.Internal = tableIsInternal;
+			table.Implemented = true;
+			if (!String.IsNullOrEmpty(parentTableName))
 			{
-				Table.ParentTable = Get(ParentTableName);
+				table.ParentTable = Get(parentTableName);
 			}
-			return Table;
+			return table;
 		}
 
 		/// <summary>
@@ -287,52 +287,52 @@ namespace EpicGames.UHT.Utils
 		/// <exception cref="UhtIceException">Thrown if there are problems with the tables.</exception>
 		public void Merge()
 		{
-			List<TTable> OrderedList = new List<TTable>(this.Tables.Count);
-			List<TTable> RemainingList = new List<TTable>(this.Tables.Count);
-			HashSet<UhtLookupTableBase> DoneTables = new HashSet<UhtLookupTableBase>();
+			List<TTable> orderedList = new List<TTable>(this.Tables.Count);
+			List<TTable> remainingList = new List<TTable>(this.Tables.Count);
+			HashSet<UhtLookupTableBase> doneTables = new HashSet<UhtLookupTableBase>();
 
 			// Collect all of the tables
-			foreach (KeyValuePair<string, TTable> KVP in this.Tables)
+			foreach (KeyValuePair<string, TTable> kvp in this.Tables)
 			{
-				if (!KVP.Value.Implemented)
+				if (!kvp.Value.Implemented)
 				{
-					throw new UhtIceException($"{this.Name} table '{KVP.Value.TableName}' has been referenced but not implemented");
+					throw new UhtIceException($"{this.Name} table '{kvp.Value.TableName}' has been referenced but not implemented");
 				}
-				RemainingList.Add(KVP.Value);
+				remainingList.Add(kvp.Value);
 			}
 
 			// Perform a topological sort of the tables
-			while (RemainingList.Count != 0)
+			while (remainingList.Count != 0)
 			{
-				bool bAddedOne = false;
-				for (int i = 0; i < RemainingList.Count;)
+				bool addedOne = false;
+				for (int i = 0; i < remainingList.Count;)
 				{
-					TTable Table = RemainingList[i];
-					if (Table.ParentTable == null || DoneTables.Contains(Table.ParentTable))
+					TTable table = remainingList[i];
+					if (table.ParentTable == null || doneTables.Contains(table.ParentTable))
 					{
-						OrderedList.Add(Table);
-						DoneTables.Add(Table);
-						RemainingList[i] = RemainingList[RemainingList.Count - 1];
-						RemainingList.RemoveAt(RemainingList.Count - 1);
-						bAddedOne = true;
+						orderedList.Add(table);
+						doneTables.Add(table);
+						remainingList[i] = remainingList[^1];
+						remainingList.RemoveAt(remainingList.Count - 1);
+						addedOne = true;
 					}
 					else
 					{
 						++i;
 					}
 				}
-				if (!bAddedOne)
+				if (!addedOne)
 				{
 					throw new UhtIceException($"Circular dependency in {this.GetType().Name}.{this.Name} tables detected");
 				}
 			}
 
 			// Merge the tables
-			foreach (TTable Table in OrderedList)
+			foreach (TTable table in orderedList)
 			{
-				if (Table.ParentTable != null)
+				if (table.ParentTable != null)
 				{
-					Table.Merge((TTable)Table.ParentTable);
+					table.Merge((TTable)table.ParentTable);
 				}
 			}
 		}
@@ -364,45 +364,45 @@ namespace EpicGames.UHT.Utils
 		/// <summary>
 		/// Create all of the standard scope tables.
 		/// </summary>
-		/// <param name="Tables">UHT tables</param>
-		public static void InitStandardTables(UhtTables Tables)
+		/// <param name="tables">UHT tables</param>
+		public static void InitStandardTables(UhtTables tables)
 		{
-			CreateTables(Tables, UhtTableNames.Default, "Default", EUhtCreateTablesFlags.Specifiers | EUhtCreateTablesFlags.Keyword, null, true);
-			CreateTables(Tables, UhtTableNames.Global, "Global", EUhtCreateTablesFlags.Keyword, UhtTableNames.Default, false);
-			CreateTables(Tables, UhtTableNames.PropertyArgument, "Argument/Return", EUhtCreateTablesFlags.Specifiers, UhtTableNames.Default, false);
-			CreateTables(Tables, UhtTableNames.PropertyMember, "Member", EUhtCreateTablesFlags.Specifiers, UhtTableNames.Default, false);
-			CreateTables(Tables, UhtTableNames.Object, "Object", EUhtCreateTablesFlags.Specifiers | EUhtCreateTablesFlags.Keyword, UhtTableNames.Default, true);
-			CreateTables(Tables, UhtTableNames.Field, "Field", EUhtCreateTablesFlags.Specifiers | EUhtCreateTablesFlags.Keyword, UhtTableNames.Object, true);
-			CreateTables(Tables, UhtTableNames.Enum, "Enum", EUhtCreateTablesFlags.Specifiers | EUhtCreateTablesFlags.Keyword, UhtTableNames.Field, false);
-			CreateTables(Tables, UhtTableNames.Struct, "Struct", EUhtCreateTablesFlags.Specifiers | EUhtCreateTablesFlags.Keyword, UhtTableNames.Field, true);
-			CreateTables(Tables, UhtTableNames.ClassBase, "ClassBase", EUhtCreateTablesFlags.Specifiers | EUhtCreateTablesFlags.Keyword, UhtTableNames.Struct, true);
-			CreateTables(Tables, UhtTableNames.Class, "Class", EUhtCreateTablesFlags.Specifiers | EUhtCreateTablesFlags.Keyword, UhtTableNames.ClassBase, false);
-			CreateTables(Tables, UhtTableNames.Interface, "Interface", EUhtCreateTablesFlags.Specifiers | EUhtCreateTablesFlags.Keyword, UhtTableNames.ClassBase, false);
-			CreateTables(Tables, UhtTableNames.NativeInterface, "IInterface", EUhtCreateTablesFlags.Specifiers | EUhtCreateTablesFlags.Keyword, UhtTableNames.ClassBase, false);
-			CreateTables(Tables, UhtTableNames.Function, "Function", EUhtCreateTablesFlags.Specifiers | EUhtCreateTablesFlags.Keyword, UhtTableNames.Struct, false);
-			CreateTables(Tables, UhtTableNames.ScriptStruct, "Struct", EUhtCreateTablesFlags.Specifiers | EUhtCreateTablesFlags.Keyword, UhtTableNames.Struct, false);
+			CreateTables(tables, UhtTableNames.Default, "Default", EUhtCreateTablesFlags.Specifiers | EUhtCreateTablesFlags.Keyword, null, true);
+			CreateTables(tables, UhtTableNames.Global, "Global", EUhtCreateTablesFlags.Keyword, UhtTableNames.Default, false);
+			CreateTables(tables, UhtTableNames.PropertyArgument, "Argument/Return", EUhtCreateTablesFlags.Specifiers, UhtTableNames.Default, false);
+			CreateTables(tables, UhtTableNames.PropertyMember, "Member", EUhtCreateTablesFlags.Specifiers, UhtTableNames.Default, false);
+			CreateTables(tables, UhtTableNames.Object, "Object", EUhtCreateTablesFlags.Specifiers | EUhtCreateTablesFlags.Keyword, UhtTableNames.Default, true);
+			CreateTables(tables, UhtTableNames.Field, "Field", EUhtCreateTablesFlags.Specifiers | EUhtCreateTablesFlags.Keyword, UhtTableNames.Object, true);
+			CreateTables(tables, UhtTableNames.Enum, "Enum", EUhtCreateTablesFlags.Specifiers | EUhtCreateTablesFlags.Keyword, UhtTableNames.Field, false);
+			CreateTables(tables, UhtTableNames.Struct, "Struct", EUhtCreateTablesFlags.Specifiers | EUhtCreateTablesFlags.Keyword, UhtTableNames.Field, true);
+			CreateTables(tables, UhtTableNames.ClassBase, "ClassBase", EUhtCreateTablesFlags.Specifiers | EUhtCreateTablesFlags.Keyword, UhtTableNames.Struct, true);
+			CreateTables(tables, UhtTableNames.Class, "Class", EUhtCreateTablesFlags.Specifiers | EUhtCreateTablesFlags.Keyword, UhtTableNames.ClassBase, false);
+			CreateTables(tables, UhtTableNames.Interface, "Interface", EUhtCreateTablesFlags.Specifiers | EUhtCreateTablesFlags.Keyword, UhtTableNames.ClassBase, false);
+			CreateTables(tables, UhtTableNames.NativeInterface, "IInterface", EUhtCreateTablesFlags.Specifiers | EUhtCreateTablesFlags.Keyword, UhtTableNames.ClassBase, false);
+			CreateTables(tables, UhtTableNames.Function, "Function", EUhtCreateTablesFlags.Specifiers | EUhtCreateTablesFlags.Keyword, UhtTableNames.Struct, false);
+			CreateTables(tables, UhtTableNames.ScriptStruct, "Struct", EUhtCreateTablesFlags.Specifiers | EUhtCreateTablesFlags.Keyword, UhtTableNames.Struct, false);
 		}
 
 		/// <summary>
 		/// Creates a series of tables given the supplied setup
 		/// </summary>
-		/// <param name="Tables">UHT tables</param>
-		/// <param name="TableName">Name of the table.</param>
-		/// <param name="TableUserName">Name presented to the user via error messages.</param>
-		/// <param name="CreateTables">Types of tables to be created.</param>
-		/// <param name="ParentTableName">Name of the parent table or null for none.</param>
-		/// <param name="bInternal">If true, this table will not be included in any error messages.</param>
-		public static void CreateTables(UhtTables Tables, string TableName, string TableUserName, 
-			EUhtCreateTablesFlags CreateTables, string? ParentTableName, bool bInternal = false)
+		/// <param name="tables">UHT tables</param>
+		/// <param name="tableName">Name of the table.</param>
+		/// <param name="tableUserName">Name presented to the user via error messages.</param>
+		/// <param name="createTables">Types of tables to be created.</param>
+		/// <param name="parentTableName">Name of the parent table or null for none.</param>
+		/// <param name="tableIsInternal">If true, this table will not be included in any error messages.</param>
+		public static void CreateTables(UhtTables tables, string tableName, string tableUserName,
+			EUhtCreateTablesFlags createTables, string? parentTableName, bool tableIsInternal = false)
 		{
-			if (CreateTables.HasFlag(EUhtCreateTablesFlags.Keyword))
+			if (createTables.HasFlag(EUhtCreateTablesFlags.Keyword))
 			{
-				Tables.KeywordTables.Create(TableName, TableUserName, ParentTableName, bInternal);
+				tables.KeywordTables.Create(tableName, tableUserName, parentTableName, tableIsInternal);
 			}
-			if (CreateTables.HasFlag(EUhtCreateTablesFlags.Specifiers))
+			if (createTables.HasFlag(EUhtCreateTablesFlags.Specifiers))
 			{
-				Tables.SpecifierTables.Create(TableName, TableUserName, ParentTableName, bInternal);
-				Tables.SpecifierValidatorTables.Create(TableName, TableUserName, ParentTableName, bInternal);
+				tables.SpecifierTables.Create(tableName, tableUserName, parentTableName, tableIsInternal);
+				tables.SpecifierValidatorTables.Create(tableName, tableUserName, parentTableName, tableIsInternal);
 			}
 		}
 	}
@@ -418,7 +418,7 @@ namespace EpicGames.UHT.Utils
 		/// If specified, this method will be invoked once during the scan for attributes.
 		/// It can be used to perform some one time initialization.
 		/// </summary>
-		public string InitMethod { get; set; } = string.Empty;
+		public string InitMethod { get; set; } = String.Empty;
 	}
 
 	/// <summary>
@@ -502,12 +502,12 @@ namespace EpicGames.UHT.Utils
 		/// <summary>
 		/// Add a collection of plugin assembly file paths
 		/// </summary>
-		/// <param name="PluginAssembliesFilePaths">Collection of plugins to load</param>
-		public void AddPlugins(IEnumerable<string> PluginAssembliesFilePaths)
+		/// <param name="pluginAssembliesFilePaths">Collection of plugins to load</param>
+		public void AddPlugins(IEnumerable<string> pluginAssembliesFilePaths)
 		{
-			foreach (string AssemblyFilePath in PluginAssembliesFilePaths)
+			foreach (string assemblyFilePath in pluginAssembliesFilePaths)
 			{
-				CheckForAttributes(LoadAssembly(AssemblyFilePath));
+				CheckForAttributes(LoadAssembly(assemblyFilePath));
 			}
 			PerformPostInitialization();
 		}
@@ -515,20 +515,20 @@ namespace EpicGames.UHT.Utils
 		/// <summary>
 		/// Check to see if the assembly is a UHT plugin
 		/// </summary>
-		/// <param name="AssemblyFilePath">Path to the assembly file</param>
+		/// <param name="assemblyFilePath">Path to the assembly file</param>
 		/// <returns></returns>
-		public static bool IsUhtPlugin(string AssemblyFilePath)
+		public static bool IsUhtPlugin(string assemblyFilePath)
 		{
-			Assembly? Assembly = LoadAssembly(AssemblyFilePath);       
-			if (Assembly != null)
+			Assembly? assembly = LoadAssembly(assemblyFilePath);
+			if (assembly != null)
 			{
-				foreach (Type Type in Assembly.SafeGetLoadedTypes())
+				foreach (Type type in assembly.SafeGetLoadedTypes())
 				{
-					if (Type.IsClass)
+					if (type.IsClass)
 					{
-						foreach (Attribute ClassAttribute in Type.GetCustomAttributes(false))
+						foreach (Attribute classAttribute in type.GetCustomAttributes(false))
 						{
-							if (ClassAttribute is UnrealHeaderToolAttribute || ClassAttribute is UhtEngineClassAttribute)
+							if (classAttribute is UnrealHeaderToolAttribute || classAttribute is UhtEngineClassAttribute)
 							{
 								return true;
 							}
@@ -542,33 +542,33 @@ namespace EpicGames.UHT.Utils
 		/// <summary>
 		/// Load the given assembly from the file path
 		/// </summary>
-		/// <param name="AssemblyFilePath">Path to the file</param>
+		/// <param name="assemblyFilePath">Path to the file</param>
 		/// <returns>Assembly if it is already loaded or could be loaded</returns>
-		public static Assembly? LoadAssembly(string AssemblyFilePath)
+		public static Assembly? LoadAssembly(string assemblyFilePath)
 		{
-			Assembly? Assembly = FindAssemblyByName(Path.GetFileNameWithoutExtension(AssemblyFilePath));
-			if (Assembly == null)
+			Assembly? assembly = FindAssemblyByName(Path.GetFileNameWithoutExtension(assemblyFilePath));
+			if (assembly == null)
 			{
-				Assembly = Assembly.LoadFile(AssemblyFilePath);
+				assembly = Assembly.LoadFile(assemblyFilePath);
 			}
-			return Assembly;
+			return assembly;
 		}
 
 		/// <summary>
 		/// Locate the assembly by name
 		/// </summary>
-		/// <param name="Name">Name of the assembly</param>
+		/// <param name="name">Name of the assembly</param>
 		/// <returns>Assembly or null</returns>
-		private static Assembly? FindAssemblyByName(string? Name)
+		private static Assembly? FindAssemblyByName(string? name)
 		{
-			if (Name != null)
+			if (name != null)
 			{
-				foreach (Assembly Assembly in AppDomain.CurrentDomain.GetAssemblies())
+				foreach (Assembly assembly in AppDomain.CurrentDomain.GetAssemblies())
 				{
-					AssemblyName AssemblyName = Assembly.GetName();
-					if (AssemblyName.Name != null && String.Equals(AssemblyName.Name, Name, StringComparison.OrdinalIgnoreCase))
+					AssemblyName assemblyName = assembly.GetName();
+					if (assemblyName.Name != null && String.Equals(assemblyName.Name, name, StringComparison.OrdinalIgnoreCase))
 					{
-						return Assembly;
+						return assembly;
 					}
 				}
 			}
@@ -578,14 +578,14 @@ namespace EpicGames.UHT.Utils
 		/// <summary>
 		/// Check for UHT attributes on all types in the given assembly
 		/// </summary>
-		/// <param name="Assembly">Assembly to scan</param>
-		private void CheckForAttributes(Assembly? Assembly)
+		/// <param name="assembly">Assembly to scan</param>
+		private void CheckForAttributes(Assembly? assembly)
 		{
-			if (Assembly != null)
+			if (assembly != null)
 			{
-				foreach (Type Type in Assembly.SafeGetLoadedTypes())
+				foreach (Type type in assembly.SafeGetLoadedTypes())
 				{
-					CheckForAttributes(Type);
+					CheckForAttributes(type);
 				}
 			}
 		}
@@ -593,22 +593,22 @@ namespace EpicGames.UHT.Utils
 		/// <summary>
 		/// For the given type, look for any table related attributes 
 		/// </summary>
-		/// <param name="Type">Type in question</param>
-		private void CheckForAttributes(Type Type)
+		/// <param name="type">Type in question</param>
+		private void CheckForAttributes(Type type)
 		{
-			if (Type.IsClass)
+			if (type.IsClass)
 			{
 
 				// Loop through the attributes
-				foreach (Attribute ClassAttribute in Type.GetCustomAttributes(false))
+				foreach (Attribute classAttribute in type.GetCustomAttributes(false))
 				{
-					if (ClassAttribute is UnrealHeaderToolAttribute ParserAttribute)
+					if (classAttribute is UnrealHeaderToolAttribute parserAttribute)
 					{
-						HandleUnrealHeaderToolAttribute(Type, ParserAttribute);
+						HandleUnrealHeaderToolAttribute(type, parserAttribute);
 					}
-					else if (ClassAttribute is UhtEngineClassAttribute EngineClassAttribute)
+					else if (classAttribute is UhtEngineClassAttribute engineClassAttribute)
 					{
-						this.EngineClassTable.OnEngineClassAttribute(EngineClassAttribute);
+						this.EngineClassTable.OnEngineClassAttribute(engineClassAttribute);
 					}
 				}
 			}
@@ -625,56 +625,56 @@ namespace EpicGames.UHT.Utils
 			_ = this.StructDefaultValueTable.Default;
 		}
 
-		private void HandleUnrealHeaderToolAttribute(Type Type, UnrealHeaderToolAttribute ParserAttribute)
+		private void HandleUnrealHeaderToolAttribute(Type type, UnrealHeaderToolAttribute parserAttribute)
 		{
-			if (!string.IsNullOrEmpty(ParserAttribute.InitMethod))
+			if (!String.IsNullOrEmpty(parserAttribute.InitMethod))
 			{
-				MethodInfo? InitInfo = Type.GetMethod(ParserAttribute.InitMethod, BindingFlags.Static | BindingFlags.Public | BindingFlags.NonPublic);
-				if (InitInfo == null)
+				MethodInfo? initInfo = type.GetMethod(parserAttribute.InitMethod, BindingFlags.Static | BindingFlags.Public | BindingFlags.NonPublic);
+				if (initInfo == null)
 				{
-					throw new Exception($"Unable to find UnrealHeaderTool attribute InitMethod {ParserAttribute.InitMethod}");
+					throw new Exception($"Unable to find UnrealHeaderTool attribute InitMethod {parserAttribute.InitMethod}");
 				}
-				InitInfo.Invoke(null, Array.Empty<object>());
+				initInfo.Invoke(null, Array.Empty<object>());
 			}
 
 			// Scan the methods for things we are interested in
-			foreach (MethodInfo MethodInfo in Type.GetMethods(BindingFlags.Static | BindingFlags.Public | BindingFlags.NonPublic))
+			foreach (MethodInfo methodInfo in type.GetMethods(BindingFlags.Static | BindingFlags.Public | BindingFlags.NonPublic))
 			{
 
 				// Scan for attributes we care about
-				foreach (Attribute MethodAttribute in MethodInfo.GetCustomAttributes())
+				foreach (Attribute methodAttribute in methodInfo.GetCustomAttributes())
 				{
-					if (MethodAttribute is UhtSpecifierAttribute SpecifierAttribute)
+					if (methodAttribute is UhtSpecifierAttribute specifierAttribute)
 					{
-						this.SpecifierTables.OnSpecifierAttribute(Type, MethodInfo, SpecifierAttribute);
+						this.SpecifierTables.OnSpecifierAttribute(type, methodInfo, specifierAttribute);
 					}
-					else if (MethodAttribute is UhtSpecifierValidatorAttribute SpecifierValidatorAttribute)
+					else if (methodAttribute is UhtSpecifierValidatorAttribute specifierValidatorAttribute)
 					{
-						this.SpecifierValidatorTables.OnSpecifierValidatorAttribute(Type, MethodInfo, SpecifierValidatorAttribute);
+						this.SpecifierValidatorTables.OnSpecifierValidatorAttribute(type, methodInfo, specifierValidatorAttribute);
 					}
-					else if (MethodAttribute is UhtKeywordCatchAllAttribute KeywordCatchAllAttribute)
+					else if (methodAttribute is UhtKeywordCatchAllAttribute keywordCatchAllAttribute)
 					{
-						this.KeywordTables.OnKeywordCatchAllAttribute(Type, MethodInfo, KeywordCatchAllAttribute);
+						this.KeywordTables.OnKeywordCatchAllAttribute(type, methodInfo, keywordCatchAllAttribute);
 					}
-					else if (MethodAttribute is UhtKeywordAttribute KeywordAttribute)
+					else if (methodAttribute is UhtKeywordAttribute keywordAttribute)
 					{
-						this.KeywordTables.OnKeywordAttribute(Type, MethodInfo, KeywordAttribute);
+						this.KeywordTables.OnKeywordAttribute(type, methodInfo, keywordAttribute);
 					}
-					else if (MethodAttribute is UhtPropertyTypeAttribute PropertyTypeAttribute)
+					else if (methodAttribute is UhtPropertyTypeAttribute propertyTypeAttribute)
 					{
-						this.PropertyTypeTable.OnPropertyTypeAttribute(MethodInfo, PropertyTypeAttribute);
+						this.PropertyTypeTable.OnPropertyTypeAttribute(methodInfo, propertyTypeAttribute);
 					}
-					else if (MethodAttribute is UhtStructDefaultValueAttribute StructDefaultValueAttribute)
+					else if (methodAttribute is UhtStructDefaultValueAttribute structDefaultValueAttribute)
 					{
-						this.StructDefaultValueTable.OnStructDefaultValueAttribute(MethodInfo, StructDefaultValueAttribute);
+						this.StructDefaultValueTable.OnStructDefaultValueAttribute(methodInfo, structDefaultValueAttribute);
 					}
-					else if (MethodAttribute is UhtExporterAttribute ExporterAttribute)
+					else if (methodAttribute is UhtExporterAttribute exporterAttribute)
 					{
-						this.ExporterTable.OnExporterAttribute(Type, MethodInfo, ExporterAttribute);
+						this.ExporterTable.OnExporterAttribute(type, methodInfo, exporterAttribute);
 					}
-					else if (MethodAttribute is UhtLocTextDefaultValueAttribute LocTextDefaultValueAttribute)
+					else if (methodAttribute is UhtLocTextDefaultValueAttribute locTextDefaultValueAttribute)
 					{
-						this.LocTextDefaultValueTable.OnLocTextDefaultValueAttribute(MethodInfo, LocTextDefaultValueAttribute);
+						this.LocTextDefaultValueTable.OnLocTextDefaultValueAttribute(methodInfo, locTextDefaultValueAttribute);
 					}
 				}
 			}
diff --git a/Engine/Source/Programs/Shared/EpicGames.UHT/Utils/UhtUtilities.cs b/Engine/Source/Programs/Shared/EpicGames.UHT/Utils/UhtUtilities.cs
--- a/Engine/Source/Programs/Shared/EpicGames.UHT/Utils/UhtUtilities.cs
+++ b/Engine/Source/Programs/Shared/EpicGames.UHT/Utils/UhtUtilities.cs
@@ -1,10 +1,10 @@
 // Copyright Epic Games, Inc. All Rights Reserved.
 
-using EpicGames.Core;
 using System;
 using System.Collections.Generic;
 using System.Linq;
 using System.Text;
+using EpicGames.Core;
 
 namespace EpicGames.UHT.Utils
 {
@@ -18,31 +18,31 @@ namespace EpicGames.UHT.Utils
 		/// <summary>
 		/// Add the given value if it isn't already contained within the list
 		/// </summary>
-		/// <param name="Container">Destination container</param>
-		/// <param name="Value">Value to be added</param>
+		/// <param name="container">Destination container</param>
+		/// <param name="value">Value to be added</param>
 		/// <returns>True if the value was added, false if it was already present</returns>
-		public static bool AddUnique(this List<string> Container, string Value)
+		public static bool AddUnique(this List<string> container, string value)
 		{
-			if (Container.Contains(Value, StringComparer.OrdinalIgnoreCase))
+			if (container.Contains(value, StringComparer.OrdinalIgnoreCase))
 			{
 				return false;
 			}
-			Container.Add(Value);
+			container.Add(value);
 			return true;
 		}
 
 		/// <summary>
 		/// Add the given values if they aren't already contained within the list
 		/// </summary>
-		/// <param name="Container">Destination container</param>
-		/// <param name="Values">Values to be added</param>
-		public static void AddUniqueRange(this List<string> Container, IEnumerable<StringView>? Values)
+		/// <param name="container">Destination container</param>
+		/// <param name="values">Values to be added</param>
+		public static void AddUniqueRange(this List<string> container, IEnumerable<StringView>? values)
 		{
-			if (Values != null)
+			if (values != null)
 			{
-				foreach (StringView Value in Values)
+				foreach (StringView value in values)
 				{
-					AddUnique(Container, Value.ToString());
+					AddUnique(container, value.ToString());
 				}
 			}
 		}
@@ -50,19 +50,19 @@ namespace EpicGames.UHT.Utils
 		/// <summary>
 		/// Remove the given value but swap the last entry into the eliminated slot
 		/// </summary>
-		/// <param name="Container">Container being modified</param>
-		/// <param name="Value">Value to be removed</param>
+		/// <param name="container">Container being modified</param>
+		/// <param name="value">Value to be removed</param>
 		/// <returns>True if the value was removed, false if not</returns>
-		public static bool RemoveSwap(this List<string> Container, string Value)
+		public static bool RemoveSwap(this List<string> container, string value)
 		{
-			int Index = Container.FindIndex(n => Value.Equals(n, StringComparison.OrdinalIgnoreCase));
-			if (Index >= 0)
+			int index = container.FindIndex(n => value.Equals(n, StringComparison.OrdinalIgnoreCase));
+			if (index >= 0)
 			{
-				if (Index != Container.Count - 1)
+				if (index != container.Count - 1)
 				{
-					Container[Index] = Container[Container.Count - 1];
+					container[index] = container[^1];
 				}
-				Container.RemoveAt(Container.Count - 1);
+				container.RemoveAt(container.Count - 1);
 				return true;
 			}
 			return false;
@@ -71,15 +71,15 @@ namespace EpicGames.UHT.Utils
 		/// <summary>
 		/// Remove a range of values from a container using swapping
 		/// </summary>
-		/// <param name="Container">Container to be modified</param>
-		/// <param name="Values">List of values to be removed</param>
-		public static void RemoveSwapRange(this List<string> Container, IEnumerable<StringView>? Values)
+		/// <param name="container">Container to be modified</param>
+		/// <param name="values">List of values to be removed</param>
+		public static void RemoveSwapRange(this List<string> container, IEnumerable<StringView>? values)
 		{
-			if (Values != null)
+			if (values != null)
 			{
-				foreach (StringView Value in Values)
+				foreach (StringView value in values)
 				{
-					RemoveSwap(Container, Value.ToString());
+					RemoveSwap(container, value.ToString());
 				}
 			}
 		}
@@ -105,7 +105,7 @@ namespace EpicGames.UHT.Utils
 		/// <summary>
 		/// The name contained the "DEPRECATED" text which has been removed from the engine name
 		/// </summary>
-		public bool bIsDeprecated { get; set; }
+		public bool IsDeprecated { get; set; }
 	}
 
 	/// <summary>
@@ -117,82 +117,82 @@ namespace EpicGames.UHT.Utils
 		/// <summary>
 		/// Given a collection of names, return a string containing the text of those names concatenated.
 		/// </summary>
-		/// <param name="TypeNames">Collect of names to be merged</param>
-		/// <param name="AndOr">Text used to separate the names</param>
-		/// <param name="bQuote">If true, add quotes around the names</param>
+		/// <param name="typeNames">Collect of names to be merged</param>
+		/// <param name="andOr">Text used to separate the names</param>
+		/// <param name="quote">If true, add quotes around the names</param>
 		/// <returns>Merged names</returns>
-		public static string MergeTypeNames(IEnumerable<string> TypeNames, string AndOr, bool bQuote = false)
+		public static string MergeTypeNames(IEnumerable<string> typeNames, string andOr, bool quote = false)
 		{
-			List<string> Local = new List<string>(TypeNames);
+			List<string> local = new List<string>(typeNames);
 
-			if (Local.Count == 0)
+			if (local.Count == 0)
 			{
 				return "";
 			}
 
-			Local.Sort();
+			local.Sort();
 
-			StringBuilder Builder = new StringBuilder();
-			for (int Index = 0; Index < Local.Count; ++Index)
+			StringBuilder builder = new StringBuilder();
+			for (int index = 0; index < local.Count; ++index)
 			{
-				if (Index != 0)
+				if (index != 0)
 				{
-					Builder.Append(", ");
-					if (Index == Local.Count - 1)
+					builder.Append(", ");
+					if (index == local.Count - 1)
 					{
-						Builder.Append(AndOr);
-						Builder.Append(' ');
+						builder.Append(andOr);
+						builder.Append(' ');
 					}
 				}
-				if (bQuote)
+				if (quote)
 				{
-					Builder.Append('\'');
-					Builder.Append(Local[Index]);
-					Builder.Append('\'');
+					builder.Append('\'');
+					builder.Append(local[index]);
+					builder.Append('\'');
 				}
 				else
 				{
-					Builder.Append(Local[Index]);
+					builder.Append(local[index]);
 				}
 			}
-			return Builder.ToString();
+			return builder.ToString();
 		}
 
 		/// <summary>
 		/// Split the given source name into the engine name parts
 		/// </summary>
-		/// <param name="SourceName">Source name</param>
+		/// <param name="sourceName">Source name</param>
 		/// <returns>Resulting engine name parts</returns>
-		public static UhtEngineNameParts GetEngineNameParts(StringView SourceName)
+		public static UhtEngineNameParts GetEngineNameParts(StringView sourceName)
 		{
-			if (SourceName.Span.Length == 0)
+			if (sourceName.Span.Length == 0)
 			{
-				return new UhtEngineNameParts { Prefix = new StringView(string.Empty), EngineName = new StringView(string.Empty), bIsDeprecated = false };
+				return new UhtEngineNameParts { Prefix = new StringView(String.Empty), EngineName = new StringView(String.Empty), IsDeprecated = false };
 			}
 
-			switch (SourceName.Span[0])
+			switch (sourceName.Span[0])
 			{
 				case 'I':
 				case 'A':
 				case 'U':
 					// If it is a class prefix, check for deprecated class prefix also
-					if (SourceName.Span.Length > 12 && SourceName.Span.Slice(1).StartsWith("DEPRECATED_"))
+					if (sourceName.Span.Length > 12 && sourceName.Span.Slice(1).StartsWith("DEPRECATED_"))
 					{
-						return new UhtEngineNameParts { Prefix = new StringView(SourceName, 0, 12), EngineName = new StringView(SourceName, 12), bIsDeprecated = true };
+						return new UhtEngineNameParts { Prefix = new StringView(sourceName, 0, 12), EngineName = new StringView(sourceName, 12), IsDeprecated = true };
 					}
 					else
 					{
-						return new UhtEngineNameParts { Prefix = new StringView(SourceName, 0, 1), EngineName = new StringView(SourceName, 1), bIsDeprecated = false };
+						return new UhtEngineNameParts { Prefix = new StringView(sourceName, 0, 1), EngineName = new StringView(sourceName, 1), IsDeprecated = false };
 					}
 
 				case 'F':
 				case 'T':
 					// Struct prefixes are also fine.
-					return new UhtEngineNameParts { Prefix = new StringView(SourceName, 0, 1), EngineName = new StringView(SourceName, 1), bIsDeprecated = false };
+					return new UhtEngineNameParts { Prefix = new StringView(sourceName, 0, 1), EngineName = new StringView(sourceName, 1), IsDeprecated = false };
 
 				default:
 					// If it's not a class or struct prefix, it's invalid
-					return new UhtEngineNameParts { Prefix = new StringView(string.Empty), EngineName = new StringView(SourceName), bIsDeprecated = false };
+					return new UhtEngineNameParts { Prefix = new StringView(String.Empty), EngineName = new StringView(sourceName), IsDeprecated = false };
 			}
 		}
 	}
@@ -207,17 +207,17 @@ namespace EpicGames.UHT.Utils
 		/// <summary>
 		/// When only a string view has been appended, this references that StringView
 		/// </summary>
-		private StringView StringViewInternal = new StringView();
+		private StringView _stringViewInternal = new StringView();
 
 		/// <summary>
 		/// Represents more complex data being appended
 		/// </summary>
-		private StringBuilder? StringBuilderInternal = null;
+		private StringBuilder? _stringBuilderInternal = null;
 
 		/// <summary>
 		/// Set to true when the builder has switched to a StringBuilder (NOTE: This can probably be removed)
 		/// </summary>
-		private bool bUseStringBuilder = false;
+		private bool _useStringBuilder = false;
 
 		/// <summary>
 		/// The length of the appended data
@@ -226,13 +226,13 @@ namespace EpicGames.UHT.Utils
 		{
 			get
 			{
-				if (this.bUseStringBuilder && this.StringBuilderInternal != null)
+				if (this._useStringBuilder && this._stringBuilderInternal != null)
 				{
-					return this.StringBuilderInternal.Length;
+					return this._stringBuilderInternal.Length;
 				}
 				else
 				{
-					return this.StringViewInternal.Span.Length;
+					return this._stringViewInternal.Span.Length;
 				}
 			}
 		}
@@ -243,7 +243,7 @@ namespace EpicGames.UHT.Utils
 		/// <returns>Contents of the builder</returns>
 		public StringView ToStringView()
 		{
-			return this.bUseStringBuilder ? new StringView(this.StringBuilderInternal!.ToString()) : this.StringViewInternal;
+			return this._useStringBuilder ? new StringView(this._stringBuilderInternal!.ToString()) : this._stringViewInternal;
 		}
 
 		/// <summary>
@@ -252,28 +252,28 @@ namespace EpicGames.UHT.Utils
 		/// <returns>Contents of the builder</returns>
 		public override string ToString()
 		{
-			return this.bUseStringBuilder ? this.StringBuilderInternal!.ToString() : this.StringViewInternal.ToString();
+			return this._useStringBuilder ? this._stringBuilderInternal!.ToString() : this._stringViewInternal.ToString();
 		}
 
 		/// <summary>
 		/// Append a StringView
 		/// </summary>
-		/// <param name="Text">Text to be appended</param>
+		/// <param name="text">Text to be appended</param>
 		/// <returns>The string builder</returns>
-		public StringViewBuilder Append(StringView Text)
+		public StringViewBuilder Append(StringView text)
 		{
-			if (this.bUseStringBuilder)
+			if (this._useStringBuilder)
 			{
-				this.StringBuilderInternal!.Append(Text.Span);
+				this._stringBuilderInternal!.Append(text.Span);
 			}
-			else if (this.StringViewInternal.Span.Length > 0)
+			else if (this._stringViewInternal.Span.Length > 0)
 			{
 				SwitchToStringBuilder();
-				this.StringBuilderInternal!.Append(Text.Span);
+				this._stringBuilderInternal!.Append(text.Span);
 			}
 			else
 			{
-				this.StringViewInternal = Text;
+				this._stringViewInternal = text;
 			}
 			return this;
 		}
@@ -281,12 +281,12 @@ namespace EpicGames.UHT.Utils
 		/// <summary>
 		/// Append a character
 		/// </summary>
-		/// <param name="C">Character to be appended</param>
+		/// <param name="c">Character to be appended</param>
 		/// <returns>The string builder</returns>
-		public StringViewBuilder Append(char C)
+		public StringViewBuilder Append(char c)
 		{
 			SwitchToStringBuilder();
-			this.StringBuilderInternal!.Append(C);
+			this._stringBuilderInternal!.Append(c);
 			return this;
 		}
 
@@ -295,14 +295,14 @@ namespace EpicGames.UHT.Utils
 		/// </summary>
 		private void SwitchToStringBuilder()
 		{
-			if (!this.bUseStringBuilder)
+			if (!this._useStringBuilder)
 			{
-				if (this.StringBuilderInternal == null)
+				if (this._stringBuilderInternal == null)
 				{
-					this.StringBuilderInternal = new StringBuilder();
+					this._stringBuilderInternal = new StringBuilder();
 				}
-				this.bUseStringBuilder = true;
-				this.StringBuilderInternal.Append(this.StringViewInternal.Span);
+				this._useStringBuilder = true;
+				this._stringBuilderInternal.Append(this._stringViewInternal.Span);
 			}
 		}
 	}
