More coding style cleanups

diff --git a/Engine/Source/Programs/Shared/EpicGames.UHT/Exporters/CodeGen/UhtHeaderCodeGeneratorHFile.cs b/Engine/Source/Programs/Shared/EpicGames.UHT/Exporters/CodeGen/UhtHeaderCodeGeneratorHFile.cs
--- a/Engine/Source/Programs/Shared/EpicGames.UHT/Exporters/CodeGen/UhtHeaderCodeGeneratorHFile.cs
+++ b/Engine/Source/Programs/Shared/EpicGames.UHT/Exporters/CodeGen/UhtHeaderCodeGeneratorHFile.cs
@@ -35,8 +35,8 @@ namespace EpicGames.UHT.Exporters.CodeGen
 		public void Generate(IUhtExportFactory factory)
 		{
 			ref UhtCodeGenerator.HeaderInfo headerInfo = ref this.HeaderInfos[this.HeaderFile.HeaderFileTypeIndex];
-			using (BorrowStringBuilder borrower = new BorrowStringBuilder(StringBuilderCache.Big))
 			{
+				using BorrowStringBuilder borrower = new(StringBuilderCache.Big);
 				StringBuilder builder = borrower.StringBuilder;
 
 				builder.Append(HeaderCopyright);
@@ -139,7 +139,7 @@ namespace EpicGames.UHT.Exporters.CodeGen
 
 			if (scriptStruct.ScriptStructFlags.HasAnyFlags(EStructFlags.Native))
 			{
-				using (UhtMacroCreator macro = new UhtMacroCreator(builder, this, scriptStruct, GeneratedBodyMacroSuffix))
+				using (UhtMacroCreator macro = new(builder, this, scriptStruct, GeneratedBodyMacroSuffix))
 				{
 					builder.Append("\tfriend struct Z_Construct_UScriptStruct_").Append(scriptStruct.SourceName).Append("_Statics; \\\r\n");
 					builder.Append('\t');
@@ -290,7 +290,7 @@ namespace EpicGames.UHT.Exporters.CodeGen
 
 		private StringBuilder AppendDelegate(StringBuilder builder, UhtFunction function)
 		{
-			using (UhtMacroCreator macro = new UhtMacroCreator(builder, this, function, DelegateMacroSuffix))
+			using (UhtMacroCreator macro = new(builder, this, function, DelegateMacroSuffix))
 			{
 				int tabs = 0;
 				string strippedFunctionName = function.StrippedFunctionName;
@@ -323,7 +323,7 @@ namespace EpicGames.UHT.Exporters.CodeGen
 			AppendSparseDeclarations(builder, classObj);
 
 			// Collect the functions in reversed order
-			List<UhtFunction> reversedFunctions = new List<UhtFunction>(classObj.Functions);
+			List<UhtFunction> reversedFunctions = new(classObj.Functions);
 			reversedFunctions.Reverse();
 
 			// Check to see if we have any RPC functions for the editor
@@ -340,7 +340,7 @@ namespace EpicGames.UHT.Exporters.CodeGen
 			AppendPropertyAccessors(builder, classObj);
 
 			// Collect the callback function and sort by name to make the order stable
-			List<UhtFunction> callbackFunctions = new List<UhtFunction>(classObj.Functions.Where(x => x.FunctionFlags.HasAnyFlags(EFunctionFlags.Event) && x.SuperFunction == null));
+			List<UhtFunction> callbackFunctions = new(classObj.Functions.Where(x => x.FunctionFlags.HasAnyFlags(EFunctionFlags.Event) && x.SuperFunction == null));
 			callbackFunctions.Sort((x, y) => StringComparerUE.OrdinalIgnoreCase.Compare(x.EngineName, y.EngineName));
 
 			// Generate the callback parameter structures
@@ -361,12 +361,12 @@ namespace EpicGames.UHT.Exporters.CodeGen
 				AppendStandardConstructors(builder, classObj, api);
 				AppendEnhancedConstructors(builder, classObj, api);
 
-				using (UhtMacroCreator macro = new UhtMacroCreator(builder, this, classObj, GeneratedUInterfaceBodyMacroSuffix))
+				using (UhtMacroCreator macro = new(builder, this, classObj, GeneratedUInterfaceBodyMacroSuffix))
 				{
 					AppendCommonGeneratedBody(builder, classObj, api);
 				}
 
-				using (UhtMacroCreator macro = new UhtMacroCreator(builder, this, classObj, GeneratedBodyLegacyMacroSuffix))
+				using (UhtMacroCreator macro = new(builder, this, classObj, GeneratedBodyLegacyMacroSuffix))
 				{
 					builder.Append('\t');
 					AppendGeneratedMacroDeprecationWarning(builder, "GENERATED_UINTERFACE_BODY");
@@ -376,7 +376,7 @@ namespace EpicGames.UHT.Exporters.CodeGen
 					builder.Append('\t').Append(EnableDeprecationWarnings).Append(" \\\r\n");
 				}
 
-				using (UhtMacroCreator macro = new UhtMacroCreator(builder, this, classObj, GeneratedBodyMacroSuffix))
+				using (UhtMacroCreator macro = new(builder, this, classObj, GeneratedBodyMacroSuffix))
 				{
 					builder.Append('\t').Append(DisableDeprecationWarnings).Append(" \\\r\n");
 					builder.Append('\t').AppendMacroName(this, classObj, GeneratedUInterfaceBodyMacroSuffix).Append(" \\\r\n");
@@ -386,24 +386,24 @@ namespace EpicGames.UHT.Exporters.CodeGen
 					builder.Append('\t').Append(EnableDeprecationWarnings).Append(" \\\r\n");
 				}
 
-				using (UhtMacroCreator macro = new UhtMacroCreator(builder, this, classObj, InClassIInterfaceNoPureDeclsMacroSuffix))
+				using (UhtMacroCreator macro = new(builder, this, classObj, InClassIInterfaceNoPureDeclsMacroSuffix))
 				{
 					AppendInClassIInterface(builder, classObj, callbackFunctions, api);
 				}
 
-				using (UhtMacroCreator macro = new UhtMacroCreator(builder, this, classObj, InClassIInterfaceMacroSuffix))
+				using (UhtMacroCreator macro = new(builder, this, classObj, InClassIInterfaceMacroSuffix))
 				{
 					AppendInClassIInterface(builder, classObj, callbackFunctions, api);
 				}
 			}
 			else
 			{
-				using (UhtMacroCreator macro = new UhtMacroCreator(builder, this, classObj, InClassNoPureDeclsMacroSuffix))
+				using (UhtMacroCreator macro = new(builder, this, classObj, InClassNoPureDeclsMacroSuffix))
 				{
 					AppendClassGeneratedBody(builder, classObj, api);
 				}
 
-				using (UhtMacroCreator macro = new UhtMacroCreator(builder, this, classObj, InClassMacroSuffix))
+				using (UhtMacroCreator macro = new(builder, this, classObj, InClassMacroSuffix))
 				{
 					AppendClassGeneratedBody(builder, classObj, api);
 				}
@@ -413,7 +413,7 @@ namespace EpicGames.UHT.Exporters.CodeGen
 				AppendFieldNotify(builder, classObj);
 			}
 
-			using (UhtMacroCreator macro = new UhtMacroCreator(builder, this, classObj.PrologLineNumber, PrologMacroSuffix))
+			using (UhtMacroCreator macro = new(builder, this, classObj.PrologLineNumber, PrologMacroSuffix))
 			{
 				if (callbackFunctions.Count > 0)
 				{
@@ -450,22 +450,18 @@ namespace EpicGames.UHT.Exporters.CodeGen
 			}
 
 			// Scan the children to see what we have
-			bool hasProperties;
-			bool hasFunctions;
-			bool hasEditorFields;
-			bool allEditorFields;
-			GetFieldNotifyStats(classObj, out hasProperties, out hasFunctions, out hasEditorFields, out allEditorFields);
+			GetFieldNotifyStats(classObj, out bool hasProperties, out bool hasFunctions, out bool hasEditorFields, out bool allEditorFields);
 
 			// If we only have editor fields or no editor fields, then we only emit one block
 			if (hasEditorFields)
 			{
 				builder.Append("#if WITH_EDITORONLY_DATA\r\n");
-				using (UhtMacroCreator macro = new UhtMacroCreator(builder, this, classObj, FieldNotifyMacroSuffix))
+				using (UhtMacroCreator macro = new(builder, this, classObj, FieldNotifyMacroSuffix))
 				{
 					AppendFieldNotify(builder, classObj, hasProperties, hasFunctions, hasEditorFields, allEditorFields, true);
 				}
 				builder.Append("#else //WITH_EDITORONLY_DATA\r\n");
-				using (UhtMacroCreator macro = new UhtMacroCreator(builder, this, classObj, FieldNotifyMacroSuffix))
+				using (UhtMacroCreator macro = new(builder, this, classObj, FieldNotifyMacroSuffix))
 				{
 					AppendFieldNotify(builder, classObj, hasProperties, hasFunctions, hasEditorFields, allEditorFields, false);
 				}
@@ -473,10 +469,8 @@ namespace EpicGames.UHT.Exporters.CodeGen
 			}
 			else
 			{
-				using (UhtMacroCreator macro = new UhtMacroCreator(builder, this, classObj, FieldNotifyMacroSuffix))
-				{
-					AppendFieldNotify(builder, classObj, hasProperties, hasFunctions, hasEditorFields, allEditorFields, false);
-				}
+				using UhtMacroCreator macro = new(builder, this, classObj, FieldNotifyMacroSuffix);
+				AppendFieldNotify(builder, classObj, hasProperties, hasFunctions, hasEditorFields, allEditorFields, false);
 			}
 			return builder;
 		}
@@ -518,7 +512,7 @@ namespace EpicGames.UHT.Exporters.CodeGen
 		private StringBuilder AppendSparseDeclarations(StringBuilder builder, UhtClass classObj)
 		{
 			// Format the sparse data
-			using (UhtMacroCreator macro = new UhtMacroCreator(builder, this, classObj, SparseDataMacroSuffix))
+			using (UhtMacroCreator macro = new(builder, this, classObj, SparseDataMacroSuffix))
 			{
 				string[]? sparseDataTypes = classObj.MetaData.GetStringArray(UhtNames.SparseClassDataTypes);
 				if (sparseDataTypes != null)
@@ -554,7 +548,7 @@ namespace EpicGames.UHT.Exporters.CodeGen
 									string cleanPropertyName = propertyName;
 									if (sparseProperty is UhtBoolProperty && propertyName.StartsWith("b", StringComparison.Ordinal))
 									{
-										cleanPropertyName = propertyName.Substring(1);
+										cleanPropertyName = propertyName[1..];
 									}
 
 									bool getByRef = sparseProperty.MetaData.ContainsKey(UhtNames.GetByRef);
@@ -597,13 +591,13 @@ namespace EpicGames.UHT.Exporters.CodeGen
 		{
 			builder.AppendBeginEditorOnlyGuard(editorOnly);
 
-			using (UhtMacroCreator macro = new UhtMacroCreator(builder, this, classObj, editorOnly ? EditorOnlyRpcWrappersMacroSuffix : RpcWrappersMacroSuffix))
+			using (UhtMacroCreator macro = new(builder, this, classObj, editorOnly ? EditorOnlyRpcWrappersMacroSuffix : RpcWrappersMacroSuffix))
 			{
 				AppendAutogeneratedBlueprintFunctionDeclarations(builder, reversedFunctions, editorOnly);
 				AppendRpcWrappers(builder, reversedFunctions, editorOnly);
 			}
 
-			using (UhtMacroCreator macro = new UhtMacroCreator(builder, this, classObj, editorOnly ? EditorOnlyRpcWrappersNoPureDeclsMacroSuffix : RpcWrappersNoPureDeclsMacroSuffix))
+			using (UhtMacroCreator macro = new(builder, this, classObj, editorOnly ? EditorOnlyRpcWrappersNoPureDeclsMacroSuffix : RpcWrappersNoPureDeclsMacroSuffix))
 			{
 				if (classObj.GeneratedCodeVersion <= EGeneratedCodeVersion.V1)
 				{
@@ -615,11 +609,11 @@ namespace EpicGames.UHT.Exporters.CodeGen
 			if (editorOnly)
 			{
 				builder.Append("#else\r\n");
-				using (UhtMacroCreator macro = new UhtMacroCreator(builder, this, classObj, editorOnly ? EditorOnlyRpcWrappersMacroSuffix : RpcWrappersMacroSuffix))
+				using (UhtMacroCreator macro = new(builder, this, classObj, editorOnly ? EditorOnlyRpcWrappersMacroSuffix : RpcWrappersMacroSuffix))
 				{
 				}
 
-				using (UhtMacroCreator macro = new UhtMacroCreator(builder, this, classObj, editorOnly ? EditorOnlyRpcWrappersNoPureDeclsMacroSuffix : RpcWrappersNoPureDeclsMacroSuffix))
+				using (UhtMacroCreator macro = new(builder, this, classObj, editorOnly ? EditorOnlyRpcWrappersNoPureDeclsMacroSuffix : RpcWrappersNoPureDeclsMacroSuffix))
 				{
 				}
 				builder.AppendEndEditorOnlyGuard();
@@ -629,7 +623,7 @@ namespace EpicGames.UHT.Exporters.CodeGen
 
 		private StringBuilder AppendPropertyAccessors(StringBuilder builder, UhtClass classObj)
 		{
-			using (UhtMacroCreator macro = new UhtMacroCreator(builder, this, classObj, AccessorsMacroSuffix))
+			using (UhtMacroCreator macro = new(builder, this, classObj, AccessorsMacroSuffix))
 			{
 				foreach (UhtType type in classObj.Children)
 				{
@@ -724,12 +718,10 @@ namespace EpicGames.UHT.Exporters.CodeGen
 		{
 			if (callbackFunctions.Count > 0)
 			{
-				using (UhtMacroCreator macro = new UhtMacroCreator(builder, this, classObj, EventParamsMacroSuffix))
+				using UhtMacroCreator macro = new(builder, this, classObj, EventParamsMacroSuffix);
+				foreach (UhtFunction function in callbackFunctions)
 				{
-					foreach (UhtFunction function in callbackFunctions)
-					{
-						AppendEventParameter(builder, function, function.StrippedFunctionName, UhtPropertyTextType.EventParameterMember, true, 1, " \\\r\n");
-					}
+					AppendEventParameter(builder, function, function.StrippedFunctionName, UhtPropertyTextType.EventParameterMember, true, 1, " \\\r\n");
 				}
 			}
 			return builder;
@@ -739,16 +731,14 @@ namespace EpicGames.UHT.Exporters.CodeGen
 		{
 			if (callbackFunctions.Count > 0)
 			{
-				using (UhtMacroCreator macro = new UhtMacroCreator(builder, this, classObj, CallbackWrappersMacroSuffix))
+				using UhtMacroCreator macro = new(builder, this, classObj, CallbackWrappersMacroSuffix);
+				foreach (UhtFunction function in callbackFunctions)
 				{
-					foreach (UhtFunction function in callbackFunctions)
+					if (!function.FunctionFlags.HasAnyFlags(EFunctionFlags.NetResponse) &&
+						function.EngineName != function.MarshalAndCallName)
 					{
-						if (!function.FunctionFlags.HasAnyFlags(EFunctionFlags.NetResponse) &&
-							function.EngineName != function.MarshalAndCallName)
-						{
-							AppendNativeFunctionHeader(builder, function, UhtPropertyTextType.EventFunctionArgOrRetVal, true, null, null, UhtFunctionExportFlags.None, " \\\r\n");
-							builder.Append(" \\\r\n");
-						}
+						AppendNativeFunctionHeader(builder, function, UhtPropertyTextType.EventFunctionArgOrRetVal, true, null, null, UhtFunctionExportFlags.None, " \\\r\n");
+						builder.Append(" \\\r\n");
 					}
 				}
 			}
@@ -767,7 +757,7 @@ namespace EpicGames.UHT.Exporters.CodeGen
 				if (classObj.EnclosingDefine.Length != 0)
 				{
 					builder.Append("#else\r\n");
-					using (UhtMacroCreator macro = new UhtMacroCreator(builder, this, classObj, ArchiveSerializerMacroSuffix))
+					using (UhtMacroCreator macro = new(builder, this, classObj, ArchiveSerializerMacroSuffix))
 					{
 					}
 					builder.Append("#endif\r\n");
@@ -778,9 +768,9 @@ namespace EpicGames.UHT.Exporters.CodeGen
 
 		private StringBuilder AppendSerializerFunction(StringBuilder builder, UhtClass classObj, string api, string declare)
 		{
-			using (UhtMacroCreator macro = new UhtMacroCreator(builder, this, classObj, ArchiveSerializerMacroSuffix))
+			using (UhtMacroCreator macro = new(builder, this, classObj, ArchiveSerializerMacroSuffix))
 			{
-				builder.Append('\t').Append(declare).Append('(').Append(classObj.SourceName).Append(", ").Append(api.Substring(0, api.Length - 1)).Append(") \\\r\n");
+				builder.Append('\t').Append(declare).Append('(').Append(classObj.SourceName).Append(", ").Append(api[0..^1]).Append(") \\\r\n");
 			}
 			return builder;
 		}
@@ -794,7 +784,7 @@ namespace EpicGames.UHT.Exporters.CodeGen
 		/// <returns>Output builder</returns>
 		private StringBuilder AppendStandardConstructors(StringBuilder builder, UhtClass classObj, string api)
 		{
-			using (UhtMacroCreator macro = new UhtMacroCreator(builder, this, classObj, StandardConstructorsMacroSuffix))
+			using (UhtMacroCreator macro = new(builder, this, classObj, StandardConstructorsMacroSuffix))
 			{
 				if (!classObj.ClassExportFlags.HasAnyFlags(UhtClassExportFlags.HasCustomConstructor))
 				{
@@ -830,7 +820,7 @@ namespace EpicGames.UHT.Exporters.CodeGen
 		/// <returns>Output builder</returns>
 		private StringBuilder AppendEnhancedConstructors(StringBuilder builder, UhtClass classObj, string api)
 		{
-			using (UhtMacroCreator macro = new UhtMacroCreator(builder, this, classObj, EnchancedConstructorsMacroSuffix))
+			using (UhtMacroCreator macro = new(builder, this, classObj, EnchancedConstructorsMacroSuffix))
 			{
 				AppendConstructorDefinition(builder, classObj, api);
 				AppendVTableHelperCtorAndCaller(builder, classObj, api);
@@ -850,7 +840,7 @@ namespace EpicGames.UHT.Exporters.CodeGen
 		{
 			if (!classObj.ClassExportFlags.HasAnyFlags(UhtClassExportFlags.HasCustomVTableHelperConstructor))
 			{
-				builder.Append("\tDECLARE_VTABLE_PTR_HELPER_CTOR(").Append(api.Substring(0, api.Length - 1)).Append(", ").Append(classObj.SourceName).Append("); \\\r\n");
+				builder.Append("\tDECLARE_VTABLE_PTR_HELPER_CTOR(").Append(api[0..^1]).Append(", ").Append(classObj.SourceName).Append("); \\\r\n");
 			}
 			builder.Append("\tDEFINE_VTABLE_PTR_HELPER_CTOR_CALLER(").Append(classObj.SourceName).Append("); \\\r\n");
 			return builder;
@@ -1037,7 +1027,7 @@ namespace EpicGames.UHT.Exporters.CodeGen
 			{
 				builder.Append("CASTCLASS_None");
 			}
-			builder.Append(", TEXT(\"").Append(this.Package.SourceName).Append("\"), ").Append(api.Substring(0, api.Length - 1)).Append(") \\\r\n");
+			builder.Append(", TEXT(\"").Append(this.Package.SourceName).Append("\"), ").Append(api[0..^1]).Append(") \\\r\n");
 
 			builder.Append("\tDECLARE_SERIALIZER(").Append(classObj.SourceName).Append(") \\\r\n");
 
@@ -1283,7 +1273,7 @@ namespace EpicGames.UHT.Exporters.CodeGen
 		private StringBuilder AppendGeneratedBodyMacroBlock(StringBuilder builder, UhtClass classObj, UhtClass bodyClassObj, bool isLegacy, bool hasEditorRpc, bool hasCallbacks, string? deprecatedMacroName)
 		{
 			bool isInterface = classObj.ClassFlags.HasAnyFlags(EClassFlags.Interface);
-			using (UhtMacroCreator macro = new UhtMacroCreator(builder, this, bodyClassObj, isLegacy ? GeneratedBodyLegacyMacroSuffix : GeneratedBodyMacroSuffix))
+			using (UhtMacroCreator macro = new(builder, this, bodyClassObj, isLegacy ? GeneratedBodyLegacyMacroSuffix : GeneratedBodyMacroSuffix))
 			{
 				if (deprecatedMacroName != null)
 				{
@@ -1375,7 +1365,7 @@ namespace EpicGames.UHT.Exporters.CodeGen
 
 		private static StringView GenerateEnumPrefix(UhtEnum enumObj)
 		{
-			StringView prefix = new StringView();
+			StringView prefix = new();
 			if (enumObj.EnumValues.Count > 0)
 			{
 				prefix = enumObj.EnumValues[0].Name;
diff --git a/Engine/Source/Programs/Shared/EpicGames.UHT/Exporters/CodeGen/UhtPackageCodeGeneratorCppFile.cs b/Engine/Source/Programs/Shared/EpicGames.UHT/Exporters/CodeGen/UhtPackageCodeGeneratorCppFile.cs
--- a/Engine/Source/Programs/Shared/EpicGames.UHT/Exporters/CodeGen/UhtPackageCodeGeneratorCppFile.cs
+++ b/Engine/Source/Programs/Shared/EpicGames.UHT/Exporters/CodeGen/UhtPackageCodeGeneratorCppFile.cs
@@ -28,14 +28,14 @@ namespace EpicGames.UHT.Exporters.CodeGen
 		/// <param name="packageSortedHeaders">Sorted list of headers by name of all headers in the package</param>
 		public void Generate(IUhtExportFactory factory, List<UhtHeaderFile> packageSortedHeaders)
 		{
-			using (BorrowStringBuilder borrower = new BorrowStringBuilder(StringBuilderCache.Big))
 			{
+				using BorrowStringBuilder borrower = new(StringBuilderCache.Big);
 				const string MetaDataParamsName = "Package_MetaDataParams";
 				StringBuilder builder = borrower.StringBuilder;
 
 				// Collect information from all of the headers
-				List<UhtField> singletons = new List<UhtField>();
-				StringBuilder declarations = new StringBuilder();
+				List<UhtField> singletons = new();
+				StringBuilder declarations = new();
 				uint bodiesHash = 0;
 				foreach (UhtHeaderFile headerFile in packageSortedHeaders)
 				{
diff --git a/Engine/Source/Programs/Shared/EpicGames.UHT/Exporters/CodeGen/UhtPackageCodeGeneratorHFile.cs b/Engine/Source/Programs/Shared/EpicGames.UHT/Exporters/CodeGen/UhtPackageCodeGeneratorHFile.cs
--- a/Engine/Source/Programs/Shared/EpicGames.UHT/Exporters/CodeGen/UhtPackageCodeGeneratorHFile.cs
+++ b/Engine/Source/Programs/Shared/EpicGames.UHT/Exporters/CodeGen/UhtPackageCodeGeneratorHFile.cs
@@ -28,8 +28,8 @@ namespace EpicGames.UHT.Exporters.CodeGen
 		/// <param name="packageSortedHeaders">Sorted list of headers by name of all headers in the package</param>
 		public void Generate(IUhtExportFactory factory, List<UhtHeaderFile> packageSortedHeaders)
 		{
-			using (BorrowStringBuilder borrower = new BorrowStringBuilder(StringBuilderCache.Big))
 			{
+				using BorrowStringBuilder borrower = new(StringBuilderCache.Big);
 				StringBuilder builder = borrower.StringBuilder;
 
 				builder.Append(HeaderCopyright);
@@ -37,7 +37,7 @@ namespace EpicGames.UHT.Exporters.CodeGen
 				builder.Append("\r\n");
 				builder.Append("\r\n");
 
-				List<UhtHeaderFile> headerFiles = new List<UhtHeaderFile>(this.Package.Children.Count * 2);
+				List<UhtHeaderFile> headerFiles = new(this.Package.Children.Count * 2);
 				headerFiles.AddRange(packageSortedHeaders);
 
 				foreach (UhtHeaderFile headerFile in this.Package.Children)
@@ -48,7 +48,7 @@ namespace EpicGames.UHT.Exporters.CodeGen
 					}
 				}
 
-				List<UhtHeaderFile> sortedHeaderFiles = new List<UhtHeaderFile>(headerFiles.Distinct());
+				List<UhtHeaderFile> sortedHeaderFiles = new(headerFiles.Distinct());
 				sortedHeaderFiles.Sort((x, y) => StringComparerUE.OrdinalIgnoreCase.Compare(x.FilePath, y.FilePath));
 
 				foreach (UhtHeaderFile headerFile in sortedHeaderFiles)
diff --git a/Engine/Source/Programs/Shared/EpicGames.UHT/Exporters/CodeGen/UhtRigVMStringBuilderExtensions.cs b/Engine/Source/Programs/Shared/EpicGames.UHT/Exporters/CodeGen/UhtRigVMStringBuilderExtensions.cs
--- a/Engine/Source/Programs/Shared/EpicGames.UHT/Exporters/CodeGen/UhtRigVMStringBuilderExtensions.cs
+++ b/Engine/Source/Programs/Shared/EpicGames.UHT/Exporters/CodeGen/UhtRigVMStringBuilderExtensions.cs
@@ -41,7 +41,7 @@ namespace EpicGames.UHT.Exporters.CodeGen
 			ReadOnlySpan<char> parameterTypeSpan = paramterType.AsSpan();
 			if (parameterTypeSpan.EndsWith("&"))
 			{
-				parameterTypeSpan = parameterTypeSpan.Slice(0, parameterTypeSpan.Length - 1);
+				parameterTypeSpan = parameterTypeSpan[0..^1];
 			}
 			builder.Append(parameterTypeSpan);
 			if (parameterTypeSpan.Length > 0 && (parameterTypeSpan[0] == 'T' || parameterTypeSpan[0] == 'F'))
@@ -72,7 +72,7 @@ namespace EpicGames.UHT.Exporters.CodeGen
 			ReadOnlySpan<char> parameterTypeSpan = paramterType.AsSpan();
 			if (parameterTypeSpan.EndsWith("&"))
 			{
-				parameterTypeSpan = parameterTypeSpan.Slice(0, parameterTypeSpan.Length - 1);
+				parameterTypeSpan = parameterTypeSpan[0..^1];
 			}
 			return builder.Append(parameterTypeSpan);
 		}
diff --git a/Engine/Source/Programs/Shared/EpicGames.UHT/Exporters/Json/UhtJsonExporter.cs b/Engine/Source/Programs/Shared/EpicGames.UHT/Exporters/Json/UhtJsonExporter.cs
--- a/Engine/Source/Programs/Shared/EpicGames.UHT/Exporters/Json/UhtJsonExporter.cs
+++ b/Engine/Source/Programs/Shared/EpicGames.UHT/Exporters/Json/UhtJsonExporter.cs
@@ -32,7 +32,7 @@ namespace EpicGames.UHT.Exporters.Json
 		private void Export()
 		{
 			// Generate the files for the packages
-			List<Task?> generatedPackages = new List<Task?>(this.Session.PackageTypeCount);
+			List<Task?> generatedPackages = new(this.Session.PackageTypeCount);
 			foreach (UhtPackage package in this.Session.Packages)
 			{
 				UHTManifest.Module module = package.Module;
@@ -40,13 +40,13 @@ namespace EpicGames.UHT.Exporters.Json
 					(IUhtExportFactory factory) =>
 					{
 						string jsonPath = factory.MakePath(package, ".json");
-						JsonSerializerOptions options = new JsonSerializerOptions { WriteIndented = true, DefaultIgnoreCondition = JsonIgnoreCondition.WhenWritingNull };
+						JsonSerializerOptions options = new() { WriteIndented = true, DefaultIgnoreCondition = JsonIgnoreCondition.WhenWritingNull };
 						factory.CommitOutput(jsonPath, JsonSerializer.Serialize(package, options));
 					}));
 			}
 
 			// Wait for all the packages to complete
-			List<Task> packageTasks = new List<Task>(this.Session.PackageTypeCount);
+			List<Task> packageTasks = new(this.Session.PackageTypeCount);
 			foreach (Task? output in generatedPackages)
 			{
 				if (output != null)
diff --git a/Engine/Source/Programs/Shared/EpicGames.UHT/Exporters/Stats/UhtStatsExporter.cs b/Engine/Source/Programs/Shared/EpicGames.UHT/Exporters/Stats/UhtStatsExporter.cs
--- a/Engine/Source/Programs/Shared/EpicGames.UHT/Exporters/Stats/UhtStatsExporter.cs
+++ b/Engine/Source/Programs/Shared/EpicGames.UHT/Exporters/Stats/UhtStatsExporter.cs
@@ -16,7 +16,7 @@ namespace EpicGames.UHT.Exporters.Stats
 		[SuppressMessage("CodeQuality", "IDE0051:Remove unused private members", Justification = "Attribute accessed method")]
 		private static void StatsExporter(IUhtExportFactory factory)
 		{
-			SortedDictionary<string, int> countByType = new SortedDictionary<string, int>();
+			SortedDictionary<string, int> countByType = new();
 			foreach (UhtType type in factory.Session.Packages)
 			{
 				Collect(countByType, type);
diff --git a/Engine/Source/Programs/Shared/EpicGames.UHT/Parsers/UhtClassParser.cs b/Engine/Source/Programs/Shared/EpicGames.UHT/Parsers/UhtClassParser.cs
--- a/Engine/Source/Programs/Shared/EpicGames.UHT/Parsers/UhtClassParser.cs
+++ b/Engine/Source/Programs/Shared/EpicGames.UHT/Parsers/UhtClassParser.cs
@@ -211,7 +211,7 @@ namespace EpicGames.UHT.Parsers
 			// If this class has new show categories, then merge them
 			if (this.ShowCategories.Count != 0)
 			{
-				StringBuilder subCategoryPath = new StringBuilder();
+				StringBuilder subCategoryPath = new();
 				foreach (string value in this.ShowCategories)
 				{
 
@@ -343,7 +343,7 @@ namespace EpicGames.UHT.Parsers
 
 		private static List<string> GetStringListMetaData(UhtType? type, string key)
 		{
-			List<string> outStrings = new List<string>();
+			List<string> outStrings = new();
 			AppendStringListMetaData(type, key, outStrings);
 			return outStrings;
 		}
@@ -403,15 +403,15 @@ namespace EpicGames.UHT.Parsers
 
 		private static UhtParseResult ParseUClass(UhtParsingScope parentScope, ref UhtToken token)
 		{
-			UhtClassParser classObj = new UhtClassParser(parentScope.ScopeType, token.InputLine);
-			using (UhtParsingScope topScope = new UhtParsingScope(parentScope, classObj, parentScope.Session.GetKeywordTable(UhtTableNames.Class), UhtAccessSpecifier.Private))
+			UhtClassParser classObj = new(parentScope.ScopeType, token.InputLine);
 			{
+				using UhtParsingScope topScope = new(parentScope, classObj, parentScope.Session.GetKeywordTable(UhtTableNames.Class), UhtAccessSpecifier.Private);
 				const string ScopeName = "class";
 
-				using (UhtMessageContext tokenContext = new UhtMessageContext(ScopeName))
 				{
+					using UhtMessageContext tokenContext = new(ScopeName);
 					// Parse the specifiers
-					UhtSpecifierContext specifierContext = new UhtSpecifierContext(topScope, topScope.TokenReader, classObj.MetaData);
+					UhtSpecifierContext specifierContext = new(topScope, topScope.TokenReader, classObj.MetaData);
 					UhtSpecifierParser specifiers = topScope.HeaderParser.GetCachedSpecifierParser(specifierContext, ScopeName, parentScope.Session.GetSpecifierTable(UhtTableNames.Class));
 					specifiers.ParseSpecifiers();
 					classObj.PrologLineNumber = topScope.TokenReader.InputLine;
@@ -425,8 +425,7 @@ namespace EpicGames.UHT.Parsers
 					topScope.TokenReader.SkipAlignasAndDeprecatedMacroIfNecessary();
 
 					// Read the class name and possible API macro name
-					UhtToken apiMacroToken;
-					topScope.TokenReader.TryOptionalAPIMacro(out apiMacroToken);
+					topScope.TokenReader.TryOptionalAPIMacro(out UhtToken apiMacroToken);
 					classObj.SourceName = topScope.TokenReader.GetIdentifier().Value.ToString();
 
 					// Update the context for better error messages
@@ -446,9 +445,7 @@ namespace EpicGames.UHT.Parsers
 					topScope.TokenReader.Optional("final");
 
 					// Parse the inheritance
-					UhtToken superIdentifier;
-					List<UhtToken[]>? baseIdentifiers;
-					UhtParserHelpers.ParseInheritance(topScope.TokenReader, topScope.Session.Config!, out superIdentifier, out baseIdentifiers);
+					UhtParserHelpers.ParseInheritance(topScope.TokenReader, topScope.Session.Config!, out UhtToken superIdentifier, out List<UhtToken[]>? baseIdentifiers);
 					classObj.SuperIdentifier = superIdentifier;
 					classObj.BaseIdentifiers = baseIdentifiers;
 
diff --git a/Engine/Source/Programs/Shared/EpicGames.UHT/Parsers/UhtEnumParser.cs b/Engine/Source/Programs/Shared/EpicGames.UHT/Parsers/UhtEnumParser.cs
--- a/Engine/Source/Programs/Shared/EpicGames.UHT/Parsers/UhtEnumParser.cs
+++ b/Engine/Source/Programs/Shared/EpicGames.UHT/Parsers/UhtEnumParser.cs
@@ -36,16 +36,16 @@ namespace EpicGames.UHT.Parsers
 
 		private static UhtParseResult ParseUEnum(UhtParsingScope parentScope, UhtToken keywordToken)
 		{
-			UhtEnumParser enumObject = new UhtEnumParser(parentScope.ScopeType, keywordToken.InputLine);
-			using (UhtParsingScope topScope = new UhtParsingScope(parentScope, enumObject, parentScope.Session.GetKeywordTable(UhtTableNames.Enum), UhtAccessSpecifier.Public))
+			UhtEnumParser enumObject = new(parentScope.ScopeType, keywordToken.InputLine);
 			{
+				using UhtParsingScope topScope = new(parentScope, enumObject, parentScope.Session.GetKeywordTable(UhtTableNames.Enum), UhtAccessSpecifier.Public);
 				const string ScopeName = "UENUM";
 
-				using (UhtMessageContext tokenContext = new UhtMessageContext(ScopeName))
 				{
+					using UhtMessageContext tokenContext = new(ScopeName);
 
 					// Parse the specifiers
-					UhtSpecifierContext specifierContext = new UhtSpecifierContext(topScope, topScope.TokenReader, enumObject.MetaData);
+					UhtSpecifierContext specifierContext = new(topScope, topScope.TokenReader, enumObject.MetaData);
 					UhtSpecifierParser specifiers = topScope.HeaderParser.GetCachedSpecifierParser(specifierContext, ScopeName, parentScope.Session.GetSpecifierTable(UhtTableNames.Enum));
 					specifiers.ParseSpecifiers();
 
@@ -88,8 +88,7 @@ namespace EpicGames.UHT.Parsers
 						{
 							UhtToken enumType = topScope.TokenReader.GetIdentifier("enumeration base");
 
-							UhtEnumUnderlyingType underlyingType;
-							if (!System.Enum.TryParse<UhtEnumUnderlyingType>(enumType.Value.ToString(), out underlyingType) || underlyingType == UhtEnumUnderlyingType.Unspecified)
+							if (!System.Enum.TryParse<UhtEnumUnderlyingType>(enumType.Value.ToString(), out UhtEnumUnderlyingType underlyingType) || underlyingType == UhtEnumUnderlyingType.Unspecified)
 							{
 								topScope.TokenReader.LogError(enumType.InputLine, $"Unsupported enum class base type '{enumType.Value}'");
 							}
@@ -176,8 +175,7 @@ namespace EpicGames.UHT.Parsers
 						// Try to read an optional explicit enum value specification
 						if (topScope.TokenReader.TryOptional('='))
 						{
-							long scatchValue;
-							bool parsedValue = topScope.TokenReader.TryOptionalConstLong(out scatchValue);
+							bool parsedValue = topScope.TokenReader.TryOptionalConstLong(out long scatchValue);
 							if (parsedValue)
 							{
 								currentEnumValue = scatchValue;
diff --git a/Engine/Source/Programs/Shared/EpicGames.UHT/Parsers/UhtFunctionParser.cs b/Engine/Source/Programs/Shared/EpicGames.UHT/Parsers/UhtFunctionParser.cs
--- a/Engine/Source/Programs/Shared/EpicGames.UHT/Parsers/UhtFunctionParser.cs
+++ b/Engine/Source/Programs/Shared/EpicGames.UHT/Parsers/UhtFunctionParser.cs
@@ -34,8 +34,7 @@ namespace EpicGames.UHT.Parsers
 			this._alreadyLeft = 0;
 			this._bUseNumber = false;
 
-			string? foundString;
-			if (this._metaData.TryGetValue(UhtNames.AdvancedDisplay, out foundString))
+			if (this._metaData.TryGetValue(UhtNames.AdvancedDisplay, out string? foundString))
 			{
 				this._parameterNames = foundString.ToString().Split(',', StringSplitOptions.RemoveEmptyEntries);
 				for (int index = 0, endIndex = this._parameterNames.Length; index < endIndex; ++index)
@@ -130,7 +129,7 @@ namespace EpicGames.UHT.Parsers
 								try
 								{
 									// All tokens MUST be consumed from the reader
-									StringBuilder builder = new StringBuilder();
+									StringBuilder builder = new();
 									IUhtTokenReader defaultValueReader = UhtTokenReplayReader.GetThreadInstance(property, this.HeaderFile.Data.Memory, property.DefaultValueTokens.ToArray(), UhtTokenType.EndOfDefault);
 									parsed = property.SanitizeDefaultValue(defaultValueReader, builder) && defaultValueReader.IsEOF;
 									if (parsed)
@@ -145,7 +144,7 @@ namespace EpicGames.UHT.Parsers
 
 								if (!parsed)
 								{
-									StringView defaultValueText = new StringView(this.HeaderFile.Data, property.DefaultValueTokens.First().InputStartPos,
+									StringView defaultValueText = new(this.HeaderFile.Data, property.DefaultValueTokens.First().InputStartPos,
 										property.DefaultValueTokens.Last().InputEndPos - property.DefaultValueTokens.First().InputStartPos);
 									property.LogError($"C++ Default parameter not parsed: {property.SourceName} '{defaultValueText}'");
 								}
@@ -213,21 +212,21 @@ namespace EpicGames.UHT.Parsers
 
 		private static UhtParseResult ParseUDelegate(UhtParsingScope parentScope, UhtToken token, bool hasSpecifiers)
 		{
-			UhtFunctionParser function = new UhtFunctionParser(parentScope.ScopeType, token.InputLine);
+			UhtFunctionParser function = new(parentScope.ScopeType, token.InputLine);
 
-			using (UhtParsingScope topScope = new UhtParsingScope(parentScope, function, parentScope.Session.GetKeywordTable(UhtTableNames.Function), UhtAccessSpecifier.Public))
 			{
+				using UhtParsingScope topScope = new(parentScope, function, parentScope.Session.GetKeywordTable(UhtTableNames.Function), UhtAccessSpecifier.Public);
 				const string ScopeName = "delegate declaration";
 
-				using (UhtMessageContext tokenContext = new UhtMessageContext(ScopeName))
 				{
+					using UhtMessageContext tokenContext = new(ScopeName);
 					topScope.AddModuleRelativePathToMetaData();
 
-					UhtSpecifierContext specifierContext = new UhtSpecifierContext(topScope, topScope.TokenReader, function.MetaData);
+					UhtSpecifierContext specifierContext = new(topScope, topScope.TokenReader, function.MetaData);
 					UhtSpecifierParser specifiers = topScope.HeaderParser.GetCachedSpecifierParser(specifierContext, ScopeName, parentScope.Session.GetSpecifierTable(UhtTableNames.Function));
 
 					// If this is a UDELEGATE, parse the specifiers
-					StringView delegateMacro = new StringView();
+					StringView delegateMacro = new();
 					if (hasSpecifiers)
 					{
 						specifiers.ParseSpecifiers();
@@ -315,7 +314,7 @@ namespace EpicGames.UHT.Parsers
 						{
 							topScope.TokenReader.LogError("Delegate type declarations must start with F");
 						}
-						function.StrippedFunctionName = function.SourceName.Substring(1);
+						function.StrippedFunctionName = function.SourceName[1..];
 						function.EngineName = $"{function.StrippedFunctionName}{UhtFunction.GeneratedDelegateSignatureSuffix}";
 					}
 
@@ -383,19 +382,19 @@ namespace EpicGames.UHT.Parsers
 
 		private static UhtParseResult ParseUFunction(UhtParsingScope parentScope, UhtToken token)
 		{
-			UhtFunctionParser function = new UhtFunctionParser(parentScope.ScopeType, token.InputLine);
+			UhtFunctionParser function = new(parentScope.ScopeType, token.InputLine);
 
-			using (UhtParsingScope topScope = new UhtParsingScope(parentScope, function, parentScope.Session.GetKeywordTable(UhtTableNames.Function), UhtAccessSpecifier.Public))
 			{
+				using UhtParsingScope topScope = new(parentScope, function, parentScope.Session.GetKeywordTable(UhtTableNames.Function), UhtAccessSpecifier.Public);
 				UhtParsingScope outerClassScope = topScope.CurrentClassScope;
 				UhtClass outerClass = (UhtClass)outerClassScope.ScopeType;
 				string scopeName = "function";
 
-				using (UhtMessageContext tokenContext = new UhtMessageContext(scopeName))
 				{
+					using UhtMessageContext tokenContext = new(scopeName);
 					topScope.AddModuleRelativePathToMetaData();
 
-					UhtSpecifierContext specifierContext = new UhtSpecifierContext(topScope, topScope.TokenReader, function.MetaData);
+					UhtSpecifierContext specifierContext = new(topScope, topScope.TokenReader, function.MetaData);
 					UhtSpecifierParser specifierParser = topScope.HeaderParser.GetCachedSpecifierParser(specifierContext, scopeName, parentScope.Session.GetSpecifierTable(UhtTableNames.Function));
 					specifierParser.ParseSpecifiers();
 
@@ -466,8 +465,8 @@ namespace EpicGames.UHT.Parsers
 					}
 
 					// Record the tokens so we can detect this function as a declaration later (i.e. RPC)
-					using (UhtTokenRecorder tokenRecorder = new UhtTokenRecorder(parentScope, function))
 					{
+						using UhtTokenRecorder tokenRecorder = new(parentScope, function);
 
 						if (topScope.TokenReader.TryOptional("virtual"))
 						{
@@ -477,8 +476,7 @@ namespace EpicGames.UHT.Parsers
 						bool internalOnly = function.MetaData.GetBoolean(UhtNames.BlueprintInternalUseOnly);
 
 						// Peek ahead to look for a CORE_API style DLL import/export token if present
-						UhtToken apiMacroToken;
-						if (topScope.TokenReader.TryOptionalAPIMacro(out apiMacroToken))
+						if (topScope.TokenReader.TryOptionalAPIMacro(out UhtToken apiMacroToken))
 						{
 							//@TODO: Validate the module name for RequiredAPIMacroIfPresent
 							function.FunctionFlags |= EFunctionFlags.RequiredAPI;
@@ -512,7 +510,7 @@ namespace EpicGames.UHT.Parsers
 						}
 
 						// Get return type.  C++ style functions always have a return value type, even if it's void
-						UhtToken funcNameToken = new UhtToken();
+						UhtToken funcNameToken = new();
 						UhtProperty? returnValueProperty = null;
 						topScope.HeaderParser.GetCachedPropertyParser(topScope).Parse(EPropertyFlags.None,
 							function.GetPropertyParseOptions(true), UhtParsePropertyDeclarationStyle.None, UhtPropertyCategory.Return,
@@ -584,8 +582,7 @@ namespace EpicGames.UHT.Parsers
 						// Handle C++ style functions being declared as abstract
 						if (topScope.TokenReader.TryOptional('='))
 						{
-							int zeroValue = 1;
-							bool gotZero = topScope.TokenReader.TryOptionalConstInt(out zeroValue);
+							bool gotZero = topScope.TokenReader.TryOptionalConstInt(out int zeroValue);
 							gotZero = gotZero && (zeroValue == 0);
 							if (!gotZero || zeroValue != 0)
 							{
@@ -610,7 +607,7 @@ namespace EpicGames.UHT.Parsers
 						else if (topScope.TokenReader.TryPeekOptional('{'))
 						{
 							// Skip inline function bodies
-							UhtToken tokenCopy = new UhtToken();
+							UhtToken tokenCopy = new();
 							topScope.TokenReader.SkipDeclaration(ref tokenCopy);
 						}
 					}
@@ -641,7 +638,7 @@ namespace EpicGames.UHT.Parsers
 			UhtPropertyCategory propertyCategory = isNetFunc ? UhtPropertyCategory.ReplicatedParameter : UhtPropertyCategory.RegularParameter;
 			EPropertyFlags disallowFlags = ~(EPropertyFlags.ParmFlags | EPropertyFlags.AutoWeak | EPropertyFlags.RepSkip | EPropertyFlags.UObjectWrapper | EPropertyFlags.NativeAccessSpecifiers);
 
-			UhtAdvancedDisplayParameterHandler advancedDisplay = new UhtAdvancedDisplayParameterHandler(topScope.ScopeType.MetaData);
+			UhtAdvancedDisplayParameterHandler advancedDisplay = new(topScope.ScopeType.MetaData);
 
 			topScope.TokenReader.RequireList(')', ',', false, () =>
 			{
@@ -657,7 +654,7 @@ namespace EpicGames.UHT.Parsers
 						// Default value.
 						if (topScope.TokenReader.TryOptional('='))
 						{
-							List<UhtToken> defaultValueTokens = new List<UhtToken>();
+							List<UhtToken> defaultValueTokens = new();
 							int parenthesisNestCount = 0;
 							while (!topScope.TokenReader.IsEOF)
 							{
@@ -720,7 +717,7 @@ namespace EpicGames.UHT.Parsers
 			string functionName = function.EngineName;
 			if (functionName.EndsWith(UhtFunction.GeneratedDelegateSignatureSuffix, StringComparison.Ordinal))
 			{
-				functionName = functionName.Substring(0, functionName.Length - UhtFunction.GeneratedDelegateSignatureSuffix.Length);
+				functionName = functionName[..^UhtFunction.GeneratedDelegateSignatureSuffix.Length];
 			}
 
 			function.UnMarshalAndCallName = "exec" + functionName;
diff --git a/Engine/Source/Programs/Shared/EpicGames.UHT/Parsers/UhtHeaderFileParser.cs b/Engine/Source/Programs/Shared/EpicGames.UHT/Parsers/UhtHeaderFileParser.cs
--- a/Engine/Source/Programs/Shared/EpicGames.UHT/Parsers/UhtHeaderFileParser.cs
+++ b/Engine/Source/Programs/Shared/EpicGames.UHT/Parsers/UhtHeaderFileParser.cs
@@ -217,12 +217,12 @@ namespace EpicGames.UHT.Parsers
 		/// <summary>
 		/// Stack of current #if states
 		/// </summary>
-		private readonly List<CompilerDirective> _compilerDirectives = new List<CompilerDirective>();
+		private readonly List<CompilerDirective> _compilerDirectives = new();
 
 		/// <summary>
 		/// Stack of current #if states saved as part of the preprocessor state
 		/// </summary>
-		private readonly List<CompilerDirective> _savedCompilerDirectives = new List<CompilerDirective>();
+		private readonly List<CompilerDirective> _savedCompilerDirectives = new();
 
 		/// <summary>
 		/// Current top of the parsing scopes.  Classes, structures and functions all allocate scopes.
@@ -236,30 +236,28 @@ namespace EpicGames.UHT.Parsers
 		/// <returns>Parser</returns>
 		public static UhtHeaderFileParser Parse(UhtHeaderFile headerFile)
 		{
-			UhtHeaderFileParser headerParser = new UhtHeaderFileParser(headerFile);
-			using (UhtParsingScope topScope = new UhtParsingScope(headerParser, headerParser.HeaderFile, headerFile.Session.GetKeywordTable(UhtTableNames.Global)))
-			{
-				headerParser.ParseStatements();
+			UhtHeaderFileParser headerParser = new(headerFile);
+			using UhtParsingScope topScope = new(headerParser, headerParser.HeaderFile, headerFile.Session.GetKeywordTable(UhtTableNames.Global));
+			headerParser.ParseStatements();
 
-				if (!headerParser.SpottedAutogeneratedHeaderInclude && headerParser.HeaderFile.Data.Length > 0)
+			if (!headerParser.SpottedAutogeneratedHeaderInclude && headerParser.HeaderFile.Data.Length > 0)
+			{
+				bool noExportClassesOnly = true;
+				foreach (UhtType type in headerParser.HeaderFile.Children)
 				{
-					bool noExportClassesOnly = true;
-					foreach (UhtType type in headerParser.HeaderFile.Children)
+					if (type is UhtClass classObj)
 					{
-						if (type is UhtClass classObj)
+						if (classObj.ClassType != UhtClassType.NativeInterface && !classObj.ClassExportFlags.HasAnyFlags(UhtClassExportFlags.NoExport))
 						{
-							if (classObj.ClassType != UhtClassType.NativeInterface && !classObj.ClassExportFlags.HasAnyFlags(UhtClassExportFlags.NoExport))
-							{
-								noExportClassesOnly = false;
-								break;
-							}
+							noExportClassesOnly = false;
+							break;
 						}
 					}
+				}
 
-					if (!noExportClassesOnly)
-					{
-						headerParser.HeaderFile.LogError($"Expected an include at the top of the header the follows all other includes: '#include \"{headerParser.HeaderFile.GeneratedHeaderFileName}\"'");
-					}
+				if (!noExportClassesOnly)
+				{
+					headerParser.HeaderFile.LogError($"Expected an include at the top of the header the follows all other includes: '#include \"{headerParser.HeaderFile.GeneratedHeaderFileName}\"'");
 				}
 			}
 			return headerParser;
@@ -481,22 +479,20 @@ namespace EpicGames.UHT.Parsers
 			{
 				if (topScope.ScopeType is UhtClass classObj)
 				{
-					using (UhtTokenRecorder recorder = new UhtTokenRecorder(topScope, ref token))
+					using UhtTokenRecorder recorder = new(topScope, ref token);
+					topScope.TokenReader.SkipDeclaration(ref token);
+					if (recorder.Stop())
 					{
-						topScope.TokenReader.SkipDeclaration(ref token);
-						if (recorder.Stop())
+						if (classObj.Declarations != null)
 						{
-							if (classObj.Declarations != null)
+							UhtDeclaration declaration = classObj.Declarations[classObj.Declarations.Count - 1];
+							if (topScope.HeaderParser.CheckForConstructor(classObj, declaration))
 							{
-								UhtDeclaration declaration = classObj.Declarations[classObj.Declarations.Count - 1];
-								if (topScope.HeaderParser.CheckForConstructor(classObj, declaration))
-								{
-								}
-								else if (classObj.ClassType == UhtClassType.Class)
+							}
+							else if (classObj.ClassType == UhtClassType.Class)
+							{
+								if (topScope.HeaderParser.CheckForSerialize(classObj, declaration))
 								{
-									if (topScope.HeaderParser.CheckForSerialize(classObj, declaration))
-									{
-									}
 								}
 							}
 						}
@@ -515,8 +511,7 @@ namespace EpicGames.UHT.Parsers
 			UhtParseResult parseResult = UhtParseResult.Unhandled;
 			for (UhtParsingScope? currentScope = topScope; currentScope != null && parseResult == UhtParseResult.Unhandled; currentScope = currentScope.ParentScope)
 			{
-				UhtKeyword keywordInfo;
-				if (currentScope.ScopeKeywordTable.TryGetValue(token.Value, out keywordInfo))
+				if (currentScope.ScopeKeywordTable.TryGetValue(token.Value, out UhtKeyword keywordInfo))
 				{
 					if (keywordInfo.AllScopes || topScope == currentScope)
 					{
@@ -566,7 +561,7 @@ namespace EpicGames.UHT.Parsers
 
 			// We'll guess about it being a macro based on it being fully-capitalized with at least one underscore.
 			int underscoreCount = 0;
-			foreach (char ch in token.Value.Span.Slice(1))
+			foreach (char ch in token.Value.Span[1..])
 			{
 				if (ch == '_')
 				{
@@ -626,7 +621,7 @@ namespace EpicGames.UHT.Parsers
 			// We assume that the vast majority of lines will not be continuations.  So avoid using the
 			// string builder as much as possible.
 			int startingLine = this.TokenReader.InputLine;
-			StringViewBuilder builder = new StringViewBuilder();
+			StringViewBuilder builder = new();
 			while (true)
 			{
 				UhtToken lineToken = this.TokenReader.GetLine();
@@ -647,11 +642,10 @@ namespace EpicGames.UHT.Parsers
 			StringView line = builder.ToStringView();
 
 			// Create a token reader we will use to decode
-			UhtTokenBufferReader lineTokenReader = new UhtTokenBufferReader(this.HeaderFile, line.Memory);
+			UhtTokenBufferReader lineTokenReader = new(this.HeaderFile, line.Memory);
 			lineTokenReader.InputLine = startingLine;
 
-			UhtToken directive;
-			if (!lineTokenReader.TryOptionalIdentifier(out directive))
+			if (!lineTokenReader.TryOptionalIdentifier(out UhtToken directive))
 			{
 				if (isBeingIncluded)
 				{
@@ -673,8 +667,7 @@ namespace EpicGames.UHT.Parsers
 			}
 			else if (directive.IsValue("linenumber"))
 			{
-				int newInputLine;
-				if (!lineTokenReader.TryOptionalConstInt(out newInputLine))
+				if (!lineTokenReader.TryOptionalConstInt(out int newInputLine))
 				{
 					throw new UhtException(this.TokenReader, directive.InputLine, "Missing line number in line number directive");
 				}
@@ -845,7 +838,7 @@ namespace EpicGames.UHT.Parsers
 		/// <param name="compilerDirective">Directive to be added</param>
 		private void PushCompilerDirective(UhtCompilerDirective compilerDirective)
 		{
-			CompilerDirective newCompileDirective = new CompilerDirective();
+			CompilerDirective newCompileDirective = new();
 			newCompileDirective._element = compilerDirective;
 			newCompileDirective._composite = GetCurrentCompositeCompilerDirective() | compilerDirective;
 			this._compilerDirectives.Add(newCompileDirective);
diff --git a/Engine/Source/Programs/Shared/EpicGames.UHT/Parsers/UhtInterfaceClassParser.cs b/Engine/Source/Programs/Shared/EpicGames.UHT/Parsers/UhtInterfaceClassParser.cs
--- a/Engine/Source/Programs/Shared/EpicGames.UHT/Parsers/UhtInterfaceClassParser.cs
+++ b/Engine/Source/Programs/Shared/EpicGames.UHT/Parsers/UhtInterfaceClassParser.cs
@@ -108,17 +108,18 @@ namespace EpicGames.UHT.Parsers
 
 		private static UhtParseResult ParseUInterface(UhtParsingScope parentScope, ref UhtToken token)
 		{
-			UhtInterfaceClassParser classObj = new UhtInterfaceClassParser(parentScope.ScopeType, token.InputLine);
+			UhtInterfaceClassParser classObj = new(parentScope.ScopeType, token.InputLine);
 			classObj.ClassType = UhtClassType.Interface;
 
-			using (UhtParsingScope topScope = new UhtParsingScope(parentScope, classObj, parentScope.Session.GetKeywordTable(UhtTableNames.Interface), UhtAccessSpecifier.Private))
 			{
+				using UhtParsingScope topScope = new(parentScope, classObj, parentScope.Session.GetKeywordTable(UhtTableNames.Interface), UhtAccessSpecifier.Private);
 				const string ScopeName = "interface";
 
-				using (UhtMessageContext tokenContext = new UhtMessageContext(ScopeName))
 				{
+					using UhtMessageContext tokenContext = new(ScopeName);
+
 					// Parse the specifiers
-					UhtSpecifierContext specifierContext = new UhtSpecifierContext(topScope, topScope.TokenReader, classObj.MetaData);
+					UhtSpecifierContext specifierContext = new(topScope, topScope.TokenReader, classObj.MetaData);
 					UhtSpecifierParser specifiers = topScope.HeaderParser.GetCachedSpecifierParser(specifierContext, ScopeName, parentScope.Session.GetSpecifierTable(UhtTableNames.Interface));
 					specifiers.ParseSpecifiers();
 					classObj.PrologLineNumber = topScope.TokenReader.InputLine;
@@ -129,8 +130,7 @@ namespace EpicGames.UHT.Parsers
 					topScope.TokenReader.Require("class");
 
 					// Get the optional API macro
-					UhtToken apiMacroToken;
-					topScope.TokenReader.TryOptionalAPIMacro(out apiMacroToken);
+					topScope.TokenReader.TryOptionalAPIMacro(out UhtToken apiMacroToken);
 
 					// Get the name of the interface
 					classObj.SourceName = topScope.TokenReader.GetIdentifier().Value.ToString();
diff --git a/Engine/Source/Programs/Shared/EpicGames.UHT/Parsers/UhtNativeInterfaceClassParser.cs b/Engine/Source/Programs/Shared/EpicGames.UHT/Parsers/UhtNativeInterfaceClassParser.cs
--- a/Engine/Source/Programs/Shared/EpicGames.UHT/Parsers/UhtNativeInterfaceClassParser.cs
+++ b/Engine/Source/Programs/Shared/EpicGames.UHT/Parsers/UhtNativeInterfaceClassParser.cs
@@ -71,10 +71,10 @@ namespace EpicGames.UHT.Parsers
 		[SuppressMessage("Style", "IDE0060:Remove unused parameter", Justification = "Attribute accessed method")]
 		private static UhtParseResult ClassKeyword(UhtParsingScope topScope, UhtParsingScope actionScope, ref UhtToken token)
 		{
-			using (UhtTokenSaveState saveState = new UhtTokenSaveState(topScope.TokenReader))
 			{
-				UhtToken sourceName = new UhtToken();
-				UhtToken superName = new UhtToken();
+				using UhtTokenSaveState saveState = new(topScope.TokenReader);
+				UhtToken sourceName = new();
+				UhtToken superName = new();
 				if (TryParseIInterface(topScope, out sourceName, out superName))
 				{
 					saveState.AbandonState();
@@ -152,7 +152,7 @@ namespace EpicGames.UHT.Parsers
 		{
 			IUhtTokenReader tokenReader = parentScope.TokenReader;
 
-			UhtParserNativeInterfaceClass classObj = new UhtParserNativeInterfaceClass(parentScope.ScopeType, token.InputLine);
+			UhtParserNativeInterfaceClass classObj = new(parentScope.ScopeType, token.InputLine);
 			classObj.ClassType = UhtClassType.NativeInterface;
 			classObj.SourceName = sourceName.Value.ToString();
 			classObj.ClassFlags |= EClassFlags.Native | EClassFlags.Interface;
@@ -174,15 +174,11 @@ namespace EpicGames.UHT.Parsers
 
 			classObj.SuperIdentifier = superName;
 
-			using (UhtParsingScope topScope = new UhtParsingScope(parentScope, classObj, parentScope.Session.GetKeywordTable(UhtTableNames.NativeInterface), UhtAccessSpecifier.Private))
 			{
-				const string ScopeName = "native interface";
-
-				using (UhtMessageContext tokenContext = new UhtMessageContext(ScopeName))
-				{
-					topScope.HeaderParser.ParseStatements('{', '}', false);
-					tokenReader.Require(';');
-				}
+				using UhtParsingScope topScope = new(parentScope, classObj, parentScope.Session.GetKeywordTable(UhtTableNames.NativeInterface), UhtAccessSpecifier.Private);
+				using UhtMessageContext tokenContext = new("native interface");
+				topScope.HeaderParser.ParseStatements('{', '}', false);
+				tokenReader.Require(';');
 			}
 			return;
 		}
diff --git a/Engine/Source/Programs/Shared/EpicGames.UHT/Parsers/UhtParserHelpers.cs b/Engine/Source/Programs/Shared/EpicGames.UHT/Parsers/UhtParserHelpers.cs
--- a/Engine/Source/Programs/Shared/EpicGames.UHT/Parsers/UhtParserHelpers.cs
+++ b/Engine/Source/Programs/Shared/EpicGames.UHT/Parsers/UhtParserHelpers.cs
@@ -24,7 +24,7 @@ namespace EpicGames.UHT.Parsers
 		/// <param name="baseIdentifiers">Output base identifiers</param>
 		public static void ParseInheritance(IUhtTokenReader tokenReader, IUhtConfig config, out UhtToken superIdentifier, out List<UhtToken[]>? baseIdentifiers)
 		{
-			UhtToken superIdentifierTemp = new UhtToken();
+			UhtToken superIdentifierTemp = new();
 			List<UhtToken[]>? baseIdentifiersTemp = null;
 			tokenReader.OptionalInheritance(
 				(ref UhtToken identifier) =>
diff --git a/Engine/Source/Programs/Shared/EpicGames.UHT/Parsers/UhtParsingScope.cs b/Engine/Source/Programs/Shared/EpicGames.UHT/Parsers/UhtParsingScope.cs
--- a/Engine/Source/Programs/Shared/EpicGames.UHT/Parsers/UhtParsingScope.cs
+++ b/Engine/Source/Programs/Shared/EpicGames.UHT/Parsers/UhtParsingScope.cs
@@ -191,16 +191,14 @@ namespace EpicGames.UHT.Parsers
 			}
 			else
 			{
-				using (BorrowStringBuilder borrower = new BorrowStringBuilder(StringBuilderCache.Small))
+				using BorrowStringBuilder borrower = new(StringBuilderCache.Small);
+				StringBuilder builder = borrower.StringBuilder;
+				foreach (StringView comment in comments)
 				{
-					StringBuilder builder = borrower.StringBuilder;
-					foreach (StringView comment in comments)
-					{
-						builder.Append(comment);
-					}
-					mergedString = builder.ToString();
-					type.MetaData.Add(UhtNames.Comment, metaNameIndex, mergedString);
+					builder.Append(comment);
 				}
+				mergedString = builder.ToString();
+				type.MetaData.Add(UhtNames.Comment, metaNameIndex, mergedString);
 			}
 
 			// Format the tooltip and set the metadata
@@ -365,7 +363,7 @@ namespace EpicGames.UHT.Parsers
 							++pos;
 						}
 
-						line = line.Substring(pos + 1);
+						line = line[(pos + 1)..];
 					}
 				}
 				lines[index] = line;
@@ -445,7 +443,7 @@ namespace EpicGames.UHT.Parsers
 
 					if (pos > 0)
 					{
-						line = line.Substring(pos);
+						line = line[pos..];
 					}
 
 					if (index > firstIndex)
@@ -464,13 +462,13 @@ namespace EpicGames.UHT.Parsers
 			// Make sure it doesn't start with a newline
 			if (result.Length != 0 && UhtFCString.IsLinebreak(result[0]))
 			{
-				result = result.Substring(1);
+				result = result[1..];
 			}
 
 			// Make sure it doesn't end with a dead newline
 			if (result.Length != 0 && UhtFCString.IsLinebreak(result[^1]))
 			{
-				result = result.Substring(0, result.Length - 1);
+				result = result[0..^1];
 			}
 
 			// Done.
diff --git a/Engine/Source/Programs/Shared/EpicGames.UHT/Parsers/UhtPropertyParser.cs b/Engine/Source/Programs/Shared/EpicGames.UHT/Parsers/UhtPropertyParser.cs
--- a/Engine/Source/Programs/Shared/EpicGames.UHT/Parsers/UhtPropertyParser.cs
+++ b/Engine/Source/Programs/Shared/EpicGames.UHT/Parsers/UhtPropertyParser.cs
@@ -442,13 +442,13 @@ namespace EpicGames.UHT.Parsers
 		private UhtParsingScope TopScope => this._specifierContext.Scope;
 
 		// Scratch pad variables used by actions
-		private List<UhtToken> _currentTypeTokens = new List<UhtToken>();
+		private List<UhtToken> _currentTypeTokens = new();
 		private int _currentTemplateDepth = 0;
 
 		// Actions/Delegates
 		private readonly UhtTokensUntilDelegate _gatherTypeTokensDelegate;
 
-		private static readonly Dictionary<StringView, UhtLayoutMacroType> s_layoutMacroTypes = new Dictionary<StringView, UhtLayoutMacroType>(new[]
+		private static readonly Dictionary<StringView, UhtLayoutMacroType> s_layoutMacroTypes = new(new[]
 		{
 			UhtLayoutMacroType.Array.MacroNameAndValue(),
 			UhtLayoutMacroType.ArrayEditorOnly.MacroNameAndValue(),
@@ -509,10 +509,8 @@ namespace EpicGames.UHT.Parsers
 			this._currentTypeTokens = new List<UhtToken>();
 			this._currentTemplateDepth = 0;
 
-			using (UhtMessageContext tokenContext = new UhtMessageContext(this))
-			{
-				ParseInternal(propertyDelegate);
-			}
+			using UhtMessageContext tokenContext = new(this);
+			ParseInternal(propertyDelegate);
 			return this;
 		}
 
@@ -522,7 +520,7 @@ namespace EpicGames.UHT.Parsers
 		{
 			get
 			{
-				Stack<object?> extraContext = new Stack<object?>(1);
+				Stack<object?> extraContext = new(1);
 				extraContext.Push(this._specifierContext.PropertySettings.PropertyCategory.GetHintText());
 				return extraContext;
 			}
@@ -544,7 +542,6 @@ namespace EpicGames.UHT.Parsers
 			IUhtTokenReader replayReader = UhtTokenReplayReader.GetThreadInstance(propertySettings.Outer, data, typeTokens, UhtTokenType.EndOfType);
 
 			// Loop through the tokens until we find a known property type or the start of a template argument list
-			UhtPropertyType propertyType;
 			for (int index = 0; index < typeTokens.Length; ++index)
 			{
 				if (typeTokens.Span[index].IsSymbol())
@@ -557,7 +554,7 @@ namespace EpicGames.UHT.Parsers
 				}
 				else
 				{
-					if (session.TryGetPropertyType(typeTokens.Span[index].Value, out propertyType))
+					if (session.TryGetPropertyType(typeTokens.Span[index].Value, out UhtPropertyType propertyType))
 					{
 						if (resolvePhase == UhtPropertyResolvePhase.Resolving || propertyType.Options.HasAnyFlags(UhtPropertyTypeOptions.Immediate))
 						{
@@ -575,7 +572,7 @@ namespace EpicGames.UHT.Parsers
 			return null;
 		}
 
-		private static readonly ThreadLocal<UhtPropertySettings> s_tlsPropertySettings = new ThreadLocal<UhtPropertySettings>(() => { return new UhtPropertySettings(); });
+		private static readonly ThreadLocal<UhtPropertySettings> s_tlsPropertySettings = new(() => { return new UhtPropertySettings(); });
 
 		/// <summary>
 		/// Given a type with children, resolve any children that couldn't be resolved during the parsing phase.
@@ -628,11 +625,11 @@ namespace EpicGames.UHT.Parsers
 			UhtSession session = parentPropertySettings.Outer.Session;
 
 			// Save the token reader state.  We need this to restore back to the start when invoking the resolve methods.
-			using (UhtTokenSaveState saveState = new UhtTokenSaveState(tokenReader))
 			{
-				UhtPropertySettings propertySettings = new UhtPropertySettings(parentPropertySettings, paramName.ToString(), tokenReader);
+				using UhtTokenSaveState saveState = new(tokenReader);
+				UhtPropertySettings propertySettings = new(parentPropertySettings, paramName.ToString(), tokenReader);
 
-				UhtPropertyType propertyType = new UhtPropertyType();
+				UhtPropertyType propertyType = new();
 				while (!tokenReader.IsEOF)
 				{
 					UhtToken token = tokenReader.GetToken();
@@ -963,7 +960,7 @@ namespace EpicGames.UHT.Parsers
 				this._currentTypeTokens.RemoveAt(this._currentTypeTokens.Count - 1);
 				CheckForOptionalParts(this._specifierContext.PropertySettings, ref nameToken);
 
-				ReadOnlyMemory<UhtToken> typeTokens = new ReadOnlyMemory<UhtToken>(this._currentTypeTokens.ToArray());
+				ReadOnlyMemory<UhtToken> typeTokens = new(this._currentTypeTokens.ToArray());
 
 				while (true)
 				{
@@ -1026,7 +1023,7 @@ namespace EpicGames.UHT.Parsers
 			}
 			else
 			{
-				UhtToken nameToken = new UhtToken();
+				UhtToken nameToken = new();
 				CheckForOptionalParts(this._specifierContext.PropertySettings, ref nameToken);
 				Finalize(ref nameToken, new ReadOnlyMemory<UhtToken>(this._currentTypeTokens.ToArray()), layoutMacroType, propertyDelegate);
 			}
@@ -1098,7 +1095,7 @@ namespace EpicGames.UHT.Parsers
 					}
 
 					newProperty.PropertyFlags |= EPropertyFlags.Deprecated;
-					newProperty.EngineName = newProperty.SourceName.Substring(0, deprecatedIndex);
+					newProperty.EngineName = newProperty.SourceName[..deprecatedIndex];
 				}
 			}
 
@@ -1141,8 +1138,7 @@ namespace EpicGames.UHT.Parsers
 
 		private void RequireBitfield(UhtPropertySettings propertySettings, ref UhtToken nameToken)
 		{
-			int bitfieldSize;
-			if (!this.TokenReader.TryOptionalConstInt(out bitfieldSize) || bitfieldSize != 1)
+			if (!this.TokenReader.TryOptionalConstInt(out int bitfieldSize) || bitfieldSize != 1)
 			{
 				throw new UhtException(this.TokenReader, $"Bad or missing bit field size for '{nameToken.Value}', must be 1.");
 			}
diff --git a/Engine/Source/Programs/Shared/EpicGames.UHT/Parsers/UhtScriptStructParser.cs b/Engine/Source/Programs/Shared/EpicGames.UHT/Parsers/UhtScriptStructParser.cs
--- a/Engine/Source/Programs/Shared/EpicGames.UHT/Parsers/UhtScriptStructParser.cs
+++ b/Engine/Source/Programs/Shared/EpicGames.UHT/Parsers/UhtScriptStructParser.cs
@@ -121,15 +121,16 @@ namespace EpicGames.UHT.Parsers
 
 		private static UhtParseResult ParseUScriptStruct(UhtParsingScope parentScope, UhtToken keywordToken)
 		{
-			UhtScriptStructParser scriptStruct = new UhtScriptStructParser(parentScope.ScopeType, keywordToken.InputLine);
-			using (UhtParsingScope topScope = new UhtParsingScope(parentScope, scriptStruct, parentScope.Session.GetKeywordTable(UhtTableNames.ScriptStruct), UhtAccessSpecifier.Public))
+			UhtScriptStructParser scriptStruct = new(parentScope.ScopeType, keywordToken.InputLine);
 			{
-				const string ScopeName = "struct";
+				using UhtParsingScope topScope = new(parentScope, scriptStruct, parentScope.Session.GetKeywordTable(UhtTableNames.ScriptStruct), UhtAccessSpecifier.Public);
 
-				using (UhtMessageContext tokenContext = new UhtMessageContext(ScopeName))
 				{
+					const string ScopeName = "struct";
+					using UhtMessageContext tokenContext = new(ScopeName);
+
 					// Parse the specifiers
-					UhtSpecifierContext specifierContext = new UhtSpecifierContext(topScope, topScope.TokenReader, scriptStruct.MetaData);
+					UhtSpecifierContext specifierContext = new(topScope, topScope.TokenReader, scriptStruct.MetaData);
 					UhtSpecifierParser specifiers = topScope.HeaderParser.GetCachedSpecifierParser(specifierContext, ScopeName, parentScope.Session.GetSpecifierTable(UhtTableNames.ScriptStruct));
 					specifiers.ParseSpecifiers();
 
@@ -139,8 +140,7 @@ namespace EpicGames.UHT.Parsers
 					topScope.TokenReader.SkipAlignasAndDeprecatedMacroIfNecessary();
 
 					// Read the struct name and possible API macro name
-					UhtToken apiMacroToken;
-					topScope.TokenReader.TryOptionalAPIMacro(out apiMacroToken);
+					topScope.TokenReader.TryOptionalAPIMacro(out UhtToken apiMacroToken);
 					scriptStruct.SourceName = topScope.TokenReader.GetIdentifier().Value.ToString();
 
 					topScope.AddModuleRelativePathToMetaData();
@@ -148,7 +148,7 @@ namespace EpicGames.UHT.Parsers
 					// Strip the name
 					if (scriptStruct.SourceName[0] == 'T' || scriptStruct.SourceName[0] == 'F')
 					{
-						scriptStruct.EngineName = scriptStruct.SourceName.Substring(1);
+						scriptStruct.EngineName = scriptStruct.SourceName[1..];
 					}
 					else
 					{
@@ -163,9 +163,7 @@ namespace EpicGames.UHT.Parsers
 					}
 
 					// Parse the inheritance
-					UhtToken superIdentifier;
-					List<UhtToken[]>? baseIdentifiers;
-					UhtParserHelpers.ParseInheritance(topScope.TokenReader, topScope.Session.Config!, out superIdentifier, out baseIdentifiers);
+					UhtParserHelpers.ParseInheritance(topScope.TokenReader, topScope.Session.Config!, out UhtToken superIdentifier, out List<UhtToken[]>? baseIdentifiers);
 					scriptStruct.SuperIdentifier = superIdentifier;
 					scriptStruct.BaseIdentifiers = baseIdentifiers;
 
@@ -209,8 +207,8 @@ namespace EpicGames.UHT.Parsers
 		{
 			UhtScriptStruct scriptStruct = (UhtScriptStruct)topScope.ScopeType;
 
-			using (UhtMessageContext tokenContext = new UhtMessageContext("RIGVM_METHOD"))
 			{
+				using UhtMessageContext tokenContext = new("RIGVM_METHOD");
 
 				// Create the RigVM information if it doesn't already exist
 				UhtRigVMStructInfo? structInfo = scriptStruct.RigVMStructInfo;
@@ -222,7 +220,7 @@ namespace EpicGames.UHT.Parsers
 				}
 
 				// Create a new method information and add it
-				UhtRigVMMethodInfo methodInfo = new UhtRigVMMethodInfo();
+				UhtRigVMMethodInfo methodInfo = new();
 
 				// NOTE: The argument list reader doesn't handle templates with multiple arguments (i.e. the ',' issue)
 				topScope.TokenReader
@@ -249,8 +247,8 @@ namespace EpicGames.UHT.Parsers
 					{
 						if (!isGetUpgradeInfo && !isGetNextAggregateName)
 						{
-							StringViewBuilder builder = new StringViewBuilder();
-							UhtToken lastToken = new UhtToken();
+							StringViewBuilder builder = new();
+							UhtToken lastToken = new();
 							foreach (UhtToken token in tokens)
 							{
 								if (token.IsSymbol('='))
diff --git a/Engine/Source/Programs/Shared/EpicGames.UHT/Parsers/UhtSpecifierParser.cs b/Engine/Source/Programs/Shared/EpicGames.UHT/Parsers/UhtSpecifierParser.cs
--- a/Engine/Source/Programs/Shared/EpicGames.UHT/Parsers/UhtSpecifierParser.cs
+++ b/Engine/Source/Programs/Shared/EpicGames.UHT/Parsers/UhtSpecifierParser.cs
@@ -23,14 +23,14 @@ namespace EpicGames.UHT.Parsers
 			public object? _value;
 		}
 
-		private static readonly List<KeyValuePair<StringView, StringView>> s_emptyKVPValues = new List<KeyValuePair<StringView, StringView>>();
+		private static readonly List<KeyValuePair<StringView, StringView>> s_emptyKVPValues = new();
 
 		private UhtSpecifierContext _specifierContext;
 		private IUhtTokenReader _tokenReader;
 		private UhtSpecifierTable _table;
 		private StringView _context;
-		private StringView _currentSpecifier = new StringView();
-		private List<DeferredSpecifier> _deferredSpecifiers = new List<DeferredSpecifier>();
+		private StringView _currentSpecifier = new();
+		private List<DeferredSpecifier> _deferredSpecifiers = new();
 		private bool _isParsingFieldMetaData = false;
 		private List<KeyValuePair<StringView, StringView>>? _currentKVPValues = null;
 		private List<StringView>? _currentStringValues = null;
@@ -97,10 +97,8 @@ namespace EpicGames.UHT.Parsers
 			this._isParsingFieldMetaData = false;
 			this._specifierContext.MetaData.LineNumber = _tokenReader.InputLine;
 
-			using (UhtMessageContext tokenContext = new UhtMessageContext(this))
-			{
-				this._tokenReader.RequireList('(', ')', ',', false, this._parseAction);
-			}
+			using UhtMessageContext tokenContext = new(this);
+			this._tokenReader.RequireList('(', ')', ',', false, this._parseAction);
 			return this;
 		}
 
@@ -113,16 +111,14 @@ namespace EpicGames.UHT.Parsers
 			this._tokenReader = this._specifierContext.Scope.TokenReader;
 			this._isParsingFieldMetaData = true;
 
-			using (UhtMessageContext tokenContext = new UhtMessageContext(this))
+			using UhtMessageContext tokenContext = new(this);
+			if (_tokenReader.TryOptional("UMETA"))
 			{
-				if (_tokenReader.TryOptional("UMETA"))
+				this._umetaElementsParsed = 0;
+				_tokenReader.RequireList('(', ')', ',', false, this._parseFieldMetaDataAction);
+				if (this._umetaElementsParsed == 0)
 				{
-					this._umetaElementsParsed = 0;
-					_tokenReader.RequireList('(', ')', ',', false, this._parseFieldMetaDataAction);
-					if (this._umetaElementsParsed == 0)
-					{
-						this._tokenReader.LogError($"No metadata specified while parsing {UhtMessage.FormatContext(this)}");
-					}
+					this._tokenReader.LogError($"No metadata specified while parsing {UhtMessage.FormatContext(this)}");
 				}
 			}
 			return this;
@@ -147,7 +143,7 @@ namespace EpicGames.UHT.Parsers
 		{
 			get
 			{
-				Stack<object?> extraContext = new Stack<object?>(1);
+				Stack<object?> extraContext = new(1);
 				string what = this._isParsingFieldMetaData ? "metadata" : "specifiers";
 				if (this._context.Span.Length > 0)
 				{
@@ -167,8 +163,7 @@ namespace EpicGames.UHT.Parsers
 			UhtToken identifier = this._tokenReader.GetIdentifier();
 
 			this._currentSpecifier = identifier.Value;
-			UhtSpecifier? specifier;
-			if (this._table.TryGetValue(_currentSpecifier, out specifier))
+			if (this._table.TryGetValue(_currentSpecifier, out UhtSpecifier? specifier))
 			{
 				if (TryParseValue(specifier.ValueType, out object? value))
 				{
@@ -194,13 +189,12 @@ namespace EpicGames.UHT.Parsers
 
 		private void ParseFieldMetaDataInternal()
 		{
-			UhtToken key;
-			if (!_tokenReader.TryOptionalIdentifier(out key))
+			if (!_tokenReader.TryOptionalIdentifier(out UhtToken key))
 			{
 				throw new UhtException(this._tokenReader, $"UMETA expects a key and optional value", this);
 			}
 
-			StringViewBuilder builder = new StringViewBuilder();
+			StringViewBuilder builder = new();
 			if (_tokenReader.TryOptional('='))
 			{
 				if (!ReadValue(_tokenReader, builder, true))
@@ -331,8 +325,7 @@ namespace EpicGames.UHT.Parsers
 
 		private KeyValuePair<StringView, StringView> ReadKVP()
 		{
-			UhtToken key;
-			if (!this._tokenReader.TryOptionalIdentifier(out key))
+			if (!this._tokenReader.TryOptionalIdentifier(out UhtToken key))
 			{
 				throw new UhtException(this._tokenReader, $"The specifier '{this._currentSpecifier}' expects a key and optional value", this);
 			}
@@ -368,7 +361,7 @@ namespace EpicGames.UHT.Parsers
 
 		private StringView ReadValue()
 		{
-			StringViewBuilder builder = new StringViewBuilder();
+			StringViewBuilder builder = new();
 			if (!ReadValue(this._tokenReader, builder, false))
 			{
 				throw new UhtException(this._tokenReader, $"The specifier '{this._currentSpecifier}' expects a value", this);
diff --git a/Engine/Source/Programs/Shared/EpicGames.UHT/Specifiers/UhtFunctionSpecifiers.cs b/Engine/Source/Programs/Shared/EpicGames.UHT/Specifiers/UhtFunctionSpecifiers.cs
--- a/Engine/Source/Programs/Shared/EpicGames.UHT/Specifiers/UhtFunctionSpecifiers.cs
+++ b/Engine/Source/Programs/Shared/EpicGames.UHT/Specifiers/UhtFunctionSpecifiers.cs
@@ -422,8 +422,7 @@ namespace EpicGames.UHT.Parsers
 				{
 					if (kvp.Key.Span.StartsWith("Id", StringComparison.OrdinalIgnoreCase))
 					{
-						int id;
-						if (!Int32.TryParse(kvp.Value.Span, out id) || id <= 0 || id > UInt16.MaxValue)
+						if (!Int32.TryParse(kvp.Value.Span, out int id) || id <= 0 || id > UInt16.MaxValue)
 						{
 							tokenReader.LogError($"Invalid network identifier {kvp.Key} for function");
 						}
@@ -434,8 +433,7 @@ namespace EpicGames.UHT.Parsers
 					}
 					else if (kvp.Key.Span.StartsWith("ResponseId", StringComparison.OrdinalIgnoreCase) || kvp.Key.Span.StartsWith("Priority", StringComparison.OrdinalIgnoreCase))
 					{
-						int id;
-						if (!Int32.TryParse(kvp.Value.Span, out id) || id <= 0 || id > UInt16.MaxValue)
+						if (!Int32.TryParse(kvp.Value.Span, out int id) || id <= 0 || id > UInt16.MaxValue)
 						{
 							tokenReader.LogError($"Invalid network identifier {kvp.Key} for function");
 						}
diff --git a/Engine/Source/Programs/Shared/EpicGames.UHT/Specifiers/UhtPropertyMemberSpecifiers.cs b/Engine/Source/Programs/Shared/EpicGames.UHT/Specifiers/UhtPropertyMemberSpecifiers.cs
--- a/Engine/Source/Programs/Shared/EpicGames.UHT/Specifiers/UhtPropertyMemberSpecifiers.cs
+++ b/Engine/Source/Programs/Shared/EpicGames.UHT/Specifiers/UhtPropertyMemberSpecifiers.cs
@@ -104,8 +104,7 @@ namespace EpicGames.UHT.Parsers
 				context.MessageSite.LogError("Cannot specify a property as being both BlueprintReadOnly and BlueprintReadWrite.");
 			}
 
-			string? privateAccessMD;
-			bool allowPrivateAccess = context.MetaData.TryGetValue(UhtNames.AllowPrivateAccess, out privateAccessMD) && !privateAccessMD.Equals("false", StringComparison.OrdinalIgnoreCase);
+			bool allowPrivateAccess = context.MetaData.TryGetValue(UhtNames.AllowPrivateAccess, out string? privateAccessMD) && !privateAccessMD.Equals("false", StringComparison.OrdinalIgnoreCase);
 			if (specifierContext.Scope.AccessSpecifier == UhtAccessSpecifier.Private && !allowPrivateAccess)
 			{
 				context.MessageSite.LogError("BlueprintReadWrite should not be used on private members");
@@ -130,8 +129,7 @@ namespace EpicGames.UHT.Parsers
 				context.MessageSite.LogError("Cannot specify both BlueprintReadOnly and BlueprintReadWrite or BlueprintSetter.");
 			}
 
-			string? privateAccessMD;
-			bool allowPrivateAccess = context.MetaData.TryGetValue(UhtNames.AllowPrivateAccess, out privateAccessMD) && !privateAccessMD.Equals("false", StringComparison.OrdinalIgnoreCase);
+			bool allowPrivateAccess = context.MetaData.TryGetValue(UhtNames.AllowPrivateAccess, out string? privateAccessMD) && !privateAccessMD.Equals("false", StringComparison.OrdinalIgnoreCase);
 			if (specifierContext.Scope.AccessSpecifier == UhtAccessSpecifier.Private && !allowPrivateAccess)
 			{
 				context.MessageSite.LogError("BlueprintReadOnly should not be used on private members");
diff --git a/Engine/Source/Programs/Shared/EpicGames.UHT/Tables/UhtEngineClassTable.cs b/Engine/Source/Programs/Shared/EpicGames.UHT/Tables/UhtEngineClassTable.cs
--- a/Engine/Source/Programs/Shared/EpicGames.UHT/Tables/UhtEngineClassTable.cs
+++ b/Engine/Source/Programs/Shared/EpicGames.UHT/Tables/UhtEngineClassTable.cs
@@ -50,7 +50,7 @@ namespace EpicGames.UHT.Tables
 		/// <summary>
 		/// Internal mapping from engine class name to information
 		/// </summary>
-		private readonly Dictionary<StringView, UhtEngineClass> _engineClasses = new Dictionary<StringView, UhtEngineClass>();
+		private readonly Dictionary<StringView, UhtEngineClass> _engineClasses = new();
 
 		/// <summary>
 		/// Test to see if the given class name is a property
diff --git a/Engine/Source/Programs/Shared/EpicGames.UHT/Tables/UhtExporterTable.cs b/Engine/Source/Programs/Shared/EpicGames.UHT/Tables/UhtExporterTable.cs
--- a/Engine/Source/Programs/Shared/EpicGames.UHT/Tables/UhtExporterTable.cs
+++ b/Engine/Source/Programs/Shared/EpicGames.UHT/Tables/UhtExporterTable.cs
@@ -188,7 +188,7 @@ namespace EpicGames.UHT.Tables
 	public class UhtExporterTable : IEnumerable<UhtExporter>
 	{
 
-		private readonly Dictionary<string, UhtExporter> _exporterValues = new Dictionary<string, UhtExporter>(StringComparer.OrdinalIgnoreCase);
+		private readonly Dictionary<string, UhtExporter> _exporterValues = new(StringComparer.OrdinalIgnoreCase);
 
 		/// <summary>
 		/// Return the exporter associated with the given name
@@ -223,7 +223,7 @@ namespace EpicGames.UHT.Tables
 				}
 			}
 
-			UhtExporter exporterValue = new UhtExporter(exporterAttribute, (UhtExporterDelegate)Delegate.CreateDelegate(typeof(UhtExporterDelegate), methodInfo));
+			UhtExporter exporterValue = new(exporterAttribute, (UhtExporterDelegate)Delegate.CreateDelegate(typeof(UhtExporterDelegate), methodInfo));
 			this._exporterValues.Add(exporterAttribute.Name, exporterValue);
 		}
 
diff --git a/Engine/Source/Programs/Shared/EpicGames.UHT/Tables/UhtKeywordTable.cs b/Engine/Source/Programs/Shared/EpicGames.UHT/Tables/UhtKeywordTable.cs
--- a/Engine/Source/Programs/Shared/EpicGames.UHT/Tables/UhtKeywordTable.cs
+++ b/Engine/Source/Programs/Shared/EpicGames.UHT/Tables/UhtKeywordTable.cs
@@ -245,8 +245,7 @@ namespace EpicGames.UHT.Tables
 				{
 					continue;
 				}
-				UhtKeyword info;
-				if (keywordTable.TryGetValue(token.Value, out info))
+				if (keywordTable.TryGetValue(token.Value, out UhtKeyword info))
 				{
 					if (info.DisableUsageError)
 					{
diff --git a/Engine/Source/Programs/Shared/EpicGames.UHT/Tables/UhtLocTextDefaultValueTable.cs b/Engine/Source/Programs/Shared/EpicGames.UHT/Tables/UhtLocTextDefaultValueTable.cs
--- a/Engine/Source/Programs/Shared/EpicGames.UHT/Tables/UhtLocTextDefaultValueTable.cs
+++ b/Engine/Source/Programs/Shared/EpicGames.UHT/Tables/UhtLocTextDefaultValueTable.cs
@@ -52,7 +52,7 @@ namespace EpicGames.UHT.Tables
 	public class UhtLocTextDefaultValueTable
 	{
 
-		private readonly Dictionary<StringView, UhtLocTextDefaultValue> _locTextDefaultValues = new Dictionary<StringView, UhtLocTextDefaultValue>();
+		private readonly Dictionary<StringView, UhtLocTextDefaultValue> _locTextDefaultValues = new();
 
 		/// <summary>
 		/// Return the loc text default value associated with the given name
@@ -78,7 +78,7 @@ namespace EpicGames.UHT.Tables
 				throw new UhtIceException("A loc text default value attribute must have a name");
 			}
 
-			UhtLocTextDefaultValue locTextDefaultValue = new UhtLocTextDefaultValue
+			UhtLocTextDefaultValue locTextDefaultValue = new()
 			{
 				Delegate = (UhtLocTextDefaultValueDelegate)Delegate.CreateDelegate(typeof(UhtLocTextDefaultValueDelegate), methodInfo)
 			};
diff --git a/Engine/Source/Programs/Shared/EpicGames.UHT/Tables/UhtPropertyTypeTable.cs b/Engine/Source/Programs/Shared/EpicGames.UHT/Tables/UhtPropertyTypeTable.cs
--- a/Engine/Source/Programs/Shared/EpicGames.UHT/Tables/UhtPropertyTypeTable.cs
+++ b/Engine/Source/Programs/Shared/EpicGames.UHT/Tables/UhtPropertyTypeTable.cs
@@ -154,9 +154,9 @@ namespace EpicGames.UHT.Tables
 	/// </summary>
 	public class UhtPropertyTypeTable
 	{
-		private readonly Dictionary<StringView, UhtPropertyType> _caseSensitive = new Dictionary<StringView, UhtPropertyType>();
-		private readonly Dictionary<StringView, UhtPropertyType> _caseInsensitive = new Dictionary<StringView, UhtPropertyType>(StringViewComparer.OrdinalIgnoreCase);
-		private UhtPropertyType? _defaultInternal = null;
+		private readonly Dictionary<StringView, UhtPropertyType> _caseSensitive = new();
+		private readonly Dictionary<StringView, UhtPropertyType> _caseInsensitive = new(StringViewComparer.OrdinalIgnoreCase);
+		private UhtPropertyType? _default = null;
 
 		/// <summary>
 		/// Return the default processor
@@ -165,11 +165,11 @@ namespace EpicGames.UHT.Tables
 		{
 			get
 			{
-				if (this._defaultInternal == null)
+				if (this._default == null)
 				{
 					throw new UhtIceException("No property type has been marked as default");
 				}
-				return (UhtPropertyType)this._defaultInternal;
+				return (UhtPropertyType)this._default;
 			}
 		}
 
@@ -199,7 +199,7 @@ namespace EpicGames.UHT.Tables
 				throw new UhtIceException("A property type must have a keyword or be marked as default");
 			}
 
-			UhtPropertyType propertyType = new UhtPropertyType
+			UhtPropertyType propertyType = new()
 			{
 				Delegate = (UhtResolvePropertyDelegate)Delegate.CreateDelegate(typeof(UhtResolvePropertyDelegate), methodInfo),
 				Options = propertyTypeAttribute.Options,
@@ -207,11 +207,11 @@ namespace EpicGames.UHT.Tables
 
 			if (propertyTypeAttribute.Options.HasAnyFlags(UhtPropertyTypeOptions.Default))
 			{
-				if (this._defaultInternal != null)
+				if (this._default != null)
 				{
 					throw new UhtIceException("Only one property type dispatcher can be marked as default");
 				}
-				this._defaultInternal = propertyType;
+				this._default = propertyType;
 			}
 			else if (!String.IsNullOrEmpty(propertyTypeAttribute.Keyword))
 			{
diff --git a/Engine/Source/Programs/Shared/EpicGames.UHT/Tables/UhtStructDefaultValueTable.cs b/Engine/Source/Programs/Shared/EpicGames.UHT/Tables/UhtStructDefaultValueTable.cs
--- a/Engine/Source/Programs/Shared/EpicGames.UHT/Tables/UhtStructDefaultValueTable.cs
+++ b/Engine/Source/Programs/Shared/EpicGames.UHT/Tables/UhtStructDefaultValueTable.cs
@@ -116,8 +116,8 @@ namespace EpicGames.UHT.Tables
 	/// </summary>
 	public class UhtStructDefaultValueTable
 	{
-		private readonly Dictionary<StringView, UhtStructDefaultValue> _structDefaultValues = new Dictionary<StringView, UhtStructDefaultValue>();
-		private UhtStructDefaultValue? _defaultInternal = null;
+		private readonly Dictionary<StringView, UhtStructDefaultValue> _structDefaultValues = new();
+		private UhtStructDefaultValue? _default = null;
 
 		/// <summary>
 		/// Fetch the default sanitizer
@@ -126,11 +126,11 @@ namespace EpicGames.UHT.Tables
 		{
 			get
 			{
-				if (this._defaultInternal == null)
+				if (this._default == null)
 				{
 					throw new UhtIceException("No struct default value has been marked as default");
 				}
-				return (UhtStructDefaultValue)this._defaultInternal;
+				return (UhtStructDefaultValue)this._default;
 			}
 		}
 
@@ -158,18 +158,18 @@ namespace EpicGames.UHT.Tables
 				throw new UhtIceException("A struct default value attribute must have a name or be marked as default");
 			}
 
-			UhtStructDefaultValue structDefaultValue = new UhtStructDefaultValue
+			UhtStructDefaultValue structDefaultValue = new()
 			{
 				Delegate = (UhtStructDefaultValueDelegate)Delegate.CreateDelegate(typeof(UhtStructDefaultValueDelegate), methodInfo)
 			};
 
 			if (structDefaultValueAttribute.Options.HasAnyFlags(UhtStructDefaultValueOptions.Default))
 			{
-				if (this._defaultInternal != null)
+				if (this._default != null)
 				{
 					throw new UhtIceException("Only one struct default value attribute dispatcher can be marked as default");
 				}
-				this._defaultInternal = structDefaultValue;
+				this._default = structDefaultValue;
 			}
 			else if (!String.IsNullOrEmpty(structDefaultValueAttribute.Name))
 			{
diff --git a/Engine/Source/Programs/Shared/EpicGames.UHT/Tokenizer/UhtToken.cs b/Engine/Source/Programs/Shared/EpicGames.UHT/Tokenizer/UhtToken.cs
--- a/Engine/Source/Programs/Shared/EpicGames.UHT/Tokenizer/UhtToken.cs
+++ b/Engine/Source/Programs/Shared/EpicGames.UHT/Tokenizer/UhtToken.cs
@@ -460,7 +460,7 @@ namespace EpicGames.UHT.Tokenizer
 			switch (this.TokenType)
 			{
 				case UhtTokenType.StringConst:
-					ReadOnlySpan<char> span = this.Value.Span.Slice(1, this.Value.Span.Length - 2);
+					ReadOnlySpan<char> span = this.Value.Span[1..^1];
 					int index = span.IndexOf('\\');
 					if (index == -1)
 					{
@@ -468,15 +468,15 @@ namespace EpicGames.UHT.Tokenizer
 					}
 					else
 					{
-						StringBuilder builder = new StringBuilder();
+						StringBuilder builder = new();
 						while (index >= 0)
 						{
-							builder.Append(span.Slice(0, index));
+							builder.Append(span[..index]);
 							if (span[index + 1] == 'n')
 							{
 								builder.Append('\n');
 							}
-							span = span.Slice(index + 1);
+							span = span[(index + 1)..];
 							index = span.IndexOf('\\');
 						}
 						builder.Append(span);
@@ -537,11 +537,11 @@ namespace EpicGames.UHT.Tokenizer
 		/// <returns>Resulting string</returns>
 		public StringView GetTokenString(bool respectQuotes = false)
 		{
-			StringViewBuilder builder = new StringViewBuilder();
+			StringViewBuilder builder = new();
 			switch (this.TokenType)
 			{
 				case UhtTokenType.StringConst:
-					StringView subValue = new StringView(this.Value, 1, this.Value.Span.Length - 2);
+					StringView subValue = new(this.Value, 1, this.Value.Span.Length - 2);
 					while (subValue.Span.Length > 0)
 					{
 						int slashIndex = subValue.Span.IndexOf('\\');
@@ -610,7 +610,7 @@ namespace EpicGames.UHT.Tokenizer
 		/// <returns>Joined strings</returns>
 		public static string Join(IEnumerable<UhtToken> tokens)
 		{
-			StringBuilder builder = new StringBuilder();
+			StringBuilder builder = new();
 			foreach (UhtToken token in tokens)
 			{
 				builder.Append(token.Value.ToString());
@@ -626,7 +626,7 @@ namespace EpicGames.UHT.Tokenizer
 		/// <returns>Joined strings</returns>
 		public static string Join(char separator, IEnumerable<UhtToken> tokens)
 		{
-			StringBuilder builder = new StringBuilder();
+			StringBuilder builder = new();
 			bool includeSeperator = false;
 			foreach (UhtToken token in tokens)
 			{
@@ -651,7 +651,7 @@ namespace EpicGames.UHT.Tokenizer
 		/// <returns>Joined strings</returns>
 		public static string Join(string separator, IEnumerable<UhtToken> tokens)
 		{
-			StringBuilder builder = new StringBuilder();
+			StringBuilder builder = new();
 			bool includeSeperator = false;
 			foreach (UhtToken token in tokens)
 			{
@@ -690,7 +690,7 @@ namespace EpicGames.UHT.Tokenizer
 			{
 				if (UhtFCString.IsFloatMarker(Value.Span[^1]))
 				{
-					return Single.Parse(Value.Span.Slice(0, Value.Span.Length - 1));
+					return Single.Parse(Value.Span[0..^1]);
 				}
 			}
 			return Single.Parse(Value.Span);
@@ -702,7 +702,7 @@ namespace EpicGames.UHT.Tokenizer
 			{
 				if (UhtFCString.IsFloatMarker(Value.Span[^1]))
 				{
-					return Double.Parse(Value.Span.Slice(0, Value.Span.Length - 1));
+					return Double.Parse(Value.Span[0..^1]);
 				}
 			}
 			return Double.Parse(Value.Span);
@@ -717,12 +717,12 @@ namespace EpicGames.UHT.Tokenizer
 				char c = span[^1];
 				if (UhtFCString.IsLongMarker(c))
 				{
-					span = span.Slice(0, span.Length - 1);
+					span = span[0..^1];
 				}
 				else if (UhtFCString.IsUnsignedMarker(c))
 				{
 					isUnsigned = true;
-					span = span.Slice(0, span.Length - 1);
+					span = span[0..^1];
 				}
 				else
 				{
