More coding style cleanups

diff --git a/Engine/Source/Programs/Shared/EpicGames.UHT/Tokenizer/UhtTokenBufferReader.cs b/Engine/Source/Programs/Shared/EpicGames.UHT/Tokenizer/UhtTokenBufferReader.cs
--- a/Engine/Source/Programs/Shared/EpicGames.UHT/Tokenizer/UhtTokenBufferReader.cs
+++ b/Engine/Source/Programs/Shared/EpicGames.UHT/Tokenizer/UhtTokenBufferReader.cs
@@ -15,19 +15,19 @@ namespace EpicGames.UHT.Tokenizer
 	{
 		private static readonly StringView[] s_emptyComments = Array.Empty<StringView>();
 
-		private readonly IUhtMessageSite _messageSiteInternal;
-		private List<StringView>? _commentsInternal = null;
-		private readonly List<UhtToken> _recordedTokensInternal = new List<UhtToken>();
-		private IUhtTokenPreprocessor? _tokenPreprocessorInternal = null;
-		private UhtToken _currentToken = new UhtToken(); // PeekToken must have been invoked first
+		private readonly IUhtMessageSite _messageSite;
+		private List<StringView>? _comments = null;
+		private readonly List<UhtToken> _recordedTokens = new();
+		private IUhtTokenPreprocessor? _tokenPreprocessor = null;
+		private UhtToken _currentToken = new(); // PeekToken must have been invoked first
 		private readonly ReadOnlyMemory<char> _data;
 		private int _prevPos = 0;
 		private int _prevLine = 1;
 		private bool _hasToken = false;
 		private int _preCurrentTokenInputPos = 0;
 		private int _preCurrentTokenInputLine = 1;
-		private int _inputPosInternal = 0;
-		private int _inputLineInternal = 1;
+		private int _inputPos = 0;
+		private int _inputLine = 1;
 		private int _commentsDisableCount = 0;
 		private int _committedComments = 0;
 		private List<StringView>? _savedComments = null;
@@ -44,13 +44,13 @@ namespace EpicGames.UHT.Tokenizer
 		/// <param name="input">Input source</param>
 		public UhtTokenBufferReader(IUhtMessageSite messageSite, ReadOnlyMemory<char> input)
 		{
-			this._messageSiteInternal = messageSite;
+			this._messageSite = messageSite;
 			this._data = input;
 		}
 
 		#region IUHTMessageSite Implementation
-		IUhtMessageSession IUhtMessageSite.MessageSession => this._messageSiteInternal.MessageSession;
-		IUhtMessageSource? IUhtMessageSite.MessageSource => this._messageSiteInternal.MessageSource;
+		IUhtMessageSession IUhtMessageSite.MessageSession => this._messageSite.MessageSession;
+		IUhtMessageSource? IUhtMessageSite.MessageSource => this._messageSite.MessageSource;
 		IUhtMessageLineNumber? IUhtMessageSite.MessageLineNumber => this;
 		#endregion
 
@@ -76,16 +76,16 @@ namespace EpicGames.UHT.Tokenizer
 		}
 
 		/// <inheritdoc/>
-		public int InputPos => this._hasToken ? this._preCurrentTokenInputPos : this._inputPosInternal;
+		public int InputPos => this._hasToken ? this._preCurrentTokenInputPos : this._inputPos;
 
 		/// <inheritdoc/>
 		public int InputLine
 		{
-			get => this._hasToken ? this._preCurrentTokenInputLine : this._inputLineInternal;
+			get => this._hasToken ? this._preCurrentTokenInputLine : this._inputLine;
 			set
 			{
 				ClearToken();
-				this._inputLineInternal = value;
+				this._inputLine = value;
 			}
 		}
 
@@ -94,9 +94,9 @@ namespace EpicGames.UHT.Tokenizer
 		{
 			get
 			{
-				if (this._commentsInternal != null && this._committedComments != 0)
+				if (this._comments != null && this._committedComments != 0)
 				{
-					return new ReadOnlySpan<StringView>(this._commentsInternal.ToArray(), 0, this._committedComments);
+					return new ReadOnlySpan<StringView>(this._comments.ToArray(), 0, this._committedComments);
 				}
 				else
 				{
@@ -106,7 +106,7 @@ namespace EpicGames.UHT.Tokenizer
 		}
 
 		/// <inheritdoc/>
-		public IUhtTokenPreprocessor? TokenPreprocessor { get => this._tokenPreprocessorInternal; set => this._tokenPreprocessorInternal = value; }
+		public IUhtTokenPreprocessor? TokenPreprocessor { get => this._tokenPreprocessor; set => this._tokenPreprocessor = value; }
 
 		/// <inheritdoc/>
 		public ref UhtToken PeekToken()
@@ -131,16 +131,16 @@ namespace EpicGames.UHT.Tokenizer
 		{
 			if (this._recordTokens && !this._currentToken.IsEndType())
 			{
-				this._recordedTokensInternal.Add(this._currentToken);
+				this._recordedTokens.Add(this._currentToken);
 			}
 			this._hasToken = false;
 
 			// When comments are disabled, we are still collecting comments, but aren't committing them
 			if (this._commentsDisableCount == 0)
 			{
-				if (this._commentsInternal != null)
+				if (this._comments != null)
 				{
-					this._committedComments = this._commentsInternal.Count;
+					this._committedComments = this._comments.Count;
 				}
 			}
 			else
@@ -180,7 +180,7 @@ namespace EpicGames.UHT.Tokenizer
 				// Check for end of line
 				if (c == '\r' || c == '\n')
 				{
-					--this._inputPosInternal;
+					--this._inputPos;
 					break;
 				}
 
@@ -189,7 +189,7 @@ namespace EpicGames.UHT.Tokenizer
 				{
 					if (options.HasAnyFlags(UhtRawStringOptions.DontConsumeTerminator))
 					{
-						--this._inputPosInternal;
+						--this._inputPos;
 					}
 					break;
 				}
@@ -200,7 +200,7 @@ namespace EpicGames.UHT.Tokenizer
 					char p = InternalPeekChar(span);
 					if (p == '*' || p == '/')
 					{
-						--this._inputPosInternal;
+						--this._inputPos;
 						break;
 					}
 				}
@@ -243,15 +243,15 @@ namespace EpicGames.UHT.Tokenizer
 			ReadOnlySpan<char> span = this._data.Span;
 
 			ClearToken();
-			this._prevPos = this._inputPosInternal;
-			this._prevLine = this._inputLineInternal;
+			this._prevPos = this._inputPos;
+			this._prevLine = this._inputLine;
 
 			if (this._prevPos == span.Length)
 			{
 				return new UhtToken(UhtTokenType.EndOfFile, this._prevPos, this._prevLine, this._prevPos, this._prevLine, new StringView(this._data.Slice(this._prevPos, 0)));
 			}
 
-			int lastPos = this._inputPosInternal;
+			int lastPos = this._inputPos;
 			while (true)
 			{
 				char c = InternalGetChar(span);
@@ -264,16 +264,16 @@ namespace EpicGames.UHT.Tokenizer
 				}
 				else if (c == '\n')
 				{
-					++this._inputLineInternal;
+					++this._inputLine;
 					break;
 				}
 				else
 				{
-					lastPos = this._inputPosInternal;
+					lastPos = this._inputPos;
 				}
 			}
 
-			return new UhtToken(UhtTokenType.Line, this._prevPos, this._prevLine, this._prevPos, this._prevLine, new StringView(this._data.Slice(this._prevPos, lastPos - this._prevPos)));
+			return new UhtToken(UhtTokenType.Line, this._prevPos, this._prevLine, this._prevPos, this._prevLine, new StringView(this._data[this._prevPos..lastPos]));
 		}
 
 		/// <inheritdoc/>
@@ -285,11 +285,11 @@ namespace EpicGames.UHT.Tokenizer
 		/// <inheritdoc/>
 		public void ClearComments()
 		{
-			if (this._commentsInternal != null)
+			if (this._comments != null)
 			{
 
 				// Clearing comments does not remove any uncommitted comments
-				this._commentsInternal.RemoveRange(0, this._committedComments);
+				this._comments.RemoveRange(0, this._committedComments);
 				this._committedComments = 0;
 			}
 		}
@@ -309,9 +309,9 @@ namespace EpicGames.UHT.Tokenizer
 		/// <inheritdoc/>
 		public void CommitPendingComments()
 		{
-			if (this._commentsInternal != null)
+			if (this._comments != null)
 			{
-				this._committedComments = this._commentsInternal.Count;
+				this._committedComments = this._comments.Count;
 			}
 		}
 
@@ -331,18 +331,18 @@ namespace EpicGames.UHT.Tokenizer
 			this._hasSavedState = true;
 			this._savedInputLine = this.InputLine;
 			this._savedInputPos = this.InputPos;
-			if (this._commentsInternal != null)
+			if (this._comments != null)
 			{
 				if (this._savedComments == null)
 				{
 					this._savedComments = new List<StringView>();
 				}
 				this._savedComments.Clear();
-				this._savedComments.AddRange(this._commentsInternal);
+				this._savedComments.AddRange(this._comments);
 			}
-			if (this._tokenPreprocessorInternal != null)
+			if (this._tokenPreprocessor != null)
 			{
-				this._tokenPreprocessorInternal.SaveState();
+				this._tokenPreprocessor.SaveState();
 			}
 		}
 
@@ -355,16 +355,16 @@ namespace EpicGames.UHT.Tokenizer
 			}
 			this._hasSavedState = false;
 			ClearToken();
-			this._inputPosInternal = this._savedInputPos;
-			this._inputLineInternal = this._savedInputLine;
-			if (this._savedComments != null && this._commentsInternal != null)
+			this._inputPos = this._savedInputPos;
+			this._inputLine = this._savedInputLine;
+			if (this._savedComments != null && this._comments != null)
 			{
-				this._commentsInternal.Clear();
-				this._commentsInternal.AddRange(this._savedComments);
+				this._comments.Clear();
+				this._comments.AddRange(this._savedComments);
 			}
-			if (this._tokenPreprocessorInternal != null)
+			if (this._tokenPreprocessor != null)
 			{
-				this._tokenPreprocessorInternal.RestoreState();
+				this._tokenPreprocessor.RestoreState();
 			}
 		}
 
@@ -396,7 +396,7 @@ namespace EpicGames.UHT.Tokenizer
 			{
 				throw new UhtIceException("Attempt to disable token recording when it isn't enabled");
 			}
-			this._recordedTokensInternal.Add(token);
+			this._recordedTokens.Add(token);
 		}
 
 		/// <inheritdoc/>
@@ -406,7 +406,7 @@ namespace EpicGames.UHT.Tokenizer
 			{
 				throw new UhtIceException("Attempt to disable token recording when it isn't enabled");
 			}
-			this._recordedTokensInternal.Clear();
+			this._recordedTokens.Clear();
 			this._recordTokens = false;
 		}
 
@@ -419,7 +419,7 @@ namespace EpicGames.UHT.Tokenizer
 				{
 					throw new UhtIceException("Attempt to get recorded tokens when it isn't enabled");
 				}
-				return this._recordedTokensInternal;
+				return this._recordedTokens;
 			}
 		}
 		#endregion
@@ -427,19 +427,19 @@ namespace EpicGames.UHT.Tokenizer
 		#region Internals
 		private void ClearAllComments()
 		{
-			if (this._commentsInternal != null)
+			if (this._comments != null)
 			{
-				this._commentsInternal.Clear();
+				this._comments.Clear();
 				this._committedComments = 0;
 			}
 		}
 
 		private void ClearPendingComments()
 		{
-			if (this._commentsInternal != null)
+			if (this._comments != null)
 			{
 				int startingComment = this._committedComments + this._preprocessorPendingCommentsCount;
-				this._commentsInternal.RemoveRange(startingComment, this._commentsInternal.Count - startingComment);
+				this._comments.RemoveRange(startingComment, this._comments.Count - startingComment);
 			}
 		}
 
@@ -471,15 +471,15 @@ namespace EpicGames.UHT.Tokenizer
 			ReadOnlySpan<char> span = this._data.Span;
 			bool gotInlineComment = false;
 Restart:
-			this._preCurrentTokenInputLine = this._inputLineInternal;
-			this._preCurrentTokenInputPos = this._inputPosInternal;
+			this._preCurrentTokenInputLine = this._inputLine;
+			this._preCurrentTokenInputPos = this._inputPos;
 			SkipWhitespaceAndCommentsInternal(ref gotInlineComment, enablePreprocessor);
-			this._prevPos = this._inputPosInternal;
-			this._prevLine = this._inputLineInternal;
+			this._prevPos = this._inputPos;
+			this._prevLine = this._inputLine;
 
 			UhtTokenType tokenType = UhtTokenType.EndOfFile;
-			int startPos = this._inputPosInternal;
-			int startLine = this._inputLineInternal;
+			int startPos = this._inputPos;
+			int startLine = this._inputLine;
 
 			char c = InternalGetChar(span);
 			if (c == 0)
@@ -489,16 +489,16 @@ Restart:
 			else if (UhtFCString.IsAlpha(c) || c == '_')
 			{
 
-				for (; this._inputPosInternal < span.Length; ++this._inputPosInternal)
+				for (; this._inputPos < span.Length; ++this._inputPos)
 				{
-					c = span[this._inputPosInternal];
+					c = span[this._inputPos];
 					if (!(UhtFCString.IsAlpha(c) || UhtFCString.IsDigit(c) || c == '_'))
 					{
 						break;
 					}
 				}
 
-				if (this._inputPosInternal - startPos >= UhtToken.MaxNameLength)
+				if (this._inputPos - startPos >= UhtToken.MaxNameLength)
 				{
 					throw new UhtException(this, $"Identifier length exceeds maximum of {UhtToken.MaxNameLength}");
 				}
@@ -533,7 +533,7 @@ Restart:
 				// If we have a hex constant
 				if (isHex)
 				{
-					for (; this._inputPosInternal < span.Length && UhtFCString.IsHexDigit(span[this._inputPosInternal]); ++this._inputPosInternal)
+					for (; this._inputPos < span.Length && UhtFCString.IsHexDigit(span[this._inputPos]); ++this._inputPos)
 					{
 					}
 				}
@@ -543,56 +543,56 @@ Restart:
 				{
 
 					// Skip all digits
-					for (; this._inputPosInternal < span.Length && UhtFCString.IsDigit(span[this._inputPosInternal]); ++this._inputPosInternal)
+					for (; this._inputPos < span.Length && UhtFCString.IsDigit(span[this._inputPos]); ++this._inputPos)
 					{
 					}
 
 					// If we have a '.'
-					if (this._inputPosInternal < span.Length && span[this._inputPosInternal] == '.')
+					if (this._inputPos < span.Length && span[this._inputPos] == '.')
 					{
 						isFloat = true;
-						++this._inputPosInternal;
+						++this._inputPos;
 
 						// Skip all digits
-						for (; this._inputPosInternal < span.Length && UhtFCString.IsDigit(span[this._inputPosInternal]); ++this._inputPosInternal)
+						for (; this._inputPos < span.Length && UhtFCString.IsDigit(span[this._inputPos]); ++this._inputPos)
 						{
 						}
 					}
 
 					// If we have a 'e'
-					if (this._inputPosInternal < span.Length && UhtFCString.IsExponentMarker(span[this._inputPosInternal]))
+					if (this._inputPos < span.Length && UhtFCString.IsExponentMarker(span[this._inputPos]))
 					{
 						isFloat = true;
-						++this._inputPosInternal;
+						++this._inputPos;
 
 						// Skip any signs
-						if (this._inputPosInternal < span.Length && UhtFCString.IsSign(span[this._inputPosInternal]))
+						if (this._inputPos < span.Length && UhtFCString.IsSign(span[this._inputPos]))
 						{
-							++this._inputPosInternal;
+							++this._inputPos;
 						}
 
 						// Skip all digits
-						for (; this._inputPosInternal < span.Length && UhtFCString.IsDigit(span[this._inputPosInternal]); ++this._inputPosInternal)
+						for (; this._inputPos < span.Length && UhtFCString.IsDigit(span[this._inputPos]); ++this._inputPos)
 						{
 						}
 					}
 
 					// If we have a 'f'
-					if (this._inputPosInternal < span.Length && UhtFCString.IsFloatMarker(span[this._inputPosInternal]))
+					if (this._inputPos < span.Length && UhtFCString.IsFloatMarker(span[this._inputPos]))
 					{
 						isFloat = true;
-						++this._inputPosInternal;
+						++this._inputPos;
 					}
 
 					// Check for u/l markers
-					while (this._inputPosInternal < span.Length &&
-						(UhtFCString.IsUnsignedMarker(span[this._inputPosInternal]) || UhtFCString.IsLongMarker(span[this._inputPosInternal])))
+					while (this._inputPos < span.Length &&
+						(UhtFCString.IsUnsignedMarker(span[this._inputPos]) || UhtFCString.IsLongMarker(span[this._inputPos])))
 					{
-						++this._inputPosInternal;
+						++this._inputPos;
 					}
 				}
 
-				if (this._inputPosInternal - startPos >= UhtToken.MaxNameLength)
+				if (this._inputPos - startPos >= UhtToken.MaxNameLength)
 				{
 					throw new UhtException(this, $"Number length exceeds maximum of {UhtToken.MaxNameLength}");
 				}
@@ -646,7 +646,7 @@ Restart:
 					}
 				}
 
-				if (this._inputPosInternal - startPos >= UhtToken.MaxStringLength)
+				if (this._inputPos - startPos >= UhtToken.MaxStringLength)
 				{
 					throw new UhtException(this, $"String constant exceeds maximum of {UhtToken.MaxStringLength} characters");
 				}
@@ -659,9 +659,9 @@ Restart:
 			else
 			{
 				{
-					if (this._inputPosInternal < span.Length)
+					if (this._inputPos < span.Length)
 					{
-						char d = span[this._inputPosInternal];
+						char d = span[this._inputPos];
 						if ((c == '<' && d == '<')
 							|| (c == '!' && d == '=')
 							|| (c == '<' && d == '=')
@@ -680,7 +680,7 @@ Restart:
 							|| (c == '~' && d == '=')
 							|| (c == ':' && d == ':'))
 						{
-							++this._inputPosInternal;
+							++this._inputPos;
 						}
 					}
 
@@ -691,11 +691,11 @@ Restart:
 					// 1) If the block is being skipped, we want to preserve any comments ahead of the token and eliminate all other comments found in the #if block
 					// 2) If the block isn't being skipped or in the case of things such as #pragma or #include, they were considered statements in the UHT
 					//	  and would result in the comments being dropped.  In the new UHT, we just eliminate all pending tokens.
-					if (c == '#' && enablePreprocessor && this._tokenPreprocessorInternal != null && IsFirstTokenInLine(span, startPos))
+					if (c == '#' && enablePreprocessor && this._tokenPreprocessor != null && IsFirstTokenInLine(span, startPos))
 					{
-						this._preprocessorPendingCommentsCount = this._commentsInternal != null ? this._commentsInternal.Count - this._committedComments : 0;
-						UhtToken temp = new UhtToken(tokenType, startPos, startLine, startPos, this._inputLineInternal, new StringView(this._data.Slice(startPos, this._inputPosInternal - startPos)));
-						bool include = this._tokenPreprocessorInternal.ParsePreprocessorDirective(ref temp, true, out bool clearComments, out bool illegalContentsCheck);
+						this._preprocessorPendingCommentsCount = this._comments != null ? this._comments.Count - this._committedComments : 0;
+						UhtToken temp = new(tokenType, startPos, startLine, startPos, this._inputLine, new StringView(this._data[startPos..this._inputPos]));
+						bool include = this._tokenPreprocessor.ParsePreprocessorDirective(ref temp, true, out bool clearComments, out bool illegalContentsCheck);
 						if (!include)
 						{
 							++this._commentsDisableCount;
@@ -710,7 +710,7 @@ Restart:
 								ReadOnlySpan<char> localValueSpan = localToken.Value.Span;
 								if (localToken.IsSymbol('#') && localValueSpan.Length == 1 && localValueSpan[0] == '#' && IsFirstTokenInLine(span, localToken.InputStartPos))
 								{
-									if (this._tokenPreprocessorInternal.ParsePreprocessorDirective(ref temp, false, out bool _, out bool scratchIllegalContentsCheck))
+									if (this._tokenPreprocessor.ParsePreprocessorDirective(ref temp, false, out bool _, out bool scratchIllegalContentsCheck))
 									{
 										break;
 									}
@@ -826,7 +826,7 @@ ResetStateMachineCheck:
 				tokenType = UhtTokenType.Symbol;
 			}
 
-			return new UhtToken(tokenType, startPos, startLine, startPos, this._inputLineInternal, new StringView(this._data.Slice(startPos, this._inputPosInternal - startPos)));
+			return new UhtToken(tokenType, startPos, startLine, startPos, this._inputLine, new StringView(this._data[startPos..this._inputPos]));
 		}
 
 		/// <summary>
@@ -848,7 +848,7 @@ ResetStateMachineCheck:
 				else if (c == '\n')
 				{
 					gotNewlineBetweenComments |= gotInlineComment;
-					++this._inputLineInternal;
+					++this._inputLine;
 				}
 				else if (c == '\r' || c == '\t' || c == ' ')
 				{
@@ -862,8 +862,8 @@ ResetStateMachineCheck:
 						{
 							ClearAllComments();
 						}
-						int commentStart = this._inputPosInternal - 1;
-						++this._inputPosInternal;
+						int commentStart = this._inputPos - 1;
+						++this._inputPos;
 						for (; ; )
 						{
 							char commentChar = InternalGetChar(span);
@@ -877,17 +877,17 @@ ResetStateMachineCheck:
 							}
 							else if (commentChar == '\n')
 							{
-								++this._inputLineInternal;
+								++this._inputLine;
 							}
 							else if (commentChar == '*' && InternalPeekChar(span) == '/')
 							{
-								++this._inputPosInternal;
+								++this._inputPos;
 								break;
 							}
 						}
 						if (enablePreprocessor)
 						{
-							AddComment(new StringView(this._data.Slice(commentStart, this._inputPosInternal - commentStart)));
+							AddComment(new StringView(this._data[commentStart..this._inputPos]));
 						}
 					}
 					else if (nextChar == '/')
@@ -901,8 +901,8 @@ ResetStateMachineCheck:
 							}
 						}
 						gotInlineComment = true;
-						int commentStart = this._inputPosInternal - 1;
-						++this._inputPosInternal;
+						int commentStart = this._inputPos - 1;
+						++this._inputPos;
 
 						// Scan to the end of the line
 						for (; ; )
@@ -918,24 +918,24 @@ ResetStateMachineCheck:
 							}
 							else if (commentChar == '\n')
 							{
-								++_inputLineInternal;
+								++_inputLine;
 								break;
 							}
 						}
 						if (enablePreprocessor)
 						{
-							AddComment(new StringView(this._data.Slice(commentStart, this._inputPosInternal - commentStart)));
+							AddComment(new StringView(this._data[commentStart..this._inputPos]));
 						}
 					}
 					else
 					{
-						--this._inputPosInternal;
+						--this._inputPos;
 						break;
 					}
 				}
 				else
 				{
-					--this._inputPosInternal;
+					--this._inputPos;
 					break;
 				}
 			}
@@ -954,29 +954,29 @@ ResetStateMachineCheck:
 			// Check for [+-]...
 			if (UhtFCString.IsSign(c))
 			{
-				if (this._inputPosInternal == span.Length)
+				if (this._inputPos == span.Length)
 				{
 					return false;
 				}
 
 				// Check for [+-][0..9]...
-				if (UhtFCString.IsDigit(span[this._inputPosInternal]))
+				if (UhtFCString.IsDigit(span[this._inputPos]))
 				{
 					return true;
 				}
 
 				// Check for [+-][.][0..9]...
-				if (span[this._inputPosInternal] != '.')
+				if (span[this._inputPos] != '.')
 				{
 					return false;
 				}
 
-				if (this._inputPosInternal + 1 == span.Length)
+				if (this._inputPos + 1 == span.Length)
 				{
 					return false;
 				}
 
-				if (UhtFCString.IsDigit(span[this._inputPosInternal + 1]))
+				if (UhtFCString.IsDigit(span[this._inputPos + 1]))
 				{
 					return true;
 				}
@@ -985,13 +985,13 @@ ResetStateMachineCheck:
 
 			if (c == '.')
 			{
-				if (this._inputPosInternal == span.Length)
+				if (this._inputPos == span.Length)
 				{
 					return false;
 				}
 
 				// Check for [.][0..9]...
-				if (UhtFCString.IsDigit(span[this._inputPosInternal]))
+				if (UhtFCString.IsDigit(span[this._inputPos]))
 				{
 					return true;
 				}
@@ -1007,7 +1007,7 @@ ResetStateMachineCheck:
 		/// <returns>Next character in the stream or zero</returns>
 		private char InternalPeekChar(ReadOnlySpan<char> span)
 		{
-			return this._inputPosInternal < span.Length ? span[this._inputPosInternal] : '\0';
+			return this._inputPos < span.Length ? span[this._inputPos] : '\0';
 		}
 
 		/// <summary>
@@ -1018,7 +1018,7 @@ ResetStateMachineCheck:
 		/// <returns>Next character in the stream or zero</returns>
 		private char InternalGetChar(ReadOnlySpan<char> span)
 		{
-			return this._inputPosInternal < span.Length ? span[this._inputPosInternal++] : '\0';
+			return this._inputPos < span.Length ? span[this._inputPos++] : '\0';
 		}
 
 		/// <summary>
@@ -1029,22 +1029,22 @@ ResetStateMachineCheck:
 			if (this._hasToken)
 			{
 				this._hasToken = false;
-				if (this._commentsInternal != null && this._commentsInternal.Count > this._committedComments)
+				if (this._comments != null && this._comments.Count > this._committedComments)
 				{
-					this._commentsInternal.RemoveRange(this._committedComments, this._commentsInternal.Count - this._committedComments);
+					this._comments.RemoveRange(this._committedComments, this._comments.Count - this._committedComments);
 				}
-				this._inputPosInternal = this._currentToken.UngetPos;
-				this._inputLineInternal = this._currentToken.UngetLine;
+				this._inputPos = this._currentToken.UngetPos;
+				this._inputLine = this._currentToken.UngetLine;
 			}
 		}
 
 		private void AddComment(StringView comment)
 		{
-			if (this._commentsInternal == null)
+			if (this._comments == null)
 			{
-				this._commentsInternal = new List<StringView>(4);
+				this._comments = new List<StringView>(4);
 			}
-			this._commentsInternal.Add(comment);
+			this._comments.Add(comment);
 		}
 		#endregion
 	}
diff --git a/Engine/Source/Programs/Shared/EpicGames.UHT/Tokenizer/UhtTokenReader.cs b/Engine/Source/Programs/Shared/EpicGames.UHT/Tokenizer/UhtTokenReader.cs
--- a/Engine/Source/Programs/Shared/EpicGames.UHT/Tokenizer/UhtTokenReader.cs
+++ b/Engine/Source/Programs/Shared/EpicGames.UHT/Tokenizer/UhtTokenReader.cs
@@ -291,7 +291,7 @@ namespace EpicGames.UHT.Tokenizer
 			{
 				return this.Token.Value.ToString();
 			}
-			StringBuilder builder = new StringBuilder();
+			StringBuilder builder = new();
 			Join(builder, seperator);
 			return builder.ToString();
 		}
@@ -324,7 +324,7 @@ namespace EpicGames.UHT.Tokenizer
 	/// </summary>
 	public static class UhtTokenListCache
 	{
-		private static readonly ThreadLocal<UhtTokenList?> s_tls = new ThreadLocal<UhtTokenList?>(() => null);
+		private static readonly ThreadLocal<UhtTokenList?> s_tls = new(() => null);
 
 		/// <summary>
 		/// Borrow a token list
diff --git a/Engine/Source/Programs/Shared/EpicGames.UHT/Tokenizer/UhtTokenReaderFloatExtensions.cs b/Engine/Source/Programs/Shared/EpicGames.UHT/Tokenizer/UhtTokenReaderFloatExtensions.cs
--- a/Engine/Source/Programs/Shared/EpicGames.UHT/Tokenizer/UhtTokenReaderFloatExtensions.cs
+++ b/Engine/Source/Programs/Shared/EpicGames.UHT/Tokenizer/UhtTokenReaderFloatExtensions.cs
@@ -77,8 +77,7 @@ namespace EpicGames.UHT.Tokenizer
 		/// <returns>The floating point value of the token</returns>
 		public static float GetConstFloat(this IUhtTokenReader tokenReader, object? exceptionContext = null)
 		{
-			float value;
-			if (!tokenReader.TryOptionalConstFloat(out value))
+			if (!tokenReader.TryOptionalConstFloat(out float value))
 			{
 				throw new UhtTokenException(tokenReader, tokenReader.PeekToken(), "constant float", exceptionContext);
 			}
diff --git a/Engine/Source/Programs/Shared/EpicGames.UHT/Tokenizer/UhtTokenReaderFlowControlExtensions.cs b/Engine/Source/Programs/Shared/EpicGames.UHT/Tokenizer/UhtTokenReaderFlowControlExtensions.cs
--- a/Engine/Source/Programs/Shared/EpicGames.UHT/Tokenizer/UhtTokenReaderFlowControlExtensions.cs
+++ b/Engine/Source/Programs/Shared/EpicGames.UHT/Tokenizer/UhtTokenReaderFlowControlExtensions.cs
@@ -216,7 +216,7 @@ namespace EpicGames.UHT.Tokenizer
 			// Process the body
 			while (true)
 			{
-				List<UhtToken> tokens = new List<UhtToken>();
+				List<UhtToken> tokens = new();
 
 				// Read the tokens until we hit the end
 				while (true)
diff --git a/Engine/Source/Programs/Shared/EpicGames.UHT/Tokenizer/UhtTokenReaderIdentifierExtensions.cs b/Engine/Source/Programs/Shared/EpicGames.UHT/Tokenizer/UhtTokenReaderIdentifierExtensions.cs
--- a/Engine/Source/Programs/Shared/EpicGames.UHT/Tokenizer/UhtTokenReaderIdentifierExtensions.cs
+++ b/Engine/Source/Programs/Shared/EpicGames.UHT/Tokenizer/UhtTokenReaderIdentifierExtensions.cs
@@ -205,81 +205,78 @@ namespace EpicGames.UHT.Tokenizer
 		public static bool TryOptionalCppIdentifier(this IUhtTokenReader tokenReader, [NotNullWhen(true)] out IEnumerable<UhtToken>? identifier, UhtCppIdentifierOptions options = UhtCppIdentifierOptions.None)
 		{
 			identifier = null;
-			List<UhtToken> localIdentifier = new List<UhtToken>();
+			List<UhtToken> localIdentifier = new();
 
-			using (UhtTokenSaveState savedState = new UhtTokenSaveState(tokenReader))
-			{
-				UhtToken token;
+			using UhtTokenSaveState savedState = new(tokenReader);
 
-				if (tokenReader.TryOptionalIdentifier(out token))
-				{
-					localIdentifier.Add(token);
-				}
-				else
-				{
-					return false;
-				}
+			if (tokenReader.TryOptionalIdentifier(out UhtToken token))
+			{
+				localIdentifier.Add(token);
+			}
+			else
+			{
+				return false;
+			}
 
-				if (options.HasAnyFlags(UhtCppIdentifierOptions.AllowTemplates))
+			if (options.HasAnyFlags(UhtCppIdentifierOptions.AllowTemplates))
+			{
+				while (true)
 				{
-					while (true)
+					if (tokenReader.TryPeekOptional('<'))
 					{
-						if (tokenReader.TryPeekOptional('<'))
+						localIdentifier.Add(tokenReader.GetToken());
+						int nestedScopes = 1;
+						while (nestedScopes > 0)
 						{
-							localIdentifier.Add(tokenReader.GetToken());
-							int nestedScopes = 1;
-							while (nestedScopes > 0)
+							token = tokenReader.GetToken();
+							if (token.TokenType.IsEndType())
 							{
-								token = tokenReader.GetToken();
-								if (token.TokenType.IsEndType())
-								{
-									return false;
-								}
-								localIdentifier.Add(token);
-								if (token.IsSymbol('<'))
-								{
-									++nestedScopes;
-								}
-								else if (token.IsSymbol('>'))
-								{
-									--nestedScopes;
-								}
+								return false;
 							}
-						}
-
-						if (!tokenReader.TryPeekOptional("::"))
-						{
-							break;
-						}
-						localIdentifier.Add(tokenReader.GetToken());
-						if (!tokenReader.TryOptionalIdentifier(out token))
-						{
 							localIdentifier.Add(token);
+							if (token.IsSymbol('<'))
+							{
+								++nestedScopes;
+							}
+							else if (token.IsSymbol('>'))
+							{
+								--nestedScopes;
+							}
 						}
-						else
-						{
-							return false;
-						}
+					}
+
+					if (!tokenReader.TryPeekOptional("::"))
+					{
+						break;
+					}
+					localIdentifier.Add(tokenReader.GetToken());
+					if (!tokenReader.TryOptionalIdentifier(out token))
+					{
+						localIdentifier.Add(token);
+					}
+					else
+					{
+						return false;
 					}
 				}
-				else
+			}
+			else
+			{
+				while (tokenReader.TryOptional("::"))
 				{
-					while (tokenReader.TryOptional("::"))
+					if (tokenReader.TryOptionalIdentifier(out token))
 					{
-						if (tokenReader.TryOptionalIdentifier(out token))
-						{
-							localIdentifier.Add(token);
-						}
-						else
-						{
-							return false;
-						}
+						localIdentifier.Add(token);
+					}
+					else
+					{
+						return false;
 					}
 				}
-				identifier = localIdentifier;
-				savedState.AbandonState();
-				return true;
 			}
+			identifier = localIdentifier;
+			savedState.AbandonState();
+			return true;
 		}
 
 		/// <summary>
diff --git a/Engine/Source/Programs/Shared/EpicGames.UHT/Tokenizer/UhtTokenReaderIntegerExtensions.cs b/Engine/Source/Programs/Shared/EpicGames.UHT/Tokenizer/UhtTokenReaderIntegerExtensions.cs
--- a/Engine/Source/Programs/Shared/EpicGames.UHT/Tokenizer/UhtTokenReaderIntegerExtensions.cs
+++ b/Engine/Source/Programs/Shared/EpicGames.UHT/Tokenizer/UhtTokenReaderIntegerExtensions.cs
@@ -70,8 +70,7 @@ namespace EpicGames.UHT.Tokenizer
 		/// <returns>The value of the constant</returns>
 		public static int GetConstInt(this IUhtTokenReader tokenReader, object? exceptionContext = null)
 		{
-			int value;
-			if (!tokenReader.TryOptionalConstInt(out value))
+			if (!tokenReader.TryOptionalConstInt(out int value))
 			{
 				throw new UhtTokenException(tokenReader, tokenReader.PeekToken(), "constant integer", exceptionContext);
 			}
@@ -135,8 +134,7 @@ namespace EpicGames.UHT.Tokenizer
 		/// <returns>The value of the constant</returns>
 		public static long GetConstLong(this IUhtTokenReader tokenReader, object? exceptionContext = null)
 		{
-			long value;
-			if (!tokenReader.TryOptionalConstLong(out value))
+			if (!tokenReader.TryOptionalConstLong(out long value))
 			{
 				throw new UhtTokenException(tokenReader, tokenReader.PeekToken(), "constant long integer", exceptionContext);
 			}
@@ -151,30 +149,28 @@ namespace EpicGames.UHT.Tokenizer
 		/// <returns>True if the next token was an parsed, false if not.</returns>
 		public static bool TryOptionalLeadingSignConstNumeric(this IUhtTokenReader tokenReader, UhtParseMergedSignToken tokenDelegate)
 		{
-			using (UhtTokenSaveState savedState = new UhtTokenSaveState(tokenReader))
+			using UhtTokenSaveState savedState = new(tokenReader);
+			// Check for a leading sign token
+			char sign = ' ';
+			UhtToken token = tokenReader.PeekToken();
+			if (token.IsSymbol() && token.Value.Length == 1 && UhtFCString.IsSign(token.Value.Span[0]))
 			{
-				// Check for a leading sign token
-				char sign = ' ';
-				UhtToken token = tokenReader.PeekToken();
-				if (token.IsSymbol() && token.Value.Length == 1 && UhtFCString.IsSign(token.Value.Span[0]))
-				{
-					sign = token.Value.Span[0];
-					tokenReader.ConsumeToken();
-					token = tokenReader.PeekToken();
-					if (UhtFCString.IsSign(token.Value.Span[0]))
-					{
-						return false;
-					}
-					token.Value = new StringView($"{sign}{token.Value}");
-				}
-				if (tokenDelegate(ref token))
+				sign = token.Value.Span[0];
+				tokenReader.ConsumeToken();
+				token = tokenReader.PeekToken();
+				if (UhtFCString.IsSign(token.Value.Span[0]))
 				{
-					tokenReader.ConsumeToken();
-					savedState.AbandonState();
-					return true;
+					return false;
 				}
-				return false;
+				token.Value = new StringView($"{sign}{token.Value}");
+			}
+			if (tokenDelegate(ref token))
+			{
+				tokenReader.ConsumeToken();
+				savedState.AbandonState();
+				return true;
 			}
+			return false;
 		}
 
 		/// <summary>
diff --git a/Engine/Source/Programs/Shared/EpicGames.UHT/Tokenizer/UhtTokenReaderUtilityExtensions.cs b/Engine/Source/Programs/Shared/EpicGames.UHT/Tokenizer/UhtTokenReaderUtilityExtensions.cs
--- a/Engine/Source/Programs/Shared/EpicGames.UHT/Tokenizer/UhtTokenReaderUtilityExtensions.cs
+++ b/Engine/Source/Programs/Shared/EpicGames.UHT/Tokenizer/UhtTokenReaderUtilityExtensions.cs
@@ -13,7 +13,7 @@ namespace EpicGames.UHT.Tokenizer
 	/// </summary>
 	public static class UhtTokenReaderUtilityExtensions
 	{
-		private static readonly HashSet<StringView> s_skipDeclarationWarningStrings = new HashSet<StringView>
+		private static readonly HashSet<StringView> s_skipDeclarationWarningStrings = new()
 		{
 			"GENERATED_BODY",
 			"GENERATED_IINTERFACE_BODY",
@@ -128,8 +128,8 @@ namespace EpicGames.UHT.Tokenizer
 			bool endOfDeclarationFound = false;
 
 			// Store the current value of PrevComment so it can be restored after we parsed everything.
-			using (UhtTokenDisableComments disableComments = new UhtTokenDisableComments(tokenReader))
 			{
+				using UhtTokenDisableComments disableComments = new(tokenReader);
 
 				// Check if this is a class/struct declaration in which case it can be followed by member variable declaration.	
 				bool possiblyClassDeclaration = token.IsIdentifier() && (token.IsValue("class") || token.IsValue("struct"));
@@ -281,7 +281,7 @@ namespace EpicGames.UHT.Tokenizer
 			}
 
 			// Failing that, we'll guess about it being a macro based on it being a fully-capitalized identifier.
-			foreach (char ch in span.Slice(1))
+			foreach (char ch in span[1..])
 			{
 				if (ch != '_' && (ch < 'A' || ch > 'Z') && (ch < '0' || ch > '9'))
 				{
diff --git a/Engine/Source/Programs/Shared/EpicGames.UHT/Tokenizer/UhtTokenReplayReader.cs b/Engine/Source/Programs/Shared/EpicGames.UHT/Tokenizer/UhtTokenReplayReader.cs
--- a/Engine/Source/Programs/Shared/EpicGames.UHT/Tokenizer/UhtTokenReplayReader.cs
+++ b/Engine/Source/Programs/Shared/EpicGames.UHT/Tokenizer/UhtTokenReplayReader.cs
@@ -14,21 +14,21 @@ namespace EpicGames.UHT.Tokenizer
 	/// </summary>
 	public sealed class UhtTokenReplayReader : IUhtTokenReader, IUhtMessageLineNumber
 	{
-		private static readonly ThreadLocal<UhtTokenReplayReader> s_tls = new ThreadLocal<UhtTokenReplayReader>(() => new UhtTokenReplayReader());
+		private static readonly ThreadLocal<UhtTokenReplayReader> s_tls = new(() => new UhtTokenReplayReader());
 
 		const int MaxSavedStates = 2;
 		private struct SavedState
 		{
-			public int _tokenIndex;
-			public bool _hasToken;
+			public int TokenIndex { get; set; }
+			public bool HasToken { get; set; }
 		}
 
-		private IUhtMessageSite _messageSiteInternal;
+		private IUhtMessageSite _messageSite;
 		private ReadOnlyMemory<UhtToken> _tokens;
 		private ReadOnlyMemory<char> _data;
 		private int _currentTokenIndex = -1;
 		private bool _hasToken = false;
-		private UhtToken _currentToken = new UhtToken();
+		private UhtToken _currentToken = new();
 		private readonly SavedState[] _savedStates = new SavedState[MaxSavedStates];
 		private int _savedStateCount = 0;
 		private UhtTokenType _endTokenType = UhtTokenType.EndOfFile;
@@ -42,7 +42,7 @@ namespace EpicGames.UHT.Tokenizer
 		/// <param name="endTokenType">Token type to return when end of tokens reached</param>
 		public UhtTokenReplayReader(IUhtMessageSite messageSite, ReadOnlyMemory<char> data, ReadOnlyMemory<UhtToken> tokens, UhtTokenType endTokenType)
 		{
-			this._messageSiteInternal = messageSite;
+			this._messageSite = messageSite;
 			this._tokens = tokens;
 			this._data = data;
 			this._endTokenType = endTokenType;
@@ -54,7 +54,7 @@ namespace EpicGames.UHT.Tokenizer
 		/// </summary>
 		public UhtTokenReplayReader()
 		{
-			this._messageSiteInternal = new UhtEmptyMessageSite();
+			this._messageSite = new UhtEmptyMessageSite();
 			this._tokens = Array.Empty<UhtToken>().AsMemory();
 			this._data = Array.Empty<char>().AsMemory();
 		}
@@ -69,7 +69,7 @@ namespace EpicGames.UHT.Tokenizer
 		/// <returns>The replay reader</returns>
 		public UhtTokenReplayReader Reset(IUhtMessageSite messageSite, ReadOnlyMemory<char> data, ReadOnlyMemory<UhtToken> tokens, UhtTokenType endTokenType)
 		{
-			this._messageSiteInternal = messageSite;
+			this._messageSite = messageSite;
 			this._tokens = tokens;
 			this._data = data;
 			this._currentTokenIndex = -1;
@@ -101,8 +101,8 @@ namespace EpicGames.UHT.Tokenizer
 		}
 
 		#region IUHTMessageSite Implementation
-		IUhtMessageSession IUhtMessageSite.MessageSession => this._messageSiteInternal.MessageSession;
-		IUhtMessageSource? IUhtMessageSite.MessageSource => this._messageSiteInternal.MessageSource;
+		IUhtMessageSession IUhtMessageSite.MessageSession => this._messageSite.MessageSession;
+		IUhtMessageSource? IUhtMessageSite.MessageSource => this._messageSite.MessageSource;
 		IUhtMessageLineNumber? IUhtMessageSite.MessageLineNumber => this;
 		#endregion
 
@@ -241,7 +241,7 @@ namespace EpicGames.UHT.Tokenizer
 			{
 				throw new UhtIceException("Token reader saved states full");
 			}
-			this._savedStates[this._savedStateCount] = new SavedState { _tokenIndex = this._currentTokenIndex, _hasToken = this._hasToken };
+			this._savedStates[this._savedStateCount] = new SavedState { TokenIndex = this._currentTokenIndex, HasToken = this._hasToken };
 			++this._savedStateCount;
 		}
 
@@ -254,8 +254,8 @@ namespace EpicGames.UHT.Tokenizer
 			}
 
 			--this._savedStateCount;
-			this._currentTokenIndex = this._savedStates[this._savedStateCount]._tokenIndex;
-			this._hasToken = this._savedStates[this._savedStateCount]._hasToken;
+			this._currentTokenIndex = this._savedStates[this._savedStateCount].TokenIndex;
+			this._hasToken = this._savedStates[this._savedStateCount].HasToken;
 		}
 
 		/// <inheritdoc/>
diff --git a/Engine/Source/Programs/Shared/EpicGames.UHT/Types/Properties/UhtArrayProperty.cs b/Engine/Source/Programs/Shared/EpicGames.UHT/Types/Properties/UhtArrayProperty.cs
--- a/Engine/Source/Programs/Shared/EpicGames.UHT/Types/Properties/UhtArrayProperty.cs
+++ b/Engine/Source/Programs/Shared/EpicGames.UHT/Types/Properties/UhtArrayProperty.cs
@@ -212,50 +212,48 @@ namespace EpicGames.UHT.Types
 		[SuppressMessage("CodeQuality", "IDE0051:Remove unused private members", Justification = "Attribute accessed method")]
 		private static UhtProperty? ArrayProperty(UhtPropertyResolvePhase resolvePhase, UhtPropertySettings propertySettings, IUhtTokenReader tokenReader, UhtToken matchedToken)
 		{
-			using (UhtMessageContext tokenContext = new UhtMessageContext("TArray"))
+			using UhtMessageContext tokenContext = new("TArray");
+			if (!tokenReader.SkipExpectedType(matchedToken.Value, propertySettings.PropertyCategory == UhtPropertyCategory.Member))
 			{
-				if (!tokenReader.SkipExpectedType(matchedToken.Value, propertySettings.PropertyCategory == UhtPropertyCategory.Member))
-				{
-					return null;
-				}
-				tokenReader.Require('<');
+				return null;
+			}
+			tokenReader.Require('<');
+
+			// Parse the value type
+			UhtProperty? value = UhtPropertyParser.ParseTemplateParam(resolvePhase, propertySettings, propertySettings.SourceName, tokenReader);
+			if (value == null)
+			{
+				return null;
+			}
 
-				// Parse the value type
-				UhtProperty? value = UhtPropertyParser.ParseTemplateParam(resolvePhase, propertySettings, propertySettings.SourceName, tokenReader);
-				if (value == null)
+			if (!value.PropertyCaps.HasAnyFlags(UhtPropertyCaps.CanBeContainerValue))
+			{
+				tokenReader.LogError($"The type \'{value.GetUserFacingDecl()}\' can not be used as a value in a TArray");
+			}
+
+			if (tokenReader.TryOptional(','))
+			{
+				// If we found a comma, read the next thing, assume it's an allocator, and report that
+				UhtToken allocatorToken = tokenReader.GetIdentifier();
+				if (allocatorToken.IsIdentifier("FMemoryImageAllocator"))
 				{
-					return null;
+					propertySettings.Allocator = UhtPropertyAllocator.MemoryImage;
 				}
-
-				if (!value.PropertyCaps.HasAnyFlags(UhtPropertyCaps.CanBeContainerValue))
+				else if (allocatorToken.IsIdentifier("TMemoryImageAllocator"))
 				{
-					tokenReader.LogError($"The type \'{value.GetUserFacingDecl()}\' can not be used as a value in a TArray");
+					tokenReader.RequireList('<', '>', "TMemoryImageAllocator template arguments");
+					propertySettings.Allocator = UhtPropertyAllocator.MemoryImage;
 				}
-
-				if (tokenReader.TryOptional(','))
+				else
 				{
-					// If we found a comma, read the next thing, assume it's an allocator, and report that
-					UhtToken allocatorToken = tokenReader.GetIdentifier();
-					if (allocatorToken.IsIdentifier("FMemoryImageAllocator"))
-					{
-						propertySettings.Allocator = UhtPropertyAllocator.MemoryImage;
-					}
-					else if (allocatorToken.IsIdentifier("TMemoryImageAllocator"))
-					{
-						tokenReader.RequireList('<', '>', "TMemoryImageAllocator template arguments");
-						propertySettings.Allocator = UhtPropertyAllocator.MemoryImage;
-					}
-					else
-					{
-						throw new UhtException(tokenReader, $"Found '{allocatorToken.Value}' - explicit allocators are not supported in TArray properties.");
-					}
+					throw new UhtException(tokenReader, $"Found '{allocatorToken.Value}' - explicit allocators are not supported in TArray properties.");
 				}
-				tokenReader.Require('>');
+			}
+			tokenReader.Require('>');
 
-				//@TODO: Prevent sparse delegate types from being used in a container
+			//@TODO: Prevent sparse delegate types from being used in a container
 
-				return new UhtArrayProperty(propertySettings, value);
-			}
+			return new UhtArrayProperty(propertySettings, value);
 		}
 		#endregion
 	}
diff --git a/Engine/Source/Programs/Shared/EpicGames.UHT/Types/Properties/UhtEnumProperty.cs b/Engine/Source/Programs/Shared/EpicGames.UHT/Types/Properties/UhtEnumProperty.cs
--- a/Engine/Source/Programs/Shared/EpicGames.UHT/Types/Properties/UhtEnumProperty.cs
+++ b/Engine/Source/Programs/Shared/EpicGames.UHT/Types/Properties/UhtEnumProperty.cs
@@ -307,7 +307,7 @@ namespace EpicGames.UHT.Types
 
 		private UhtProperty CreateUnderlyingProperty()
 		{
-			UhtPropertySettings propertySettings = new UhtPropertySettings();
+			UhtPropertySettings propertySettings = new();
 			propertySettings.Reset(this, 0, this.PropertyCategory, 0);
 			propertySettings.SourceName = "UnderlyingType";
 			switch (this.UnderlyingType)
@@ -341,18 +341,16 @@ namespace EpicGames.UHT.Types
 		{
 			UhtType outer = propertySettings.Outer;
 			UhtEnum? enumObj = null;
-			using (UhtMessageContext tokenContext = new UhtMessageContext("TEnumAsByte"))
-			{
-				tokenReader
-					.Require("TEnumAsByte")
-					.Require('<')
-					.Optional("enum")
-					.RequireCppIdentifier(UhtCppIdentifierOptions.None, (UhtTokenList tokenList) =>
-					{
-						enumObj = outer.FindType(UhtFindOptions.Enum | UhtFindOptions.SourceName, tokenList, tokenReader) as UhtEnum;
-					})
-					.Require('>');
-			}
+			using UhtMessageContext tokenContext = new("TEnumAsByte");
+			tokenReader
+				.Require("TEnumAsByte")
+				.Require('<')
+				.Optional("enum")
+				.RequireCppIdentifier(UhtCppIdentifierOptions.None, (UhtTokenList tokenList) =>
+				{
+					enumObj = outer.FindType(UhtFindOptions.Enum | UhtFindOptions.SourceName, tokenList, tokenReader) as UhtEnum;
+				})
+				.Require('>');
 			return enumObj != null ? new UhtByteProperty(propertySettings, UhtPropertyIntType.None, enumObj) : null;
 		}
 		#endregion
diff --git a/Engine/Source/Programs/Shared/EpicGames.UHT/Types/Properties/UhtFieldPathProperty.cs b/Engine/Source/Programs/Shared/EpicGames.UHT/Types/Properties/UhtFieldPathProperty.cs
--- a/Engine/Source/Programs/Shared/EpicGames.UHT/Types/Properties/UhtFieldPathProperty.cs
+++ b/Engine/Source/Programs/Shared/EpicGames.UHT/Types/Properties/UhtFieldPathProperty.cs
@@ -110,22 +110,20 @@ namespace EpicGames.UHT.Types
 		[SuppressMessage("Style", "IDE0060:Remove unused parameter", Justification = "Attribute accessed method")]
 		private static UhtProperty? FieldPathProperty(UhtPropertyResolvePhase resolvePhase, UhtPropertySettings propertySettings, IUhtTokenReader tokenReader, UhtToken matchedToken)
 		{
-			using (UhtMessageContext tokenContext = new UhtMessageContext("TFieldPath"))
+			using UhtMessageContext tokenContext = new("TFieldPath");
+			UhtToken identifier = new();
+			tokenReader
+				.Require("TFieldPath")
+				.Require('<')
+				.RequireIdentifier((ref UhtToken token) => identifier = token)
+				.Require('>');
+
+			StringView fieldClassName = new(identifier.Value, 1);
+			if (!propertySettings.Outer.Session.IsValidPropertyTypeName(fieldClassName))
 			{
-				UhtToken identifier = new UhtToken();
-				tokenReader
-					.Require("TFieldPath")
-					.Require('<')
-					.RequireIdentifier((ref UhtToken token) => identifier = token)
-					.Require('>');
-
-				StringView fieldClassName = new StringView(identifier.Value, 1);
-				if (!propertySettings.Outer.Session.IsValidPropertyTypeName(fieldClassName))
-				{
-					throw new UhtException($"Undefined property type: {identifier.Value}");
-				}
-				return new UhtFieldPathProperty(propertySettings, fieldClassName.ToString());
+				throw new UhtException($"Undefined property type: {identifier.Value}");
 			}
+			return new UhtFieldPathProperty(propertySettings, fieldClassName.ToString());
 		}
 		#endregion
 	}
diff --git a/Engine/Source/Programs/Shared/EpicGames.UHT/Types/Properties/UhtMapProperty.cs b/Engine/Source/Programs/Shared/EpicGames.UHT/Types/Properties/UhtMapProperty.cs
--- a/Engine/Source/Programs/Shared/EpicGames.UHT/Types/Properties/UhtMapProperty.cs
+++ b/Engine/Source/Programs/Shared/EpicGames.UHT/Types/Properties/UhtMapProperty.cs
@@ -260,63 +260,61 @@ namespace EpicGames.UHT.Types
 		[SuppressMessage("CodeQuality", "IDE0051:Remove unused private members", Justification = "Attribute accessed method")]
 		private static UhtProperty? MapProperty(UhtPropertyResolvePhase resolvePhase, UhtPropertySettings propertySettings, IUhtTokenReader tokenReader, UhtToken matchedToken)
 		{
-			using (UhtMessageContext tokenContext = new UhtMessageContext("TMap"))
+			using UhtMessageContext tokenContext = new("TMap");
+			if (!tokenReader.SkipExpectedType(matchedToken.Value, propertySettings.PropertyCategory == UhtPropertyCategory.Member))
 			{
-				if (!tokenReader.SkipExpectedType(matchedToken.Value, propertySettings.PropertyCategory == UhtPropertyCategory.Member))
-				{
-					return null;
-				}
-				tokenReader.Require('<');
+				return null;
+			}
+			tokenReader.Require('<');
 
-				// Parse the key type
-				UhtProperty? key = UhtPropertyParser.ParseTemplateParam(resolvePhase, propertySettings, "Key", tokenReader);
-				if (key == null)
-				{
-					return null;
-				}
+			// Parse the key type
+			UhtProperty? key = UhtPropertyParser.ParseTemplateParam(resolvePhase, propertySettings, "Key", tokenReader);
+			if (key == null)
+			{
+				return null;
+			}
 
-				if (!key.PropertyCaps.HasAnyFlags(UhtPropertyCaps.CanBeContainerKey))
-				{
-					tokenReader.LogError($"The type \'{key.GetUserFacingDecl()}\' can not be used as a key in a TMap");
-				}
+			if (!key.PropertyCaps.HasAnyFlags(UhtPropertyCaps.CanBeContainerKey))
+			{
+				tokenReader.LogError($"The type \'{key.GetUserFacingDecl()}\' can not be used as a key in a TMap");
+			}
 
-				if (propertySettings.PropertyFlags.HasAnyFlags(EPropertyFlags.Net))
-				{
-					tokenReader.LogError("Replicated maps are not supported.");
-				}
+			if (propertySettings.PropertyFlags.HasAnyFlags(EPropertyFlags.Net))
+			{
+				tokenReader.LogError("Replicated maps are not supported.");
+			}
 
-				tokenReader.Require(',');
+			tokenReader.Require(',');
 
-				// Parse the value type
-				UhtProperty? value = UhtPropertyParser.ParseTemplateParam(resolvePhase, propertySettings, "Value", tokenReader);
-				if (value == null)
-				{
-					return null;
-				}
+			// Parse the value type
+			UhtProperty? value = UhtPropertyParser.ParseTemplateParam(resolvePhase, propertySettings, "Value", tokenReader);
+			if (value == null)
+			{
+				return null;
+			}
 
-				if (!value.PropertyCaps.HasAnyFlags(UhtPropertyCaps.CanBeContainerValue))
+			if (!value.PropertyCaps.HasAnyFlags(UhtPropertyCaps.CanBeContainerValue))
+			{
+				tokenReader.LogError($"The type \'{value.GetUserFacingDecl()}\' can not be used as a value in a TMap");
+			}
+
+			if (tokenReader.TryOptional(','))
+			{
+				UhtToken allocatorToken = tokenReader.GetIdentifier();
+				if (allocatorToken.IsIdentifier("FMemoryImageSetAllocator"))
 				{
-					tokenReader.LogError($"The type \'{value.GetUserFacingDecl()}\' can not be used as a value in a TMap");
+					propertySettings.Allocator = UhtPropertyAllocator.MemoryImage;
 				}
-
-				if (tokenReader.TryOptional(','))
+				else
 				{
-					UhtToken allocatorToken = tokenReader.GetIdentifier();
-					if (allocatorToken.IsIdentifier("FMemoryImageSetAllocator"))
-					{
-						propertySettings.Allocator = UhtPropertyAllocator.MemoryImage;
-					}
-					else
-					{
-						tokenReader.LogError($"Found '{allocatorToken.Value}' - explicit allocators are not supported in TMap properties.");
-					}
+					tokenReader.LogError($"Found '{allocatorToken.Value}' - explicit allocators are not supported in TMap properties.");
 				}
-				tokenReader.Require('>');
+			}
+			tokenReader.Require('>');
 
-				//@TODO: Prevent sparse delegate types from being used in a container
+			//@TODO: Prevent sparse delegate types from being used in a container
 
-				return new UhtMapProperty(propertySettings, key, value);
-			}
+			return new UhtMapProperty(propertySettings, key, value);
 		}
 		#endregion
 	}
diff --git a/Engine/Source/Programs/Shared/EpicGames.UHT/Types/Properties/UhtObjectPropertyBase.cs b/Engine/Source/Programs/Shared/EpicGames.UHT/Types/Properties/UhtObjectPropertyBase.cs
--- a/Engine/Source/Programs/Shared/EpicGames.UHT/Types/Properties/UhtObjectPropertyBase.cs
+++ b/Engine/Source/Programs/Shared/EpicGames.UHT/Types/Properties/UhtObjectPropertyBase.cs
@@ -115,10 +115,9 @@ namespace EpicGames.UHT.Types
 		/// <inheritdoc/>
 		public override bool SanitizeDefaultValue(IUhtTokenReader defaultValueReader, StringBuilder innerDefaultValue)
 		{
-			int value;
 			if (defaultValueReader.TryOptional("NULL") ||
 				defaultValueReader.TryOptional("nullptr") ||
-				(defaultValueReader.TryOptionalConstInt(out value) && value == 0))
+				(defaultValueReader.TryOptionalConstInt(out int value) && value == 0))
 			{
 				innerDefaultValue.Append("None");
 				return true;
@@ -161,7 +160,7 @@ namespace EpicGames.UHT.Types
 				return null;
 			}
 			bool isNativeConstTemplateArg = false;
-			UhtToken identifier = new UhtToken();
+			UhtToken identifier = new();
 			tokenReader
 				.Require('<')
 				.Optional("const", () => { isNativeConstTemplateArg = true; })
@@ -193,7 +192,7 @@ namespace EpicGames.UHT.Types
 		public static UhtClass? ParseTemplateClass(UhtPropertySettings propertySettings, IUhtTokenReader tokenReader, UhtToken matchedToken)
 		{
 			UhtSession session = propertySettings.Outer.Session;
-			UhtToken identifier = new UhtToken();
+			UhtToken identifier = new();
 
 			if (tokenReader.TryOptional("const"))
 			{
diff --git a/Engine/Source/Programs/Shared/EpicGames.UHT/Types/Properties/UhtSetProperty.cs b/Engine/Source/Programs/Shared/EpicGames.UHT/Types/Properties/UhtSetProperty.cs
--- a/Engine/Source/Programs/Shared/EpicGames.UHT/Types/Properties/UhtSetProperty.cs
+++ b/Engine/Source/Programs/Shared/EpicGames.UHT/Types/Properties/UhtSetProperty.cs
@@ -211,44 +211,42 @@ namespace EpicGames.UHT.Types
 		[SuppressMessage("CodeQuality", "IDE0051:Remove unused private members", Justification = "Attribute accessed method")]
 		private static UhtProperty? SetProperty(UhtPropertyResolvePhase resolvePhase, UhtPropertySettings propertySettings, IUhtTokenReader tokenReader, UhtToken matchedToken)
 		{
-			using (UhtMessageContext tokenContext = new UhtMessageContext("TSet"))
+			using UhtMessageContext tokenContext = new("TSet");
+			if (!tokenReader.SkipExpectedType(matchedToken.Value, propertySettings.PropertyCategory == UhtPropertyCategory.Member))
 			{
-				if (!tokenReader.SkipExpectedType(matchedToken.Value, propertySettings.PropertyCategory == UhtPropertyCategory.Member))
-				{
-					return null;
-				}
-				tokenReader.Require('<');
+				return null;
+			}
+			tokenReader.Require('<');
 
-				// Parse the value type
-				UhtProperty? value = UhtPropertyParser.ParseTemplateParam(resolvePhase, propertySettings, "Value", tokenReader);
-				if (value == null)
-				{
-					return null;
-				}
+			// Parse the value type
+			UhtProperty? value = UhtPropertyParser.ParseTemplateParam(resolvePhase, propertySettings, "Value", tokenReader);
+			if (value == null)
+			{
+				return null;
+			}
 
-				if (!value.PropertyCaps.HasAnyFlags(UhtPropertyCaps.CanBeContainerKey))
-				{
-					tokenReader.LogError($"The type \'{value.GetUserFacingDecl()}\' can not be used as a key in a TSet");
-				}
+			if (!value.PropertyCaps.HasAnyFlags(UhtPropertyCaps.CanBeContainerKey))
+			{
+				tokenReader.LogError($"The type \'{value.GetUserFacingDecl()}\' can not be used as a key in a TSet");
+			}
 
-				if (propertySettings.PropertyFlags.HasAnyFlags(EPropertyFlags.Net))
-				{
-					tokenReader.LogError("Replicated sets are not supported.");
-				}
+			if (propertySettings.PropertyFlags.HasAnyFlags(EPropertyFlags.Net))
+			{
+				tokenReader.LogError("Replicated sets are not supported.");
+			}
 
-				if (tokenReader.TryOptional(','))
-				{
-					// If we found a comma, read the next thing, assume it's a keyfuncs, and report that
-					UhtToken keyFuncToken = tokenReader.GetIdentifier();
-					throw new UhtException(tokenReader, $"Found '{keyFuncToken.Value}' - explicit KeyFuncs are not supported in TSet properties.");
-				}
+			if (tokenReader.TryOptional(','))
+			{
+				// If we found a comma, read the next thing, assume it's a keyfuncs, and report that
+				UhtToken keyFuncToken = tokenReader.GetIdentifier();
+				throw new UhtException(tokenReader, $"Found '{keyFuncToken.Value}' - explicit KeyFuncs are not supported in TSet properties.");
+			}
 
-				tokenReader.Require('>');
+			tokenReader.Require('>');
 
-				//@TODO: Prevent sparse delegate types from being used in a container
+			//@TODO: Prevent sparse delegate types from being used in a container
 
-				return new UhtSetProperty(propertySettings, value);
-			}
+			return new UhtSetProperty(propertySettings, value);
 		}
 		#endregion
 	}
diff --git a/Engine/Source/Programs/Shared/EpicGames.UHT/Types/Properties/UhtStrProperty.cs b/Engine/Source/Programs/Shared/EpicGames.UHT/Types/Properties/UhtStrProperty.cs
--- a/Engine/Source/Programs/Shared/EpicGames.UHT/Types/Properties/UhtStrProperty.cs
+++ b/Engine/Source/Programs/Shared/EpicGames.UHT/Types/Properties/UhtStrProperty.cs
@@ -111,7 +111,7 @@ namespace EpicGames.UHT.Types
 			{
 				return null;
 			}
-			UhtStrProperty property = new UhtStrProperty(propertySettings);
+			UhtStrProperty property = new(propertySettings);
 			if (property.PropertyCategory != UhtPropertyCategory.Member)
 			{
 				if (tokenReader.TryOptional('&'))
diff --git a/Engine/Source/Programs/Shared/EpicGames.UHT/Types/Properties/UhtStructProperty.cs b/Engine/Source/Programs/Shared/EpicGames.UHT/Types/Properties/UhtStructProperty.cs
--- a/Engine/Source/Programs/Shared/EpicGames.UHT/Types/Properties/UhtStructProperty.cs
+++ b/Engine/Source/Programs/Shared/EpicGames.UHT/Types/Properties/UhtStructProperty.cs
@@ -159,8 +159,7 @@ namespace EpicGames.UHT.Types
 		/// <inheritdoc/>
 		public override bool SanitizeDefaultValue(IUhtTokenReader defaultValueReader, StringBuilder innerDefaultValue)
 		{
-			UhtStructDefaultValue structDefaultValue;
-			if (!this.Session.TryGetStructDefaultValue(this.ScriptStruct.SourceName, out structDefaultValue))
+			if (!this.Session.TryGetStructDefaultValue(this.ScriptStruct.SourceName, out UhtStructDefaultValue structDefaultValue))
 			{
 				structDefaultValue = this.Session.DefaultStructDefaultValue;
 			}
diff --git a/Engine/Source/Programs/Shared/EpicGames.UHT/Types/Properties/UhtTextProperty.cs b/Engine/Source/Programs/Shared/EpicGames.UHT/Types/Properties/UhtTextProperty.cs
--- a/Engine/Source/Programs/Shared/EpicGames.UHT/Types/Properties/UhtTextProperty.cs
+++ b/Engine/Source/Programs/Shared/EpicGames.UHT/Types/Properties/UhtTextProperty.cs
@@ -170,7 +170,7 @@ namespace EpicGames.UHT.Types
 			ReadOnlySpan<char> strippedNS = namespaceString.Span;
 			if (strippedNS.Length > 1)
 			{
-				strippedNS = strippedNS.Slice(1, strippedNS.Length - 2).Trim();
+				strippedNS = strippedNS[1..^1].Trim();
 				if (strippedNS.Length > 0)
 				{
 					if (strippedNS[^1] == ']')
@@ -178,7 +178,7 @@ namespace EpicGames.UHT.Types
 						int index = strippedNS.LastIndexOf('[');
 						if (index != -1)
 						{
-							strippedNS = strippedNS.Slice(0, index).TrimEnd();
+							strippedNS = strippedNS[..index].TrimEnd();
 						}
 					}
 				}
@@ -333,9 +333,8 @@ namespace EpicGames.UHT.Types
 		{
 			defaultValueReader.Require('(');
 
-			double baseValue;
 			UhtToken token = defaultValueReader.GetToken();
-			if (!token.GetConstDouble(out baseValue))
+			if (!token.GetConstDouble(out double baseValue))
 			{
 				return false;
 			}
diff --git a/Engine/Source/Programs/Shared/EpicGames.UHT/Types/UhtClass.cs b/Engine/Source/Programs/Shared/EpicGames.UHT/Types/UhtClass.cs
--- a/Engine/Source/Programs/Shared/EpicGames.UHT/Types/UhtClass.cs
+++ b/Engine/Source/Programs/Shared/EpicGames.UHT/Types/UhtClass.cs
@@ -277,7 +277,7 @@ namespace EpicGames.UHT.Types
 	[UhtEngineClass(Name = "Class")]
 	public class UhtClass : UhtStruct
 	{
-		private List<UhtDeclaration>? _declarationsInternal = null;
+		private List<UhtDeclaration>? _declarations = null;
 
 		/// <summary>
 		/// Configuration section
@@ -333,7 +333,7 @@ namespace EpicGames.UHT.Types
 		/// Collection of functions and other declarations found in the class
 		/// </summary>
 		[JsonIgnore]
-		public IList<UhtDeclaration>? Declarations => this._declarationsInternal;
+		public IList<UhtDeclaration>? Declarations => this._declarations;
 
 		/// <summary>
 		/// If this, this class is a UINTERFACE and NativeInterface is the associated native interface
@@ -463,11 +463,11 @@ namespace EpicGames.UHT.Types
 		/// <param name="function">If parsed as part of a UFUNCTION, this will reference it</param>
 		public void AddDeclaration(UhtCompilerDirective compilerDirectives, List<UhtToken> tokens, UhtFunction? function)
 		{
-			if (this._declarationsInternal == null)
+			if (this._declarations == null)
 			{
-				this._declarationsInternal = new List<UhtDeclaration>();
+				this._declarations = new List<UhtDeclaration>();
 			}
-			this._declarationsInternal.Add(new UhtDeclaration { CompilerDirectives = compilerDirectives, Tokens = tokens.ToArray(), Function = function });
+			this._declarations.Add(new UhtDeclaration { CompilerDirectives = compilerDirectives, Tokens = tokens.ToArray(), Function = function });
 		}
 
 		/// <summary>
@@ -819,8 +819,8 @@ namespace EpicGames.UHT.Types
 			}
 
 			// Verify the type
-			using (BorrowStringBuilder borrower = new BorrowStringBuilder(StringBuilderCache.Small))
 			{
+				using BorrowStringBuilder borrower = new(StringBuilderCache.Small);
 				StringBuilder builder = borrower.StringBuilder;
 				builder.AppendPropertyText(property, UhtPropertyTextType.GetterSetterArg);
 				if (property.IsStaticArray)
@@ -861,7 +861,7 @@ namespace EpicGames.UHT.Types
 					}
 
 					// Extract the remaining part of the string we expect to match
-					expectedSpan = expectedSpan.Slice(tokenSpan.Length).TrimStart();
+					expectedSpan = expectedSpan[tokenSpan.Length..].TrimStart();
 					typeIndex++;
 				}
 
@@ -1321,33 +1321,29 @@ namespace EpicGames.UHT.Types
 					if (property.PropertyExportFlags.HasExactFlags(UhtPropertyExportFlags.GetterSpecified | UhtPropertyExportFlags.GetterSpecifiedNone | 
 						UhtPropertyExportFlags.GetterFound, UhtPropertyExportFlags.GetterSpecified))
 					{
+						using BorrowStringBuilder borrower = new(StringBuilderCache.Small);
+						StringBuilder builder = borrower.StringBuilder;
 						string expectedName = property.Getter ?? $"Get{property.SourceName}";
-						using (BorrowStringBuilder borrower = new BorrowStringBuilder(StringBuilderCache.Small))
+						builder.AppendPropertyText(property, UhtPropertyTextType.GetterSetterArg);
+						if (property.IsStaticArray)
 						{
-							StringBuilder builder = borrower.StringBuilder;
-							builder.AppendPropertyText(property, UhtPropertyTextType.GetterSetterArg);
-							if (property.IsStaticArray)
-							{
-								builder.Append($"*/[{property.ArrayDimensions}]");
-							}
-							string expectedType = builder.ToString();
-							property.LogError($"Property '{property.SourceName}' expected a getter function '[const] {expectedType} [&] {expectedName}() const', but it was not found");
+							builder.Append($"*/[{property.ArrayDimensions}]");
 						}
+						string expectedType = builder.ToString();
+						property.LogError($"Property '{property.SourceName}' expected a getter function '[const] {expectedType} [&] {expectedName}() const', but it was not found");
 					}
 					if (property.PropertyExportFlags.HasExactFlags(UhtPropertyExportFlags.SetterSpecified | UhtPropertyExportFlags.SetterSpecifiedNone | UhtPropertyExportFlags.SetterFound, UhtPropertyExportFlags.SetterSpecified))
 					{
+						using BorrowStringBuilder borrower = new(StringBuilderCache.Small);
+						StringBuilder builder = borrower.StringBuilder;
 						string expectedName = property.Setter ?? $"Set{property.SourceName}";
-						using (BorrowStringBuilder borrower = new BorrowStringBuilder(StringBuilderCache.Small))
+						builder.AppendPropertyText(property, UhtPropertyTextType.GetterSetterArg);
+						if (property.IsStaticArray)
 						{
-							StringBuilder builder = borrower.StringBuilder;
-							builder.AppendPropertyText(property, UhtPropertyTextType.GetterSetterArg);
-							if (property.IsStaticArray)
-							{
-								builder.Append($"*/[{property.ArrayDimensions}]");
-							}
-							string expectedType = builder.ToString();
-							property.LogError($"Property '{property.SourceName}' expected a setter function 'void {expectedName}([const] {expectedType} [&] InArg)', but it was not found");
+							builder.Append($"*/[{property.ArrayDimensions}]");
 						}
+						string expectedType = builder.ToString();
+						property.LogError($"Property '{property.SourceName}' expected a setter function 'void {expectedName}([const] {expectedType} [&] InArg)', but it was not found");
 					}
 				}
 			}
diff --git a/Engine/Source/Programs/Shared/EpicGames.UHT/Types/UhtEnum.cs b/Engine/Source/Programs/Shared/EpicGames.UHT/Types/UhtEnum.cs
--- a/Engine/Source/Programs/Shared/EpicGames.UHT/Types/UhtEnum.cs
+++ b/Engine/Source/Programs/Shared/EpicGames.UHT/Types/UhtEnum.cs
@@ -209,7 +209,7 @@ namespace EpicGames.UHT.Types
 				int scopeIndex = enumName.IndexOf("::", StringComparison.Ordinal);
 				if (scopeIndex >= 0)
 				{
-					enumName = enumName.Substring(scopeIndex + 2);
+					enumName = enumName[(scopeIndex + 2)..];
 				}
 			}
 			return $"{enumName}.{name}";
@@ -259,7 +259,7 @@ namespace EpicGames.UHT.Types
 		private string CleanEnumValueName(string name)
 		{
 			int lastColons = name.LastIndexOf("::", StringComparison.Ordinal);
-			return lastColons == -1 ? GetFullEnumName(name) : GetFullEnumName(name.Substring(lastColons + 2));
+			return lastColons == -1 ? GetFullEnumName(name) : GetFullEnumName(name[(lastColons + 2)..]);
 		}
 
 		#region Validation support
@@ -274,24 +274,21 @@ namespace EpicGames.UHT.Types
 				}
 			}
 
-			Dictionary<string, string> toolTips = new Dictionary<string, string>(StringComparer.OrdinalIgnoreCase);
+			Dictionary<string, string> toolTips = new(StringComparer.OrdinalIgnoreCase);
 			for (int enumIndex = 0; enumIndex < this.EnumValues.Count; ++enumIndex)
 			{
-				string? entryName;
-				if (!this.MetaData.TryGetValue(UhtNames.Name, enumIndex, out entryName))
+				if (!this.MetaData.TryGetValue(UhtNames.Name, enumIndex, out string? entryName))
 				{
 					continue;
 				}
 
-				string? toolTip;
-				if (!this.MetaData.TryGetValue(UhtNames.ToolTip, enumIndex, out toolTip))
+				if (!this.MetaData.TryGetValue(UhtNames.ToolTip, enumIndex, out string? toolTip))
 				{
 					this.LogError(this.MetaData.LineNumber, $"Enum entry '{this.SourceName}::{entryName}' does not provide a tooltip / comment (DocumentationPolicy)");
 					continue;
 				}
 
-				string? dupName;
-				if (toolTips.TryGetValue(toolTip, out dupName))
+				if (toolTips.TryGetValue(toolTip, out string? dupName))
 				{
 					this.LogError(this.MetaData.LineNumber, $"Enum entries '{this.SourceName}::{entryName}' and '{this.SourceName}::{dupName}' have identical tooltips / comments (DocumentationPolicy)");
 				}
diff --git a/Engine/Source/Programs/Shared/EpicGames.UHT/Types/UhtFunction.cs b/Engine/Source/Programs/Shared/EpicGames.UHT/Types/UhtFunction.cs
--- a/Engine/Source/Programs/Shared/EpicGames.UHT/Types/UhtFunction.cs
+++ b/Engine/Source/Programs/Shared/EpicGames.UHT/Types/UhtFunction.cs
@@ -196,7 +196,7 @@ namespace EpicGames.UHT.Types
 	/// </summary>
 	public class UhtFunction : UhtStruct
 	{
-		private string? _strippedFunctionNameInternal = null;
+		private string? _strippedFunctionName = null;
 
 		/// <summary>
 		/// Suffix added to delegate engine names 
@@ -270,8 +270,8 @@ namespace EpicGames.UHT.Types
 		[JsonIgnore]
 		public string StrippedFunctionName
 		{
-			get => this._strippedFunctionNameInternal ?? this.EngineName;
-			set => this._strippedFunctionNameInternal = value;
+			get => this._strippedFunctionName ?? this.EngineName;
+			set => this._strippedFunctionName = value;
 		}
 
 		///<inheritdoc/>
@@ -339,7 +339,7 @@ namespace EpicGames.UHT.Types
 		/// Return read only memory of all the function parameters
 		/// </summary>
 		[JsonIgnore]
-		public ReadOnlyMemory<UhtType> ParameterProperties => new ReadOnlyMemory<UhtType>(this.Children.ToArray(), 0, this.HasReturnProperty ? this.Children.Count - 1 : this.Children.Count);
+		public ReadOnlyMemory<UhtType> ParameterProperties => new(this.Children.ToArray(), 0, this.HasReturnProperty ? this.Children.Count - 1 : this.Children.Count);
 
 		/// <summary>
 		/// True if the function has any outputs including a return value
@@ -419,8 +419,7 @@ namespace EpicGames.UHT.Types
 							!this.FunctionExportFlags.HasAnyFlags(UhtFunctionExportFlags.CustomThunk) &&
 							this.CppImplName != this.SourceName)
 						{
-							UhtFoundDeclaration declaration;
-							if (classObj.TryGetDeclaration(this.CppImplName, out declaration))
+							if (classObj.TryGetDeclaration(this.CppImplName, out UhtFoundDeclaration declaration))
 							{
 								this.FunctionExportFlags |= UhtFunctionExportFlags.ImplFound;
 								if (declaration.IsVirtual)
@@ -765,7 +764,7 @@ namespace EpicGames.UHT.Types
 
 		private void LogRpcFunctionError(bool virtualError, string functionName)
 		{
-			StringBuilder builder = new StringBuilder();
+			StringBuilder builder = new();
 
 			if (virtualError)
 			{
@@ -920,7 +919,7 @@ namespace EpicGames.UHT.Types
 				if (hasAnyParamToolTips)
 				{
 					// ensure each parameter has a tooltip
-					HashSet<StringView> existingFields = new HashSet<StringView>(StringViewComparer.OrdinalIgnoreCase);
+					HashSet<StringView> existingFields = new(StringViewComparer.OrdinalIgnoreCase);
 					foreach (UhtType parameter in this.ParameterProperties.Span)
 					{
 						existingFields.Add(parameter.SourceName);
@@ -945,7 +944,7 @@ namespace EpicGames.UHT.Types
 					}
 
 					// check for duplicate tooltips
-					Dictionary<StringView, StringView> toolTipToParam = new Dictionary<StringView, StringView>(StringViewComparer.OrdinalIgnoreCase);
+					Dictionary<StringView, StringView> toolTipToParam = new(StringViewComparer.OrdinalIgnoreCase);
 					foreach (KeyValuePair<StringView, StringView> kvp in paramToolTips)
 					{
 						if (kvp.Key == UhtNames.ReturnValue)
@@ -1002,7 +1001,7 @@ namespace EpicGames.UHT.Types
 			const string ParamTag = "@param";
 			const string ReturnTag = "@return";
 
-			Dictionary<StringView, StringView> paramMap = new Dictionary<StringView, StringView>(StringViewComparer.OrdinalIgnoreCase);
+			Dictionary<StringView, StringView> paramMap = new(StringViewComparer.OrdinalIgnoreCase);
 
 			// Loop until we are out of text
 			while (input.Length > 0)
diff --git a/Engine/Source/Programs/Shared/EpicGames.UHT/Types/UhtHeaderFile.cs b/Engine/Source/Programs/Shared/EpicGames.UHT/Types/UhtHeaderFile.cs
--- a/Engine/Source/Programs/Shared/EpicGames.UHT/Types/UhtHeaderFile.cs
+++ b/Engine/Source/Programs/Shared/EpicGames.UHT/Types/UhtHeaderFile.cs
@@ -105,7 +105,7 @@ namespace EpicGames.UHT.Types
 	{
 		private readonly UhtSimpleMessageSite _messageSite;
 		private readonly UhtSourceFile _sourceFile;
-		private readonly List<UhtHeaderFile> _referencedHeaders = new List<UhtHeaderFile>();
+		private readonly List<UhtHeaderFile> _referencedHeaders = new();
 
 		/// <summary>
 		/// Contents of the header
@@ -333,8 +333,8 @@ namespace EpicGames.UHT.Types
 		{
 			options = base.Validate(options | UhtValidationOptions.Shadowing);
 
-			Dictionary<int, UhtFunction> usedRPCIds = new Dictionary<int, UhtFunction>();
-			Dictionary<int, UhtFunction> rpcsNeedingHookup = new Dictionary<int, UhtFunction>();
+			Dictionary<int, UhtFunction> usedRPCIds = new();
+			Dictionary<int, UhtFunction> rpcsNeedingHookup = new();
 			foreach (UhtType type in this.Children)
 			{
 				if (type is UhtClass classObj)
@@ -350,8 +350,7 @@ namespace EpicGames.UHT.Types
 
 							if (function.RPCId > 0)
 							{
-								UhtFunction? existingFunc;
-								if (usedRPCIds.TryGetValue(function.RPCId, out existingFunc))
+								if (usedRPCIds.TryGetValue(function.RPCId, out UhtFunction? existingFunc))
 								{
 									function.LogError($"Function {existingFunc.SourceName} already uses identifier {function.RPCId}");
 								}
diff --git a/Engine/Source/Programs/Shared/EpicGames.UHT/Types/UhtMetaData.cs b/Engine/Source/Programs/Shared/EpicGames.UHT/Types/UhtMetaData.cs
--- a/Engine/Source/Programs/Shared/EpicGames.UHT/Types/UhtMetaData.cs
+++ b/Engine/Source/Programs/Shared/EpicGames.UHT/Types/UhtMetaData.cs
@@ -125,12 +125,12 @@ namespace EpicGames.UHT.Types
 		/// <summary>
 		/// Helper comparer for meta data keys
 		/// </summary>
-		private static readonly UhtMetaDataKeyComparer s_comparer = new UhtMetaDataKeyComparer();
+		private static readonly UhtMetaDataKeyComparer s_comparer = new();
 
 		/// <summary>
 		/// Empty collection of meta data
 		/// </summary>
-		public static readonly UhtMetaData Empty = new UhtMetaData(null, null);
+		public static readonly UhtMetaData Empty = new(null, null);
 
 		/// <summary>
 		/// The meta data of the outer object for the type that owns this meta data
@@ -313,8 +313,7 @@ namespace EpicGames.UHT.Types
 		/// <returns>Meta data value or empty string if not found.</returns>
 		public string GetValueOrDefault(string name, int nameIndex = IndexNone)
 		{
-			string? output;
-			if (TryGetValue(name, nameIndex, out output))
+			if (TryGetValue(name, nameIndex, out string? output))
 			{
 				return output;
 			}
@@ -329,8 +328,7 @@ namespace EpicGames.UHT.Types
 		/// <returns>Meta data value or empty string if not found.</returns>
 		public string GetValueOrDefaultHierarchical(string name, int nameIndex = IndexNone)
 		{
-			string? output;
-			if (TryGetValueHierarchical(name, nameIndex, out output))
+			if (TryGetValueHierarchical(name, nameIndex, out string? output))
 			{
 				return output;
 			}
@@ -368,8 +366,7 @@ namespace EpicGames.UHT.Types
 		public double GetDouble(string name, int nameIndex = IndexNone)
 		{
 			string value = GetValueOrDefault(name, nameIndex);
-			double result;
-			if (!Double.TryParse(value, out result))
+			if (!Double.TryParse(value, out double result))
 			{
 				result = 0;
 			}
@@ -384,8 +381,7 @@ namespace EpicGames.UHT.Types
 		/// <returns>String array or null if not found</returns>
 		public string[]? GetStringArray(string name, int nameIndex = IndexNone)
 		{
-			string? temp;
-			if (TryGetValue(name, nameIndex, out temp))
+			if (TryGetValue(name, nameIndex, out string? temp))
 			{
 				return temp.ToString().Split(' ', StringSplitOptions.RemoveEmptyEntries);
 			}
@@ -400,8 +396,7 @@ namespace EpicGames.UHT.Types
 		/// <returns>String array or null if not found</returns>
 		public string[]? GetStringArrayHierarchical(string name, int nameIndex = IndexNone)
 		{
-			string? temp;
-			if (TryGetValueHierarchical(name, nameIndex, out temp))
+			if (TryGetValueHierarchical(name, nameIndex, out string? temp))
 			{
 				return temp.ToString().Split(' ', StringSplitOptions.RemoveEmptyEntries);
 			}
@@ -426,10 +421,9 @@ namespace EpicGames.UHT.Types
 		/// <param name="value">Value of the meta data</param>
 		public void Add(string name, int nameIndex, string value)
 		{
-			string remappedName;
 			if (this.Config != null)
 			{
-				if (this.Config.RedirectMetaDataKey(name, out remappedName))
+				if (this.Config.RedirectMetaDataKey(name, out string remappedName))
 				{
 					if (this.MessageSite != null)
 					{
@@ -544,7 +538,7 @@ namespace EpicGames.UHT.Types
 		/// <returns>List of meta data key and value pairs</returns>
 		public List<KeyValuePair<string, string>> GetSorted()
 		{
-			List<KeyValuePair<string, string>> output = new List<KeyValuePair<string, string>>(this.Dictionary != null ? this.Dictionary.Count : 0);
+			List<KeyValuePair<string, string>> output = new(this.Dictionary != null ? this.Dictionary.Count : 0);
 			if (this.Dictionary != null && this.Dictionary.Count > 0)
 			{
 				foreach (KeyValuePair<UhtMetaDataKey, string> kvp in this.Dictionary)
diff --git a/Engine/Source/Programs/Shared/EpicGames.UHT/Types/UhtPackage.cs b/Engine/Source/Programs/Shared/EpicGames.UHT/Types/UhtPackage.cs
--- a/Engine/Source/Programs/Shared/EpicGames.UHT/Types/UhtPackage.cs
+++ b/Engine/Source/Programs/Shared/EpicGames.UHT/Types/UhtPackage.cs
@@ -109,7 +109,7 @@ namespace EpicGames.UHT.Types
 			else
 			{
 				this.SourceName = module.Name;
-				this.ShortName = this.SourceName.Substring(lastSlashIndex + 1);
+				this.ShortName = this.SourceName[(lastSlashIndex + 1)..];
 			}
 		}
 	}
diff --git a/Engine/Source/Programs/Shared/EpicGames.UHT/Types/UhtProperty.cs b/Engine/Source/Programs/Shared/EpicGames.UHT/Types/UhtProperty.cs
--- a/Engine/Source/Programs/Shared/EpicGames.UHT/Types/UhtProperty.cs
+++ b/Engine/Source/Programs/Shared/EpicGames.UHT/Types/UhtProperty.cs
@@ -845,7 +845,7 @@ namespace EpicGames.UHT.Types
 		/// <summary>
 		/// Collection of invalid names for parameters 
 		/// </summary>
-		private static readonly HashSet<string> s_invalidParamNames = new HashSet<string>(StringComparer.OrdinalIgnoreCase) { "self" };
+		private static readonly HashSet<string> s_invalidParamNames = new(StringComparer.OrdinalIgnoreCase) { "self" };
 
 		/// <summary>
 		/// Standard object flags for properties
@@ -1377,7 +1377,7 @@ namespace EpicGames.UHT.Types
 		/// <returns></returns>
 		public string GetUserFacingDecl()
 		{
-			StringBuilder builder = new StringBuilder();
+			StringBuilder builder = new();
 			AppendText(builder, UhtPropertyTextType.Generic);
 			return builder.ToString();
 		}
@@ -1388,11 +1388,9 @@ namespace EpicGames.UHT.Types
 		/// <returns></returns>
 		public string GetRigVMType()
 		{
-			using (BorrowStringBuilder borrower = new BorrowStringBuilder(StringBuilderCache.Small))
-			{
-				AppendText(borrower.StringBuilder, UhtPropertyTextType.RigVMType);
-				return borrower.StringBuilder.ToString();
-			}
+			using BorrowStringBuilder borrower = new(StringBuilderCache.Small);
+			AppendText(borrower.StringBuilder, UhtPropertyTextType.RigVMType);
+			return borrower.StringBuilder.ToString();
 		}
 
 		/// <summary>
@@ -1447,7 +1445,7 @@ namespace EpicGames.UHT.Types
 										{
 											if (index == dim.Length - 1)
 											{
-												dim = dim.Slice(1, index - 1);
+												dim = dim[1..index];
 												again = true;
 											}
 											break;
@@ -1467,7 +1465,7 @@ namespace EpicGames.UHT.Types
 								{
 									if (dim.StartsWith(cast))
 									{
-										dim = dim.Slice(cast.Length);
+										dim = dim[cast.Length..];
 										again = true;
 										break;
 									}
@@ -1759,8 +1757,7 @@ namespace EpicGames.UHT.Types
 			if (!options.HasAnyFlags(UhtValidationOptions.IsKeyOrValue))
 			{
 				// First check if the category was specified at all and if the property was exposed to the editor.
-				string? category;
-				if (!this.MetaData.TryGetValue(UhtNames.Category, out category))
+				if (!this.MetaData.TryGetValue(UhtNames.Category, out string? category))
 				{
 					if (this.PropertyFlags.HasAnyFlags(EPropertyFlags.Edit | EPropertyFlags.BlueprintVisible))
 					{
diff --git a/Engine/Source/Programs/Shared/EpicGames.UHT/Types/UhtScriptStruct.cs b/Engine/Source/Programs/Shared/EpicGames.UHT/Types/UhtScriptStruct.cs
--- a/Engine/Source/Programs/Shared/EpicGames.UHT/Types/UhtScriptStruct.cs
+++ b/Engine/Source/Programs/Shared/EpicGames.UHT/Types/UhtScriptStruct.cs
@@ -219,7 +219,7 @@ namespace EpicGames.UHT.Types
 					{
 						string extendedType = this.ExtendedType(false);
 						this.CastName = $"{this.Name}_{index}_Array";
-						this.CastType = $"TArrayView<const {extendedType.Substring(1, extendedType.Length - 2)}>";
+						this.CastType = $"TArrayView<const {extendedType[1..^1]}>";
 					}
 				}
 			}
@@ -283,7 +283,7 @@ namespace EpicGames.UHT.Types
 			int lesserPos = typeOriginal.IndexOf('<', StringComparison.Ordinal);
 			if (lesserPos >= 0)
 			{
-				return typeOriginal.Substring(0, lesserPos);
+				return typeOriginal[..lesserPos];
 			}
 			else
 			{
@@ -303,7 +303,7 @@ namespace EpicGames.UHT.Types
 			int lesserPos = typeOriginal.IndexOf('<', StringComparison.Ordinal);
 			if (lesserPos >= 0)
 			{
-				return typeOriginal.Substring(lesserPos);
+				return typeOriginal[lesserPos..];
 			}
 			else
 			{
@@ -350,7 +350,7 @@ namespace EpicGames.UHT.Types
 			string typeOriginal = TypeOriginal(castType);
 			if (typeOriginal.EndsWith("&", StringComparison.Ordinal))
 			{
-				return typeOriginal.Substring(0, typeOriginal.Length - 1);
+				return typeOriginal[0..^1];
 			}
 			else
 			{
@@ -404,7 +404,7 @@ namespace EpicGames.UHT.Types
 	/// </summary>
 	public class UhtRigVMMethodInfo
 	{
-		private static readonly string s_noPrefixInternal = String.Empty;
+		private static readonly string s_noPrefix = String.Empty;
 		private const string ReturnPrefixInternal = "return ";
 
 		/// <summary>
@@ -428,7 +428,7 @@ namespace EpicGames.UHT.Types
 		/// <returns>Prefix required for the return value</returns>
 		public string ReturnPrefix()
 		{
-			return (ReturnType.Length == 0 || ReturnType == "void") ? s_noPrefixInternal : ReturnPrefixInternal;
+			return (ReturnType.Length == 0 || ReturnType == "void") ? s_noPrefix : ReturnPrefixInternal;
 		}
 	}
 
diff --git a/Engine/Source/Programs/Shared/EpicGames.UHT/Types/UhtStruct.cs b/Engine/Source/Programs/Shared/EpicGames.UHT/Types/UhtStruct.cs
--- a/Engine/Source/Programs/Shared/EpicGames.UHT/Types/UhtStruct.cs
+++ b/Engine/Source/Programs/Shared/EpicGames.UHT/Types/UhtStruct.cs
@@ -86,7 +86,7 @@ namespace EpicGames.UHT.Types
 		/// </summary>
 		//[JsonConverter(typeof(UhtNullableTypeListJsonConverter<UhtStruct>))]
 		[JsonIgnore]
-		public IReadOnlyList<UhtStruct> Bases => _basesInternal ?? s_emptyBases;
+		public IReadOnlyList<UhtStruct> Bases => _bases ?? s_emptyBases;
 
 		/// <summary>
 		/// Super struct type
@@ -124,8 +124,8 @@ namespace EpicGames.UHT.Types
 			return false;
 		}
 
-		private static readonly List<UhtStruct> s_emptyBases = new List<UhtStruct>();
-		private List<UhtStruct>? _basesInternal = null;
+		private static readonly List<UhtStruct> s_emptyBases = new();
+		private List<UhtStruct>? _bases = null;
 
 		#region Resolution support
 		/// <inheritdoc/>
@@ -208,11 +208,11 @@ namespace EpicGames.UHT.Types
 					UhtStruct? baseStruct = (UhtStruct?)this.FindType(findOptions | UhtFindOptions.Class | UhtFindOptions.ScriptStruct | UhtFindOptions.SourceName | UhtFindOptions.NoSelf, baseIdentifier);
 					if (baseStruct != null)
 					{
-						if (this._basesInternal == null)
+						if (this._bases == null)
 						{
-							this._basesInternal = new List<UhtStruct>();
+							this._bases = new List<UhtStruct>();
 						}
-						this._basesInternal.Add(baseStruct);
+						this._bases.Add(baseStruct);
 						this.HeaderFile.AddReferencedHeader(baseStruct);
 						baseStruct.Resolve(UhtResolvePhase.Bases);
 					}
@@ -240,14 +240,12 @@ namespace EpicGames.UHT.Types
 			}
 
 			// NOTE: Old UHT didn't handle parse errors
-			double minValue;
-			if (!Double.TryParse(uiMin, out minValue))
+			if (!Double.TryParse(uiMin, out double minValue))
 			{
 				minValue = 0;
 			}
 
-			double maxValue;
-			if (!Double.TryParse(uiMax, out maxValue))
+			if (!Double.TryParse(uiMax, out double maxValue))
 			{
 				maxValue = 0;
 			}
@@ -270,7 +268,7 @@ namespace EpicGames.UHT.Types
 
 			if (policy.MemberToolTipsRequired)
 			{
-				Dictionary<string, UhtProperty> toolTipToType = new Dictionary<string, UhtProperty>();
+				Dictionary<string, UhtProperty> toolTipToType = new();
 				foreach (UhtProperty property in this.Properties)
 				{
 					string toolTip = property.GetToolTipText();
@@ -280,8 +278,7 @@ namespace EpicGames.UHT.Types
 						continue;
 					}
 
-					UhtProperty? existing;
-					if (toolTipToType.TryGetValue(toolTip, out existing))
+					if (toolTipToType.TryGetValue(toolTip, out UhtProperty? existing))
 					{
 						property.LogError($"Property '{this.SourceName}::{existing.SourceName}' and '{this.SourceName}::{property.SourceName}' are using identical tooltips (DocumentationPolicy).");
 					}
@@ -311,7 +308,7 @@ namespace EpicGames.UHT.Types
 			{
 				if (this is UhtClass)
 				{
-					Dictionary<string, UhtType> toolTipToType = new Dictionary<string, UhtType>();
+					Dictionary<string, UhtType> toolTipToType = new();
 					foreach (UhtType child in this.Children)
 					{
 						if (child is UhtFunction function)
@@ -324,8 +321,7 @@ namespace EpicGames.UHT.Types
 								continue;
 							}
 
-							UhtType? existing;
-							if (toolTipToType.TryGetValue(toolTip, out existing))
+							if (toolTipToType.TryGetValue(toolTip, out UhtType? existing))
 							{
 								child.LogError($"Functions '{this.SourceName}::{existing.SourceName}' and '{this.SourceName}::{function.SourceName}' are using identical tooltips / comments (DocumentationPolicy).");
 							}
diff --git a/Engine/Source/Programs/Shared/EpicGames.UHT/Types/UhtType.cs b/Engine/Source/Programs/Shared/EpicGames.UHT/Types/UhtType.cs
--- a/Engine/Source/Programs/Shared/EpicGames.UHT/Types/UhtType.cs
+++ b/Engine/Source/Programs/Shared/EpicGames.UHT/Types/UhtType.cs
@@ -587,7 +587,7 @@ namespace EpicGames.UHT.Types
 	/// </summary>
 	public class UhtDocumentationPolicy
 	{
-		private static readonly UhtDocumentationPolicy s_strictDocumentationPolicy = new UhtDocumentationPolicy
+		private static readonly UhtDocumentationPolicy s_strictDocumentationPolicy = new()
 		{
 			PolicySet = true,
 			ClassOrStructCommentRequired = true,
@@ -597,7 +597,7 @@ namespace EpicGames.UHT.Types
 			FloatRangesRequired = true,
 		};
 
-		private static readonly UhtDocumentationPolicy s_noDocumentationPolicy = new UhtDocumentationPolicy
+		private static readonly UhtDocumentationPolicy s_noDocumentationPolicy = new()
 		{
 			PolicySet = false,
 			ClassOrStructCommentRequired = false,
@@ -665,7 +665,7 @@ namespace EpicGames.UHT.Types
 		/// <summary>
 		/// Empty list of type used when children are requested but no children have been added.
 		/// </summary>
-		private static readonly List<UhtType> s_emptyTypeList = new List<UhtType>();
+		private static readonly List<UhtType> s_emptyTypeList = new();
 
 		/// <summary>
 		/// All UHT runs are associated with a given session.  The session holds all the global information for a run.
@@ -764,7 +764,7 @@ namespace EpicGames.UHT.Types
 		{
 			get
 			{
-				StringBuilder builder = new StringBuilder();
+				StringBuilder builder = new();
 				AppendFullName(builder);
 				return builder.ToString();
 			}
@@ -778,7 +778,7 @@ namespace EpicGames.UHT.Types
 		{
 			get
 			{
-				StringBuilder builder = new StringBuilder();
+				StringBuilder builder = new();
 				AppendPathName(builder, null);
 				return builder.ToString();
 			}
@@ -1089,8 +1089,7 @@ namespace EpicGames.UHT.Types
 				{
 					foreach (KeyValuePair<UhtMetaDataKey, string> kvp in this.MetaData.Dictionary)
 					{
-						UhtSpecifierValidator? specifier;
-						if (table.TryGetValue(kvp.Key.Name, out specifier))
+						if (table.TryGetValue(kvp.Key.Name, out UhtSpecifierValidator? specifier))
 						{
 							specifier.Delegate(this, MetaData, kvp.Key, kvp.Value);
 						}
@@ -1202,8 +1201,7 @@ namespace EpicGames.UHT.Types
 		/// <returns>The display name for this object.</returns>
 		public string GetDisplayNameText()
 		{
-			string? nativeDisplayName;
-			if (this.MetaData.TryGetValue(UhtNames.DisplayName, out nativeDisplayName))
+			if (this.MetaData.TryGetValue(UhtNames.DisplayName, out string? nativeDisplayName))
 			{
 				return nativeDisplayName;
 			}
diff --git a/Engine/Source/Programs/Shared/EpicGames.UHT/Utils/UhtException.cs b/Engine/Source/Programs/Shared/EpicGames.UHT/Utils/UhtException.cs
--- a/Engine/Source/Programs/Shared/EpicGames.UHT/Utils/UhtException.cs
+++ b/Engine/Source/Programs/Shared/EpicGames.UHT/Utils/UhtException.cs
@@ -96,7 +96,7 @@ namespace EpicGames.UHT.Utils
 		/// <returns></returns>
 		private static string FormatExtraContext(object? extraContext = null)
 		{
-			StringBuilder builder = new StringBuilder(" while parsing ");
+			StringBuilder builder = new(" while parsing ");
 			int startingLength = builder.Length;
 			if (extraContext != null)
 			{
diff --git a/Engine/Source/Programs/Shared/EpicGames.UHT/Utils/UhtFCString.cs b/Engine/Source/Programs/Shared/EpicGames.UHT/Utils/UhtFCString.cs
--- a/Engine/Source/Programs/Shared/EpicGames.UHT/Utils/UhtFCString.cs
+++ b/Engine/Source/Programs/Shared/EpicGames.UHT/Utils/UhtFCString.cs
@@ -46,8 +46,7 @@ namespace EpicGames.UHT.Utils
 			}
 			else
 			{
-				int intValue;
-				return Int32.TryParse(span, out intValue) && intValue != 0;
+				return Int32.TryParse(span, out int intValue) && intValue != 0;
 			}
 		}
 
@@ -76,8 +75,7 @@ namespace EpicGames.UHT.Utils
 			}
 			else
 			{
-				int intValue;
-				return Int32.TryParse(value, out intValue) && intValue != 0;
+				return Int32.TryParse(value, out int intValue) && intValue != 0;
 			}
 		}
 
@@ -252,7 +250,7 @@ namespace EpicGames.UHT.Utils
 		/// <returns>Resulting string</returns>
 		public static string UnescapeText(string text)
 		{
-			StringBuilder result = new StringBuilder();
+			StringBuilder result = new();
 			for (int idx = 0; idx < text.Length; idx++)
 			{
 				if (text[idx] == '\\' && idx + 1 < text.Length)
@@ -278,8 +276,7 @@ namespace EpicGames.UHT.Utils
 							bool parsed = false;
 							if (idx + 4 < text.Length)
 							{
-								int value;
-								if (Int32.TryParse(text.Substring(idx + 1, 4), System.Globalization.NumberStyles.HexNumber, null, out value))
+								if (Int32.TryParse(text.Substring(idx + 1, 4), System.Globalization.NumberStyles.HexNumber, null, out int value))
 								{
 									parsed = true;
 									result.Append((char)value);
@@ -346,7 +343,7 @@ namespace EpicGames.UHT.Utils
 			bool wasNumber = false;
 			bool wasMinusSign = false;
 
-			StringBuilder builder = new StringBuilder();
+			StringBuilder builder = new();
 			for (int charIndex = 0; charIndex < chars.Length; ++charIndex)
 			{
 				char ch = chars[charIndex];
@@ -487,47 +484,45 @@ namespace EpicGames.UHT.Utils
 		private static string TabsToSpacesInternal(ReadOnlySpan<char> span, int tabSpacing, bool emulateCrBug, int tabIndex)
 		{
 			// Locate the last \n since all tabs have to be computed relative to this.
-			int crPos = span.Slice(0, tabIndex).LastIndexOf('\n') + 1;
+			int crPos = span[..tabIndex].LastIndexOf('\n') + 1;
 
-			using (BorrowStringBuilder borrower = new BorrowStringBuilder(StringBuilderCache.Small))
+			using BorrowStringBuilder borrower = new(StringBuilderCache.Small);
+			StringBuilder builder = borrower.StringBuilder;
+			int committed = 0;
+			do
 			{
-				StringBuilder builder = borrower.StringBuilder;
-				int committed = 0;
-				do
+				// Commit everything prior to the tab to the builder
+				builder.Append(span[..tabIndex]);
+				span = span[(tabIndex + 1)..];
+				committed += tabIndex;
+
+				// Add the appropriate number of spaces
+				int adjustedCrPos = crPos;
+				if (emulateCrBug && adjustedCrPos > 0)
 				{
-					// Commit everything prior to the tab to the builder
-					builder.Append(span.Slice(0, tabIndex));
-					span = span.Slice(tabIndex + 1);
-					committed += tabIndex;
-
-					// Add the appropriate number of spaces
-					int adjustedCrPos = crPos;
-					if (emulateCrBug && adjustedCrPos > 0)
+					--adjustedCrPos;
+				}
+				int spacesToInsert = tabSpacing - (committed - adjustedCrPos) % tabSpacing;
+				builder.AppendSpaces(spacesToInsert);
+				committed += spacesToInsert;
+
+				// Search for the next \t or \n
+				for (tabIndex = 0; tabIndex < span.Length; ++tabIndex)
+				{
+					if (span[tabIndex] == '\n')
 					{
-						--adjustedCrPos;
+						crPos = committed + tabIndex + 1;
 					}
-					int spacesToInsert = tabSpacing - (committed - adjustedCrPos) % tabSpacing;
-					builder.AppendSpaces(spacesToInsert);
-					committed += spacesToInsert;
-
-					// Search for the next \t or \n
-					for (tabIndex = 0; tabIndex < span.Length; ++tabIndex)
+					else if (span[tabIndex] == '\t')
 					{
-						if (span[tabIndex] == '\n')
-						{
-							crPos = committed + tabIndex + 1;
-						}
-						else if (span[tabIndex] == '\t')
-						{
-							break;
-						}
+						break;
 					}
-				} while (tabIndex < span.Length);
+				}
+			} while (tabIndex < span.Length);
 
-				// Commit the remaining data
-				builder.Append(span);
-				return builder.ToString();
-			}
+			// Commit the remaining data
+			builder.Append(span);
+			return builder.ToString();
 		}
 	}
 }
diff --git a/Engine/Source/Programs/Shared/EpicGames.UHT/Utils/UhtFileManager.cs b/Engine/Source/Programs/Shared/EpicGames.UHT/Utils/UhtFileManager.cs
--- a/Engine/Source/Programs/Shared/EpicGames.UHT/Utils/UhtFileManager.cs
+++ b/Engine/Source/Programs/Shared/EpicGames.UHT/Utils/UhtFileManager.cs
@@ -95,32 +95,28 @@ namespace EpicGames.UHT.Utils
 
 			try
 			{
-				using (FileStream fs = new FileStream(filePath, FileMode.Open, FileAccess.Read, FileShare.Read, 4 * 1024, FileOptions.SequentialScan))
+				using FileStream fs = new(filePath, FileMode.Open, FileAccess.Read, FileShare.Read, 4 * 1024, FileOptions.SequentialScan);
+				using StreamReader sr = new(fs, Encoding.UTF8, true);
+
+				// Try to read the whole file into a buffer created by hand.  This avoids a LOT of memory allocations which in turn reduces the
+				// GC stress on the system.  Removing the StreamReader would be nice in the future.
+				long rawFileLength = fs.Length;
+				UhtBuffer initialBuffer = UhtBuffer.Borrow((int)rawFileLength);
+				int readLength = sr.Read(initialBuffer.Memory.Span);
+				if (sr.EndOfStream)
 				{
-					using (StreamReader sr = new StreamReader(fs, Encoding.UTF8, true))
-					{
-
-						// Try to read the whole file into a buffer created by hand.  This avoids a LOT of memory allocations which in turn reduces the
-						// GC stress on the system.  Removing the StreamReader would be nice in the future.
-						long rawFileLength = fs.Length;
-						UhtBuffer initialBuffer = UhtBuffer.Borrow((int)rawFileLength);
-						int readLength = sr.Read(initialBuffer.Memory.Span);
-						if (sr.EndOfStream)
-						{
-							initialBuffer.Reset(readLength);
-							return initialBuffer;
-						}
-						else
-						{
-							string remaining = sr.ReadToEnd();
-							long totalSize = readLength + remaining.Length;
-							UhtBuffer combined = UhtBuffer.Borrow((int)totalSize);
-							Buffer.BlockCopy(initialBuffer.Block, 0, combined.Block, 0, readLength * sizeof(char));
-							Buffer.BlockCopy(remaining.ToArray(), 0, combined.Block, readLength * sizeof(char), remaining.Length * sizeof(char));
-							UhtBuffer.Return(initialBuffer);
-							return combined;
-						}
-					}
+					initialBuffer.Reset(readLength);
+					return initialBuffer;
+				}
+				else
+				{
+					string remaining = sr.ReadToEnd();
+					long totalSize = readLength + remaining.Length;
+					UhtBuffer combined = UhtBuffer.Borrow((int)totalSize);
+					Buffer.BlockCopy(initialBuffer.Block, 0, combined.Block, 0, readLength * sizeof(char));
+					Buffer.BlockCopy(remaining.ToArray(), 0, combined.Block, readLength * sizeof(char), remaining.Length * sizeof(char));
+					UhtBuffer.Return(initialBuffer);
+					return combined;
 				}
 			}
 			catch (IOException)
@@ -139,10 +135,8 @@ namespace EpicGames.UHT.Utils
 				{
 					Directory.CreateDirectory(fileDirectory);
 				}
-				using (StreamWriter writer = new StreamWriter(filePath, false, new UTF8Encoding(false, true), 16 * 1024))
-				{
-					writer.Write(contents);
-				}
+				using StreamWriter writer = new(filePath, false, new UTF8Encoding(false, true), 16 * 1024);
+				writer.Write(contents);
 				return true;
 			}
 			catch (Exception)
@@ -182,28 +176,24 @@ namespace EpicGames.UHT.Utils
 
 			try
 			{
-				using (FileStream fs = new FileStream(filePath, FileMode.Open, FileAccess.Read, FileShare.Read, 4 * 1024, FileOptions.SequentialScan))
+				using FileStream fs = new(filePath, FileMode.Open, FileAccess.Read, FileShare.Read, 4 * 1024, FileOptions.SequentialScan);
+				using StreamReader sr = new(fs, Encoding.UTF8, true);
+
+				// Try to read the whole file into a buffer created by hand.  This avoids a LOT of memory allocations which in turn reduces the
+				// GC stress on the system.  Removing the StreamReader would be nice in the future.
+				long rawFileLength = fs.Length;
+				char[] initialBuffer = new char[rawFileLength];
+				int readLength = sr.Read(initialBuffer, 0, (int)rawFileLength);
+				if (sr.EndOfStream)
 				{
-					using (StreamReader sr = new StreamReader(fs, Encoding.UTF8, true))
-					{
-
-						// Try to read the whole file into a buffer created by hand.  This avoids a LOT of memory allocations which in turn reduces the
-						// GC stress on the system.  Removing the StreamReader would be nice in the future.
-						long rawFileLength = fs.Length;
-						char[] initialBuffer = new char[rawFileLength];
-						int readLength = sr.Read(initialBuffer, 0, (int)rawFileLength);
-						if (sr.EndOfStream)
-						{
-							contents = new StringView(new ReadOnlyMemory<char>(initialBuffer, 0, readLength));
-						}
-						else
-						{
-							string remaining = sr.ReadToEnd();
-							contents = new StringView(String.Concat(new ReadOnlySpan<char>(initialBuffer, 0, readLength), remaining));
-						}
-						return true;
-					}
+					contents = new StringView(new ReadOnlyMemory<char>(initialBuffer, 0, readLength));
 				}
+				else
+				{
+					string remaining = sr.ReadToEnd();
+					contents = new StringView(String.Concat(new ReadOnlySpan<char>(initialBuffer, 0, readLength), remaining));
+				}
+				return true;
 			}
 			catch (IOException)
 			{
diff --git a/Engine/Source/Programs/Shared/EpicGames.UHT/Utils/UhtHash.cs b/Engine/Source/Programs/Shared/EpicGames.UHT/Utils/UhtHash.cs
--- a/Engine/Source/Programs/Shared/EpicGames.UHT/Utils/UhtHash.cs
+++ b/Engine/Source/Programs/Shared/EpicGames.UHT/Utils/UhtHash.cs
@@ -11,14 +11,14 @@ namespace EpicGames.UHT.Utils
 	/// </summary>
 	public struct UhtHash
 	{
-		private ulong _hashInternal;
+		private ulong _hash;
 
 		/// <summary>
 		/// Start the hash computation
 		/// </summary>
 		public void Begin()
 		{
-			this._hashInternal = 0;
+			this._hash = 0;
 		}
 
 		/// <summary>
@@ -32,13 +32,13 @@ namespace EpicGames.UHT.Utils
 				int crPos = text.IndexOf('\r');
 				if (crPos == -1)
 				{
-					HashText(text.Slice(0, text.Length));
+					HashText(text[..]);
 					break;
 				}
 				else
 				{
-					HashText(text.Slice(0, crPos));
-					text = text.Slice(crPos + 1);
+					HashText(text[..crPos]);
+					text = text[(crPos + 1)..];
 				}
 			}
 		}
@@ -49,7 +49,7 @@ namespace EpicGames.UHT.Utils
 		/// <returns>Final hash value</returns>
 		public uint End()
 		{
-			return (uint)(this._hashInternal + (this._hashInternal >> 32));
+			return (uint)(this._hash + (this._hash >> 32));
 		}
 
 		/// <summary>
@@ -59,7 +59,7 @@ namespace EpicGames.UHT.Utils
 		/// <returns>Hash value</returns>
 		public static uint GenenerateTextHash(ReadOnlySpan<char> text)
 		{
-			UhtHash hash = new UhtHash();
+			UhtHash hash = new();
 			hash.Begin();
 			hash.Add(text);
 			return hash.End();
@@ -73,7 +73,7 @@ namespace EpicGames.UHT.Utils
 				{
 					fixed (char* textPtr = text)
 					{
-						this._hashInternal = CityHash.CityHash64WithSeed((byte*)textPtr, (uint)(text.Length * sizeof(char)), this._hashInternal);
+						this._hash = CityHash.CityHash64WithSeed((byte*)textPtr, (uint)(text.Length * sizeof(char)), this._hash);
 					}
 				}
 			}
