More coding style cleanups

diff --git a/Engine/Source/Programs/Shared/EpicGames.UHT/Utils/UhtMessage.cs b/Engine/Source/Programs/Shared/EpicGames.UHT/Utils/UhtMessage.cs
--- a/Engine/Source/Programs/Shared/EpicGames.UHT/Utils/UhtMessage.cs
+++ b/Engine/Source/Programs/Shared/EpicGames.UHT/Utils/UhtMessage.cs
@@ -198,7 +198,7 @@ namespace EpicGames.UHT.Utils
 		{
 			if (context is IUhtMessageExtraContext extraContextInterface)
 			{
-				StringBuilder builder = new StringBuilder();
+				StringBuilder builder = new();
 				Append(builder, extraContextInterface, false);
 				return builder.ToString();
 			}
@@ -286,14 +286,14 @@ namespace EpicGames.UHT.Utils
 	/// </summary>
 	public class UhtSimpleMessageSite : IUhtMessageSite
 	{
-		private readonly IUhtMessageSession _messageSessionInternal;
-		private IUhtMessageSource? _messageSourceInternal;
+		private readonly IUhtMessageSession _messageSession;
+		private IUhtMessageSource? _messageSource;
 
 		#region IUHTMessageSite implementation
 		/// <inheritdoc/>
-		public IUhtMessageSession MessageSession => this._messageSessionInternal;
+		public IUhtMessageSession MessageSession => this._messageSession;
 		/// <inheritdoc/>
-		public IUhtMessageSource? MessageSource { get => this._messageSourceInternal; set => this._messageSourceInternal = value; }
+		public IUhtMessageSource? MessageSource { get => this._messageSource; set => this._messageSource = value; }
 		/// <inheritdoc/>
 		public IUhtMessageLineNumber? MessageLineNumber => null;
 		#endregion
@@ -305,8 +305,8 @@ namespace EpicGames.UHT.Utils
 		/// <param name="messageSource">Source for the messages</param>
 		public UhtSimpleMessageSite(IUhtMessageSession messageSession, IUhtMessageSource? messageSource = null)
 		{
-			this._messageSessionInternal = messageSession;
-			this._messageSourceInternal = messageSource;
+			this._messageSession = messageSession;
+			this._messageSource = messageSource;
 		}
 	}
 
@@ -487,7 +487,7 @@ namespace EpicGames.UHT.Utils
 	public sealed class UhtTlsMessageExtraContext : IUhtMessageExtraContext
 	{
 		private Stack<object?>? _extraContexts;
-		private static readonly ThreadLocal<UhtTlsMessageExtraContext> s_tls = new ThreadLocal<UhtTlsMessageExtraContext>(() => new UhtTlsMessageExtraContext());
+		private static readonly ThreadLocal<UhtTlsMessageExtraContext> s_tls = new(() => new UhtTlsMessageExtraContext());
 
 		#region IUHTMessageExtraContext implementation
 		IEnumerable<object?>? IUhtMessageExtraContext.MessageExtraContext => this._extraContexts;
diff --git a/Engine/Source/Programs/Shared/EpicGames.UHT/Utils/UhtReferenceCollector.cs b/Engine/Source/Programs/Shared/EpicGames.UHT/Utils/UhtReferenceCollector.cs
--- a/Engine/Source/Programs/Shared/EpicGames.UHT/Utils/UhtReferenceCollector.cs
+++ b/Engine/Source/Programs/Shared/EpicGames.UHT/Utils/UhtReferenceCollector.cs
@@ -125,9 +125,7 @@ namespace EpicGames.UHT.Utils
 			for (int index = 0; index < this.References.Count; ++index)
 			{
 				int key = this.References[index];
-				int objectIndex;
-				bool registered;
-				UngetKey(key, out objectIndex, out registered);
+				UngetKey(key, out int objectIndex, out bool registered);
 				sorted[index] = referenceStringDelegate(objectIndex, registered);
 			}
 
diff --git a/Engine/Source/Programs/Shared/EpicGames.UHT/Utils/UhtSession.cs b/Engine/Source/Programs/Shared/EpicGames.UHT/Utils/UhtSession.cs
--- a/Engine/Source/Programs/Shared/EpicGames.UHT/Utils/UhtSession.cs
+++ b/Engine/Source/Programs/Shared/EpicGames.UHT/Utils/UhtSession.cs
@@ -54,26 +54,26 @@ namespace EpicGames.UHT.Utils
 	{
 		public struct Output
 		{
-			public string _filePath;
-			public string _tempFilePath;
-			public bool _saved;
+			public string FilePath { get; set; }
+			public string TempFilePath { get; set; }
+			public bool Saved { get; set; }
 		}
 
 		/// <summary>
 		/// UHT session
 		/// </summary>
-		private readonly UhtSession _sessionInternal;
+		private readonly UhtSession _session;
 
 		/// <summary>
 		/// Module associated with the plugin
 		/// </summary>
-		private readonly UHTManifest.Module? _pluginModuleInternal;
+		private readonly UHTManifest.Module? _pluginModule;
 
 		/// <summary>
 		/// Limiter for the number of files being saved to the reference directory.
 		/// The OS can get swamped on high core systems
 		/// </summary>
-		private static readonly Semaphore s_writeRefSemaphore = new Semaphore(32, 32);
+		private static readonly Semaphore s_writeRefSemaphore = new(32, 32);
 
 		/// <summary>
 		/// Requesting exporter
@@ -83,12 +83,12 @@ namespace EpicGames.UHT.Utils
 		/// <summary>
 		/// UHT Session
 		/// </summary>
-		public UhtSession Session => this._sessionInternal;
+		public UhtSession Session => this._session;
 
 		/// <summary>
 		/// Plugin module
 		/// </summary>
-		public UHTManifest.Module? PluginModule => this._pluginModuleInternal;
+		public UHTManifest.Module? PluginModule => this._pluginModule;
 
 		/// <summary>
 		/// Collection of error from mismatches with the reference files
@@ -123,8 +123,8 @@ namespace EpicGames.UHT.Utils
 		/// <param name="exporter">Exporter being run</param>
 		public UhtExportFactory(UhtSession session, UHTManifest.Module? pluginModule, UhtExporter exporter)
 		{
-			this._sessionInternal = session;
-			this._pluginModuleInternal = pluginModule;
+			this._session = session;
+			this._pluginModule = pluginModule;
 			this.Exporter = exporter;
 			if (this.Session.ReferenceMode != UhtReferenceMode.None)
 			{
@@ -143,11 +143,9 @@ namespace EpicGames.UHT.Utils
 		/// <param name="builder">Source for the content</param>
 		public void CommitOutput(string filePath, StringBuilder builder)
 		{
-			using (UhtBorrowBuffer borrowBuffer = new UhtBorrowBuffer(builder))
-			{
-				string tempFilePath = filePath + ".tmp";
-				SaveIfChanged(filePath, tempFilePath, new StringView(borrowBuffer.Buffer.Memory));
-			}
+			using UhtBorrowBuffer borrowBuffer = new(builder);
+			string tempFilePath = filePath + ".tmp";
+			SaveIfChanged(filePath, tempFilePath, new StringView(borrowBuffer.Buffer.Memory));
 		}
 
 		/// <summary>
@@ -353,7 +351,7 @@ namespace EpicGames.UHT.Utils
 			// Add this to the list of outputs
 			lock (this.Outputs)
 			{
-				this.Outputs.Add(new Output { _filePath = filePath, _tempFilePath = tempFilePath, _saved = save });
+				this.Outputs.Add(new Output { FilePath = filePath, TempFilePath = tempFilePath, Saved = save });
 			}
 		}
 
@@ -371,28 +369,27 @@ namespace EpicGames.UHT.Utils
 			{
 
 				// These outputs are used to cull old outputs from the directories
-				Dictionary<string, HashSet<string>> outputsByDirectory = new Dictionary<string, HashSet<string>>(StringComparer.OrdinalIgnoreCase);
-				List<UhtExportFactory.Output> saves = new List<UhtExportFactory.Output>();
+				Dictionary<string, HashSet<string>> outputsByDirectory = new(StringComparer.OrdinalIgnoreCase);
+				List<UhtExportFactory.Output> saves = new();
 
 				// Collect information about the outputs
 				foreach (UhtExportFactory.Output output in this.Outputs)
 				{
 
 					// Add this output to the list of expected outputs by directory
-					string? fileDirectory = Path.GetDirectoryName(output._filePath);
+					string? fileDirectory = Path.GetDirectoryName(output.FilePath);
 					if (fileDirectory != null)
 					{
-						HashSet<string>? files;
-						if (!outputsByDirectory.TryGetValue(fileDirectory, out files))
+						if (!outputsByDirectory.TryGetValue(fileDirectory, out HashSet<string>? files))
 						{
 							files = new HashSet<string>(StringComparer.OrdinalIgnoreCase);
 							outputsByDirectory.Add(fileDirectory, files);
 						}
-						files.Add(Path.GetFileName(output._filePath));
+						files.Add(Path.GetFileName(output.FilePath));
 					}
 
 					// Add the save task
-					if (output._saved)
+					if (output.Saved)
 					{
 						saves.Add(output);
 					}
@@ -443,7 +440,7 @@ namespace EpicGames.UHT.Utils
 		/// <param name="output">The output file to rename</param>
 		private void RenameSource(UhtExportFactory.Output output)
 		{
-			this.Session.RenameSource(output._tempFilePath, output._filePath);
+			this.Session.RenameSource(output.TempFilePath, output.FilePath);
 		}
 
 		/// <summary>
@@ -561,7 +558,7 @@ namespace EpicGames.UHT.Utils
 		/// <summary>
 		/// Collection of reserved names
 		/// </summary>
-		private static readonly HashSet<string> s_reservedNames = new HashSet<string> { "none" };
+		private static readonly HashSet<string> s_reservedNames = new() { "none" };
 
 		#region Configurable settings
 
@@ -675,24 +672,24 @@ namespace EpicGames.UHT.Utils
 		/// <summary>
 		/// Collection of packages from the manifest
 		/// </summary>
-		public IReadOnlyList<UhtPackage> Packages => this._packagesInternal;
+		public IReadOnlyList<UhtPackage> Packages => this._packages;
 
 		/// <summary>
 		/// Collection of header files from the manifest.  The header files will also appear as the children 
 		/// of the packages
 		/// </summary>
-		public IReadOnlyList<UhtHeaderFile> HeaderFiles => this._headerFilesInternal;
+		public IReadOnlyList<UhtHeaderFile> HeaderFiles => this._headerFiles;
 
 		/// <summary>
 		/// Collection of header files topologically sorted.  This will not be populated until after header files
 		/// are parsed and resolved.
 		/// </summary>
-		public IReadOnlyList<UhtHeaderFile> SortedHeaderFiles => this._sortedHeaderFilesInternal;
+		public IReadOnlyList<UhtHeaderFile> SortedHeaderFiles => this._sortedHeaderFiles;
 
 		/// <summary>
 		/// Dictionary of stripped file name to the header file
 		/// </summary>
-		public IReadOnlyDictionary<string, UhtHeaderFile> HeaderFileDictionary => this._headerFileDictionaryInternal;
+		public IReadOnlyDictionary<string, UhtHeaderFile> HeaderFileDictionary => this._headerFileDictionary;
 
 		/// <summary>
 		/// After headers are parsed, returns the UObject class.
@@ -766,40 +763,40 @@ namespace EpicGames.UHT.Utils
 		/// </summary>
 		public UhtClass? INotifyFieldValueChanged { get; set; } = null;
 
-		private readonly List<UhtPackage> _packagesInternal = new List<UhtPackage>();
-		private readonly List<UhtHeaderFile> _headerFilesInternal = new List<UhtHeaderFile>();
-		private readonly List<UhtHeaderFile> _sortedHeaderFilesInternal = new List<UhtHeaderFile>();
-		private readonly Dictionary<string, UhtHeaderFile> _headerFileDictionaryInternal = new Dictionary<string, UhtHeaderFile>(StringComparer.OrdinalIgnoreCase);
-		private long _errorCountInternal = 0;
-		private long _warningCountInternal = 0;
-		private readonly List<UhtMessage> _messages = new List<UhtMessage>();
+		private readonly List<UhtPackage> _packages = new();
+		private readonly List<UhtHeaderFile> _headerFiles = new();
+		private readonly List<UhtHeaderFile> _sortedHeaderFiles = new();
+		private readonly Dictionary<string, UhtHeaderFile> _headerFileDictionary = new(StringComparer.OrdinalIgnoreCase);
+		private long _errorCount = 0;
+		private long _warningCount = 0;
+		private readonly List<UhtMessage> _messages = new();
 		private Task? _messageTask = null;
 		private UhtClass? _uobject = null;
 		private UhtClass? _uclass = null;
 		private UhtClass? _uinterface = null;
 		private UhtClass? _iinterface = null;
-		private readonly TypeCounter _typeCounterInternal = new TypeCounter();
-		private readonly TypeCounter _packageTypeCountInternal = new TypeCounter();
-		private readonly TypeCounter _headerFileTypeCountInternal = new TypeCounter();
-		private readonly TypeCounter _objectTypeCountInternal = new TypeCounter();
-		private UhtSymbolTable _sourceNameSymbolTable = new UhtSymbolTable(0);
-		private UhtSymbolTable _engineNameSymbolTable = new UhtSymbolTable(0);
+		private readonly TypeCounter _typeCounter = new();
+		private readonly TypeCounter _packageTypeCount = new();
+		private readonly TypeCounter _headerFileTypeCount = new();
+		private readonly TypeCounter _objectTypeCount = new();
+		private UhtSymbolTable _sourceNameSymbolTable = new(0);
+		private UhtSymbolTable _engineNameSymbolTable = new(0);
 		private bool _symbolTablePopulated = false;
 		private Task? _referenceDeleteTask = null;
-		private readonly Dictionary<string, bool> _exporterStates = new Dictionary<string, bool>(StringComparer.OrdinalIgnoreCase);
-		private readonly Dictionary<string, EnumAndValue> _fullEnumValueLookup = new Dictionary<string, EnumAndValue>();
-		private readonly Dictionary<string, UhtEnum> _shortEnumValueLookup = new Dictionary<string, UhtEnum>();
+		private readonly Dictionary<string, bool> _exporterStates = new(StringComparer.OrdinalIgnoreCase);
+		private readonly Dictionary<string, EnumAndValue> _fullEnumValueLookup = new();
+		private readonly Dictionary<string, UhtEnum> _shortEnumValueLookup = new();
 		private JsonDocument? _projectJson = null;
 
 		/// <summary>
 		/// The number of errors
 		/// </summary>
-		public long ErrorCount => Interlocked.Read(ref this._errorCountInternal);
+		public long ErrorCount => Interlocked.Read(ref this._errorCount);
 
 		/// <summary>
 		/// The number of warnings
 		/// </summary>
-		public long WarningCount => Interlocked.Read(ref this._warningCountInternal);
+		public long WarningCount => Interlocked.Read(ref this._warningCount);
 
 		/// <summary>
 		/// True if any errors have occurred or warnings if warnings are to be treated as errors 
@@ -821,13 +818,13 @@ namespace EpicGames.UHT.Utils
 		/// <returns>New index</returns>
 		public int GetNextTypeIndex()
 		{
-			return this._typeCounterInternal.GetNext();
+			return this._typeCounter.GetNext();
 		}
 
 		/// <summary>
 		/// Return the number of types that have been defined.  This includes all types.
 		/// </summary>
-		public int TypeCount => this._typeCounterInternal.Count;
+		public int TypeCount => this._typeCounter.Count;
 
 		/// <summary>
 		/// Return the index for a newly defined packaging
@@ -835,13 +832,13 @@ namespace EpicGames.UHT.Utils
 		/// <returns>New index</returns>
 		public int GetNextHeaderFileTypeIndex()
 		{
-			return this._headerFileTypeCountInternal.GetNext();
+			return this._headerFileTypeCount.GetNext();
 		}
 
 		/// <summary>
 		/// Return the number of headers that have been defined
 		/// </summary>
-		public int HeaderFileTypeCount => this._headerFileTypeCountInternal.Count;
+		public int HeaderFileTypeCount => this._headerFileTypeCount.Count;
 
 		/// <summary>
 		/// Return the index for a newly defined package
@@ -849,13 +846,13 @@ namespace EpicGames.UHT.Utils
 		/// <returns>New index</returns>
 		public int GetNextPackageTypeIndex()
 		{
-			return this._packageTypeCountInternal.GetNext();
+			return this._packageTypeCount.GetNext();
 		}
 
 		/// <summary>
 		/// Return the number of UPackage types that have been defined
 		/// </summary>
-		public int PackageTypeCount => this._packageTypeCountInternal.Count;
+		public int PackageTypeCount => this._packageTypeCount.Count;
 
 		/// <summary>
 		/// Return the index for a newly defined object
@@ -863,13 +860,13 @@ namespace EpicGames.UHT.Utils
 		/// <returns>New index</returns>
 		public int GetNextObjectTypeIndex()
 		{
-			return this._objectTypeCountInternal.GetNext();
+			return this._objectTypeCount.GetNext();
 		}
 
 		/// <summary>
 		/// Return the total number of UObject types that have been defined
 		/// </summary>
-		public int ObjectTypeCount => this._objectTypeCountInternal.Count;
+		public int ObjectTypeCount => this._objectTypeCount.Count;
 
 		/// <summary>
 		/// Return the collection of exporters
@@ -977,21 +974,21 @@ namespace EpicGames.UHT.Utils
 		{
 			if (this.FileManager == null)
 			{
-				Interlocked.Increment(ref this._errorCountInternal);
+				Interlocked.Increment(ref this._errorCount);
 				Log.Logger.LogError("No file manager supplied, aborting.");
 				return;
 			}
 
 			if (this.Config == null)
 			{
-				Interlocked.Increment(ref this._errorCountInternal);
+				Interlocked.Increment(ref this._errorCount);
 				Log.Logger.LogError("No configuration supplied, aborting.");
 				return;
 			}
 
 			if (this.Tables == null)
 			{
-				Interlocked.Increment(ref this._errorCountInternal);
+				Interlocked.Increment(ref this._errorCount);
 				Log.Logger.LogError("No parsing tables supplied, aborting.");
 				return;
 			}
@@ -1412,7 +1409,7 @@ namespace EpicGames.UHT.Utils
 		/// <param name="name">The name of the symbol</param>
 		private static void FindTypeError(IUhtMessageSite messageSite, int lineNumber, UhtFindOptions options, string name)
 		{
-			List<string> types = new List<string>();
+			List<string> types = new();
 			if (options.HasAnyFlags(UhtFindOptions.Enum))
 			{
 				types.Add("'enum'");
@@ -1448,7 +1445,7 @@ namespace EpicGames.UHT.Utils
 		/// <returns></returns>
 		public UhtHeaderFile? FindHeaderFile(string name)
 		{
-			if (this._headerFileDictionaryInternal.TryGetValue(name, out UhtHeaderFile? headerFile))
+			if (this._headerFileDictionary.TryGetValue(name, out UhtHeaderFile? headerFile))
 			{
 				return headerFile;
 			}
@@ -1478,11 +1475,11 @@ namespace EpicGames.UHT.Utils
 			{
 				case UhtMessageType.Error:
 				case UhtMessageType.Ice:
-					Interlocked.Increment(ref this._errorCountInternal);
+					Interlocked.Increment(ref this._errorCount);
 					break;
 
 				case UhtMessageType.Warning:
-					Interlocked.Increment(ref this._warningCountInternal);
+					Interlocked.Increment(ref this._warningCount);
 					break;
 
 				case UhtMessageType.Info:
@@ -1567,7 +1564,7 @@ namespace EpicGames.UHT.Utils
 		/// <returns>List of all the messages</returns>
 		public List<string> CollectMessages()
 		{
-			List<string> messages = new List<string>();
+			List<string> messages = new();
 			foreach (UhtMessage message in FetchOrderedMessages())
 			{
 				messages.Add(FormatMessage(message));
@@ -1597,7 +1594,7 @@ namespace EpicGames.UHT.Utils
 		/// <returns>Enumerator</returns>
 		private IOrderedEnumerable<UhtMessage> FetchOrderedMessages()
 		{
-			List<UhtMessage> messages = new List<UhtMessage>();
+			List<UhtMessage> messages = new();
 			lock (this._messages)
 			{
 				messages.AddRange(this._messages);
@@ -1792,8 +1789,8 @@ namespace EpicGames.UHT.Utils
 							break;
 					}
 
-					UhtPackage package = new UhtPackage(this, module, packageFlags);
-					this._packagesInternal.Add(package);
+					UhtPackage package = new(this, module, packageFlags);
+					this._packages.Add(package);
 				}
 			});
 		}
@@ -1813,8 +1810,7 @@ namespace EpicGames.UHT.Utils
 				// Make sure this isn't a duplicate
 				string normalizedFullFilePath = GetNormalizedFullFilePath(headerFilePath);
 				string fileName = Path.GetFileName(normalizedFullFilePath);
-				UhtHeaderFile? existingHeaderFile;
-				if (_headerFileDictionaryInternal.TryGetValue(fileName, out existingHeaderFile) && existingHeaderFile != null)
+				if (_headerFileDictionary.TryGetValue(fileName, out UhtHeaderFile? existingHeaderFile) && existingHeaderFile != null)
 				{
 					string normalizedExistingFullFilePath = GetNormalizedFullFilePath(existingHeaderFile.FilePath);
 					if (!String.Equals(normalizedFullFilePath, normalizedExistingFullFilePath, StringComparison.OrdinalIgnoreCase))
@@ -1826,10 +1822,10 @@ namespace EpicGames.UHT.Utils
 				}
 
 				// Create the header file and add to the collections
-				UhtHeaderFile headerFile = new UhtHeaderFile(package, headerFilePath);
+				UhtHeaderFile headerFile = new(package, headerFilePath);
 				headerFile.HeaderFileType = headerFileType;
-				_headerFilesInternal.Add(headerFile);
-				_headerFileDictionaryInternal.Add(fileName, headerFile);
+				_headerFiles.Add(headerFile);
+				_headerFileDictionary.Add(fileName, headerFile);
 				package.AddChild(headerFile);
 
 				// Save metadata for the class path, both for it's include path and relative to the module base directory
@@ -1841,14 +1837,14 @@ namespace EpicGames.UHT.Utils
 						++stripLength;
 					}
 
-					headerFile.ModuleRelativeFilePath = normalizedFullFilePath.Substring(stripLength);
+					headerFile.ModuleRelativeFilePath = normalizedFullFilePath[stripLength..];
 
-					if (normalizedFullFilePath.Substring(stripLength).StartsWith(typeDirectory, true, null))
+					if (normalizedFullFilePath[stripLength..].StartsWith(typeDirectory, true, null))
 					{
 						stripLength += typeDirectory.Length;
 					}
 
-					headerFile.IncludeFilePath = normalizedFullFilePath.Substring(stripLength);
+					headerFile.IncludeFilePath = normalizedFullFilePath[stripLength..];
 				}
 			}
 		}
@@ -1864,7 +1860,7 @@ namespace EpicGames.UHT.Utils
 			{
 				Log.Logger.LogTrace("Step - Prepare Headers");
 
-				foreach (UhtPackage package in this._packagesInternal)
+				foreach (UhtPackage package in this._packages)
 				{
 					if (package.Module != null)
 					{
@@ -1876,9 +1872,9 @@ namespace EpicGames.UHT.Utils
 				}
 
 				// Locate the NoExportTypes.h file and add it to every other header file
-				if (this._headerFileDictionaryInternal.TryGetValue("NoExportTypes.h", out UhtHeaderFile? noExportTypes))
+				if (this._headerFileDictionary.TryGetValue("NoExportTypes.h", out UhtHeaderFile? noExportTypes))
 				{
-					foreach (UhtPackage package in this._packagesInternal)
+					foreach (UhtPackage package in this._packages)
 					{
 						foreach (UhtHeaderFile headerFile in package.Children)
 						{
@@ -1903,14 +1899,14 @@ namespace EpicGames.UHT.Utils
 
 			if (this.GoWide)
 			{
-				Parallel.ForEach(this._headerFilesInternal, headerFile =>
+				Parallel.ForEach(this._headerFiles, headerFile =>
 				{
 					ParseHeaderFile(headerFile);
 				});
 			}
 			else
 			{
-				foreach (UhtHeaderFile headerFile in this._headerFilesInternal)
+				foreach (UhtHeaderFile headerFile in this._headerFiles)
 				{
 					ParseHeaderFile(headerFile);
 				}
@@ -1996,14 +1992,14 @@ namespace EpicGames.UHT.Utils
 
 			if (this.GoWide)
 			{
-				Parallel.ForEach(this._headerFilesInternal, headerFile =>
+				Parallel.ForEach(this._headerFiles, headerFile =>
 				{
 					stepDelegate(headerFile);
 				});
 			}
 			else
 			{
-				foreach (UhtHeaderFile headerFile in this._headerFilesInternal)
+				foreach (UhtHeaderFile headerFile in this._headerFiles)
 				{
 					stepDelegate(headerFile);
 				}
@@ -2014,7 +2010,7 @@ namespace EpicGames.UHT.Utils
 		#region Symbol table initialization
 		private void PopulateSymbolTable()
 		{
-			foreach (UhtHeaderFile headerFile in this._headerFilesInternal)
+			foreach (UhtHeaderFile headerFile in this._headerFiles)
 			{
 				AddTypeToSymbolTable(headerFile);
 			}
@@ -2052,7 +2048,7 @@ namespace EpicGames.UHT.Utils
 					{
 						if (value.Name.StartsWith(checkName, StringComparison.Ordinal))
 						{
-							this._shortEnumValueLookup.TryAdd(value.Name.Substring(checkName.Length), enumObj);
+							this._shortEnumValueLookup.TryAdd(value.Name[checkName.Length..], enumObj);
 						}
 					}
 				}
@@ -2141,7 +2137,7 @@ namespace EpicGames.UHT.Utils
 						}
 					}
 					states[visit.HeaderFileTypeIndex] = TopologicalState.Permanent;
-					this._sortedHeaderFilesInternal.Add(visit);
+					this._sortedHeaderFiles.Add(visit);
 					return null;
 
 				case TopologicalState.Temporary:
@@ -2162,8 +2158,8 @@ namespace EpicGames.UHT.Utils
 				Log.Logger.LogTrace("Step - Topological Sort Header Files");
 
 				// Initialize a scratch table for topological states
-				this._sortedHeaderFilesInternal.Capacity = this.HeaderFileTypeCount;
-				List<TopologicalState> states = new List<TopologicalState>(this.HeaderFileTypeCount);
+				this._sortedHeaderFiles.Capacity = this.HeaderFileTypeCount;
+				List<TopologicalState> states = new(this.HeaderFileTypeCount);
 				for (int index = 0; index < this.HeaderFileTypeCount; ++index)
 				{
 					states.Add(TopologicalState.Unmarked);
@@ -2187,7 +2183,7 @@ namespace EpicGames.UHT.Utils
 		#endregion
 
 		#region Validation helpers
-		private readonly HashSet<UhtScriptStruct> _scriptStructsValidForNet = new HashSet<UhtScriptStruct>();
+		private readonly HashSet<UhtScriptStruct> _scriptStructsValidForNet = new();
 
 		/// <summary>
 		/// Validate that the given referenced script structure is valid for network operations.  If the structure
@@ -2278,7 +2274,7 @@ namespace EpicGames.UHT.Utils
 				return false;
 			}
 
-			JsonObject rootObject = new JsonObject(this._projectJson.RootElement);
+			JsonObject rootObject = new(this._projectJson.RootElement);
 			if (rootObject.TryGetObjectArrayField("Plugins", out JsonObject[]? plugins))
 			{
 				foreach (JsonObject plugin in plugins)
@@ -2316,7 +2312,7 @@ namespace EpicGames.UHT.Utils
 
 		private void StepExport()
 		{
-			HashSet<string> externalDependencies = new HashSet<string>();
+			HashSet<string> externalDependencies = new();
 			long totalWrittenFiles = 0;
 			Try(null, () =>
 			{
@@ -2324,8 +2320,7 @@ namespace EpicGames.UHT.Utils
 
 				foreach (UhtExporter exporter in this.ExporterTable)
 				{
-					bool run = false;
-					if (!this._exporterStates.TryGetValue(exporter.Name, out run))
+					if (!this._exporterStates.TryGetValue(exporter.Name, out bool run))
 					{
 						run = Config!.IsExporterEnabled(exporter.Name) ||
 							(exporter.Options.HasAnyFlags(UhtExporterOptions.Default) && !this.NoDefaultExporters);
@@ -2352,11 +2347,11 @@ namespace EpicGames.UHT.Utils
 					if (run)
 					{
 						Log.TraceLog($"       Running exporter {exporter.Name}");
-						UhtExportFactory factory = new UhtExportFactory(this, pluginModule, exporter);
+						UhtExportFactory factory = new(this, pluginModule, exporter);
 						factory.Run();
 						foreach (UhtExportFactory.Output output in factory.Outputs)
 						{
-							if (output._saved)
+							if (output.Saved)
 							{
 								totalWrittenFiles++;
 							}
@@ -2375,12 +2370,10 @@ namespace EpicGames.UHT.Utils
 				// Save the collected external dependencies
 				if (!String.IsNullOrEmpty(this.Manifest!.ExternalDependenciesFile))
 				{
-					using (StreamWriter output = new StreamWriter(this.Manifest!.ExternalDependenciesFile))
+					using StreamWriter output = new(this.Manifest!.ExternalDependenciesFile);
+					foreach (string dep in externalDependencies)
 					{
-						foreach (string dep in externalDependencies)
-						{
-							output.WriteLine(dep);
-						}
+						output.WriteLine(dep);
 					}
 				}
 			});
diff --git a/Engine/Source/Programs/Shared/EpicGames.UHT/Utils/UhtStringBuilderExtensions.cs b/Engine/Source/Programs/Shared/EpicGames.UHT/Utils/UhtStringBuilderExtensions.cs
--- a/Engine/Source/Programs/Shared/EpicGames.UHT/Utils/UhtStringBuilderExtensions.cs
+++ b/Engine/Source/Programs/Shared/EpicGames.UHT/Utils/UhtStringBuilderExtensions.cs
@@ -13,12 +13,12 @@ namespace EpicGames.UHT.Utils
 		/// <summary>
 		/// String of tabs used to generate code with proper indentation
 		/// </summary>
-		public static StringView TabsString = new StringView(new string('\t', 128));
+		public static StringView TabsString = new(new string('\t', 128));
 
 		/// <summary>
 		/// String of spaces used to generate code with proper indentation
 		/// </summary>
-		public static StringView SpacesString = new StringView(new string(' ', 128));
+		public static StringView SpacesString = new(new string(' ', 128));
 
 		/// <summary>
 		/// Append tabs to the builder
@@ -35,7 +35,7 @@ namespace EpicGames.UHT.Utils
 			}
 			else if (tabs > 0)
 			{
-				builder.Append(TabsString.Span.Slice(0, tabs));
+				builder.Append(TabsString.Span[..tabs]);
 			}
 			return builder;
 		}
@@ -55,7 +55,7 @@ namespace EpicGames.UHT.Utils
 			}
 			else if (spaces > 0)
 			{
-				builder.Append(SpacesString.Span.Slice(0, spaces));
+				builder.Append(SpacesString.Span[..spaces]);
 			}
 			return builder;
 		}
@@ -347,7 +347,7 @@ namespace EpicGames.UHT.Utils
 						{
 							builder.Append("\"\"");
 						}
-						builder.Append(span.Slice(startIndex, index - startIndex));
+						builder.Append(span[startIndex..index]);
 					}
 
 					// We have either reached the end of the string, break
@@ -492,12 +492,12 @@ namespace EpicGames.UHT.Utils
 		/// <summary>
 		/// Cache of StringBuilders with large initial buffer sizes
 		/// </summary>
-		public static readonly StringBuilderCache Big = new StringBuilderCache(256, 256 * 1024);
+		public static readonly StringBuilderCache Big = new(256, 256 * 1024);
 
 		/// <summary>
 		/// Cache of StringBuilders with small initial buffer sizes
 		/// </summary>
-		public static readonly StringBuilderCache Small = new StringBuilderCache(256, 1 * 1024);
+		public static readonly StringBuilderCache Small = new(256, 1 * 1024);
 
 		/// <summary>
 		/// Capacity of the cache
diff --git a/Engine/Source/Programs/Shared/EpicGames.UHT/Utils/UhtSymbolTable.cs b/Engine/Source/Programs/Shared/EpicGames.UHT/Utils/UhtSymbolTable.cs
--- a/Engine/Source/Programs/Shared/EpicGames.UHT/Utils/UhtSymbolTable.cs
+++ b/Engine/Source/Programs/Shared/EpicGames.UHT/Utils/UhtSymbolTable.cs
@@ -73,12 +73,12 @@ namespace EpicGames.UHT.Utils
 		/// <summary>
 		/// Case name lookup table that returns the symbol index and the case index
 		/// </summary>
-		private readonly Dictionary<string, Lookup> _casedDictionary = new Dictionary<string, Lookup>(StringComparer.Ordinal);
+		private readonly Dictionary<string, Lookup> _casedDictionary = new(StringComparer.Ordinal);
 
 		/// <summary>
 		/// Caseless name lookup table that returns the symbol index
 		/// </summary>
-		private readonly Dictionary<string, int> _caselessDictionary = new Dictionary<string, int>(StringComparer.OrdinalIgnoreCase);
+		private readonly Dictionary<string, int> _caselessDictionary = new(StringComparer.OrdinalIgnoreCase);
 
 		/// <summary>
 		/// Collection of symbols in the table
@@ -101,8 +101,7 @@ namespace EpicGames.UHT.Utils
 		/// <param name="name">The name of the type which could be the source name or the engine name</param>
 		public void Add(UhtType type, string name)
 		{
-			Lookup existing;
-			if (this._casedDictionary.TryGetValue(name, out existing))
+			if (this._casedDictionary.TryGetValue(name, out Lookup existing))
 			{
 				AddExisting(type, existing.CasedIndex, existing.SymbolIndex);
 				return;
@@ -110,8 +109,7 @@ namespace EpicGames.UHT.Utils
 
 			int casedIndex = ++this._casedCount;
 
-			int symbolIndex;
-			if (this._caselessDictionary.TryGetValue(name, out symbolIndex))
+			if (this._caselessDictionary.TryGetValue(name, out int symbolIndex))
 			{
 				this._casedDictionary.Add(name, new Lookup { SymbolIndex = symbolIndex, CasedIndex = casedIndex });
 				AddExisting(type, casedIndex, symbolIndex);
@@ -135,8 +133,7 @@ namespace EpicGames.UHT.Utils
 		/// <exception cref="UhtIceException">Thrown if the symbol wasn't found.</exception>
 		public void Replace(UhtType oldType, UhtType newType, string name)
 		{
-			int symbolIndex;
-			if (this._caselessDictionary.TryGetValue(name, out symbolIndex))
+			if (this._caselessDictionary.TryGetValue(name, out int symbolIndex))
 			{
 				for (; symbolIndex != 0; symbolIndex = this._symbols[symbolIndex].NextIndex)
 				{
@@ -159,8 +156,7 @@ namespace EpicGames.UHT.Utils
 		/// <returns>Found type or null if not found.</returns>
 		public UhtType? FindCasedType(UhtType? startingType, UhtFindOptions options, string name)
 		{
-			Lookup existing;
-			if (this._casedDictionary.TryGetValue(name, out existing))
+			if (this._casedDictionary.TryGetValue(name, out Lookup existing))
 			{
 				return FindType(startingType, options, existing);
 			}
@@ -176,8 +172,7 @@ namespace EpicGames.UHT.Utils
 		/// <returns>Found type or null if not found.</returns>
 		public UhtType? FindCaselessType(UhtType? startingType, UhtFindOptions options, string name)
 		{
-			int symbolIndex;
-			if (this._caselessDictionary.TryGetValue(name, out symbolIndex))
+			if (this._caselessDictionary.TryGetValue(name, out int symbolIndex))
 			{
 				return FindType(startingType, options, new Lookup { SymbolIndex = symbolIndex, CasedIndex = 0 });
 			}
diff --git a/Engine/Source/Programs/Shared/EpicGames.UHT/Utils/UhtTables.cs b/Engine/Source/Programs/Shared/EpicGames.UHT/Utils/UhtTables.cs
--- a/Engine/Source/Programs/Shared/EpicGames.UHT/Utils/UhtTables.cs
+++ b/Engine/Source/Programs/Shared/EpicGames.UHT/Utils/UhtTables.cs
@@ -108,8 +108,8 @@ namespace EpicGames.UHT.Utils
 		/// </summary>
 		public string UserName
 		{
-			get => String.IsNullOrEmpty(this.UserNameInternal) ? this.TableName : this.UserNameInternal;
-			set => this.UserNameInternal = value;
+			get => String.IsNullOrEmpty(this._userName) ? this.TableName : this._userName;
+			set => this._userName = value;
 		}
 
 		/// <summary>
@@ -125,7 +125,7 @@ namespace EpicGames.UHT.Utils
 		/// <summary>
 		/// Internal version of the user name.  If it hasn't been set, then the table name will be used
 		/// </summary>
-		private string UserNameInternal { get; set; } = String.Empty;
+		private string _userName = String.Empty;
 
 		/// <summary>
 		/// Merge the lookup table.  Duplicates will be ignored.
@@ -150,7 +150,7 @@ namespace EpicGames.UHT.Utils
 			{
 				if (methodInfo.Name.EndsWith(suffix, StringComparison.Ordinal))
 				{
-					name = methodInfo.Name.Substring(0, methodInfo.Name.Length - suffix.Length);
+					name = methodInfo.Name[..^suffix.Length];
 				}
 				else
 				{
@@ -249,8 +249,7 @@ namespace EpicGames.UHT.Utils
 		/// <returns>The table associated with the name.</returns>
 		public TTable Get(string tableName)
 		{
-			TTable? table;
-			if (!this.Tables.TryGetValue(tableName, out table))
+			if (!this.Tables.TryGetValue(tableName, out TTable? table))
 			{
 				table = new TTable();
 				table.TableName = tableName;
@@ -287,9 +286,9 @@ namespace EpicGames.UHT.Utils
 		/// <exception cref="UhtIceException">Thrown if there are problems with the tables.</exception>
 		public void Merge()
 		{
-			List<TTable> orderedList = new List<TTable>(this.Tables.Count);
-			List<TTable> remainingList = new List<TTable>(this.Tables.Count);
-			HashSet<UhtLookupTableBase> doneTables = new HashSet<UhtLookupTableBase>();
+			List<TTable> orderedList = new(this.Tables.Count);
+			List<TTable> remainingList = new(this.Tables.Count);
+			HashSet<UhtLookupTableBase> doneTables = new();
 
 			// Collect all of the tables
 			foreach (KeyValuePair<string, TTable> kvp in this.Tables)
diff --git a/Engine/Source/Programs/Shared/EpicGames.UHT/Utils/UhtUtilities.cs b/Engine/Source/Programs/Shared/EpicGames.UHT/Utils/UhtUtilities.cs
--- a/Engine/Source/Programs/Shared/EpicGames.UHT/Utils/UhtUtilities.cs
+++ b/Engine/Source/Programs/Shared/EpicGames.UHT/Utils/UhtUtilities.cs
@@ -123,7 +123,7 @@ namespace EpicGames.UHT.Utils
 		/// <returns>Merged names</returns>
 		public static string MergeTypeNames(IEnumerable<string> typeNames, string andOr, bool quote = false)
 		{
-			List<string> local = new List<string>(typeNames);
+			List<string> local = new(typeNames);
 
 			if (local.Count == 0)
 			{
@@ -132,7 +132,7 @@ namespace EpicGames.UHT.Utils
 
 			local.Sort();
 
-			StringBuilder builder = new StringBuilder();
+			StringBuilder builder = new();
 			for (int index = 0; index < local.Count; ++index)
 			{
 				if (index != 0)
@@ -176,7 +176,7 @@ namespace EpicGames.UHT.Utils
 				case 'A':
 				case 'U':
 					// If it is a class prefix, check for deprecated class prefix also
-					if (sourceName.Span.Length > 12 && sourceName.Span.Slice(1).StartsWith("DEPRECATED_"))
+					if (sourceName.Span.Length > 12 && sourceName.Span[1..].StartsWith("DEPRECATED_"))
 					{
 						return new UhtEngineNameParts { Prefix = new StringView(sourceName, 0, 12), EngineName = new StringView(sourceName, 12), IsDeprecated = true };
 					}
@@ -207,12 +207,12 @@ namespace EpicGames.UHT.Utils
 		/// <summary>
 		/// When only a string view has been appended, this references that StringView
 		/// </summary>
-		private StringView _stringViewInternal = new StringView();
+		private StringView _stringView = new();
 
 		/// <summary>
 		/// Represents more complex data being appended
 		/// </summary>
-		private StringBuilder? _stringBuilderInternal = null;
+		private StringBuilder? _stringBuilder = null;
 
 		/// <summary>
 		/// Set to true when the builder has switched to a StringBuilder (NOTE: This can probably be removed)
@@ -226,13 +226,13 @@ namespace EpicGames.UHT.Utils
 		{
 			get
 			{
-				if (this._useStringBuilder && this._stringBuilderInternal != null)
+				if (this._useStringBuilder && this._stringBuilder != null)
 				{
-					return this._stringBuilderInternal.Length;
+					return this._stringBuilder.Length;
 				}
 				else
 				{
-					return this._stringViewInternal.Span.Length;
+					return this._stringView.Span.Length;
 				}
 			}
 		}
@@ -243,7 +243,7 @@ namespace EpicGames.UHT.Utils
 		/// <returns>Contents of the builder</returns>
 		public StringView ToStringView()
 		{
-			return this._useStringBuilder ? new StringView(this._stringBuilderInternal!.ToString()) : this._stringViewInternal;
+			return this._useStringBuilder ? new StringView(this._stringBuilder!.ToString()) : this._stringView;
 		}
 
 		/// <summary>
@@ -252,7 +252,7 @@ namespace EpicGames.UHT.Utils
 		/// <returns>Contents of the builder</returns>
 		public override string ToString()
 		{
-			return this._useStringBuilder ? this._stringBuilderInternal!.ToString() : this._stringViewInternal.ToString();
+			return this._useStringBuilder ? this._stringBuilder!.ToString() : this._stringView.ToString();
 		}
 
 		/// <summary>
@@ -264,16 +264,16 @@ namespace EpicGames.UHT.Utils
 		{
 			if (this._useStringBuilder)
 			{
-				this._stringBuilderInternal!.Append(text.Span);
+				this._stringBuilder!.Append(text.Span);
 			}
-			else if (this._stringViewInternal.Span.Length > 0)
+			else if (this._stringView.Span.Length > 0)
 			{
 				SwitchToStringBuilder();
-				this._stringBuilderInternal!.Append(text.Span);
+				this._stringBuilder!.Append(text.Span);
 			}
 			else
 			{
-				this._stringViewInternal = text;
+				this._stringView = text;
 			}
 			return this;
 		}
@@ -286,7 +286,7 @@ namespace EpicGames.UHT.Utils
 		public StringViewBuilder Append(char c)
 		{
 			SwitchToStringBuilder();
-			this._stringBuilderInternal!.Append(c);
+			this._stringBuilder!.Append(c);
 			return this;
 		}
 
@@ -297,12 +297,12 @@ namespace EpicGames.UHT.Utils
 		{
 			if (!this._useStringBuilder)
 			{
-				if (this._stringBuilderInternal == null)
+				if (this._stringBuilder == null)
 				{
-					this._stringBuilderInternal = new StringBuilder();
+					this._stringBuilder = new StringBuilder();
 				}
 				this._useStringBuilder = true;
-				this._stringBuilderInternal.Append(this._stringViewInternal.Span);
+				this._stringBuilder.Append(this._stringView.Span);
 			}
 		}
 	}
