Fixed issue in C# UHT where duplicate enumerations would result in a crash.

#rnx

diff --git a/Engine/Source/Programs/Shared/EpicGames.UHT/Utils/UhtSession.cs b/Engine/Source/Programs/Shared/EpicGames.UHT/Utils/UhtSession.cs
--- a/Engine/Source/Programs/Shared/EpicGames.UHT/Utils/UhtSession.cs
+++ b/Engine/Source/Programs/Shared/EpicGames.UHT/Utils/UhtSession.cs
@@ -2043,7 +2043,10 @@ namespace EpicGames.UHT.Utils
 				string checkName = $"{enumObj.SourceName}::";
 				foreach (UhtEnumValue value in enumObj.EnumValues)
 				{
-					this._fullEnumValueLookup.Add(value.Name, new EnumAndValue { Enum = enumObj, Value = value.Value });
+					if (!this._fullEnumValueLookup.TryAdd(value.Name, new EnumAndValue { Enum = enumObj, Value = value.Value }))
+					{
+						//TODO - add a warning
+					}
 					if (addShortNames)
 					{
 						if (value.Name.StartsWith(checkName, StringComparison.Ordinal))
