Minor cleanup of some UHT implementations

#rnx

diff --git a/Engine/Source/Programs/Shared/EpicGames.UHT/Parsers/UhtClassBaseParser.cs b/Engine/Source/Programs/Shared/EpicGames.UHT/Parsers/UhtClassBaseParser.cs
deleted file mode 100644
diff --git a/Engine/Source/Programs/Shared/EpicGames.UHT/Parsers/UhtClassParser.cs b/Engine/Source/Programs/Shared/EpicGames.UHT/Parsers/UhtClassParser.cs
--- a/Engine/Source/Programs/Shared/EpicGames.UHT/Parsers/UhtClassParser.cs
+++ b/Engine/Source/Programs/Shared/EpicGames.UHT/Parsers/UhtClassParser.cs
@@ -1,9 +1,7 @@
 // Copyright Epic Games, Inc. All Rights Reserved.
 
-using System;
 using System.Collections.Generic;
 using System.Diagnostics.CodeAnalysis;
-using System.Text;
 using EpicGames.Core;
 using EpicGames.UHT.Tables;
 using EpicGames.UHT.Tokenizer;
@@ -16,354 +14,8 @@ namespace EpicGames.UHT.Parsers
 	/// UCLASS parser
 	/// </summary>
 	[UnrealHeaderTool]
-	public class UhtClassParser : UhtClassBaseParser
+	public static class UhtClassParser
 	{
-		/// <summary>
-		/// Engine class flags removed
-		/// </summary>
-		public EClassFlags RemovedClassFlags { get; set; } = EClassFlags.None;
-
-		/// <summary>
-		/// Class within identifier
-		/// </summary>
-		public string ClassWithinIdentifier { get; set; } = String.Empty;
-
-		/// <summary>
-		/// Collection of show categories
-		/// </summary>
-		public List<string> ShowCategories { get; } = new List<string>();
-
-		/// <summary>
-		/// Collection of hide categories
-		/// </summary>
-		public List<string> HideCategories { get; } = new List<string>();
-
-		/// <summary>
-		/// Collection of auto expand categories
-		/// </summary>
-		public List<string> AutoExpandCategories { get; } = new List<string>();
-
-		/// <summary>
-		/// Collection of auto collapse categories
-		/// </summary>
-		public List<string> AutoCollapseCategories { get; } = new List<string>();
-
-		/// <summary>
-		/// Collection of prioritize categories
-		/// </summary>
-		public List<string> PrioritizeCategories { get; } = new List<string>();
-
-		/// <summary>
-		/// Collection of show functions
-		/// </summary>
-		public List<string> ShowFunctions { get; } = new List<string>();
-
-		/// <summary>
-		/// Collection of hide functions
-		/// </summary>
-		public List<string> HideFunctions { get; } = new List<string>();
-
-		/// <summary>
-		/// Sparse class data types
-		/// </summary>
-		public List<string> SparseClassDataTypes { get; } = new List<string>();
-
-		/// <summary>
-		/// Class group names
-		/// </summary>
-		public List<string> ClassGroupNames { get; } = new List<string>();
-
-		/// <summary>
-		/// Add the given class flags
-		/// </summary>
-		/// <param name="flags">Flags to add</param>
-		public void AddClassFlags(EClassFlags flags)
-		{
-			this.ClassFlags |= flags;
-			this.RemovedClassFlags &= ~flags;
-		}
-
-		/// <summary>
-		/// Remove the given class flags
-		/// </summary>
-		/// <param name="flags">Flags to remove</param>
-		public void RemoveClassFlags(EClassFlags flags)
-		{
-			this.RemovedClassFlags |= flags;
-			this.ClassFlags &= ~flags;
-		}
-
-		/// <summary>
-		/// Construct a new class parser
-		/// </summary>
-		/// <param name="outer">Outer type</param>
-		/// <param name="lineNumber">Line number</param>
-		public UhtClassParser(UhtType outer, int lineNumber) : base(outer, lineNumber)
-		{
-		}
-
-		/// <inheritdoc/>
-		protected override void ResolveSuper(UhtResolvePhase resolvePhase)
-		{
-			base.ResolveSuper(resolvePhase);
-
-			// Make sure the class within is resolved.  We have to make sure we don't try to resolve ourselves.
-			if (this.ClassWithin != null && this.ClassWithin != this)
-			{
-				this.ClassWithin.Resolve(resolvePhase);
-			}
-
-			switch (resolvePhase)
-			{
-				case UhtResolvePhase.Bases:
-					UhtClass? superClass = this.SuperClass;
-
-					// Merge the super class flags
-					if (superClass != null)
-					{
-						this.ClassFlags |= superClass.ClassFlags & EClassFlags.ScriptInherit & ~this.RemovedClassFlags;
-					}
-
-					foreach (UhtStruct baseStruct in this.Bases)
-					{
-						if (baseStruct is UhtClass baseClass)
-						{
-							if (!baseClass.ClassFlags.HasAnyFlags(EClassFlags.Interface))
-							{
-								this.LogError($"Class '{baseClass.SourceName}' is not an interface; Can only inherit from non-UObjects or UInterface derived interfaces");
-							}
-							this.ClassFlags |= baseClass.ClassFlags & EClassFlags.ScriptInherit & ~this.RemovedClassFlags;
-						}
-					}
-
-					// These following collections only inherit from he parent if they are empty in this class
-					if (this.SparseClassDataTypes.Count == 0)
-					{
-						AppendStringListMetaData(this.SuperClass, UhtNames.SparseClassDataTypes, this.SparseClassDataTypes);
-					}
-					if (this.PrioritizeCategories.Count == 0)
-					{
-						AppendStringListMetaData(this.SuperClass, UhtNames.PrioritizeCategories, this.PrioritizeCategories);
-					}
-
-					// Merge with categories inherited from the parent.
-					MergeCategories();
-
-					SetAndValidateWithinClass(resolvePhase);
-					SetAndValidateConfigName();
-
-					// Copy all of the lists back to the meta data
-					this.MetaData.AddIfNotEmpty(UhtNames.ClassGroupNames, this.ClassGroupNames);
-					this.MetaData.AddIfNotEmpty(UhtNames.AutoCollapseCategories, this.AutoCollapseCategories);
-					this.MetaData.AddIfNotEmpty(UhtNames.AutoExpandCategories, this.AutoExpandCategories);
-					this.MetaData.AddIfNotEmpty(UhtNames.PrioritizeCategories, this.PrioritizeCategories);
-					this.MetaData.AddIfNotEmpty(UhtNames.HideCategories, this.HideCategories);
-					this.MetaData.AddIfNotEmpty(UhtNames.ShowCategories, this.ShowCategories);
-					this.MetaData.AddIfNotEmpty(UhtNames.SparseClassDataTypes, this.SparseClassDataTypes);
-					this.MetaData.AddIfNotEmpty(UhtNames.HideFunctions, this.HideFunctions);
-
-					this.MetaData.Add(UhtNames.IncludePath, this.HeaderFile.IncludeFilePath);
-					break;
-			}
-		}
-
-		private void MergeCategories()
-		{
-			MergeShowCategories();
-
-			// Merge ShowFunctions and HideFunctions
-			AppendStringListMetaData(this.SuperClass, UhtNames.HideFunctions, this.HideFunctions);
-			foreach (string value in this.ShowFunctions)
-			{
-				this.HideFunctions.RemoveSwap(value);
-			}
-			this.ShowFunctions.Clear();
-
-			// Merge AutoExpandCategories and AutoCollapseCategories (we still want to keep AutoExpandCategories though!)
-			List<string> parentAutoExpandCategories = GetStringListMetaData(this.SuperClass, UhtNames.AutoExpandCategories);
-			List<string> parentAutoCollapseCategories = GetStringListMetaData(this.SuperClass, UhtNames.AutoCollapseCategories);
-
-			foreach (string value in this.AutoExpandCategories)
-			{
-				this.AutoCollapseCategories.RemoveSwap(value);
-				parentAutoCollapseCategories.RemoveSwap(value);
-			}
-
-			// Do the same as above but the other way around
-			foreach (string value in this.AutoCollapseCategories)
-			{
-				this.AutoExpandCategories.RemoveSwap(value);
-				parentAutoExpandCategories.RemoveSwap(value);
-			}
-
-			// Once AutoExpandCategories and AutoCollapseCategories for THIS class have been parsed, add the parent inherited categories
-			this.AutoCollapseCategories.AddRange(parentAutoCollapseCategories);
-			this.AutoExpandCategories.AddRange(parentAutoExpandCategories);
-		}
-
-		private void MergeShowCategories()
-		{
-
-			// Add the super class hide categories and prime the output show categories with the parent
-			List<string> outShowCategories = GetStringListMetaData(this.SuperClass, UhtNames.ShowCategories);
-			AppendStringListMetaData(this.SuperClass, UhtNames.HideCategories, this.HideCategories);
-
-			// If this class has new show categories, then merge them
-			if (this.ShowCategories.Count != 0)
-			{
-				StringBuilder subCategoryPath = new();
-				foreach (string value in this.ShowCategories)
-				{
-
-					// if we didn't find this specific category path in the HideCategories metadata
-					if (!this.HideCategories.RemoveSwap(value))
-					{
-						string[] subCategories = value.ToString().Split('|', StringSplitOptions.RemoveEmptyEntries);
-
-						subCategoryPath.Clear();
-						// look to see if any of the parent paths are excluded in the HideCategories list
-						for (int categoryPathIndex = 0; categoryPathIndex < subCategories.Length - 1; ++categoryPathIndex)
-						{
-							subCategoryPath.Append(subCategories[categoryPathIndex]);
-							// if we're hiding a parent category, then we need to flag this sub category for show
-							if (HideCategories.Contains(subCategoryPath.ToString()))
-							{
-								outShowCategories.AddUnique(value);
-								break;
-							}
-							subCategoryPath.Append('|');
-						}
-					}
-				}
-			}
-
-			// Replace the show categories
-			this.ShowCategories.Clear();
-			this.ShowCategories.AddRange(outShowCategories);
-		}
-
-		private void SetAndValidateWithinClass(UhtResolvePhase resolvePhase)
-		{
-			// The class within must be a child of any super class within
-			UhtClass expectedClassWithin = this.SuperClass != null ? this.SuperClass.ClassWithin : this.Session.UObject;
-
-			// Process all of the class specifiers
-			if (!String.IsNullOrEmpty(this.ClassWithinIdentifier))
-			{
-				UhtClass? specifiedClassWithin = (UhtClass?)this.Session.FindType(null, UhtFindOptions.EngineName | UhtFindOptions.Class, this.ClassWithinIdentifier);
-				if (specifiedClassWithin == null)
-				{
-					this.LogError($"Within class '{this.ClassWithinIdentifier}' not found");
-				}
-				else
-				{
-
-					// Make sure the class within has been resolved to this phase.  We don't need to worry about the super since we know 
-					// that it has already been resolved.
-					if (specifiedClassWithin != this)
-					{
-						specifiedClassWithin.Resolve(resolvePhase);
-					}
-
-					if (specifiedClassWithin.IsChildOf(this.Session.UInterface))
-					{
-						this.LogError("Classes cannot be 'within' interfaces");
-					}
-					else if (!specifiedClassWithin.IsChildOf(expectedClassWithin))
-					{
-						this.LogError($"Cannot override within class with '{specifiedClassWithin.SourceName}' since it isn't a child of parent class's expected within '{expectedClassWithin.SourceName}'");
-					}
-					else
-					{
-						this.ClassWithin = specifiedClassWithin;
-					}
-				}
-			}
-
-			// If we don't have a class within set, then just use the expected within
-			else
-			{
-				this.ClassWithin = expectedClassWithin;
-			}
-		}
-
-		private void SetAndValidateConfigName()
-		{
-			// Since this flag is computed in this method, we have to re-propagate the flag from the super
-			// just in case they were defined in this source file.
-			if (this.SuperClass != null)
-			{
-				this.ClassFlags |= this.SuperClass.ClassFlags & EClassFlags.Config;
-			}
-
-			// Set the class config flag if any properties have config
-			if (!this.ClassFlags.HasAnyFlags(EClassFlags.Config))
-			{
-				foreach (UhtProperty property in this.Properties)
-				{
-					if (property.PropertyFlags.HasAnyFlags(EPropertyFlags.Config))
-					{
-						this.ClassFlags |= EClassFlags.Config;
-						break;
-					}
-				}
-			}
-
-			if (this.Config.Length > 0)
-			{
-				// if the user specified "inherit", we're just going to use the parent class's config filename
-				// this is not actually necessary but it can be useful for explicitly communicating config-ness
-				if (this.Config.Equals("inherit", StringComparison.OrdinalIgnoreCase))
-				{
-					if (this.SuperClass == null)
-					{
-						this.LogError($"Cannot inherit config filename for class '{this.SourceName}' when it has no super class");
-					}
-					else if (this.SuperClass.Config.Length == 0)
-					{
-						this.LogError($"Cannot inherit config filename for class '{this.SourceName}' when parent class '{this.SuperClass.SourceName}' has no config filename");
-					}
-					else
-					{
-						this.Config = this.SuperClass.Config;
-					}
-				}
-			}
-			else if (this.ClassFlags.HasAnyFlags(EClassFlags.Config) && this.SuperClass != null)
-			{
-				this.Config = this.SuperClass.Config;
-			}
-
-			if (this.ClassFlags.HasAnyFlags(EClassFlags.Config) && this.Config.Length == 0)
-			{
-				this.LogError("Classes with config / globalconfig member variables need to specify config file.");
-				this.Config = "Engine";
-			}
-		}
-
-		private static List<string> GetStringListMetaData(UhtType? type, string key)
-		{
-			List<string> outStrings = new();
-			AppendStringListMetaData(type, key, outStrings);
-			return outStrings;
-		}
-
-		private static void AppendStringListMetaData(UhtType? type, string key, List<string> stringList)
-		{
-			if (type != null)
-			{
-				string[]? values = type.MetaData.GetStringArray(key);
-				if (values != null)
-				{
-					foreach (string value in values)
-					{
-						//COMPATIBILITY-TODO - TEST - This preserves duplicates that are in old UHT
-						//StringList.AddUnique(Value);
-						stringList.Add(value);
-					}
-				}
-			}
-		}
 
 		#region Keywords
 		[UhtKeyword(Extends = UhtTableNames.Global)]
@@ -403,7 +55,7 @@ namespace EpicGames.UHT.Parsers
 
 		private static UhtParseResult ParseUClass(UhtParsingScope parentScope, ref UhtToken token)
 		{
-			UhtClassParser classObj = new(parentScope.ScopeType, token.InputLine);
+			UhtClass classObj = new(parentScope.ScopeType, token.InputLine);
 			{
 				using UhtParsingScope topScope = new(parentScope, classObj, parentScope.Session.GetKeywordTable(UhtTableNames.Class), UhtAccessSpecifier.Private);
 				const string ScopeName = "class";
diff --git a/Engine/Source/Programs/Shared/EpicGames.UHT/Parsers/UhtEnumParser.cs b/Engine/Source/Programs/Shared/EpicGames.UHT/Parsers/UhtEnumParser.cs
--- a/Engine/Source/Programs/Shared/EpicGames.UHT/Parsers/UhtEnumParser.cs
+++ b/Engine/Source/Programs/Shared/EpicGames.UHT/Parsers/UhtEnumParser.cs
@@ -13,17 +13,8 @@ namespace EpicGames.UHT.Parsers
 	/// UENUM parser
 	/// </summary>
 	[UnrealHeaderTool]
-	public class UhtEnumParser : UhtEnum
+	public static class UhtEnumParser
 	{
-		/// <summary>
-		/// Construct a new enum parser
-		/// </summary>
-		/// <param name="outer">Outer type</param>
-		/// <param name="lineNumber">Line number</param>
-		public UhtEnumParser(UhtType outer, int lineNumber) : base(outer, lineNumber)
-		{
-		}
-
 		#region Keywords
 		[UhtKeyword(Extends = UhtTableNames.Global)]
 		[SuppressMessage("CodeQuality", "IDE0051:Remove unused private members", Justification = "Attribute accessed method")]
@@ -36,7 +27,7 @@ namespace EpicGames.UHT.Parsers
 
 		private static UhtParseResult ParseUEnum(UhtParsingScope parentScope, UhtToken keywordToken)
 		{
-			UhtEnumParser enumObject = new(parentScope.ScopeType, keywordToken.InputLine);
+			UhtEnum enumObject = new(parentScope.ScopeType, keywordToken.InputLine);
 			{
 				using UhtParsingScope topScope = new(parentScope, enumObject, parentScope.Session.GetKeywordTable(UhtTableNames.Enum), UhtAccessSpecifier.Public);
 				const string ScopeName = "UENUM";
diff --git a/Engine/Source/Programs/Shared/EpicGames.UHT/Parsers/UhtFunctionParser.cs b/Engine/Source/Programs/Shared/EpicGames.UHT/Parsers/UhtFunctionParser.cs
--- a/Engine/Source/Programs/Shared/EpicGames.UHT/Parsers/UhtFunctionParser.cs
+++ b/Engine/Source/Programs/Shared/EpicGames.UHT/Parsers/UhtFunctionParser.cs
@@ -4,7 +4,6 @@ using System;
 using System.Collections.Generic;
 using System.Diagnostics.CodeAnalysis;
 using System.Linq;
-using System.Text;
 using EpicGames.Core;
 using EpicGames.UHT.Tables;
 using EpicGames.UHT.Tokenizer;
@@ -94,90 +93,8 @@ namespace EpicGames.UHT.Parsers
 	/// UFUNCTION parser
 	/// </summary>
 	[UnrealHeaderTool]
-	public class UhtFunctionParser : UhtFunction
+	public static class UhtFunctionParser
 	{
-		/// <summary>
-		/// Construct a new function parser
-		/// </summary>
-		/// <param name="outer">Outer object</param>
-		/// <param name="lineNumber">Line number</param>
-		public UhtFunctionParser(UhtType outer, int lineNumber) : base(outer, lineNumber)
-		{
-		}
-
-		/// <summary>
-		/// True if the function specifier has a getter/setter specified
-		/// </summary>
-		public bool SawPropertyAccessor { get; set; } = false;
-
-		/// <inheritdoc/>
-		protected override bool ResolveSelf(UhtResolvePhase resolvePhase)
-		{
-			bool result = base.ResolveSelf(resolvePhase);
-			switch (resolvePhase)
-			{
-				case UhtResolvePhase.Properties:
-					UhtPropertyParser.ResolveChildren(this, GetPropertyParseOptions(false));
-					foreach (UhtProperty property in this.Properties)
-					{
-						if (property.DefaultValueTokens != null)
-						{
-							string key = "CPP_Default_" + property.EngineName;
-							if (!this.MetaData.ContainsKey(key))
-							{
-								bool parsed = false;
-								try
-								{
-									// All tokens MUST be consumed from the reader
-									StringBuilder builder = new();
-									IUhtTokenReader defaultValueReader = UhtTokenReplayReader.GetThreadInstance(property, this.HeaderFile.Data.Memory, property.DefaultValueTokens.ToArray(), UhtTokenType.EndOfDefault);
-									parsed = property.SanitizeDefaultValue(defaultValueReader, builder) && defaultValueReader.IsEOF;
-									if (parsed)
-									{
-										this.MetaData.Add(key, builder.ToString());
-									}
-								}
-								catch (Exception)
-								{
-									// Ignore the exception for now
-								}
-
-								if (!parsed)
-								{
-									StringView defaultValueText = new(this.HeaderFile.Data, property.DefaultValueTokens.First().InputStartPos,
-										property.DefaultValueTokens.Last().InputEndPos - property.DefaultValueTokens.First().InputStartPos);
-									property.LogError($"C++ Default parameter not parsed: {property.SourceName} '{defaultValueText}'");
-								}
-							}
-						}
-					}
-					break;
-			}
-			return result;
-		}
-
-		private UhtPropertyParseOptions GetPropertyParseOptions(bool returnValue)
-		{
-			switch (this.FunctionType)
-			{
-				case UhtFunctionType.Delegate:
-				case UhtFunctionType.SparseDelegate:
-					return (returnValue ? UhtPropertyParseOptions.None : UhtPropertyParseOptions.CommaSeparatedName) | UhtPropertyParseOptions.DontAddReturn;
-
-				case UhtFunctionType.Function:
-					UhtPropertyParseOptions options = UhtPropertyParseOptions.DontAddReturn; // Fetch the function name
-					options |= returnValue ? UhtPropertyParseOptions.FunctionNameIncluded : UhtPropertyParseOptions.NameIncluded;
-					if (this.FunctionFlags.HasAllFlags(EFunctionFlags.BlueprintEvent | EFunctionFlags.Native))
-					{
-						options |= UhtPropertyParseOptions.NoAutoConst;
-					}
-					return options;
-
-				default:
-					throw new UhtIceException("Unknown enumeration value");
-			}
-		}
-
 		#region Keywords
 		[UhtKeyword(Extends = UhtTableNames.Global)]
 		[UhtKeyword(Extends = UhtTableNames.Class)]
@@ -212,7 +129,7 @@ namespace EpicGames.UHT.Parsers
 
 		private static UhtParseResult ParseUDelegate(UhtParsingScope parentScope, UhtToken token, bool hasSpecifiers)
 		{
-			UhtFunctionParser function = new(parentScope.ScopeType, token.InputLine);
+			UhtFunction function = new(parentScope.ScopeType, token.InputLine);
 
 			{
 				using UhtParsingScope topScope = new(parentScope, function, parentScope.Session.GetKeywordTable(UhtTableNames.Function), UhtAccessSpecifier.Public);
@@ -382,7 +299,7 @@ namespace EpicGames.UHT.Parsers
 
 		private static UhtParseResult ParseUFunction(UhtParsingScope parentScope, UhtToken token)
 		{
-			UhtFunctionParser function = new(parentScope.ScopeType, token.InputLine);
+			UhtFunction function = new(parentScope.ScopeType, token.InputLine);
 
 			{
 				using UhtParsingScope topScope = new(parentScope, function, parentScope.Session.GetKeywordTable(UhtTableNames.Function), UhtAccessSpecifier.Public);
diff --git a/Engine/Source/Programs/Shared/EpicGames.UHT/Parsers/UhtInterfaceClassParser.cs b/Engine/Source/Programs/Shared/EpicGames.UHT/Parsers/UhtInterfaceClassParser.cs
--- a/Engine/Source/Programs/Shared/EpicGames.UHT/Parsers/UhtInterfaceClassParser.cs
+++ b/Engine/Source/Programs/Shared/EpicGames.UHT/Parsers/UhtInterfaceClassParser.cs
@@ -14,68 +14,9 @@ namespace EpicGames.UHT.Parsers
 	/// Interface class parser
 	/// </summary>
 	[UnrealHeaderTool]
-	public class UhtInterfaceClassParser : UhtClassBaseParser
+	public static class UhtInterfaceClassParser
 	{
-		/// <summary>
-		/// Construct a new interface class parser
-		/// </summary>
-		/// <param name="outer">Outer object</param>
-		/// <param name="lineNumber">Line number</param>
-		public UhtInterfaceClassParser(UhtType outer, int lineNumber) : base(outer, lineNumber)
-		{
-		}
 
-		/// <inheritdoc/>
-		protected override void ResolveSuper(UhtResolvePhase resolvePhase)
-		{
-			base.ResolveSuper(resolvePhase);
-			switch (resolvePhase)
-			{
-				case UhtResolvePhase.Bases:
-					UhtClass? superClass = this.SuperClass;
-					if (superClass != null)
-					{
-						this.ClassFlags |= superClass.ClassFlags & EClassFlags.ScriptInherit;
-						if (!superClass.ClassFlags.HasAnyFlags(EClassFlags.Native))
-						{
-							throw new UhtException(this, $"Native classes cannot extend non-native classes");
-						}
-						this.ClassWithin = superClass.ClassWithin;
-					}
-					else
-					{
-						this.ClassWithin = this.Session.UObject;
-					}
-					break;
-			}
-		}
-
-		/// <inheritdoc/>
-		protected override bool ResolveSelf(UhtResolvePhase resolvePhase)
-		{
-			bool result = base.ResolveSelf(resolvePhase);
-
-			switch (resolvePhase)
-			{
-				case UhtResolvePhase.InvalidCheck:
-					{
-						string nativeInterfaceName = "I" + this.EngineName;
-						UhtType? nativeInterface = this.Session.FindType(null, UhtFindOptions.SourceName | UhtFindOptions.Class, nativeInterfaceName);
-						if (nativeInterface == null)
-						{
-							this.LogError($"UInterface '{this.SourceName}' parsed without a corresponding '{nativeInterfaceName}'");
-						}
-						else
-						{
-							// Copy the children
-							this.AddChildren(nativeInterface.DetachChildren());
-						}
-					}
-					break;
-			}
-
-			return result;
-		}
 		#region Keywords
 		[UhtKeyword(Extends = UhtTableNames.Global)]
 		[SuppressMessage("CodeQuality", "IDE0051:Remove unused private members", Justification = "Attribute accessed method")]
@@ -108,7 +49,7 @@ namespace EpicGames.UHT.Parsers
 
 		private static UhtParseResult ParseUInterface(UhtParsingScope parentScope, ref UhtToken token)
 		{
-			UhtInterfaceClassParser classObj = new(parentScope.ScopeType, token.InputLine);
+			UhtClass classObj = new(parentScope.ScopeType, token.InputLine);
 			classObj.ClassType = UhtClassType.Interface;
 
 			{
diff --git a/Engine/Source/Programs/Shared/EpicGames.UHT/Parsers/UhtNativeInterfaceClassParser.cs b/Engine/Source/Programs/Shared/EpicGames.UHT/Parsers/UhtNativeInterfaceClassParser.cs
--- a/Engine/Source/Programs/Shared/EpicGames.UHT/Parsers/UhtNativeInterfaceClassParser.cs
+++ b/Engine/Source/Programs/Shared/EpicGames.UHT/Parsers/UhtNativeInterfaceClassParser.cs
@@ -14,56 +14,8 @@ namespace EpicGames.UHT.Parsers
 	/// Parser object for native interfaces
 	/// </summary>
 	[UnrealHeaderTool]
-	public class UhtParserNativeInterfaceClass : UhtClassBaseParser
-	{
-		/// <summary>
-		/// Construct a new native interface
-		/// </summary>
-		/// <param name="outer">Outer object</param>
-		/// <param name="lineNumber">Line number of declaration</param>
-		public UhtParserNativeInterfaceClass(UhtType outer, int lineNumber) : base(outer, lineNumber)
-		{
-		}
-
-		/// <inheritdoc/>
-		protected override bool ResolveSelf(UhtResolvePhase resolvePhase)
-		{
-			bool result = base.ResolveSelf(resolvePhase);
-
-			switch (resolvePhase)
-			{
-				case UhtResolvePhase.InvalidCheck:
-					string interfaceName = "U" + this.EngineName;
-					UhtClass? interfaceObj = (UhtClass?)this.Session.FindType(null, UhtFindOptions.SourceName | UhtFindOptions.Class, interfaceName);
-					if (interfaceObj == null)
-					{
-						if (this.HasGeneratedBody || this.Children.Count != 0)
-						{
-							this.LogError($"Native interface '{this.SourceName}' parsed without a corresponding '{interfaceName}'");
-						}
-						else
-						{
-							this.VisibleType = false;
-							result = false;
-						}
-					}
-					else
-					{
-						this.AlternateObject = interfaceObj;
-						interfaceObj.NativeInterface = this;
-						//COMPATIBILITY-TODO - Use the native interface access specifier
-						interfaceObj.GeneratedBodyAccessSpecifier = this.GeneratedBodyAccessSpecifier;
-
-						if (this.GeneratedBodyLineNumber == -1)
-						{
-							this.LogError("Expected a GENERATED_BODY() at the start of the native interface");
-						}
-					}
-					break;
-			}
-
-			return result;
-		}
+	public static class UhtNativeInterfaceClassParser
+	{ 
 
 		#region Keywords
 		[UhtKeyword(Extends = UhtTableNames.Global, Keyword = "class", DisableUsageError = true)]
@@ -152,7 +104,7 @@ namespace EpicGames.UHT.Parsers
 		{
 			IUhtTokenReader tokenReader = parentScope.TokenReader;
 
-			UhtParserNativeInterfaceClass classObj = new(parentScope.ScopeType, token.InputLine);
+			UhtClass classObj = new(parentScope.ScopeType, token.InputLine);
 			classObj.ClassType = UhtClassType.NativeInterface;
 			classObj.SourceName = sourceName.Value.ToString();
 			classObj.ClassFlags |= EClassFlags.Native | EClassFlags.Interface;
diff --git a/Engine/Source/Programs/Shared/EpicGames.UHT/Parsers/UhtScriptStructParser.cs b/Engine/Source/Programs/Shared/EpicGames.UHT/Parsers/UhtScriptStructParser.cs
--- a/Engine/Source/Programs/Shared/EpicGames.UHT/Parsers/UhtScriptStructParser.cs
+++ b/Engine/Source/Programs/Shared/EpicGames.UHT/Parsers/UhtScriptStructParser.cs
@@ -2,7 +2,6 @@
 
 using System.Collections.Generic;
 using System.Diagnostics.CodeAnalysis;
-using System.Text.Json.Serialization;
 using EpicGames.Core;
 using EpicGames.UHT.Tables;
 using EpicGames.UHT.Tokenizer;
@@ -16,63 +15,8 @@ namespace EpicGames.UHT.Parsers
 	/// USTRUCT parser object
 	/// </summary>
 	[UnrealHeaderTool]
-	public class UhtScriptStructParser : UhtScriptStruct
+	public static class UhtScriptStructParser
 	{
-		/// <summary>
-		/// Super identifier
-		/// </summary>
-		[JsonIgnore]
-		public UhtToken SuperIdentifier { get; set; }
-
-		/// <summary>
-		/// Base identifiers
-		/// </summary>
-		[JsonIgnore]
-		[SuppressMessage("Usage", "CA2227:Collection properties should be read only", Justification = "<Pending>")]
-		public List<UhtToken[]>? BaseIdentifiers { get; set; } = null;
-
-		/// <summary>
-		/// Construct a new instance
-		/// </summary>
-		/// <param name="outer">Output type</param>
-		/// <param name="lineNumber">Line number of declaration</param>
-		public UhtScriptStructParser(UhtType outer, int lineNumber) : base(outer, lineNumber)
-		{
-		}
-
-		/// <inheritdoc/>
-		protected override void ResolveSuper(UhtResolvePhase resolvePhase)
-		{
-			base.ResolveSuper(resolvePhase);
-
-			switch (resolvePhase)
-			{
-				case UhtResolvePhase.Bases:
-					BindAndResolveSuper(this.SuperIdentifier, UhtFindOptions.ScriptStruct);
-
-					// if we have a base struct, propagate inherited struct flags now
-					UhtScriptStruct? superScriptStruct = this.SuperScriptStruct;
-					if (superScriptStruct != null)
-					{
-						this.ScriptStructFlags |= superScriptStruct.ScriptStructFlags & EStructFlags.Inherit;
-					}
-					break;
-			}
-		}
-
-		/// <inheritdoc/>
-		protected override bool ResolveSelf(UhtResolvePhase resolvePhase)
-		{
-			bool result = base.ResolveSelf(resolvePhase);
-
-			switch (resolvePhase)
-			{
-				case UhtResolvePhase.Properties:
-					UhtPropertyParser.ResolveChildren(this, UhtPropertyParseOptions.AddModuleRelativePath);
-					break;
-			}
-			return result;
-		}
 
 		#region Keywords
 		[UhtKeyword(Extends = UhtTableNames.Global)]
@@ -121,7 +65,7 @@ namespace EpicGames.UHT.Parsers
 
 		private static UhtParseResult ParseUScriptStruct(UhtParsingScope parentScope, UhtToken keywordToken)
 		{
-			UhtScriptStructParser scriptStruct = new(parentScope.ScopeType, keywordToken.InputLine);
+			UhtScriptStruct scriptStruct = new(parentScope.ScopeType, keywordToken.InputLine);
 			{
 				using UhtParsingScope topScope = new(parentScope, scriptStruct, parentScope.Session.GetKeywordTable(UhtTableNames.ScriptStruct), UhtAccessSpecifier.Public);
 
diff --git a/Engine/Source/Programs/Shared/EpicGames.UHT/Specifiers/UhtClassSpecifiers.cs b/Engine/Source/Programs/Shared/EpicGames.UHT/Specifiers/UhtClassSpecifiers.cs
--- a/Engine/Source/Programs/Shared/EpicGames.UHT/Specifiers/UhtClassSpecifiers.cs
+++ b/Engine/Source/Programs/Shared/EpicGames.UHT/Specifiers/UhtClassSpecifiers.cs
@@ -13,178 +13,157 @@ namespace EpicGames.UHT.Parsers
 	/// Collection of UCLASS specifiers
 	/// </summary>
 	[UnrealHeaderTool]
+	[SuppressMessage("CodeQuality", "IDE0051:Remove unused private members", Justification = "Attribute accessed method")]
 	public static class UhtClassSpecifiers
 	{
 		[UhtSpecifier(Extends = UhtTableNames.Class, ValueType = UhtSpecifierValueType.Legacy)]
-		[SuppressMessage("CodeQuality", "IDE0051:Remove unused private members", Justification = "Attribute accessed method")]
 		private static void NoExportSpecifier(UhtSpecifierContext specifierContext)
 		{
-			UhtClassParser classObj = (UhtClassParser)specifierContext.Scope.ScopeType;
+			UhtClass classObj = (UhtClass)specifierContext.Scope.ScopeType;
 			classObj.ClassExportFlags |= UhtClassExportFlags.NoExport;
 		}
 
 		[UhtSpecifier(Extends = UhtTableNames.Class, ValueType = UhtSpecifierValueType.Legacy)]
-		[SuppressMessage("CodeQuality", "IDE0051:Remove unused private members", Justification = "Attribute accessed method")]
 		private static void IntrinsicSpecifier(UhtSpecifierContext specifierContext)
 		{
-			UhtClassParser classObj = (UhtClassParser)specifierContext.Scope.ScopeType;
+			UhtClass classObj = (UhtClass)specifierContext.Scope.ScopeType;
 			classObj.AddClassFlags(EClassFlags.Intrinsic);
 		}
 
 		[UhtSpecifier(Extends = UhtTableNames.Class, ValueType = UhtSpecifierValueType.Legacy)]
-		[SuppressMessage("CodeQuality", "IDE0051:Remove unused private members", Justification = "Attribute accessed method")]
 		private static void ComponentWrapperClassSpecifier(UhtSpecifierContext specifierContext)
 		{
 			specifierContext.MetaData.Add(UhtNames.IgnoreCategoryKeywordsInSubclasses, true);
 		}
 
 		[UhtSpecifier(Extends = UhtTableNames.Class, ValueType = UhtSpecifierValueType.SingleString)]
-		[SuppressMessage("CodeQuality", "IDE0051:Remove unused private members", Justification = "Attribute accessed method")]
 		private static void WithinSpecifier(UhtSpecifierContext specifierContext, StringView value)
 		{
-			UhtClassParser classObj = (UhtClassParser)specifierContext.Scope.ScopeType;
+			UhtClass classObj = (UhtClass)specifierContext.Scope.ScopeType;
 			classObj.ClassWithinIdentifier = value.ToString();
 		}
 
 		[UhtSpecifier(Extends = UhtTableNames.Class, ValueType = UhtSpecifierValueType.Legacy)]
-		[SuppressMessage("CodeQuality", "IDE0051:Remove unused private members", Justification = "Attribute accessed method")]
 		private static void EditInlineNewSpecifier(UhtSpecifierContext specifierContext)
 		{
-			UhtClassParser classObj = (UhtClassParser)specifierContext.Scope.ScopeType;
+			UhtClass classObj = (UhtClass)specifierContext.Scope.ScopeType;
 			classObj.AddClassFlags(EClassFlags.EditInlineNew);
 		}
 
 		[UhtSpecifier(Extends = UhtTableNames.Class, ValueType = UhtSpecifierValueType.Legacy)]
-		[SuppressMessage("CodeQuality", "IDE0051:Remove unused private members", Justification = "Attribute accessed method")]
 		private static void NotEditInlineNewSpecifier(UhtSpecifierContext specifierContext)
 		{
-			UhtClassParser classObj = (UhtClassParser)specifierContext.Scope.ScopeType;
+			UhtClass classObj = (UhtClass)specifierContext.Scope.ScopeType;
 			classObj.RemoveClassFlags(EClassFlags.EditInlineNew);
 		}
 
 		[UhtSpecifier(Extends = UhtTableNames.Class, ValueType = UhtSpecifierValueType.Legacy)]
-		[SuppressMessage("CodeQuality", "IDE0051:Remove unused private members", Justification = "Attribute accessed method")]
 		private static void PlaceableSpecifier(UhtSpecifierContext specifierContext)
 		{
-			UhtClassParser classObj = (UhtClassParser)specifierContext.Scope.ScopeType;
+			UhtClass classObj = (UhtClass)specifierContext.Scope.ScopeType;
 			classObj.RemoveClassFlags(EClassFlags.NotPlaceable);
 		}
 
 		[UhtSpecifier(Extends = UhtTableNames.Class, ValueType = UhtSpecifierValueType.Legacy)]
-		[SuppressMessage("CodeQuality", "IDE0051:Remove unused private members", Justification = "Attribute accessed method")]
 		private static void NotPlaceableSpecifier(UhtSpecifierContext specifierContext)
 		{
-			UhtClassParser classObj = (UhtClassParser)specifierContext.Scope.ScopeType;
+			UhtClass classObj = (UhtClass)specifierContext.Scope.ScopeType;
 			classObj.AddClassFlags(EClassFlags.NotPlaceable);
 		}
 
 		[UhtSpecifier(Extends = UhtTableNames.Class, ValueType = UhtSpecifierValueType.Legacy)]
-		[SuppressMessage("CodeQuality", "IDE0051:Remove unused private members", Justification = "Attribute accessed method")]
 		private static void DefaultToInstancedSpecifier(UhtSpecifierContext specifierContext)
 		{
-			UhtClassParser classObj = (UhtClassParser)specifierContext.Scope.ScopeType;
+			UhtClass classObj = (UhtClass)specifierContext.Scope.ScopeType;
 			classObj.AddClassFlags(EClassFlags.DefaultToInstanced);
 		}
 
 		[UhtSpecifier(Extends = UhtTableNames.Class, ValueType = UhtSpecifierValueType.Legacy)]
-		[SuppressMessage("CodeQuality", "IDE0051:Remove unused private members", Justification = "Attribute accessed method")]
 		private static void HideDropdownSpecifier(UhtSpecifierContext specifierContext)
 		{
-			UhtClassParser classObj = (UhtClassParser)specifierContext.Scope.ScopeType;
+			UhtClass classObj = (UhtClass)specifierContext.Scope.ScopeType;
 			classObj.AddClassFlags(EClassFlags.HideDropDown);
 		}
 
 		[UhtSpecifier(Extends = UhtTableNames.Class, ValueType = UhtSpecifierValueType.Legacy)]
-		[SuppressMessage("CodeQuality", "IDE0051:Remove unused private members", Justification = "Attribute accessed method")]
 		private static void HiddenSpecifier(UhtSpecifierContext specifierContext)
 		{
 			// Prevents class from appearing in the editor class browser and edit inline menus.
-			UhtClassParser classObj = (UhtClassParser)specifierContext.Scope.ScopeType;
+			UhtClass classObj = (UhtClass)specifierContext.Scope.ScopeType;
 			classObj.AddClassFlags(EClassFlags.Hidden);
 		}
 
 		[UhtSpecifier(Extends = UhtTableNames.Class, ValueType = UhtSpecifierValueType.Legacy)]
-		[SuppressMessage("CodeQuality", "IDE0051:Remove unused private members", Justification = "Attribute accessed method")]
 		private static void DependsOnSpecifier(UhtSpecifierContext specifierContext)
 		{
 			specifierContext.MessageSite.LogError("The dependsOn specifier is deprecated. Please use #include \"ClassHeaderFilename.h\" instead.");
 		}
 
 		[UhtSpecifier(Extends = UhtTableNames.Class, ValueType = UhtSpecifierValueType.Legacy)]
-		[SuppressMessage("CodeQuality", "IDE0051:Remove unused private members", Justification = "Attribute accessed method")]
 		private static void MinimalAPISpecifier(UhtSpecifierContext specifierContext)
 		{
-			UhtClassParser classObj = (UhtClassParser)specifierContext.Scope.ScopeType;
+			UhtClass classObj = (UhtClass)specifierContext.Scope.ScopeType;
 			classObj.AddClassFlags(EClassFlags.MinimalAPI);
 		}
 
 		[UhtSpecifier(Extends = UhtTableNames.Class, ValueType = UhtSpecifierValueType.Legacy)]
-		[SuppressMessage("CodeQuality", "IDE0051:Remove unused private members", Justification = "Attribute accessed method")]
 		private static void ConstSpecifier(UhtSpecifierContext specifierContext)
 		{
-			UhtClassParser classObj = (UhtClassParser)specifierContext.Scope.ScopeType;
+			UhtClass classObj = (UhtClass)specifierContext.Scope.ScopeType;
 			classObj.AddClassFlags(EClassFlags.Const);
 		}
 
 		[UhtSpecifier(Extends = UhtTableNames.Class, ValueType = UhtSpecifierValueType.Legacy)]
-		[SuppressMessage("CodeQuality", "IDE0051:Remove unused private members", Justification = "Attribute accessed method")]
 		private static void PerObjectConfigSpecifier(UhtSpecifierContext specifierContext)
 		{
-			UhtClassParser classObj = (UhtClassParser)specifierContext.Scope.ScopeType;
+			UhtClass classObj = (UhtClass)specifierContext.Scope.ScopeType;
 			classObj.AddClassFlags(EClassFlags.PerObjectConfig);
 		}
 
 		[UhtSpecifier(Extends = UhtTableNames.Class, ValueType = UhtSpecifierValueType.Legacy)]
-		[SuppressMessage("CodeQuality", "IDE0051:Remove unused private members", Justification = "Attribute accessed method")]
 		private static void ConfigDoNotCheckDefaultsSpecifier(UhtSpecifierContext specifierContext)
 		{
-			UhtClassParser classObj = (UhtClassParser)specifierContext.Scope.ScopeType;
+			UhtClass classObj = (UhtClass)specifierContext.Scope.ScopeType;
 			classObj.AddClassFlags(EClassFlags.ConfigDoNotCheckDefaults);
 		}
 
 		[UhtSpecifier(Extends = UhtTableNames.Class, ValueType = UhtSpecifierValueType.Legacy)]
-		[SuppressMessage("CodeQuality", "IDE0051:Remove unused private members", Justification = "Attribute accessed method")]
 		private static void AbstractSpecifier(UhtSpecifierContext specifierContext)
 		{
-			UhtClassParser classObj = (UhtClassParser)specifierContext.Scope.ScopeType;
+			UhtClass classObj = (UhtClass)specifierContext.Scope.ScopeType;
 			classObj.AddClassFlags(EClassFlags.Abstract);
 		}
 
 		[UhtSpecifier(Extends = UhtTableNames.Class, ValueType = UhtSpecifierValueType.Legacy)]
-		[SuppressMessage("CodeQuality", "IDE0051:Remove unused private members", Justification = "Attribute accessed method")]
 		private static void DeprecatedSpecifier(UhtSpecifierContext specifierContext)
 		{
-			UhtClassParser classObj = (UhtClassParser)specifierContext.Scope.ScopeType;
+			UhtClass classObj = (UhtClass)specifierContext.Scope.ScopeType;
 			classObj.AddClassFlags(EClassFlags.Deprecated | EClassFlags.NotPlaceable);
 		}
 
 		[UhtSpecifier(Extends = UhtTableNames.Class, ValueType = UhtSpecifierValueType.Legacy)]
-		[SuppressMessage("CodeQuality", "IDE0051:Remove unused private members", Justification = "Attribute accessed method")]
 		private static void TransientSpecifier(UhtSpecifierContext specifierContext)
 		{
-			UhtClassParser classObj = (UhtClassParser)specifierContext.Scope.ScopeType;
+			UhtClass classObj = (UhtClass)specifierContext.Scope.ScopeType;
 			classObj.AddClassFlags(EClassFlags.Transient);
 		}
 
 		[UhtSpecifier(Extends = UhtTableNames.Class, ValueType = UhtSpecifierValueType.Legacy)]
-		[SuppressMessage("CodeQuality", "IDE0051:Remove unused private members", Justification = "Attribute accessed method")]
 		private static void NonTransientSpecifier(UhtSpecifierContext specifierContext)
 		{
-			UhtClassParser classObj = (UhtClassParser)specifierContext.Scope.ScopeType;
+			UhtClass classObj = (UhtClass)specifierContext.Scope.ScopeType;
 			classObj.RemoveClassFlags(EClassFlags.Transient);
 		}
 
 		[UhtSpecifier(Extends = UhtTableNames.Class, ValueType = UhtSpecifierValueType.Legacy)]
-		[SuppressMessage("CodeQuality", "IDE0051:Remove unused private members", Justification = "Attribute accessed method")]
 		private static void OptionalSpecifier(UhtSpecifierContext specifierContext)
 		{
 			// Optional class
-			UhtClassParser classObj = (UhtClassParser)specifierContext.Scope.ScopeType;
+			UhtClass classObj = (UhtClass)specifierContext.Scope.ScopeType;
 			classObj.AddClassFlags(EClassFlags.Optional);
 		}
 
 		[UhtSpecifier(Extends = UhtTableNames.Class, ValueType = UhtSpecifierValueType.Legacy)]
-		[SuppressMessage("CodeQuality", "IDE0051:Remove unused private members", Justification = "Attribute accessed method")]
 		private static void CustomConstructorSpecifier(UhtSpecifierContext specifierContext)
 		{
 			// we will not export a constructor for this class, assuming it is in the CPP block
@@ -193,7 +172,6 @@ namespace EpicGames.UHT.Parsers
 		}
 
 		[UhtSpecifier(Extends = UhtTableNames.Class, ValueType = UhtSpecifierValueType.SingleString)]
-		[SuppressMessage("CodeQuality", "IDE0051:Remove unused private members", Justification = "Attribute accessed method")]
 		private static void ConfigSpecifier(UhtSpecifierContext specifierContext, StringView value)
 		{
 			UhtClass classObj = (UhtClass)specifierContext.Scope.ScopeType;
@@ -201,84 +179,74 @@ namespace EpicGames.UHT.Parsers
 		}
 
 		[UhtSpecifier(Extends = UhtTableNames.Class, ValueType = UhtSpecifierValueType.Legacy)]
-		[SuppressMessage("CodeQuality", "IDE0051:Remove unused private members", Justification = "Attribute accessed method")]
 		private static void DefaultConfigSpecifier(UhtSpecifierContext specifierContext)
 		{
 			// Save object config only to Default INIs, never to local INIs.
-			UhtClassParser classObj = (UhtClassParser)specifierContext.Scope.ScopeType;
+			UhtClass classObj = (UhtClass)specifierContext.Scope.ScopeType;
 			classObj.AddClassFlags(EClassFlags.DefaultConfig);
 		}
 
 		[UhtSpecifier(Extends = UhtTableNames.Class, ValueType = UhtSpecifierValueType.Legacy)]
-		[SuppressMessage("CodeQuality", "IDE0051:Remove unused private members", Justification = "Attribute accessed method")]
 		private static void GlobalUserConfigSpecifier(UhtSpecifierContext specifierContext)
 		{
 			// Save object config only to global user overrides, never to local INIs
-			UhtClassParser classObj = (UhtClassParser)specifierContext.Scope.ScopeType;
+			UhtClass classObj = (UhtClass)specifierContext.Scope.ScopeType;
 			classObj.AddClassFlags(EClassFlags.GlobalUserConfig);
 		}
 
 		[UhtSpecifier(Extends = UhtTableNames.Class, ValueType = UhtSpecifierValueType.Legacy)]
-		[SuppressMessage("CodeQuality", "IDE0051:Remove unused private members", Justification = "Attribute accessed method")]
 		private static void ProjectUserConfigSpecifier(UhtSpecifierContext specifierContext)
 		{
 			// Save object config only to project user overrides, never to INIs that are checked in
-			UhtClassParser classObj = (UhtClassParser)specifierContext.Scope.ScopeType;
+			UhtClass classObj = (UhtClass)specifierContext.Scope.ScopeType;
 			classObj.AddClassFlags(EClassFlags.ProjectUserConfig);
 		}
 
 		[UhtSpecifier(Extends = UhtTableNames.Class, ValueType = UhtSpecifierValueType.SingleString)]
-		[SuppressMessage("CodeQuality", "IDE0051:Remove unused private members", Justification = "Attribute accessed method")]
 		private static void EditorConfigSpecifier(UhtSpecifierContext specifierContext, StringView value)
 		{
 			specifierContext.MetaData.Add(UhtNames.EditorConfig, value.ToString());
 		}
 
 		[UhtSpecifier(Extends = UhtTableNames.Class, ValueType = UhtSpecifierValueType.NonEmptyStringList)]
-		[SuppressMessage("CodeQuality", "IDE0051:Remove unused private members", Justification = "Attribute accessed method")]
 		private static void ShowCategoriesSpecifier(UhtSpecifierContext specifierContext, List<StringView> value)
 		{
-			UhtClassParser classObj = (UhtClassParser)specifierContext.Scope.ScopeType;
+			UhtClass classObj = (UhtClass)specifierContext.Scope.ScopeType;
 			classObj.ShowCategories.AddUniqueRange(value);
 		}
 
 		[UhtSpecifier(Extends = UhtTableNames.Class, ValueType = UhtSpecifierValueType.NonEmptyStringList)]
-		[SuppressMessage("CodeQuality", "IDE0051:Remove unused private members", Justification = "Attribute accessed method")]
 		private static void HideCategoriesSpecifier(UhtSpecifierContext specifierContext, List<StringView> value)
 		{
-			UhtClassParser classObj = (UhtClassParser)specifierContext.Scope.ScopeType;
+			UhtClass classObj = (UhtClass)specifierContext.Scope.ScopeType;
 			classObj.HideCategories.AddUniqueRange(value);
 		}
 
 		[UhtSpecifier(Extends = UhtTableNames.Class, ValueType = UhtSpecifierValueType.NonEmptyStringList)]
-		[SuppressMessage("CodeQuality", "IDE0051:Remove unused private members", Justification = "Attribute accessed method")]
 		private static void ShowFunctionsSpecifier(UhtSpecifierContext specifierContext, List<StringView> value)
 		{
-			UhtClassParser classObj = (UhtClassParser)specifierContext.Scope.ScopeType;
+			UhtClass classObj = (UhtClass)specifierContext.Scope.ScopeType;
 			classObj.ShowFunctions.AddUniqueRange(value);
 		}
 
 		[UhtSpecifier(Extends = UhtTableNames.Class, ValueType = UhtSpecifierValueType.NonEmptyStringList)]
-		[SuppressMessage("CodeQuality", "IDE0051:Remove unused private members", Justification = "Attribute accessed method")]
 		private static void HideFunctionsSpecifier(UhtSpecifierContext specifierContext, List<StringView> value)
 		{
-			UhtClassParser classObj = (UhtClassParser)specifierContext.Scope.ScopeType;
+			UhtClass classObj = (UhtClass)specifierContext.Scope.ScopeType;
 			classObj.HideFunctions.AddUniqueRange(value);
 		}
 
 		[UhtSpecifier(Extends = UhtTableNames.Class, ValueType = UhtSpecifierValueType.SingleString)]
-		[SuppressMessage("CodeQuality", "IDE0051:Remove unused private members", Justification = "Attribute accessed method")]
 		private static void SparseClassDataTypesSpecifier(UhtSpecifierContext specifierContext, StringView value)
 		{
-			UhtClassParser classObj = (UhtClassParser)specifierContext.Scope.ScopeType;
+			UhtClass classObj = (UhtClass)specifierContext.Scope.ScopeType;
 			classObj.SparseClassDataTypes.AddUnique(value.ToString());
 		}
 
 		[UhtSpecifier(Extends = UhtTableNames.Class, ValueType = UhtSpecifierValueType.NonEmptyStringList)]
-		[SuppressMessage("CodeQuality", "IDE0051:Remove unused private members", Justification = "Attribute accessed method")]
 		private static void ClassGroupSpecifier(UhtSpecifierContext specifierContext, List<StringView> value)
 		{
-			UhtClassParser classObj = (UhtClassParser)specifierContext.Scope.ScopeType;
+			UhtClass classObj = (UhtClass)specifierContext.Scope.ScopeType;
 			foreach (StringView element in value)
 			{
 				classObj.ClassGroupNames.Add(element.ToString());
@@ -286,59 +254,52 @@ namespace EpicGames.UHT.Parsers
 		}
 
 		[UhtSpecifier(Extends = UhtTableNames.Class, ValueType = UhtSpecifierValueType.NonEmptyStringList)]
-		[SuppressMessage("CodeQuality", "IDE0051:Remove unused private members", Justification = "Attribute accessed method")]
 		private static void AutoExpandCategoriesSpecifier(UhtSpecifierContext specifierContext, List<StringView> value)
 		{
-			UhtClassParser classObj = (UhtClassParser)specifierContext.Scope.ScopeType;
+			UhtClass classObj = (UhtClass)specifierContext.Scope.ScopeType;
 			classObj.AutoExpandCategories.AddUniqueRange(value);
 			classObj.AutoCollapseCategories.RemoveSwapRange(value);
 		}
 
 		[UhtSpecifier(Extends = UhtTableNames.Class, ValueType = UhtSpecifierValueType.NonEmptyStringList)]
-		[SuppressMessage("CodeQuality", "IDE0051:Remove unused private members", Justification = "Attribute accessed method")]
 		private static void AutoCollapseCategoriesSpecifier(UhtSpecifierContext specifierContext, List<StringView> value)
 		{
-			UhtClassParser classObj = (UhtClassParser)specifierContext.Scope.ScopeType;
+			UhtClass classObj = (UhtClass)specifierContext.Scope.ScopeType;
 			classObj.AutoCollapseCategories.AddUniqueRange(value);
 			classObj.AutoExpandCategories.RemoveSwapRange(value);
 		}
 
 		[UhtSpecifier(Extends = UhtTableNames.Class, ValueType = UhtSpecifierValueType.NonEmptyStringList)]
-		[SuppressMessage("CodeQuality", "IDE0051:Remove unused private members", Justification = "Attribute accessed method")]
 		private static void PrioritizeCategoriesSpecifier(UhtSpecifierContext specifierContext, List<StringView> value)
 		{
-			UhtClassParser classObj = (UhtClassParser)specifierContext.Scope.ScopeType;
+			UhtClass classObj = (UhtClass)specifierContext.Scope.ScopeType;
 			classObj.PrioritizeCategories.AddUniqueRange(value);
 		}
 
 		[UhtSpecifier(Extends = UhtTableNames.Class, ValueType = UhtSpecifierValueType.NonEmptyStringList)]
-		[SuppressMessage("CodeQuality", "IDE0051:Remove unused private members", Justification = "Attribute accessed method")]
 		private static void DontAutoCollapseCategoriesSpecifier(UhtSpecifierContext specifierContext, List<StringView> value)
 		{
-			UhtClassParser classObj = (UhtClassParser)specifierContext.Scope.ScopeType;
+			UhtClass classObj = (UhtClass)specifierContext.Scope.ScopeType;
 			classObj.AutoCollapseCategories.RemoveSwapRange(value);
 		}
 
 		[UhtSpecifier(Extends = UhtTableNames.Class, ValueType = UhtSpecifierValueType.Legacy)]
-		[SuppressMessage("CodeQuality", "IDE0051:Remove unused private members", Justification = "Attribute accessed method")]
 		private static void CollapseCategoriesSpecifier(UhtSpecifierContext specifierContext)
 		{
 			// Class' properties should not be shown categorized in the editor.
-			UhtClassParser classObj = (UhtClassParser)specifierContext.Scope.ScopeType;
+			UhtClass classObj = (UhtClass)specifierContext.Scope.ScopeType;
 			classObj.AddClassFlags(EClassFlags.CollapseCategories);
 		}
 
 		[UhtSpecifier(Extends = UhtTableNames.Class, ValueType = UhtSpecifierValueType.Legacy)]
-		[SuppressMessage("CodeQuality", "IDE0051:Remove unused private members", Justification = "Attribute accessed method")]
 		private static void DontCollapseCategoriesSpecifier(UhtSpecifierContext specifierContext)
 		{
 			// Class' properties should be shown categorized in the editor.
-			UhtClassParser classObj = (UhtClassParser)specifierContext.Scope.ScopeType;
+			UhtClass classObj = (UhtClass)specifierContext.Scope.ScopeType;
 			classObj.RemoveClassFlags(EClassFlags.CollapseCategories);
 		}
 
 		[UhtSpecifier(Extends = UhtTableNames.Class, ValueType = UhtSpecifierValueType.Legacy)]
-		[SuppressMessage("CodeQuality", "IDE0051:Remove unused private members", Justification = "Attribute accessed method")]
 		private static void AdvancedClassDisplaySpecifier(UhtSpecifierContext specifierContext)
 		{
 			// By default the class properties are shown in advanced sections in UI
@@ -347,7 +308,6 @@ namespace EpicGames.UHT.Parsers
 		}
 
 		[UhtSpecifier(Extends = UhtTableNames.Class, ValueType = UhtSpecifierValueType.Legacy)]
-		[SuppressMessage("CodeQuality", "IDE0051:Remove unused private members", Justification = "Attribute accessed method")]
 		private static void ConversionRootSpecifier(UhtSpecifierContext specifierContext)
 		{
 			UhtClass classObj = (UhtClass)specifierContext.Scope.ScopeType;
@@ -355,42 +315,38 @@ namespace EpicGames.UHT.Parsers
 		}
 
 		[UhtSpecifier(Extends = UhtTableNames.Class, ValueType = UhtSpecifierValueType.Legacy)]
-		[SuppressMessage("CodeQuality", "IDE0051:Remove unused private members", Justification = "Attribute accessed method")]
 		private static void NeedsDeferredDependencyLoadingSpecifier(UhtSpecifierContext specifierContext)
 		{
-			UhtClassParser classObj = (UhtClassParser)specifierContext.Scope.ScopeType;
+			UhtClass classObj = (UhtClass)specifierContext.Scope.ScopeType;
 			classObj.AddClassFlags(EClassFlags.NeedsDeferredDependencyLoading);
 		}
 
 		[UhtSpecifier(Extends = UhtTableNames.Class, ValueType = UhtSpecifierValueType.Legacy)]
-		[SuppressMessage("CodeQuality", "IDE0051:Remove unused private members", Justification = "Attribute accessed method")]
 		private static void MatchedSerializersSpecifier(UhtSpecifierContext specifierContext)
 		{
 			if (!specifierContext.Scope.HeaderParser.HeaderFile.IsNoExportTypes)
 			{
 				specifierContext.MessageSite.LogError("The 'MatchedSerializers' class specifier is only valid in the NoExportTypes.h file");
 			}
-			UhtClassParser classObj = (UhtClassParser)specifierContext.Scope.ScopeType;
+			UhtClass classObj = (UhtClass)specifierContext.Scope.ScopeType;
 			classObj.AddClassFlags(EClassFlags.MatchedSerializers);
 		}
 
 		[UhtSpecifier(Extends = UhtTableNames.Class, ValueType = UhtSpecifierValueType.Legacy)]
-		[SuppressMessage("CodeQuality", "IDE0051:Remove unused private members", Justification = "Attribute accessed method")]
 		private static void InterfaceSpecifier(UhtSpecifierContext specifierContext)
 		{
 			if (!specifierContext.Scope.HeaderParser.HeaderFile.IsNoExportTypes)
 			{
 				specifierContext.MessageSite.LogError("The 'Interface' class specifier is only valid in the NoExportTypes.h file");
 			}
-			UhtClassParser classObj = (UhtClassParser)specifierContext.Scope.ScopeType;
+			UhtClass classObj = (UhtClass)specifierContext.Scope.ScopeType;
 			classObj.AddClassFlags(EClassFlags.Interface);
 		}
 
 		[UhtSpecifier(Extends = UhtTableNames.Class, ValueType = UhtSpecifierValueType.None)]
-		[SuppressMessage("CodeQuality", "IDE0051:Remove unused private members", Justification = "Attribute accessed method")]
 		private static void CustomFieldNotifySpecifier(UhtSpecifierContext specifierContext)
 		{
-			UhtClassParser classObj = (UhtClassParser)specifierContext.Scope.ScopeType;
+			UhtClass classObj = (UhtClass)specifierContext.Scope.ScopeType;
 			classObj.ClassExportFlags |= UhtClassExportFlags.HasCustomFieldNotify;
 		}
 	}
diff --git a/Engine/Source/Programs/Shared/EpicGames.UHT/Specifiers/UhtDefaultSpecifiers.cs b/Engine/Source/Programs/Shared/EpicGames.UHT/Specifiers/UhtDefaultSpecifiers.cs
--- a/Engine/Source/Programs/Shared/EpicGames.UHT/Specifiers/UhtDefaultSpecifiers.cs
+++ b/Engine/Source/Programs/Shared/EpicGames.UHT/Specifiers/UhtDefaultSpecifiers.cs
@@ -15,6 +15,8 @@ namespace EpicGames.UHT.Parsers
 	/// Collection of default specifiers that apply to everything
 	/// </summary>
 	[UnrealHeaderTool]
+	[SuppressMessage("CodeQuality", "IDE0051:Remove unused private members", Justification = "Attribute accessed method")]
+	[SuppressMessage("Style", "IDE0060:Remove unused parameter", Justification = "Attribute accessed method")]
 	public static class UhtDefaultSpecifiers
 	{
 		private static void SetMetaData(UhtSpecifierContext specifierContext, StringView key, StringView value)
@@ -81,7 +83,6 @@ namespace EpicGames.UHT.Parsers
 
 		#region Specifiers
 		[UhtSpecifier(Extends = UhtTableNames.Default, ValueType = UhtSpecifierValueType.KeyValuePairList, When = UhtSpecifierWhen.Immediate)]
-		[SuppressMessage("CodeQuality", "IDE0051:Remove unused private members", Justification = "Attribute accessed method")]
 		private static void MetaSpecifier(UhtSpecifierContext specifierContext, List<KeyValuePair<StringView, StringView>> value)
 		{
 			foreach (KeyValuePair<StringView, StringView> kvp in (List<KeyValuePair<StringView, StringView>>)value)
@@ -91,21 +92,18 @@ namespace EpicGames.UHT.Parsers
 		}
 
 		[UhtSpecifier(Extends = UhtTableNames.Default, ValueType = UhtSpecifierValueType.String, When = UhtSpecifierWhen.Immediate)]
-		[SuppressMessage("CodeQuality", "IDE0051:Remove unused private members", Justification = "Attribute accessed method")]
 		private static void DisplayNameSpecifier(UhtSpecifierContext specifierContext, StringView value)
 		{
 			SetMetaData(specifierContext, "DisplayName", value);
 		}
 
 		[UhtSpecifier(Extends = UhtTableNames.Default, ValueType = UhtSpecifierValueType.String, When = UhtSpecifierWhen.Immediate)]
-		[SuppressMessage("CodeQuality", "IDE0051:Remove unused private members", Justification = "Attribute accessed method")]
 		private static void FriendlyNameSpecifier(UhtSpecifierContext specifierContext, StringView value)
 		{
 			SetMetaData(specifierContext, "FriendlyName", value);
 		}
 
 		[UhtSpecifier(Extends = UhtTableNames.Default, ValueType = UhtSpecifierValueType.None, When = UhtSpecifierWhen.Immediate)]
-		[SuppressMessage("CodeQuality", "IDE0051:Remove unused private members", Justification = "Attribute accessed method")]
 		private static void BlueprintInternalUseOnlySpecifier(UhtSpecifierContext specifierContext)
 		{
 			SetMetaData(specifierContext, "BlueprintInternalUseOnly", true);
@@ -113,7 +111,6 @@ namespace EpicGames.UHT.Parsers
 		}
 
 		[UhtSpecifier(Extends = UhtTableNames.Default, ValueType = UhtSpecifierValueType.None, When = UhtSpecifierWhen.Immediate)]
-		[SuppressMessage("CodeQuality", "IDE0051:Remove unused private members", Justification = "Attribute accessed method")]
 		private static void BlueprintInternalUseOnlyHierarchicalSpecifier(UhtSpecifierContext specifierContext)
 		{
 			SetMetaData(specifierContext, "BlueprintInternalUseOnlyHierarchical", true);
@@ -122,14 +119,12 @@ namespace EpicGames.UHT.Parsers
 		}
 
 		[UhtSpecifier(Extends = UhtTableNames.Default, ValueType = UhtSpecifierValueType.None, When = UhtSpecifierWhen.Immediate)]
-		[SuppressMessage("CodeQuality", "IDE0051:Remove unused private members", Justification = "Attribute accessed method")]
 		private static void BlueprintTypeSpecifier(UhtSpecifierContext specifierContext)
 		{
 			SetMetaData(specifierContext, "BlueprintType", true);
 		}
 
 		[UhtSpecifier(Extends = UhtTableNames.Default, ValueType = UhtSpecifierValueType.None, When = UhtSpecifierWhen.Immediate)]
-		[SuppressMessage("CodeQuality", "IDE0051:Remove unused private members", Justification = "Attribute accessed method")]
 		private static void NotBlueprintTypeSpecifier(UhtSpecifierContext specifierContext)
 		{
 			SetMetaData(specifierContext, "NotBlueprintType", true);
@@ -137,7 +132,6 @@ namespace EpicGames.UHT.Parsers
 		}
 
 		[UhtSpecifier(Extends = UhtTableNames.Default, ValueType = UhtSpecifierValueType.None, When = UhtSpecifierWhen.Immediate)]
-		[SuppressMessage("CodeQuality", "IDE0051:Remove unused private members", Justification = "Attribute accessed method")]
 		private static void BlueprintableSpecifier(UhtSpecifierContext specifierContext)
 		{
 			SetMetaData(specifierContext, "IsBlueprintBase", true);
@@ -145,14 +139,12 @@ namespace EpicGames.UHT.Parsers
 		}
 
 		[UhtSpecifier(Extends = UhtTableNames.Default, ValueType = UhtSpecifierValueType.None, When = UhtSpecifierWhen.Immediate)]
-		[SuppressMessage("CodeQuality", "IDE0051:Remove unused private members", Justification = "Attribute accessed method")]
 		private static void CallInEditorSpecifier(UhtSpecifierContext specifierContext)
 		{
 			SetMetaData(specifierContext, "CallInEditor", true);
 		}
 
 		[UhtSpecifier(Extends = UhtTableNames.Default, ValueType = UhtSpecifierValueType.None, When = UhtSpecifierWhen.Immediate)]
-		[SuppressMessage("CodeQuality", "IDE0051:Remove unused private members", Justification = "Attribute accessed method")]
 		private static void NotBlueprintableSpecifier(UhtSpecifierContext specifierContext)
 		{
 			SetMetaData(specifierContext, "IsBlueprintBase", false);
@@ -160,35 +152,30 @@ namespace EpicGames.UHT.Parsers
 		}
 
 		[UhtSpecifier(Extends = UhtTableNames.Default, ValueType = UhtSpecifierValueType.String, When = UhtSpecifierWhen.Immediate)]
-		[SuppressMessage("CodeQuality", "IDE0051:Remove unused private members", Justification = "Attribute accessed method")]
 		private static void CategorySpecifier(UhtSpecifierContext specifierContext, StringView value)
 		{
 			SetMetaData(specifierContext, "Category", value);
 		}
 
 		[UhtSpecifier(Extends = UhtTableNames.Default, ValueType = UhtSpecifierValueType.None, When = UhtSpecifierWhen.Immediate)]
-		[SuppressMessage("CodeQuality", "IDE0051:Remove unused private members", Justification = "Attribute accessed method")]
 		private static void ExperimentalSpecifier(UhtSpecifierContext specifierContext)
 		{
 			SetMetaData(specifierContext, "DevelopmentStatus", "Experimental");
 		}
 
 		[UhtSpecifier(Extends = UhtTableNames.Default, ValueType = UhtSpecifierValueType.None, When = UhtSpecifierWhen.Immediate)]
-		[SuppressMessage("CodeQuality", "IDE0051:Remove unused private members", Justification = "Attribute accessed method")]
 		private static void EarlyAccessPreviewSpecifier(UhtSpecifierContext specifierContext)
 		{
 			SetMetaData(specifierContext, "DevelopmentStatus", "EarlyAccess");
 		}
 
 		[UhtSpecifier(Extends = UhtTableNames.Default, ValueType = UhtSpecifierValueType.None, When = UhtSpecifierWhen.Immediate)]
-		[SuppressMessage("CodeQuality", "IDE0051:Remove unused private members", Justification = "Attribute accessed method")]
 		private static void DocumentationPolicySpecifier(UhtSpecifierContext specifierContext)
 		{
 			SetMetaData(specifierContext, UhtNames.DocumentationPolicy, "Strict");
 		}
 
 		[UhtSpecifier(Extends = UhtTableNames.Default, ValueType = UhtSpecifierValueType.String, When = UhtSpecifierWhen.Immediate)]
-		[SuppressMessage("CodeQuality", "IDE0051:Remove unused private members", Justification = "Attribute accessed method")]
 		private static void SparseClassDataTypeSpecifier(UhtSpecifierContext specifierContext, StringView value)
 		{
 			SetMetaData(specifierContext, "SparseClassDataType", value);
@@ -200,8 +187,6 @@ namespace EpicGames.UHT.Parsers
 		[UhtSpecifierValidator(Name = "UIMax", Extends = UhtTableNames.Default)]
 		[UhtSpecifierValidator(Name = "ClampMin", Extends = UhtTableNames.Default)]
 		[UhtSpecifierValidator(Name = "ClampMax", Extends = UhtTableNames.Default)]
-		[SuppressMessage("CodeQuality", "IDE0051:Remove unused private members", Justification = "Attribute accessed method")]
-		[SuppressMessage("Style", "IDE0060:Remove unused parameter", Justification = "Attribute accessed method")]
 		private static void ValidateNumeric(UhtType type, UhtMetaData metaData, UhtMetaDataKey key, StringView value)
 		{
 			if (!UhtFCString.IsNumeric(value.Span))
@@ -211,8 +196,6 @@ namespace EpicGames.UHT.Parsers
 		}
 
 		[UhtSpecifierValidator(Extends = UhtTableNames.Default)]
-		[SuppressMessage("CodeQuality", "IDE0051:Remove unused private members", Justification = "Attribute accessed method")]
-		[SuppressMessage("Style", "IDE0060:Remove unused parameter", Justification = "Attribute accessed method")]
 		private static void DevelopmentStatusSpecifierValidator(UhtType type, UhtMetaData metaData, UhtMetaDataKey key, StringView value)
 		{
 			string[] allowedValues = { "EarlyAccess", "Experimental" };
@@ -227,7 +210,6 @@ namespace EpicGames.UHT.Parsers
 		}
 
 		[UhtSpecifierValidator(Extends = UhtTableNames.Default)]
-		[SuppressMessage("CodeQuality", "IDE0051:Remove unused private members", Justification = "Attribute accessed method")]
 		private static void DocumentationPolicySpecifierValidator(UhtType type, UhtMetaData metaData, UhtMetaDataKey key, StringView value)
 		{
 			const string StrictValue = "Strict";
@@ -238,8 +220,6 @@ namespace EpicGames.UHT.Parsers
 		}
 
 		[UhtSpecifierValidator(Extends = UhtTableNames.Default)]
-		[SuppressMessage("CodeQuality", "IDE0051:Remove unused private members", Justification = "Attribute accessed method")]
-		[SuppressMessage("Style", "IDE0060:Remove unused parameter", Justification = "Attribute accessed method")]
 		private static void UnitsSpecifierValidator(UhtType type, UhtMetaData metaData, UhtMetaDataKey key, StringView value)
 		{
 			// Check for numeric property
diff --git a/Engine/Source/Programs/Shared/EpicGames.UHT/Specifiers/UhtFunctionSpecifiers.cs b/Engine/Source/Programs/Shared/EpicGames.UHT/Specifiers/UhtFunctionSpecifiers.cs
--- a/Engine/Source/Programs/Shared/EpicGames.UHT/Specifiers/UhtFunctionSpecifiers.cs
+++ b/Engine/Source/Programs/Shared/EpicGames.UHT/Specifiers/UhtFunctionSpecifiers.cs
@@ -15,10 +15,11 @@ namespace EpicGames.UHT.Parsers
 	/// Collection of UFUNCTION/UDELEGATE specifiers
 	/// </summary>
 	[UnrealHeaderTool]
+	[SuppressMessage("CodeQuality", "IDE0051:Remove unused private members", Justification = "Attribute accessed method")]
+	[SuppressMessage("Style", "IDE0060:Remove unused parameter", Justification = "Attribute accessed method")]
 	public static class UhtFunctionSpecifiers
 	{
 		[UhtSpecifier(Extends = UhtTableNames.Function, ValueType = UhtSpecifierValueType.Legacy)]
-		[SuppressMessage("CodeQuality", "IDE0051:Remove unused private members", Justification = "Attribute accessed method")]
 		private static void BlueprintAuthorityOnlySpecifier(UhtSpecifierContext specifierContext)
 		{
 			UhtFunction function = (UhtFunction)specifierContext.Scope.ScopeType;
@@ -26,7 +27,6 @@ namespace EpicGames.UHT.Parsers
 		}
 
 		[UhtSpecifier(Extends = UhtTableNames.Function, ValueType = UhtSpecifierValueType.Legacy)]
-		[SuppressMessage("CodeQuality", "IDE0051:Remove unused private members", Justification = "Attribute accessed method")]
 		private static void BlueprintCallableSpecifier(UhtSpecifierContext specifierContext)
 		{
 			UhtFunction function = (UhtFunction)specifierContext.Scope.ScopeType;
@@ -34,7 +34,6 @@ namespace EpicGames.UHT.Parsers
 		}
 
 		[UhtSpecifier(Extends = UhtTableNames.Function, ValueType = UhtSpecifierValueType.Legacy)]
-		[SuppressMessage("CodeQuality", "IDE0051:Remove unused private members", Justification = "Attribute accessed method")]
 		private static void BlueprintCosmeticSpecifier(UhtSpecifierContext specifierContext)
 		{
 			UhtFunction function = (UhtFunction)specifierContext.Scope.ScopeType;
@@ -42,27 +41,25 @@ namespace EpicGames.UHT.Parsers
 		}
 
 		[UhtSpecifier(Extends = UhtTableNames.Function, ValueType = UhtSpecifierValueType.Legacy)]
-		[SuppressMessage("CodeQuality", "IDE0051:Remove unused private members", Justification = "Attribute accessed method")]
 		private static void BlueprintGetterSpecifier(UhtSpecifierContext specifierContext)
 		{
-			UhtFunctionParser function = (UhtFunctionParser)specifierContext.Scope.ScopeType;
+			UhtFunction function = (UhtFunction)specifierContext.Scope.ScopeType;
 
 			if (function.FunctionFlags.HasAnyFlags(EFunctionFlags.Event))
 			{
 				specifierContext.MessageSite.LogError("Function cannot be a blueprint event and a blueprint getter.");
 			}
 
-			function.SawPropertyAccessor = true;
+			function.FunctionExportFlags |= UhtFunctionExportFlags.SawPropertyAccessor;
 			function.FunctionFlags |= EFunctionFlags.BlueprintCallable;
 			function.FunctionFlags |= EFunctionFlags.BlueprintPure;
 			function.MetaData.Add(UhtNames.BlueprintGetter, "");
 		}
 
 		[UhtSpecifier(Extends = UhtTableNames.Function, ValueType = UhtSpecifierValueType.Legacy)]
-		[SuppressMessage("CodeQuality", "IDE0051:Remove unused private members", Justification = "Attribute accessed method")]
 		private static void BlueprintImplementableEventSpecifier(UhtSpecifierContext specifierContext)
 		{
-			UhtFunctionParser function = (UhtFunctionParser)specifierContext.Scope.ScopeType;
+			UhtFunction function = (UhtFunction)specifierContext.Scope.ScopeType;
 
 			if (function.FunctionFlags.HasAnyFlags(EFunctionFlags.Net))
 			{
@@ -73,7 +70,7 @@ namespace EpicGames.UHT.Parsers
 				// already a BlueprintNativeEvent
 				specifierContext.MessageSite.LogError("A function cannot be both BlueprintNativeEvent and BlueprintImplementableEvent!");
 			}
-			else if (function.SawPropertyAccessor)
+			else if (function.FunctionExportFlags.HasAnyFlags(UhtFunctionExportFlags.SawPropertyAccessor))
 			{
 				specifierContext.MessageSite.LogError("A function cannot be both BlueprintImplementableEvent and a Blueprint Property accessor!");
 			}
@@ -88,10 +85,9 @@ namespace EpicGames.UHT.Parsers
 		}
 
 		[UhtSpecifier(Extends = UhtTableNames.Function, ValueType = UhtSpecifierValueType.Legacy)]
-		[SuppressMessage("CodeQuality", "IDE0051:Remove unused private members", Justification = "Attribute accessed method")]
 		private static void BlueprintNativeEventSpecifier(UhtSpecifierContext specifierContext)
 		{
-			UhtFunctionParser function = (UhtFunctionParser)specifierContext.Scope.ScopeType;
+			UhtFunction function = (UhtFunction)specifierContext.Scope.ScopeType;
 
 			if (function.FunctionFlags.HasAnyFlags(EFunctionFlags.Net))
 			{
@@ -102,7 +98,7 @@ namespace EpicGames.UHT.Parsers
 				// already a BlueprintImplementableEvent
 				specifierContext.MessageSite.LogError("A function cannot be both BlueprintNativeEvent and BlueprintImplementableEvent!");
 			}
-			else if (function.SawPropertyAccessor)
+			else if (function.FunctionExportFlags.HasAnyFlags(UhtFunctionExportFlags.SawPropertyAccessor))
 			{
 				specifierContext.MessageSite.LogError("A function cannot be both BlueprintNativeEvent and a Blueprint Property accessor!");
 			}
@@ -116,7 +112,6 @@ namespace EpicGames.UHT.Parsers
 		}
 
 		[UhtSpecifier(Extends = UhtTableNames.Function, ValueType = UhtSpecifierValueType.OptionalString)]
-		[SuppressMessage("CodeQuality", "IDE0051:Remove unused private members", Justification = "Attribute accessed method")]
 		private static void BlueprintPureSpecifier(UhtSpecifierContext specifierContext, StringView? value)
 		{
 			UhtFunction function = (UhtFunction)specifierContext.Scope.ScopeType;
@@ -135,23 +130,21 @@ namespace EpicGames.UHT.Parsers
 		}
 
 		[UhtSpecifier(Extends = UhtTableNames.Function, ValueType = UhtSpecifierValueType.Legacy)]
-		[SuppressMessage("CodeQuality", "IDE0051:Remove unused private members", Justification = "Attribute accessed method")]
 		private static void BlueprintSetterSpecifier(UhtSpecifierContext specifierContext)
 		{
-			UhtFunctionParser function = (UhtFunctionParser)specifierContext.Scope.ScopeType;
+			UhtFunction function = (UhtFunction)specifierContext.Scope.ScopeType;
 
 			if (function.FunctionFlags.HasAnyFlags(EFunctionFlags.Event))
 			{
 				specifierContext.MessageSite.LogError("Function cannot be a blueprint event and a blueprint setter.");
 			}
 
-			function.SawPropertyAccessor = true;
+			function.FunctionExportFlags |= UhtFunctionExportFlags.SawPropertyAccessor;
 			function.FunctionFlags |= EFunctionFlags.BlueprintCallable;
 			function.MetaData.Add(UhtNames.BlueprintSetter, "");
 		}
 
 		[UhtSpecifier(Extends = UhtTableNames.Function, ValueType = UhtSpecifierValueType.OptionalString)]
-		[SuppressMessage("CodeQuality", "IDE0051:Remove unused private members", Justification = "Attribute accessed method")]
 		private static void ClientSpecifier(UhtSpecifierContext specifierContext, StringView? value)
 		{
 			UhtFunction function = (UhtFunction)specifierContext.Scope.ScopeType;
@@ -176,7 +169,6 @@ namespace EpicGames.UHT.Parsers
 		}
 
 		[UhtSpecifier(Extends = UhtTableNames.Function, ValueType = UhtSpecifierValueType.Legacy)]
-		[SuppressMessage("CodeQuality", "IDE0051:Remove unused private members", Justification = "Attribute accessed method")]
 		private static void CustomThunkSpecifier(UhtSpecifierContext specifierContext)
 		{
 			UhtFunction function = (UhtFunction)specifierContext.Scope.ScopeType;
@@ -184,7 +176,6 @@ namespace EpicGames.UHT.Parsers
 		}
 
 		[UhtSpecifier(Extends = UhtTableNames.Function, ValueType = UhtSpecifierValueType.Legacy)]
-		[SuppressMessage("CodeQuality", "IDE0051:Remove unused private members", Justification = "Attribute accessed method")]
 		private static void ExecSpecifier(UhtSpecifierContext specifierContext)
 		{
 			UhtFunction function = (UhtFunction)specifierContext.Scope.ScopeType;
@@ -196,7 +187,6 @@ namespace EpicGames.UHT.Parsers
 		}
 
 		[UhtSpecifier(Extends = UhtTableNames.Function, ValueType = UhtSpecifierValueType.Legacy)]
-		[SuppressMessage("CodeQuality", "IDE0051:Remove unused private members", Justification = "Attribute accessed method")]
 		private static void NetMulticastSpecifier(UhtSpecifierContext specifierContext)
 		{
 			UhtFunction function = (UhtFunction)specifierContext.Scope.ScopeType;
@@ -210,7 +200,6 @@ namespace EpicGames.UHT.Parsers
 		}
 
 		[UhtSpecifier(Extends = UhtTableNames.Function, ValueType = UhtSpecifierValueType.Legacy)]
-		[SuppressMessage("CodeQuality", "IDE0051:Remove unused private members", Justification = "Attribute accessed method")]
 		private static void ReliableSpecifier(UhtSpecifierContext specifierContext)
 		{
 			UhtFunction function = (UhtFunction)specifierContext.Scope.ScopeType;
@@ -218,7 +207,6 @@ namespace EpicGames.UHT.Parsers
 		}
 
 		[UhtSpecifier(Extends = UhtTableNames.Function, ValueType = UhtSpecifierValueType.Legacy)]
-		[SuppressMessage("CodeQuality", "IDE0051:Remove unused private members", Justification = "Attribute accessed method")]
 		private static void SealedEventSpecifier(UhtSpecifierContext specifierContext)
 		{
 			UhtFunction function = (UhtFunction)specifierContext.Scope.ScopeType;
@@ -226,7 +214,6 @@ namespace EpicGames.UHT.Parsers
 		}
 
 		[UhtSpecifier(Extends = UhtTableNames.Function, ValueType = UhtSpecifierValueType.OptionalString)]
-		[SuppressMessage("CodeQuality", "IDE0051:Remove unused private members", Justification = "Attribute accessed method")]
 		private static void ServerSpecifier(UhtSpecifierContext specifierContext, StringView? value)
 		{
 			UhtFunction function = (UhtFunction)specifierContext.Scope.ScopeType;
@@ -251,10 +238,9 @@ namespace EpicGames.UHT.Parsers
 		}
 
 		[UhtSpecifier(Extends = UhtTableNames.Function, ValueType = UhtSpecifierValueType.OptionalEqualsKeyValuePairList)]
-		[SuppressMessage("CodeQuality", "IDE0051:Remove unused private members", Justification = "Attribute accessed method")]
 		private static void ServiceRequestSpecifier(UhtSpecifierContext specifierContext, List<KeyValuePair<StringView, StringView>> value)
 		{
-			UhtFunctionParser function = (UhtFunctionParser)specifierContext.Scope.ScopeType;
+			UhtFunction function = (UhtFunction)specifierContext.Scope.ScopeType;
 
 			if (function.FunctionFlags.HasAnyFlags(EFunctionFlags.BlueprintEvent))
 			{
@@ -275,10 +261,9 @@ namespace EpicGames.UHT.Parsers
 		}
 
 		[UhtSpecifier(Extends = UhtTableNames.Function, ValueType = UhtSpecifierValueType.OptionalEqualsKeyValuePairList)]
-		[SuppressMessage("CodeQuality", "IDE0051:Remove unused private members", Justification = "Attribute accessed method")]
 		private static void ServiceResponseSpecifier(UhtSpecifierContext specifierContext, List<KeyValuePair<StringView, StringView>> value)
 		{
-			UhtFunctionParser function = (UhtFunctionParser)specifierContext.Scope.ScopeType;
+			UhtFunction function = (UhtFunction)specifierContext.Scope.ScopeType;
 
 			if (function.FunctionFlags.HasAnyFlags(EFunctionFlags.BlueprintEvent))
 			{
@@ -299,7 +284,6 @@ namespace EpicGames.UHT.Parsers
 		}
 
 		[UhtSpecifier(Extends = UhtTableNames.Function, ValueType = UhtSpecifierValueType.Legacy)]
-		[SuppressMessage("CodeQuality", "IDE0051:Remove unused private members", Justification = "Attribute accessed method")]
 		private static void UnreliableSpecifier(UhtSpecifierContext specifierContext)
 		{
 			UhtFunction function = (UhtFunction)specifierContext.Scope.ScopeType;
@@ -307,7 +291,6 @@ namespace EpicGames.UHT.Parsers
 		}
 
 		[UhtSpecifier(Extends = UhtTableNames.Function, ValueType = UhtSpecifierValueType.StringList)]
-		[SuppressMessage("CodeQuality", "IDE0051:Remove unused private members", Justification = "Attribute accessed method")]
 		private static void WithValidationSpecifier(UhtSpecifierContext specifierContext, List<StringView>? value)
 		{
 			UhtFunction function = (UhtFunction)specifierContext.Scope.ScopeType;
@@ -320,7 +303,6 @@ namespace EpicGames.UHT.Parsers
 		}
 
 		[UhtSpecifier(Extends = UhtTableNames.Function, ValueType = UhtSpecifierValueType.None)]
-		[SuppressMessage("CodeQuality", "IDE0051:Remove unused private members", Justification = "Attribute accessed method")]
 		private static void FieldNotifySpecifier(UhtSpecifierContext specifierContext)
 		{
 			UhtFunction function = (UhtFunction)specifierContext.Scope.ScopeType;
@@ -328,8 +310,6 @@ namespace EpicGames.UHT.Parsers
 		}
 
 		[UhtSpecifierValidator(Extends = UhtTableNames.Function)]
-		[SuppressMessage("CodeQuality", "IDE0051:Remove unused private members", Justification = "Attribute accessed method")]
-		[SuppressMessage("Style", "IDE0060:Remove unused parameter", Justification = "Attribute accessed method")]
 		private static void BlueprintProtectedSpecifierValidator(UhtType type, UhtMetaData metaData, UhtMetaDataKey key, StringView value)
 		{
 			UhtFunction function = (UhtFunction)type;
@@ -357,8 +337,6 @@ namespace EpicGames.UHT.Parsers
 		}
 
 		[UhtSpecifierValidator(Extends = UhtTableNames.Function)]
-		[SuppressMessage("CodeQuality", "IDE0051:Remove unused private members", Justification = "Attribute accessed method")]
-		[SuppressMessage("Style", "IDE0060:Remove unused parameter", Justification = "Attribute accessed method")]
 		private static void CommutativeAssociativeBinaryOperatorSpecifierValidator(UhtType type, UhtMetaData metaData, UhtMetaDataKey key, StringView value)
 		{
 			UhtFunction function = (UhtFunction)type;
@@ -373,8 +351,6 @@ namespace EpicGames.UHT.Parsers
 
 		[UhtSpecifierValidator(Name = "ExpandBoolAsExecs", Extends = UhtTableNames.Function)]
 		[UhtSpecifierValidator(Name = "ExpandEnumAsExecs", Extends = UhtTableNames.Function)]
-		[SuppressMessage("CodeQuality", "IDE0051:Remove unused private members", Justification = "Attribute accessed method")]
-		[SuppressMessage("Style", "IDE0060:Remove unused parameter", Justification = "Attribute accessed method")]
 		private static void ExpandsSpecifierValidator(UhtType type, UhtMetaData metaData, UhtMetaDataKey key, StringView value)
 		{
 			// multiple entry parsing in the same format as eg SetParam.
@@ -414,7 +390,7 @@ namespace EpicGames.UHT.Parsers
 		private static void ParseNetServiceIdentifiers(UhtSpecifierContext specifierContext, List<KeyValuePair<StringView, StringView>> identifiers)
 		{
 			IUhtTokenReader tokenReader = specifierContext.Scope.TokenReader;
-			UhtFunctionParser function = (UhtFunctionParser)specifierContext.Scope.ScopeType;
+			UhtFunction function = (UhtFunction)specifierContext.Scope.ScopeType;
 
 			foreach (KeyValuePair<StringView, StringView> kvp in identifiers)
 			{
diff --git a/Engine/Source/Programs/Shared/EpicGames.UHT/Specifiers/UhtInterfaceClassSpecifiers.cs b/Engine/Source/Programs/Shared/EpicGames.UHT/Specifiers/UhtInterfaceClassSpecifiers.cs
--- a/Engine/Source/Programs/Shared/EpicGames.UHT/Specifiers/UhtInterfaceClassSpecifiers.cs
+++ b/Engine/Source/Programs/Shared/EpicGames.UHT/Specifiers/UhtInterfaceClassSpecifiers.cs
@@ -9,17 +9,16 @@ using EpicGames.UHT.Utils;
 namespace EpicGames.UHT.Parsers
 {
 	[UnrealHeaderTool]
+	[SuppressMessage("CodeQuality", "IDE0051:Remove unused private members", Justification = "Attribute accessed method")]
 	class UhtInterfaceClassSpecifiers
 	{
 		[UhtSpecifier(Extends = UhtTableNames.Interface, ValueType = UhtSpecifierValueType.Legacy)]
-		[SuppressMessage("CodeQuality", "IDE0051:Remove unused private members", Justification = "Attribute accessed method")]
 		private static void DependsOnSpecifier(UhtSpecifierContext specifierContext)
 		{
 			throw new UhtException(specifierContext.MessageSite, $"The dependsOn specifier is deprecated. Please use #include \"ClassHeaderFilename.h\" instead.");
 		}
 
 		[UhtSpecifier(Extends = UhtTableNames.Interface, ValueType = UhtSpecifierValueType.Legacy)]
-		[SuppressMessage("CodeQuality", "IDE0051:Remove unused private members", Justification = "Attribute accessed method")]
 		private static void MinimalAPISpecifier(UhtSpecifierContext specifierContext)
 		{
 			UhtClass clasObj = (UhtClass)specifierContext.Scope.ScopeType;
@@ -27,7 +26,6 @@ namespace EpicGames.UHT.Parsers
 		}
 
 		[UhtSpecifier(Extends = UhtTableNames.Interface, ValueType = UhtSpecifierValueType.Legacy)]
-		[SuppressMessage("CodeQuality", "IDE0051:Remove unused private members", Justification = "Attribute accessed method")]
 		private static void ConversionRootSpecifier(UhtSpecifierContext specifierContext)
 		{
 			UhtClass classObj = (UhtClass)specifierContext.Scope.ScopeType;
diff --git a/Engine/Source/Programs/Shared/EpicGames.UHT/Specifiers/UhtPropertyArgumentSpecifiers.cs b/Engine/Source/Programs/Shared/EpicGames.UHT/Specifiers/UhtPropertyArgumentSpecifiers.cs
--- a/Engine/Source/Programs/Shared/EpicGames.UHT/Specifiers/UhtPropertyArgumentSpecifiers.cs
+++ b/Engine/Source/Programs/Shared/EpicGames.UHT/Specifiers/UhtPropertyArgumentSpecifiers.cs
@@ -13,11 +13,11 @@ namespace EpicGames.UHT.Parsers
 	/// Collection of property argument specifiers
 	/// </summary>
 	[UnrealHeaderTool]
+	[SuppressMessage("CodeQuality", "IDE0051:Remove unused private members", Justification = "Attribute accessed method")]
 	public static class UhtPropertyArgumentSpecifiers
 	{
 		#region Argument Property Specifiers
 		[UhtSpecifier(Extends = UhtTableNames.PropertyArgument, ValueType = UhtSpecifierValueType.Legacy)]
-		[SuppressMessage("CodeQuality", "IDE0051:Remove unused private members", Justification = "Attribute accessed method")]
 		private static void ConstSpecifier(UhtSpecifierContext specifierContext)
 		{
 			UhtPropertySpecifierContext context = (UhtPropertySpecifierContext)specifierContext;
@@ -25,7 +25,6 @@ namespace EpicGames.UHT.Parsers
 		}
 
 		[UhtSpecifier(Extends = UhtTableNames.PropertyArgument, ValueType = UhtSpecifierValueType.Legacy)]
-		[SuppressMessage("CodeQuality", "IDE0051:Remove unused private members", Justification = "Attribute accessed method")]
 		private static void RefSpecifier(UhtSpecifierContext specifierContext)
 		{
 			UhtPropertySpecifierContext context = (UhtPropertySpecifierContext)specifierContext;
@@ -33,7 +32,6 @@ namespace EpicGames.UHT.Parsers
 		}
 
 		[UhtSpecifier(Extends = UhtTableNames.PropertyArgument, ValueType = UhtSpecifierValueType.Legacy)]
-		[SuppressMessage("CodeQuality", "IDE0051:Remove unused private members", Justification = "Attribute accessed method")]
 		private static void NotReplicatedSpecifier(UhtSpecifierContext specifierContext)
 		{
 			UhtPropertySpecifierContext context = (UhtPropertySpecifierContext)specifierContext;
diff --git a/Engine/Source/Programs/Shared/EpicGames.UHT/Specifiers/UhtPropertyMemberSpecifiers.cs b/Engine/Source/Programs/Shared/EpicGames.UHT/Specifiers/UhtPropertyMemberSpecifiers.cs
--- a/Engine/Source/Programs/Shared/EpicGames.UHT/Specifiers/UhtPropertyMemberSpecifiers.cs
+++ b/Engine/Source/Programs/Shared/EpicGames.UHT/Specifiers/UhtPropertyMemberSpecifiers.cs
@@ -14,10 +14,10 @@ namespace EpicGames.UHT.Parsers
 	/// Collection of property member specifiers
 	/// </summary>
 	[UnrealHeaderTool]
+	[SuppressMessage("CodeQuality", "IDE0051:Remove unused private members", Justification = "Attribute accessed method")]
 	public static class UhtPropertyMemberSpecifiers
 	{
 		[UhtSpecifier(Extends = UhtTableNames.PropertyMember, ValueType = UhtSpecifierValueType.Legacy)]
-		[SuppressMessage("CodeQuality", "IDE0051:Remove unused private members", Justification = "Attribute accessed method")]
 		private static void EditAnywhereSpecifier(UhtSpecifierContext specifierContext)
 		{
 			UhtPropertySpecifierContext context = (UhtPropertySpecifierContext)specifierContext;
@@ -30,7 +30,6 @@ namespace EpicGames.UHT.Parsers
 		}
 
 		[UhtSpecifier(Extends = UhtTableNames.PropertyMember, ValueType = UhtSpecifierValueType.Legacy)]
-		[SuppressMessage("CodeQuality", "IDE0051:Remove unused private members", Justification = "Attribute accessed method")]
 		private static void EditInstanceOnlySpecifier(UhtSpecifierContext specifierContext)
 		{
 			UhtPropertySpecifierContext context = (UhtPropertySpecifierContext)specifierContext;
@@ -43,7 +42,6 @@ namespace EpicGames.UHT.Parsers
 		}
 
 		[UhtSpecifier(Extends = UhtTableNames.PropertyMember, ValueType = UhtSpecifierValueType.Legacy)]
-		[SuppressMessage("CodeQuality", "IDE0051:Remove unused private members", Justification = "Attribute accessed method")]
 		private static void EditDefaultsOnlySpecifier(UhtSpecifierContext specifierContext)
 		{
 			UhtPropertySpecifierContext context = (UhtPropertySpecifierContext)specifierContext;
@@ -56,7 +54,6 @@ namespace EpicGames.UHT.Parsers
 		}
 
 		[UhtSpecifier(Extends = UhtTableNames.PropertyMember, ValueType = UhtSpecifierValueType.Legacy)]
-		[SuppressMessage("CodeQuality", "IDE0051:Remove unused private members", Justification = "Attribute accessed method")]
 		private static void VisibleAnywhereSpecifier(UhtSpecifierContext specifierContext)
 		{
 			UhtPropertySpecifierContext context = (UhtPropertySpecifierContext)specifierContext;
@@ -69,7 +66,6 @@ namespace EpicGames.UHT.Parsers
 		}
 
 		[UhtSpecifier(Extends = UhtTableNames.PropertyMember, ValueType = UhtSpecifierValueType.Legacy)]
-		[SuppressMessage("CodeQuality", "IDE0051:Remove unused private members", Justification = "Attribute accessed method")]
 		private static void VisibleInstanceOnlySpecifier(UhtSpecifierContext specifierContext)
 		{
 			UhtPropertySpecifierContext context = (UhtPropertySpecifierContext)specifierContext;
@@ -82,7 +78,6 @@ namespace EpicGames.UHT.Parsers
 		}
 
 		[UhtSpecifier(Extends = UhtTableNames.PropertyMember, ValueType = UhtSpecifierValueType.Legacy)]
-		[SuppressMessage("CodeQuality", "IDE0051:Remove unused private members", Justification = "Attribute accessed method")]
 		private static void VisibleDefaultsOnlySpecifier(UhtSpecifierContext specifierContext)
 		{
 			UhtPropertySpecifierContext context = (UhtPropertySpecifierContext)specifierContext;
@@ -95,7 +90,6 @@ namespace EpicGames.UHT.Parsers
 		}
 
 		[UhtSpecifier(Extends = UhtTableNames.PropertyMember, ValueType = UhtSpecifierValueType.Legacy)]
-		[SuppressMessage("CodeQuality", "IDE0051:Remove unused private members", Justification = "Attribute accessed method")]
 		private static void BlueprintReadWriteSpecifier(UhtSpecifierContext specifierContext)
 		{
 			UhtPropertySpecifierContext context = (UhtPropertySpecifierContext)specifierContext;
@@ -120,7 +114,6 @@ namespace EpicGames.UHT.Parsers
 		}
 
 		[UhtSpecifier(Extends = UhtTableNames.PropertyMember, ValueType = UhtSpecifierValueType.Legacy)]
-		[SuppressMessage("CodeQuality", "IDE0051:Remove unused private members", Justification = "Attribute accessed method")]
 		private static void BlueprintReadOnlySpecifier(UhtSpecifierContext specifierContext)
 		{
 			UhtPropertySpecifierContext context = (UhtPropertySpecifierContext)specifierContext;
@@ -145,7 +138,6 @@ namespace EpicGames.UHT.Parsers
 		}
 
 		[UhtSpecifier(Extends = UhtTableNames.PropertyMember, ValueType = UhtSpecifierValueType.SingleString)]
-		[SuppressMessage("CodeQuality", "IDE0051:Remove unused private members", Justification = "Attribute accessed method")]
 		private static void BlueprintSetterSpecifier(UhtSpecifierContext specifierContext, StringView value)
 		{
 			UhtPropertySpecifierContext context = (UhtPropertySpecifierContext)specifierContext;
@@ -166,7 +158,6 @@ namespace EpicGames.UHT.Parsers
 		}
 
 		[UhtSpecifier(Extends = UhtTableNames.PropertyMember, ValueType = UhtSpecifierValueType.SingleString)]
-		[SuppressMessage("CodeQuality", "IDE0051:Remove unused private members", Justification = "Attribute accessed method")]
 		private static void BlueprintGetterSpecifier(UhtSpecifierContext specifierContext, StringView value)
 		{
 			UhtPropertySpecifierContext context = (UhtPropertySpecifierContext)specifierContext;
@@ -182,7 +173,6 @@ namespace EpicGames.UHT.Parsers
 		}
 
 		[UhtSpecifier(Extends = UhtTableNames.PropertyMember, ValueType = UhtSpecifierValueType.Legacy)]
-		[SuppressMessage("CodeQuality", "IDE0051:Remove unused private members", Justification = "Attribute accessed method")]
 		private static void ConfigSpecifier(UhtSpecifierContext specifierContext)
 		{
 			UhtPropertySpecifierContext context = (UhtPropertySpecifierContext)specifierContext;
@@ -190,7 +180,6 @@ namespace EpicGames.UHT.Parsers
 		}
 
 		[UhtSpecifier(Extends = UhtTableNames.PropertyMember, ValueType = UhtSpecifierValueType.Legacy)]
-		[SuppressMessage("CodeQuality", "IDE0051:Remove unused private members", Justification = "Attribute accessed method")]
 		private static void GlobalConfigSpecifier(UhtSpecifierContext specifierContext)
 		{
 			UhtPropertySpecifierContext context = (UhtPropertySpecifierContext)specifierContext;
@@ -198,7 +187,6 @@ namespace EpicGames.UHT.Parsers
 		}
 
 		[UhtSpecifier(Extends = UhtTableNames.PropertyMember, ValueType = UhtSpecifierValueType.Legacy)]
-		[SuppressMessage("CodeQuality", "IDE0051:Remove unused private members", Justification = "Attribute accessed method")]
 		private static void LocalizedSpecifier(UhtSpecifierContext specifierContext)
 		{
 			UhtPropertySpecifierContext context = (UhtPropertySpecifierContext)specifierContext;
@@ -206,7 +194,6 @@ namespace EpicGames.UHT.Parsers
 		}
 
 		[UhtSpecifier(Extends = UhtTableNames.PropertyMember, ValueType = UhtSpecifierValueType.Legacy)]
-		[SuppressMessage("CodeQuality", "IDE0051:Remove unused private members", Justification = "Attribute accessed method")]
 		private static void TransientSpecifier(UhtSpecifierContext specifierContext)
 		{
 			UhtPropertySpecifierContext context = (UhtPropertySpecifierContext)specifierContext;
@@ -214,7 +201,6 @@ namespace EpicGames.UHT.Parsers
 		}
 
 		[UhtSpecifier(Extends = UhtTableNames.PropertyMember, ValueType = UhtSpecifierValueType.Legacy)]
-		[SuppressMessage("CodeQuality", "IDE0051:Remove unused private members", Justification = "Attribute accessed method")]
 		private static void DuplicateTransientSpecifier(UhtSpecifierContext specifierContext)
 		{
 			UhtPropertySpecifierContext context = (UhtPropertySpecifierContext)specifierContext;
@@ -222,7 +208,6 @@ namespace EpicGames.UHT.Parsers
 		}
 
 		[UhtSpecifier(Extends = UhtTableNames.PropertyMember, ValueType = UhtSpecifierValueType.Legacy)]
-		[SuppressMessage("CodeQuality", "IDE0051:Remove unused private members", Justification = "Attribute accessed method")]
 		private static void TextExportTransientSpecifier(UhtSpecifierContext specifierContext)
 		{
 			UhtPropertySpecifierContext context = (UhtPropertySpecifierContext)specifierContext;
@@ -230,7 +215,6 @@ namespace EpicGames.UHT.Parsers
 		}
 
 		[UhtSpecifier(Extends = UhtTableNames.PropertyMember, ValueType = UhtSpecifierValueType.Legacy)]
-		[SuppressMessage("CodeQuality", "IDE0051:Remove unused private members", Justification = "Attribute accessed method")]
 		private static void NonPIETransientSpecifier(UhtSpecifierContext specifierContext)
 		{
 			UhtPropertySpecifierContext context = (UhtPropertySpecifierContext)specifierContext;
@@ -239,7 +223,6 @@ namespace EpicGames.UHT.Parsers
 		}
 
 		[UhtSpecifier(Extends = UhtTableNames.PropertyMember, ValueType = UhtSpecifierValueType.Legacy)]
-		[SuppressMessage("CodeQuality", "IDE0051:Remove unused private members", Justification = "Attribute accessed method")]
 		private static void NonPIEDuplicateTransientSpecifier(UhtSpecifierContext specifierContext)
 		{
 			UhtPropertySpecifierContext context = (UhtPropertySpecifierContext)specifierContext;
@@ -247,7 +230,6 @@ namespace EpicGames.UHT.Parsers
 		}
 
 		[UhtSpecifier(Extends = UhtTableNames.PropertyMember, ValueType = UhtSpecifierValueType.Legacy)]
-		[SuppressMessage("CodeQuality", "IDE0051:Remove unused private members", Justification = "Attribute accessed method")]
 		private static void ExportSpecifier(UhtSpecifierContext specifierContext)
 		{
 			UhtPropertySpecifierContext context = (UhtPropertySpecifierContext)specifierContext;
@@ -255,7 +237,6 @@ namespace EpicGames.UHT.Parsers
 		}
 
 		[UhtSpecifier(Extends = UhtTableNames.PropertyMember, ValueType = UhtSpecifierValueType.Legacy)]
-		[SuppressMessage("CodeQuality", "IDE0051:Remove unused private members", Justification = "Attribute accessed method")]
 		private static void EditInlineSpecifier(UhtSpecifierContext specifierContext)
 		{
 			UhtPropertySpecifierContext context = (UhtPropertySpecifierContext)specifierContext;
@@ -263,7 +244,6 @@ namespace EpicGames.UHT.Parsers
 		}
 
 		[UhtSpecifier(Extends = UhtTableNames.PropertyMember, ValueType = UhtSpecifierValueType.Legacy)]
-		[SuppressMessage("CodeQuality", "IDE0051:Remove unused private members", Justification = "Attribute accessed method")]
 		private static void NoClearSpecifier(UhtSpecifierContext specifierContext)
 		{
 			UhtPropertySpecifierContext context = (UhtPropertySpecifierContext)specifierContext;
@@ -271,7 +251,6 @@ namespace EpicGames.UHT.Parsers
 		}
 
 		[UhtSpecifier(Extends = UhtTableNames.PropertyMember, ValueType = UhtSpecifierValueType.Legacy)]
-		[SuppressMessage("CodeQuality", "IDE0051:Remove unused private members", Justification = "Attribute accessed method")]
 		private static void EditFixedSizeSpecifier(UhtSpecifierContext specifierContext)
 		{
 			UhtPropertySpecifierContext context = (UhtPropertySpecifierContext)specifierContext;
@@ -279,7 +258,6 @@ namespace EpicGames.UHT.Parsers
 		}
 
 		[UhtSpecifier(Extends = UhtTableNames.PropertyMember, ValueType = UhtSpecifierValueType.Legacy)]
-		[SuppressMessage("CodeQuality", "IDE0051:Remove unused private members", Justification = "Attribute accessed method")]
 		private static void ReplicatedSpecifier(UhtSpecifierContext specifierContext)
 		{
 			UhtPropertySpecifierContext context = (UhtPropertySpecifierContext)specifierContext;
@@ -291,7 +269,6 @@ namespace EpicGames.UHT.Parsers
 		}
 
 		[UhtSpecifier(Extends = UhtTableNames.PropertyMember, ValueType = UhtSpecifierValueType.SingleString)]
-		[SuppressMessage("CodeQuality", "IDE0051:Remove unused private members", Justification = "Attribute accessed method")]
 		private static void ReplicatedUsingSpecifier(UhtSpecifierContext specifierContext, StringView value)
 		{
 			UhtPropertySpecifierContext context = (UhtPropertySpecifierContext)specifierContext;
@@ -310,7 +287,6 @@ namespace EpicGames.UHT.Parsers
 		}
 
 		[UhtSpecifier(Extends = UhtTableNames.PropertyMember, ValueType = UhtSpecifierValueType.Legacy)]
-		[SuppressMessage("CodeQuality", "IDE0051:Remove unused private members", Justification = "Attribute accessed method")]
 		private static void NotReplicatedSpecifier(UhtSpecifierContext specifierContext)
 		{
 			UhtPropertySpecifierContext context = (UhtPropertySpecifierContext)specifierContext;
@@ -322,7 +298,6 @@ namespace EpicGames.UHT.Parsers
 		}
 
 		[UhtSpecifier(Extends = UhtTableNames.PropertyMember, ValueType = UhtSpecifierValueType.Legacy)]
-		[SuppressMessage("CodeQuality", "IDE0051:Remove unused private members", Justification = "Attribute accessed method")]
 		private static void RepRetrySpecifier(UhtSpecifierContext specifierContext)
 		{
 			UhtPropertySpecifierContext context = (UhtPropertySpecifierContext)specifierContext;
@@ -330,7 +305,6 @@ namespace EpicGames.UHT.Parsers
 		}
 
 		[UhtSpecifier(Extends = UhtTableNames.PropertyMember, ValueType = UhtSpecifierValueType.Legacy)]
-		[SuppressMessage("CodeQuality", "IDE0051:Remove unused private members", Justification = "Attribute accessed method")]
 		private static void InterpSpecifier(UhtSpecifierContext specifierContext)
 		{
 			UhtPropertySpecifierContext context = (UhtPropertySpecifierContext)specifierContext;
@@ -338,7 +312,6 @@ namespace EpicGames.UHT.Parsers
 		}
 
 		[UhtSpecifier(Extends = UhtTableNames.PropertyMember, ValueType = UhtSpecifierValueType.Legacy)]
-		[SuppressMessage("CodeQuality", "IDE0051:Remove unused private members", Justification = "Attribute accessed method")]
 		private static void NonTransactionalSpecifier(UhtSpecifierContext specifierContext)
 		{
 			UhtPropertySpecifierContext context = (UhtPropertySpecifierContext)specifierContext;
@@ -346,7 +319,6 @@ namespace EpicGames.UHT.Parsers
 		}
 
 		[UhtSpecifier(Extends = UhtTableNames.PropertyMember, ValueType = UhtSpecifierValueType.Legacy)]
-		[SuppressMessage("CodeQuality", "IDE0051:Remove unused private members", Justification = "Attribute accessed method")]
 		private static void InstancedSpecifier(UhtSpecifierContext specifierContext)
 		{
 			UhtPropertySpecifierContext context = (UhtPropertySpecifierContext)specifierContext;
@@ -355,7 +327,6 @@ namespace EpicGames.UHT.Parsers
 		}
 
 		[UhtSpecifier(Extends = UhtTableNames.PropertyMember, ValueType = UhtSpecifierValueType.Legacy)]
-		[SuppressMessage("CodeQuality", "IDE0051:Remove unused private members", Justification = "Attribute accessed method")]
 		private static void BlueprintAssignableSpecifier(UhtSpecifierContext specifierContext)
 		{
 			UhtPropertySpecifierContext context = (UhtPropertySpecifierContext)specifierContext;
@@ -363,7 +334,6 @@ namespace EpicGames.UHT.Parsers
 		}
 
 		[UhtSpecifier(Extends = UhtTableNames.PropertyMember, ValueType = UhtSpecifierValueType.Legacy)]
-		[SuppressMessage("CodeQuality", "IDE0051:Remove unused private members", Justification = "Attribute accessed method")]
 		private static void BlueprintCallableSpecifier(UhtSpecifierContext specifierContext)
 		{
 			UhtPropertySpecifierContext context = (UhtPropertySpecifierContext)specifierContext;
@@ -371,7 +341,6 @@ namespace EpicGames.UHT.Parsers
 		}
 
 		[UhtSpecifier(Extends = UhtTableNames.PropertyMember, ValueType = UhtSpecifierValueType.Legacy)]
-		[SuppressMessage("CodeQuality", "IDE0051:Remove unused private members", Justification = "Attribute accessed method")]
 		private static void BlueprintAuthorityOnlySpecifier(UhtSpecifierContext specifierContext)
 		{
 			UhtPropertySpecifierContext context = (UhtPropertySpecifierContext)specifierContext;
@@ -379,7 +348,6 @@ namespace EpicGames.UHT.Parsers
 		}
 
 		[UhtSpecifier(Extends = UhtTableNames.PropertyMember, ValueType = UhtSpecifierValueType.Legacy)]
-		[SuppressMessage("CodeQuality", "IDE0051:Remove unused private members", Justification = "Attribute accessed method")]
 		private static void AssetRegistrySearchableSpecifier(UhtSpecifierContext specifierContext)
 		{
 			UhtPropertySpecifierContext context = (UhtPropertySpecifierContext)specifierContext;
@@ -387,7 +355,6 @@ namespace EpicGames.UHT.Parsers
 		}
 
 		[UhtSpecifier(Extends = UhtTableNames.PropertyMember, ValueType = UhtSpecifierValueType.Legacy)]
-		[SuppressMessage("CodeQuality", "IDE0051:Remove unused private members", Justification = "Attribute accessed method")]
 		private static void SimpleDisplaySpecifier(UhtSpecifierContext specifierContext)
 		{
 			UhtPropertySpecifierContext context = (UhtPropertySpecifierContext)specifierContext;
@@ -395,7 +362,6 @@ namespace EpicGames.UHT.Parsers
 		}
 
 		[UhtSpecifier(Extends = UhtTableNames.PropertyMember, ValueType = UhtSpecifierValueType.Legacy)]
-		[SuppressMessage("CodeQuality", "IDE0051:Remove unused private members", Justification = "Attribute accessed method")]
 		private static void AdvancedDisplaySpecifier(UhtSpecifierContext specifierContext)
 		{
 			UhtPropertySpecifierContext context = (UhtPropertySpecifierContext)specifierContext;
@@ -403,7 +369,6 @@ namespace EpicGames.UHT.Parsers
 		}
 
 		[UhtSpecifier(Extends = UhtTableNames.PropertyMember, ValueType = UhtSpecifierValueType.Legacy)]
-		[SuppressMessage("CodeQuality", "IDE0051:Remove unused private members", Justification = "Attribute accessed method")]
 		private static void SaveGameSpecifier(UhtSpecifierContext specifierContext)
 		{
 			UhtPropertySpecifierContext context = (UhtPropertySpecifierContext)specifierContext;
@@ -411,7 +376,6 @@ namespace EpicGames.UHT.Parsers
 		}
 
 		[UhtSpecifier(Extends = UhtTableNames.PropertyMember, ValueType = UhtSpecifierValueType.Legacy)]
-		[SuppressMessage("CodeQuality", "IDE0051:Remove unused private members", Justification = "Attribute accessed method")]
 		private static void SkipSerializationSpecifier(UhtSpecifierContext specifierContext)
 		{
 			UhtPropertySpecifierContext context = (UhtPropertySpecifierContext)specifierContext;
@@ -419,7 +383,6 @@ namespace EpicGames.UHT.Parsers
 		}
 
 		[UhtSpecifier(Extends = UhtTableNames.PropertyMember, ValueType = UhtSpecifierValueType.OptionalString)]
-		[SuppressMessage("CodeQuality", "IDE0051:Remove unused private members", Justification = "Attribute accessed method")]
 		private static void GetterSpecifier(UhtSpecifierContext specifierContext, StringView? value)
 		{
 			UhtPropertySpecifierContext context = (UhtPropertySpecifierContext)specifierContext;
@@ -446,7 +409,6 @@ namespace EpicGames.UHT.Parsers
 		}
 
 		[UhtSpecifier(Extends = UhtTableNames.PropertyMember, ValueType = UhtSpecifierValueType.OptionalString)]
-		[SuppressMessage("CodeQuality", "IDE0051:Remove unused private members", Justification = "Attribute accessed method")]
 		private static void SetterSpecifier(UhtSpecifierContext specifierContext, StringView? value)
 		{
 			UhtPropertySpecifierContext context = (UhtPropertySpecifierContext)specifierContext;
@@ -473,7 +435,6 @@ namespace EpicGames.UHT.Parsers
 		}
 
 		[UhtSpecifier(Extends = UhtTableNames.PropertyMember, ValueType = UhtSpecifierValueType.None)]
-		[SuppressMessage("CodeQuality", "IDE0051:Remove unused private members", Justification = "Attribute accessed method")]
 		private static void FieldNotifySpecifier(UhtSpecifierContext specifierContext)
 		{
 			UhtPropertySpecifierContext context = (UhtPropertySpecifierContext)specifierContext;
diff --git a/Engine/Source/Programs/Shared/EpicGames.UHT/Specifiers/UhtScriptStructSpecifiers.cs b/Engine/Source/Programs/Shared/EpicGames.UHT/Specifiers/UhtScriptStructSpecifiers.cs
--- a/Engine/Source/Programs/Shared/EpicGames.UHT/Specifiers/UhtScriptStructSpecifiers.cs
+++ b/Engine/Source/Programs/Shared/EpicGames.UHT/Specifiers/UhtScriptStructSpecifiers.cs
@@ -12,10 +12,10 @@ namespace EpicGames.UHT.Parsers
 	/// Collection of USCRIPT specifiers
 	/// </summary>
 	[UnrealHeaderTool]
+	[SuppressMessage("CodeQuality", "IDE0051:Remove unused private members", Justification = "Attribute accessed method")]
 	public static class UhtScriptStructSpecifiers
 	{
 		[UhtSpecifier(Extends = UhtTableNames.ScriptStruct, ValueType = UhtSpecifierValueType.Legacy)]
-		[SuppressMessage("CodeQuality", "IDE0051:Remove unused private members", Justification = "Attribute accessed method")]
 		private static void NoExportSpecifier(UhtSpecifierContext specifierContext)
 		{
 			UhtScriptStruct scriptStruct = (UhtScriptStruct)specifierContext.Scope.ScopeType;
@@ -24,7 +24,6 @@ namespace EpicGames.UHT.Parsers
 		}
 
 		[UhtSpecifier(Extends = UhtTableNames.ScriptStruct, ValueType = UhtSpecifierValueType.Legacy)]
-		[SuppressMessage("CodeQuality", "IDE0051:Remove unused private members", Justification = "Attribute accessed method")]
 		private static void AtomicSpecifier(UhtSpecifierContext specifierContext)
 		{
 			UhtScriptStruct scriptStruct = (UhtScriptStruct)specifierContext.Scope.ScopeType;
@@ -40,7 +39,6 @@ namespace EpicGames.UHT.Parsers
 		}
 
 		[UhtSpecifier(Extends = UhtTableNames.ScriptStruct, ValueType = UhtSpecifierValueType.Legacy)]
-		[SuppressMessage("CodeQuality", "IDE0051:Remove unused private members", Justification = "Attribute accessed method")]
 		private static void HasDefaultsSpecifier(UhtSpecifierContext specifierContext)
 		{
 			if (!specifierContext.Scope.HeaderParser.HeaderFile.IsNoExportTypes)
@@ -52,7 +50,6 @@ namespace EpicGames.UHT.Parsers
 		}
 
 		[UhtSpecifier(Extends = UhtTableNames.ScriptStruct, ValueType = UhtSpecifierValueType.Legacy)]
-		[SuppressMessage("CodeQuality", "IDE0051:Remove unused private members", Justification = "Attribute accessed method")]
 		private static void HasNoOpConstructorSpecifier(UhtSpecifierContext specifierContext)
 		{
 			if (!specifierContext.Scope.HeaderParser.HeaderFile.IsNoExportTypes)
@@ -64,7 +61,6 @@ namespace EpicGames.UHT.Parsers
 		}
 
 		[UhtSpecifier(Extends = UhtTableNames.ScriptStruct, ValueType = UhtSpecifierValueType.Legacy)]
-		[SuppressMessage("CodeQuality", "IDE0051:Remove unused private members", Justification = "Attribute accessed method")]
 		private static void IsAlwaysAccessibleSpecifier(UhtSpecifierContext specifierContext)
 		{
 			if (!specifierContext.Scope.HeaderParser.HeaderFile.IsNoExportTypes)
@@ -76,7 +72,6 @@ namespace EpicGames.UHT.Parsers
 		}
 
 		[UhtSpecifier(Extends = UhtTableNames.ScriptStruct, ValueType = UhtSpecifierValueType.Legacy)]
-		[SuppressMessage("CodeQuality", "IDE0051:Remove unused private members", Justification = "Attribute accessed method")]
 		private static void IsCoreTypeSpecifier(UhtSpecifierContext specifierContext)
 		{
 			if (!specifierContext.Scope.HeaderParser.HeaderFile.IsNoExportTypes)
diff --git a/Engine/Source/Programs/Shared/EpicGames.UHT/Types/UhtClass.cs b/Engine/Source/Programs/Shared/EpicGames.UHT/Types/UhtClass.cs
--- a/Engine/Source/Programs/Shared/EpicGames.UHT/Types/UhtClass.cs
+++ b/Engine/Source/Programs/Shared/EpicGames.UHT/Types/UhtClass.cs
@@ -356,6 +356,93 @@ namespace EpicGames.UHT.Types
 		/// </summary>
 		public bool HasGeneratedBody { get; set; } = false;
 
+		#region Class specifier parse time values (not interfaces)
+		/// <summary>
+		/// Engine class flags removed
+		/// </summary>
+		public EClassFlags RemovedClassFlags { get; set; } = EClassFlags.None;
+
+		/// <summary>
+		/// Class within identifier
+		/// </summary>
+		[JsonIgnore]
+		public string ClassWithinIdentifier { get; set; } = String.Empty;
+
+		/// <summary>
+		/// Collection of show categories
+		/// </summary>
+		[JsonIgnore]
+		public List<string> ShowCategories { get; } = new List<string>();
+
+		/// <summary>
+		/// Collection of hide categories
+		/// </summary>
+		[JsonIgnore]
+		public List<string> HideCategories { get; } = new List<string>();
+
+		/// <summary>
+		/// Collection of auto expand categories
+		/// </summary>
+		[JsonIgnore]
+		public List<string> AutoExpandCategories { get; } = new List<string>();
+
+		/// <summary>
+		/// Collection of auto collapse categories
+		/// </summary>
+		[JsonIgnore]
+		public List<string> AutoCollapseCategories { get; } = new List<string>();
+
+		/// <summary>
+		/// Collection of prioritize categories
+		/// </summary>
+		[JsonIgnore]
+		public List<string> PrioritizeCategories { get; } = new List<string>();
+
+		/// <summary>
+		/// Collection of show functions
+		/// </summary>
+		[JsonIgnore]
+		public List<string> ShowFunctions { get; } = new List<string>();
+
+		/// <summary>
+		/// Collection of hide functions
+		/// </summary>
+		[JsonIgnore]
+		public List<string> HideFunctions { get; } = new List<string>();
+
+		/// <summary>
+		/// Sparse class data types
+		/// </summary>
+		[JsonIgnore]
+		public List<string> SparseClassDataTypes { get; } = new List<string>();
+
+		/// <summary>
+		/// Class group names
+		/// </summary>
+		[JsonIgnore]
+		public List<string> ClassGroupNames { get; } = new List<string>();
+
+		/// <summary>
+		/// Add the given class flags
+		/// </summary>
+		/// <param name="flags">Flags to add</param>
+		public void AddClassFlags(EClassFlags flags)
+		{
+			this.ClassFlags |= flags;
+			this.RemovedClassFlags &= ~flags;
+		}
+
+		/// <summary>
+		/// Remove the given class flags
+		/// </summary>
+		/// <param name="flags">Flags to remove</param>
+		public void RemoveClassFlags(EClassFlags flags)
+		{
+			this.RemovedClassFlags |= flags;
+			this.ClassFlags &= ~flags;
+		}
+		#endregion
+
 		/// <inheritdoc/>
 		[JsonIgnore]
 		public override UhtEngineType EngineType
@@ -539,6 +626,112 @@ namespace EpicGames.UHT.Types
 			public bool Found { get; set; }
 		}
 
+		/// <inheritdoc/>
+		protected override void ResolveSuper(UhtResolvePhase resolvePhase)
+		{
+			base.ResolveSuper(resolvePhase);
+
+			// Make sure the class within is resolved.  We have to make sure we don't try to resolve ourselves.
+			if (this.ClassWithin != null && this.ClassWithin != this)
+			{
+				this.ClassWithin.Resolve(resolvePhase);
+			}
+
+			switch (resolvePhase)
+			{
+				case UhtResolvePhase.Bases:
+					BindAndResolveSuper(this.SuperIdentifier, UhtFindOptions.Class);
+					BindAndResolveBases(this.BaseIdentifiers, UhtFindOptions.Class);
+
+					// Force the MatchedSerializers on for anything being exported
+					if (!this.ClassExportFlags.HasAnyFlags(UhtClassExportFlags.NoExport))
+					{
+						this.ClassFlags |= EClassFlags.MatchedSerializers;
+					}
+
+					switch (this.ClassType)
+					{
+						case UhtClassType.Class:
+							{
+								UhtClass? superClass = this.SuperClass;
+
+								// Merge the super class flags
+								if (superClass != null)
+								{
+									this.ClassFlags |= superClass.ClassFlags & EClassFlags.ScriptInherit & ~this.RemovedClassFlags;
+								}
+
+								foreach (UhtStruct baseStruct in this.Bases)
+								{
+									if (baseStruct is UhtClass baseClass)
+									{
+										if (!baseClass.ClassFlags.HasAnyFlags(EClassFlags.Interface))
+										{
+											this.LogError($"Class '{baseClass.SourceName}' is not an interface; Can only inherit from non-UObjects or UInterface derived interfaces");
+										}
+										this.ClassFlags |= baseClass.ClassFlags & EClassFlags.ScriptInherit & ~this.RemovedClassFlags;
+									}
+								}
+
+								// These following collections only inherit from he parent if they are empty in this class
+								if (this.SparseClassDataTypes.Count == 0)
+								{
+									AppendStringListMetaData(this.SuperClass, UhtNames.SparseClassDataTypes, this.SparseClassDataTypes);
+								}
+								if (this.PrioritizeCategories.Count == 0)
+								{
+									AppendStringListMetaData(this.SuperClass, UhtNames.PrioritizeCategories, this.PrioritizeCategories);
+								}
+
+								// Merge with categories inherited from the parent.
+								MergeCategories();
+
+								SetAndValidateWithinClass(resolvePhase);
+								SetAndValidateConfigName();
+
+								// Copy all of the lists back to the meta data
+								this.MetaData.AddIfNotEmpty(UhtNames.ClassGroupNames, this.ClassGroupNames);
+								this.MetaData.AddIfNotEmpty(UhtNames.AutoCollapseCategories, this.AutoCollapseCategories);
+								this.MetaData.AddIfNotEmpty(UhtNames.AutoExpandCategories, this.AutoExpandCategories);
+								this.MetaData.AddIfNotEmpty(UhtNames.PrioritizeCategories, this.PrioritizeCategories);
+								this.MetaData.AddIfNotEmpty(UhtNames.HideCategories, this.HideCategories);
+								this.MetaData.AddIfNotEmpty(UhtNames.ShowCategories, this.ShowCategories);
+								this.MetaData.AddIfNotEmpty(UhtNames.SparseClassDataTypes, this.SparseClassDataTypes);
+								this.MetaData.AddIfNotEmpty(UhtNames.HideFunctions, this.HideFunctions);
+
+								this.MetaData.Add(UhtNames.IncludePath, this.HeaderFile.IncludeFilePath);
+							}
+							break;
+
+						case UhtClassType.Interface:
+							{
+								UhtClass? superClass = this.SuperClass;
+								if (superClass != null)
+								{
+									this.ClassFlags |= superClass.ClassFlags & EClassFlags.ScriptInherit;
+									if (!superClass.ClassFlags.HasAnyFlags(EClassFlags.Native))
+									{
+										throw new UhtException(this, $"Native classes cannot extend non-native classes");
+									}
+									this.ClassWithin = superClass.ClassWithin;
+								}
+								else
+								{
+									this.ClassWithin = this.Session.UObject;
+								}
+							}
+							break;
+
+						case UhtClassType.NativeInterface:
+							break;
+
+						default:
+							throw new UhtIceException("Unexpected class type");
+					}
+					break;
+			}
+		}
+
 		/// <inheritdoc/>
 		protected override bool ResolveSelf(UhtResolvePhase phase)
 		{
@@ -546,6 +739,64 @@ namespace EpicGames.UHT.Types
 
 			switch (phase)
 			{
+				case UhtResolvePhase.InvalidCheck:
+					switch (this.ClassType)
+					{
+						case UhtClassType.Class:
+							break;
+
+						case UhtClassType.Interface:
+							{
+								string nativeInterfaceName = "I" + this.EngineName;
+								UhtType? nativeInterface = this.Session.FindType(null, UhtFindOptions.SourceName | UhtFindOptions.Class, nativeInterfaceName);
+								if (nativeInterface == null)
+								{
+									this.LogError($"UInterface '{this.SourceName}' parsed without a corresponding '{nativeInterfaceName}'");
+								}
+								else
+								{
+									// Copy the children
+									this.AddChildren(nativeInterface.DetachChildren());
+								}
+							}
+							break;
+
+						case UhtClassType.NativeInterface:
+							{
+								string interfaceName = "U" + this.EngineName;
+								UhtClass? interfaceObj = (UhtClass?)this.Session.FindType(null, UhtFindOptions.SourceName | UhtFindOptions.Class, interfaceName);
+								if (interfaceObj == null)
+								{
+									if (this.HasGeneratedBody || this.Children.Count != 0)
+									{
+										this.LogError($"Native interface '{this.SourceName}' parsed without a corresponding '{interfaceName}'");
+									}
+									else
+									{
+										this.VisibleType = false;
+										result = false;
+									}
+								}
+								else
+								{
+									this.AlternateObject = interfaceObj;
+									interfaceObj.NativeInterface = this;
+									//COMPATIBILITY-TODO - Use the native interface access specifier
+									interfaceObj.GeneratedBodyAccessSpecifier = this.GeneratedBodyAccessSpecifier;
+
+									if (this.GeneratedBodyLineNumber == -1)
+									{
+										this.LogError("Expected a GENERATED_BODY() at the start of the native interface");
+									}
+								}
+							}
+							break;
+
+						default:
+							throw new UhtIceException("Unexpected class type");
+					}
+					break;
+
 				case UhtResolvePhase.Bases:
 
 					// Check to see if this class matches an entry in the ClassCastFlags.  If it does, that
@@ -723,6 +974,204 @@ namespace EpicGames.UHT.Types
 			return base.ScanForInstancedReferenced(deepScan);
 		}
 
+		private void MergeCategories()
+		{
+			MergeShowCategories();
+
+			// Merge ShowFunctions and HideFunctions
+			AppendStringListMetaData(this.SuperClass, UhtNames.HideFunctions, this.HideFunctions);
+			foreach (string value in this.ShowFunctions)
+			{
+				this.HideFunctions.RemoveSwap(value);
+			}
+			this.ShowFunctions.Clear();
+
+			// Merge AutoExpandCategories and AutoCollapseCategories (we still want to keep AutoExpandCategories though!)
+			List<string> parentAutoExpandCategories = GetStringListMetaData(this.SuperClass, UhtNames.AutoExpandCategories);
+			List<string> parentAutoCollapseCategories = GetStringListMetaData(this.SuperClass, UhtNames.AutoCollapseCategories);
+
+			foreach (string value in this.AutoExpandCategories)
+			{
+				this.AutoCollapseCategories.RemoveSwap(value);
+				parentAutoCollapseCategories.RemoveSwap(value);
+			}
+
+			// Do the same as above but the other way around
+			foreach (string value in this.AutoCollapseCategories)
+			{
+				this.AutoExpandCategories.RemoveSwap(value);
+				parentAutoExpandCategories.RemoveSwap(value);
+			}
+
+			// Once AutoExpandCategories and AutoCollapseCategories for THIS class have been parsed, add the parent inherited categories
+			this.AutoCollapseCategories.AddRange(parentAutoCollapseCategories);
+			this.AutoExpandCategories.AddRange(parentAutoExpandCategories);
+		}
+
+		private void MergeShowCategories()
+		{
+
+			// Add the super class hide categories and prime the output show categories with the parent
+			List<string> outShowCategories = GetStringListMetaData(this.SuperClass, UhtNames.ShowCategories);
+			AppendStringListMetaData(this.SuperClass, UhtNames.HideCategories, this.HideCategories);
+
+			// If this class has new show categories, then merge them
+			if (this.ShowCategories.Count != 0)
+			{
+				StringBuilder subCategoryPath = new();
+				foreach (string value in this.ShowCategories)
+				{
+
+					// if we didn't find this specific category path in the HideCategories metadata
+					if (!this.HideCategories.RemoveSwap(value))
+					{
+						string[] subCategories = value.ToString().Split('|', StringSplitOptions.RemoveEmptyEntries);
+
+						subCategoryPath.Clear();
+						// look to see if any of the parent paths are excluded in the HideCategories list
+						for (int categoryPathIndex = 0; categoryPathIndex < subCategories.Length - 1; ++categoryPathIndex)
+						{
+							subCategoryPath.Append(subCategories[categoryPathIndex]);
+							// if we're hiding a parent category, then we need to flag this sub category for show
+							if (HideCategories.Contains(subCategoryPath.ToString()))
+							{
+								outShowCategories.AddUnique(value);
+								break;
+							}
+							subCategoryPath.Append('|');
+						}
+					}
+				}
+			}
+
+			// Replace the show categories
+			this.ShowCategories.Clear();
+			this.ShowCategories.AddRange(outShowCategories);
+		}
+
+		private void SetAndValidateWithinClass(UhtResolvePhase resolvePhase)
+		{
+			// The class within must be a child of any super class within
+			UhtClass expectedClassWithin = this.SuperClass != null ? this.SuperClass.ClassWithin : this.Session.UObject;
+
+			// Process all of the class specifiers
+			if (!String.IsNullOrEmpty(this.ClassWithinIdentifier))
+			{
+				UhtClass? specifiedClassWithin = (UhtClass?)this.Session.FindType(null, UhtFindOptions.EngineName | UhtFindOptions.Class, this.ClassWithinIdentifier);
+				if (specifiedClassWithin == null)
+				{
+					this.LogError($"Within class '{this.ClassWithinIdentifier}' not found");
+				}
+				else
+				{
+
+					// Make sure the class within has been resolved to this phase.  We don't need to worry about the super since we know 
+					// that it has already been resolved.
+					if (specifiedClassWithin != this)
+					{
+						specifiedClassWithin.Resolve(resolvePhase);
+					}
+
+					if (specifiedClassWithin.IsChildOf(this.Session.UInterface))
+					{
+						this.LogError("Classes cannot be 'within' interfaces");
+					}
+					else if (!specifiedClassWithin.IsChildOf(expectedClassWithin))
+					{
+						this.LogError($"Cannot override within class with '{specifiedClassWithin.SourceName}' since it isn't a child of parent class's expected within '{expectedClassWithin.SourceName}'");
+					}
+					else
+					{
+						this.ClassWithin = specifiedClassWithin;
+					}
+				}
+			}
+
+			// If we don't have a class within set, then just use the expected within
+			else
+			{
+				this.ClassWithin = expectedClassWithin;
+			}
+		}
+
+		private void SetAndValidateConfigName()
+		{
+			// Since this flag is computed in this method, we have to re-propagate the flag from the super
+			// just in case they were defined in this source file.
+			if (this.SuperClass != null)
+			{
+				this.ClassFlags |= this.SuperClass.ClassFlags & EClassFlags.Config;
+			}
+
+			// Set the class config flag if any properties have config
+			if (!this.ClassFlags.HasAnyFlags(EClassFlags.Config))
+			{
+				foreach (UhtProperty property in this.Properties)
+				{
+					if (property.PropertyFlags.HasAnyFlags(EPropertyFlags.Config))
+					{
+						this.ClassFlags |= EClassFlags.Config;
+						break;
+					}
+				}
+			}
+
+			if (this.Config.Length > 0)
+			{
+				// if the user specified "inherit", we're just going to use the parent class's config filename
+				// this is not actually necessary but it can be useful for explicitly communicating config-ness
+				if (this.Config.Equals("inherit", StringComparison.OrdinalIgnoreCase))
+				{
+					if (this.SuperClass == null)
+					{
+						this.LogError($"Cannot inherit config filename for class '{this.SourceName}' when it has no super class");
+					}
+					else if (this.SuperClass.Config.Length == 0)
+					{
+						this.LogError($"Cannot inherit config filename for class '{this.SourceName}' when parent class '{this.SuperClass.SourceName}' has no config filename");
+					}
+					else
+					{
+						this.Config = this.SuperClass.Config;
+					}
+				}
+			}
+			else if (this.ClassFlags.HasAnyFlags(EClassFlags.Config) && this.SuperClass != null)
+			{
+				this.Config = this.SuperClass.Config;
+			}
+
+			if (this.ClassFlags.HasAnyFlags(EClassFlags.Config) && this.Config.Length == 0)
+			{
+				this.LogError("Classes with config / globalconfig member variables need to specify config file.");
+				this.Config = "Engine";
+			}
+		}
+
+		private static List<string> GetStringListMetaData(UhtType? type, string key)
+		{
+			List<string> outStrings = new();
+			AppendStringListMetaData(type, key, outStrings);
+			return outStrings;
+		}
+
+		private static void AppendStringListMetaData(UhtType? type, string key, List<string> stringList)
+		{
+			if (type != null)
+			{
+				string[]? values = type.MetaData.GetStringArray(key);
+				if (values != null)
+				{
+					foreach (string value in values)
+					{
+						//COMPATIBILITY-TODO - TEST - This preserves duplicates that are in old UHT
+						//StringList.AddUnique(Value);
+						stringList.Add(value);
+					}
+				}
+			}
+		}
+
 		/// <summary>
 		/// Given a token stream, verify that it matches the expected signature for a getter/setter
 		/// </summary>
diff --git a/Engine/Source/Programs/Shared/EpicGames.UHT/Types/UhtFunction.cs b/Engine/Source/Programs/Shared/EpicGames.UHT/Types/UhtFunction.cs
--- a/Engine/Source/Programs/Shared/EpicGames.UHT/Types/UhtFunction.cs
+++ b/Engine/Source/Programs/Shared/EpicGames.UHT/Types/UhtFunction.cs
@@ -2,10 +2,13 @@
 
 using System;
 using System.Collections.Generic;
+using System.Linq;
 using System.Text;
 using System.Text.Json.Serialization;
 using EpicGames.Core;
+using EpicGames.UHT.Parsers;
 using EpicGames.UHT.Tables;
+using EpicGames.UHT.Tokenizer;
 using EpicGames.UHT.Utils;
 
 namespace EpicGames.UHT.Types
@@ -108,6 +111,11 @@ namespace EpicGames.UHT.Types
 		/// Generate the entry for the FieldNotificationClassDescriptor
 		/// </summary>
 		FieldNotify = 1 << 16,
+
+		/// <summary>
+		/// True if the function specifier has a getter/setter specified
+		/// </summary>
+		SawPropertyAccessor = 1 << 17,
 	}
 
 	/// <summary>
@@ -412,6 +420,43 @@ namespace EpicGames.UHT.Types
 					}
 					break;
 
+				case UhtResolvePhase.Properties:
+					UhtPropertyParser.ResolveChildren(this, GetPropertyParseOptions(false));
+					foreach (UhtProperty property in this.Properties)
+					{
+						if (property.DefaultValueTokens != null)
+						{
+							string key = "CPP_Default_" + property.EngineName;
+							if (!this.MetaData.ContainsKey(key))
+							{
+								bool parsed = false;
+								try
+								{
+									// All tokens MUST be consumed from the reader
+									StringBuilder builder = new();
+									IUhtTokenReader defaultValueReader = UhtTokenReplayReader.GetThreadInstance(property, this.HeaderFile.Data.Memory, property.DefaultValueTokens.ToArray(), UhtTokenType.EndOfDefault);
+									parsed = property.SanitizeDefaultValue(defaultValueReader, builder) && defaultValueReader.IsEOF;
+									if (parsed)
+									{
+										this.MetaData.Add(key, builder.ToString());
+									}
+								}
+								catch (Exception)
+								{
+									// Ignore the exception for now
+								}
+
+								if (!parsed)
+								{
+									StringView defaultValueText = new(this.HeaderFile.Data, property.DefaultValueTokens.First().InputStartPos,
+										property.DefaultValueTokens.Last().InputEndPos - property.DefaultValueTokens.First().InputStartPos);
+									property.LogError($"C++ Default parameter not parsed: {property.SourceName} '{defaultValueText}'");
+								}
+							}
+						}
+					}
+					break;
+
 				case UhtResolvePhase.Final:
 					if (this.Outer is UhtClass classObj)
 					{
@@ -1052,5 +1097,33 @@ namespace EpicGames.UHT.Types
 				child.CollectReferences(collector);
 			}
 		}
+
+		/// <summary>
+		/// Return the property parse options for the given function
+		/// </summary>
+		/// <param name="returnValue">If true, return property parse options for the return value</param>
+		/// <returns>Property parse options</returns>
+		/// <exception cref="UhtIceException"></exception>
+		protected internal UhtPropertyParseOptions GetPropertyParseOptions(bool returnValue)
+		{
+			switch (this.FunctionType)
+			{
+				case UhtFunctionType.Delegate:
+				case UhtFunctionType.SparseDelegate:
+					return (returnValue ? UhtPropertyParseOptions.None : UhtPropertyParseOptions.CommaSeparatedName) | UhtPropertyParseOptions.DontAddReturn;
+
+				case UhtFunctionType.Function:
+					UhtPropertyParseOptions options = UhtPropertyParseOptions.DontAddReturn; // Fetch the function name
+					options |= returnValue ? UhtPropertyParseOptions.FunctionNameIncluded : UhtPropertyParseOptions.NameIncluded;
+					if (this.FunctionFlags.HasAllFlags(EFunctionFlags.BlueprintEvent | EFunctionFlags.Native))
+					{
+						options |= UhtPropertyParseOptions.NoAutoConst;
+					}
+					return options;
+
+				default:
+					throw new UhtIceException("Unknown enumeration value");
+			}
+		}
 	}
 }
diff --git a/Engine/Source/Programs/Shared/EpicGames.UHT/Types/UhtScriptStruct.cs b/Engine/Source/Programs/Shared/EpicGames.UHT/Types/UhtScriptStruct.cs
--- a/Engine/Source/Programs/Shared/EpicGames.UHT/Types/UhtScriptStruct.cs
+++ b/Engine/Source/Programs/Shared/EpicGames.UHT/Types/UhtScriptStruct.cs
@@ -622,6 +622,26 @@ namespace EpicGames.UHT.Types
 		{
 		}
 
+		#region Resolution support
+		/// <inheritdoc/>
+		protected override void ResolveSuper(UhtResolvePhase resolvePhase)
+		{
+			base.ResolveSuper(resolvePhase);
+
+			switch (resolvePhase)
+			{
+				case UhtResolvePhase.Bases:
+					BindAndResolveSuper(this.SuperIdentifier, UhtFindOptions.ScriptStruct);
+
+					// if we have a base struct, propagate inherited struct flags now
+					UhtScriptStruct? superScriptStruct = this.SuperScriptStruct;
+					if (superScriptStruct != null)
+					{
+						this.ScriptStructFlags |= superScriptStruct.ScriptStructFlags & EStructFlags.Inherit;
+					}
+					break;
+			}
+		}
 		/// <inheritdoc/>
 		protected override void ResolveChildren(UhtResolvePhase phase)
 		{
@@ -703,6 +723,7 @@ namespace EpicGames.UHT.Types
 
 			return base.ScanForInstancedReferenced(deepScan);
 		}
+		#endregion
 
 		#region Validation support
 		/// <inheritdoc/>
diff --git a/Engine/Source/Programs/Shared/EpicGames.UHT/Types/UhtStruct.cs b/Engine/Source/Programs/Shared/EpicGames.UHT/Types/UhtStruct.cs
--- a/Engine/Source/Programs/Shared/EpicGames.UHT/Types/UhtStruct.cs
+++ b/Engine/Source/Programs/Shared/EpicGames.UHT/Types/UhtStruct.cs
@@ -2,8 +2,10 @@
 
 using System;
 using System.Collections.Generic;
+using System.Diagnostics.CodeAnalysis;
 using System.Text.Json.Serialization;
 using EpicGames.Core;
+using EpicGames.UHT.Parsers;
 using EpicGames.UHT.Tables;
 using EpicGames.UHT.Tokenizer;
 using EpicGames.UHT.Utils;
@@ -88,6 +90,19 @@ namespace EpicGames.UHT.Types
 		[JsonIgnore]
 		public IReadOnlyList<UhtStruct> Bases => _bases ?? s_emptyBases;
 
+		/// <summary>
+		/// Super identifier
+		/// </summary>
+		[JsonIgnore]
+		public UhtToken SuperIdentifier { get; set; }
+
+		/// <summary>
+		/// Base identifiers
+		/// </summary>
+		[JsonIgnore]
+		[SuppressMessage("Usage", "CA2227:Collection properties should be read only", Justification = "<Pending>")]
+		public List<UhtToken[]>? BaseIdentifiers { get; set; } = null;
+
 		/// <summary>
 		/// Super struct type
 		/// </summary>
@@ -144,6 +159,20 @@ namespace EpicGames.UHT.Types
 			base.ResolveSuper(resolvePhase);
 		}
 
+		/// <inheritdoc/>
+		protected override bool ResolveSelf(UhtResolvePhase resolvePhase)
+		{
+			bool result = base.ResolveSelf(resolvePhase);
+
+			switch (resolvePhase)
+			{
+				case UhtResolvePhase.Properties:
+					UhtPropertyParser.ResolveChildren(this, UhtPropertyParseOptions.AddModuleRelativePath);
+					break;
+			}
+			return result;
+		}
+
 		/// <summary>
 		/// Check properties to see if any instances are referenced.
 		/// This method does NOT cache the result.
