Server Timing middleware for aspnet core

diff --git a/Engine/Source/Programs/Shared/EpicGames.AspNet/ServerTimings.cs b/Engine/Source/Programs/Shared/EpicGames.AspNet/ServerTimings.cs
new file mode 100644
--- /dev/null
+++ b/Engine/Source/Programs/Shared/EpicGames.AspNet/ServerTimings.cs
@@ -0,0 +1,166 @@
+// Copyright Epic Games, Inc. All Rights Reserved.
+
+using System;
+using System.Collections.Generic;
+using System.Globalization;
+using System.Linq;
+using System.Text;
+using System.Threading.Tasks;
+using Microsoft.AspNetCore.Http;
+using Microsoft.Extensions.DependencyInjection;
+
+namespace EpicGames.AspNet
+{
+    public struct ServerTimingMetric
+    {
+        private readonly string _metricName;
+        private readonly double? _duration;
+        private readonly string? _description;
+        private string? _serverTimingMetric;
+        
+        public ServerTimingMetric(string metricName, double? duration, string? description)
+        {
+            _metricName = metricName;
+            _duration = duration;
+            _description = description;
+
+            _serverTimingMetric = null;
+        }
+
+        public override string ToString()
+        {
+            if (_serverTimingMetric != null)
+            {
+                return _serverTimingMetric;
+            }
+
+            StringBuilder sb = new StringBuilder(_metricName);
+            if (_duration != null)
+            {
+                sb.Append(";dur=");
+                sb.Append(_duration.Value.ToString(CultureInfo.InvariantCulture));
+            }
+
+            if (!String.IsNullOrEmpty(_description))
+            {
+                sb.Append(";desc=\"");
+                sb.Append(_description);
+                sb.Append("\"");
+            }
+
+            _serverTimingMetric = sb.ToString();
+
+            return _serverTimingMetric;
+        }
+    }
+
+    public interface IServerTiming
+    {
+        public void AddServerTimingMetric(string metricName, double? duration, string? description);
+        public ServerTimingMetricScoped CreateServerTimingMetricScope(string metricName, string? description);
+
+        public IReadOnlyCollection<ServerTimingMetric> Metrics { get; }
+    }
+
+    public class ServerTimingMetricScoped : IDisposable
+    {
+        private readonly IServerTiming _timingManager;
+        private readonly string _metricName;
+        private readonly string? _description;
+        private readonly DateTime _startTime;
+
+        internal ServerTimingMetricScoped(IServerTiming timingManager, string metricName, string? description)
+        {
+            _timingManager = timingManager;
+            _metricName = metricName;
+            _description = description;
+            _startTime = DateTime.Now;
+        }
+        public void Dispose()
+        {
+            TimeSpan duration = DateTime.Now - _startTime;
+            _timingManager.AddServerTimingMetric(_metricName, duration.TotalMilliseconds, _description);
+        }
+    }
+
+    public class ServerTiming : IServerTiming
+    {
+        private readonly List<ServerTimingMetric> _metrics = new List<ServerTimingMetric>();
+
+        public void AddServerTimingMetric(string metricName, double? duration, string? description)
+        {
+            _metrics.Add(new ServerTimingMetric(metricName, duration, description));
+        }
+
+        public ServerTimingMetricScoped CreateServerTimingMetricScope(string metricName, string? description)
+        {
+            return new ServerTimingMetricScoped(this, metricName, description);
+        }
+
+        public IReadOnlyCollection<ServerTimingMetric> Metrics => _metrics;
+    }
+
+    public class ServerTimingMiddleware
+    {
+        private readonly RequestDelegate _next;
+
+        public ServerTimingMiddleware(RequestDelegate next)
+        {
+            _next = next ?? throw new ArgumentNullException(nameof(next));
+        }
+
+        public async Task Invoke(HttpContext context)
+        {
+            IServerTiming serverTiming = context.RequestServices.GetRequiredService<IServerTiming>();
+            if (AllowsTrailers(context.Request) && context.Response.SupportsTrailers())
+            {
+                await HandleServerTimingAsTrailerHeaderAsync(context, serverTiming);
+            }
+            else
+            {
+                await HandleServerTimingAsResponseHeadersAsync(context, serverTiming);
+            }
+        }
+
+        public static bool AllowsTrailers(HttpRequest request)
+        {
+	        return request.Headers.ContainsKey("TE") && request.Headers["TE"].Contains("trailers");
+        }
+
+		private async Task HandleServerTimingAsTrailerHeaderAsync(HttpContext context, IServerTiming serverTiming)
+        {
+            context.Response.DeclareTrailer("Server-Timing");
+
+            await _next(context);
+
+            string serverTimingValue = serverTiming.Metrics.Any() ? String.Join(",", serverTiming.Metrics) : "";
+            context.Response.AppendTrailer(
+                "Server-Timing",
+                serverTimingValue);
+        }
+
+        private Task HandleServerTimingAsResponseHeadersAsync(HttpContext context, IServerTiming serverTiming)
+        {
+            context.Response.OnStarting(() => {
+                if (serverTiming.Metrics.Any())
+                {
+                    string serverTimingValue = String.Join(",", serverTiming.Metrics);
+                    context.Response.Headers.Append("Server-Timing", serverTimingValue);
+                }
+
+                return Task.CompletedTask;
+            });
+
+            return _next(context);
+        }
+    }
+    public static class ServerTimingServiceCollectionExtensions
+    {
+        public static IServiceCollection AddServerTiming(this IServiceCollection services)
+        {
+            services.AddScoped<IServerTiming, ServerTiming>();
+
+            return services;
+        }
+    }
+}
