Horde Storage - Make sure we support server timings being added from seperate threads

diff --git a/Engine/Source/Programs/Shared/EpicGames.AspNet/ServerTimings.cs b/Engine/Source/Programs/Shared/EpicGames.AspNet/ServerTimings.cs
--- a/Engine/Source/Programs/Shared/EpicGames.AspNet/ServerTimings.cs
+++ b/Engine/Source/Programs/Shared/EpicGames.AspNet/ServerTimings.cs
@@ -89,8 +89,11 @@ namespace EpicGames.AspNet
 
         public void AddServerTimingMetric(string metricName, double? duration, string? description)
         {
-            _metrics.Add(new ServerTimingMetric(metricName, duration, description));
-        }
+	        lock (_metrics)
+	        {
+		        _metrics.Add(new ServerTimingMetric(metricName, duration, description));
+	        }
+		}
 
         public ServerTimingMetricScoped CreateServerTimingMetricScope(string metricName, string? description)
         {
