Ignore iTunes DLL if the version number cannot be parsed as we expect

diff --git a/Engine/Source/Programs/Shared/EpicGames.Build/Platform/Mac/ApplePlatformSDK.cs b/Engine/Source/Programs/Shared/EpicGames.Build/Platform/Mac/ApplePlatformSDK.cs
--- a/Engine/Source/Programs/Shared/EpicGames.Build/Platform/Mac/ApplePlatformSDK.cs
+++ b/Engine/Source/Programs/Shared/EpicGames.Build/Platform/Mac/ApplePlatformSDK.cs
@@ -15,6 +15,34 @@ namespace UnrealBuildBase
 	{
 		public static readonly string? InstalledSDKVersion = GetInstalledSDKVersion();
 
+		public static bool TryConvertVersionToInt(string? StringValue, out UInt64 OutValue)
+		{
+			OutValue = 0;
+
+			if (StringValue == null)
+			{
+				return false;
+			}
+
+			// 8 bits per component, with high getting extra from high 32
+			Match Result = Regex.Match(StringValue, @"^(\d+).(\d+)(.(\d+))?(.(\d+))?$");
+			if (Result.Success)
+			{
+				OutValue = UInt64.Parse(Result.Groups[1].Value) << 24 | UInt64.Parse(Result.Groups[2].Value) << 16;
+				if (Result.Groups[4].Success)
+				{
+					OutValue |= UInt64.Parse(Result.Groups[4].Value) << 8;
+				}
+				if (Result.Groups[6].Success)
+				{
+					OutValue |= UInt64.Parse(Result.Groups[6].Value) << 0;
+				}
+				return true;
+			}
+
+			return false;
+		}
+
 		private static string? GetInstalledSDKVersion()
 		{
 			// get xcode version on Mac
@@ -56,7 +84,12 @@ namespace UnrealBuildBase
 
 				if (!string.IsNullOrEmpty(DllPath) && File.Exists(DllPath))
 				{
-					return FileVersionInfo.GetVersionInfo(DllPath).FileVersion;
+					string? DllVersion = FileVersionInfo.GetVersionInfo(DllPath).FileVersion;
+					// Only return the DLL version as the SDK version if we can correctly parse it
+					if (TryConvertVersionToInt(DllVersion, out _))
+					{
+						return DllVersion;
+					}
 				}
 			}
 
diff --git a/Engine/Source/Programs/UnrealBuildTool/Platform/Mac/ApplePlatformSDK.cs b/Engine/Source/Programs/UnrealBuildTool/Platform/Mac/ApplePlatformSDK.cs
--- a/Engine/Source/Programs/UnrealBuildTool/Platform/Mac/ApplePlatformSDK.cs
+++ b/Engine/Source/Programs/UnrealBuildTool/Platform/Mac/ApplePlatformSDK.cs
@@ -20,30 +20,7 @@ namespace UnrealBuildTool
 	{
 		public override bool TryConvertVersionToInt(string? StringValue, out UInt64 OutValue, string? Hint)
 		{
-			OutValue = 0;
-
-			if (StringValue == null)
-			{
-				return false;
-			}
-
-			// 8 bits per component, with high getting extra from high 32
-			Match Result = Regex.Match(StringValue, @"^(\d+).(\d+)(.(\d+))?(.(\d+))?$");
-			if (Result.Success)
-			{
-				OutValue = UInt64.Parse(Result.Groups[1].Value) << 24 | UInt64.Parse(Result.Groups[2].Value) << 16;
-				if (Result.Groups[4].Success)
-				{
-					OutValue |= UInt64.Parse(Result.Groups[4].Value) << 8;
-				}
-				if (Result.Groups[6].Success)
-				{
-					OutValue |= UInt64.Parse(Result.Groups[6].Value) << 0;
-				}
-				return true;
-			}
-
-			return false;
+			return UnrealBuildBase.ApplePlatformSDK.TryConvertVersionToInt(StringValue, out OutValue);
 		}
 
 		protected override string? GetInstalledSDKVersion()
