- Added concept of NONE desired AutoSDK level. This can be used to turn off AutoSDK on a per-platform basis via envvar (UE_AUTOSDK_DEFAULT_LEVEL and UE_AUTOSDK_SPECIFIC_LEVELS are the important envvars here, FULL is the default DEFAULT)
#okforgitub public

diff --git a/Engine/Source/Programs/Shared/EpicGames.Build/UEBuildPlatformSDK.cs b/Engine/Source/Programs/Shared/EpicGames.Build/UEBuildPlatformSDK.cs
--- a/Engine/Source/Programs/Shared/EpicGames.Build/UEBuildPlatformSDK.cs
+++ b/Engine/Source/Programs/Shared/EpicGames.Build/UEBuildPlatformSDK.cs
@@ -1234,6 +1234,7 @@ namespace EpicGames.Core
 
 		static private string[] AutoSDKLevels = new string[]
 		{
+			"NONE",
 			"BUILD",
 			"PACKAGE",
 			"RUN",
@@ -1562,6 +1563,13 @@ namespace EpicGames.Core
 				string AutoSDKRoot = GetPathToPlatformAutoSDKs();
 				if (AutoSDKRoot != "")
 				{
+					string DesiredSDKLevel = GetAutoSDKLevelForPlatform(AutoSDKRoot);
+					// if the user doesn't want AutoSDK for this platform, then return that it is not installed, even if it actually is
+					if (string.Compare(DesiredSDKLevel, "NONE", true) == 0)
+					{
+						return SDKStatus.Invalid;
+					}
+
 					// check script version so script fixes can be propagated without touching every build machine's CurrentlyInstalled file manually.
 					string CurrentScriptVersionString;
 					if (GetLastRunScriptVersionString(AutoSDKRoot, out CurrentScriptVersionString) && CurrentScriptVersionString == GetRequiredScriptVersionString())
@@ -1577,7 +1585,7 @@ namespace EpicGames.Core
 							// match version
 							if (CurrentSDKString == GetAutoSDKDirectoryForMainVersion())
 							{
-								if (IsSDKLevelAtLeast(CurrentSDKLevel, GetAutoSDKLevelForPlatform(AutoSDKRoot)))
+								if (IsSDKLevelAtLeast(CurrentSDKLevel, DesiredSDKLevel))
 								{
 									return SDKStatus.Valid;
 								}
@@ -1671,10 +1679,19 @@ namespace EpicGames.Core
 				// run installation for autosdk if necessary.
 				if (HasRequiredAutoSDKInstalled() == SDKStatus.Invalid)
 				{
+					string AutoSDKRoot = GetPathToPlatformAutoSDKs();
+
+					string DesiredSDKLevel = GetAutoSDKLevelForPlatform(AutoSDKRoot);
+					// if the user doesn't want AutoSDK for this platform, then do nothing
+					if (string.Compare(DesiredSDKLevel, "NONE", true) == 0)
+					{
+						Log.TraceLog($"Skipping AutoSDK for {PlatformName} because NONE was specified as the desired AutoSDK level");
+						return;
+					}
+
 					//reset check status so any checking sdk status after the attempted setup will do a real check again.
 					SDKCheckStatus = -1;
 
-					string AutoSDKRoot = GetPathToPlatformAutoSDKs();
 					string CurrentSDKString;
 					string CurrentSDKLevel;
 					GetCurrentlyInstalledSDKString(AutoSDKRoot, out CurrentSDKString, out CurrentSDKLevel);
@@ -1689,7 +1706,6 @@ namespace EpicGames.Core
 					// delete Manifest file to avoid multiple uninstalls
 					InvalidateCurrentlyInstalledAutoSDK();
 
-					string DesiredSDKLevel = GetAutoSDKLevelForPlatform(AutoSDKRoot);
 					if (!RunAutoSDKHooks(AutoSDKRoot, GetAutoSDKDirectoryForMainVersion(), DesiredSDKLevel, SDKHookType.Install, false))
 					{
 						Log.TraceLog("Failed to install required SDK {0}.  Attemping to uninstall", GetAutoSDKDirectoryForMainVersion());
