Fixed issue where UPARAM in a function signature caused the getter/setter matching to fail for UFUNCTIONS.

#rnx

diff --git a/Engine/Source/Programs/Shared/EpicGames.UHT/Types/UhtClass.cs b/Engine/Source/Programs/Shared/EpicGames.UHT/Types/UhtClass.cs
--- a/Engine/Source/Programs/Shared/EpicGames.UHT/Types/UhtClass.cs
+++ b/Engine/Source/Programs/Shared/EpicGames.UHT/Types/UhtClass.cs
@@ -1220,6 +1220,34 @@ namespace EpicGames.UHT.Types
 				}
 				tokenIndex++;
 
+				// Skip any UPARAM only in actual functions
+				if (declaration.Function != null && tokenIndex < tokenCount && tokens[tokenIndex].IsIdentifier("UPARAM"))
+				{
+					tokenIndex++;
+					if (tokenIndex == tokenCount || !tokens[tokenIndex].IsValue('('))
+					{
+						return false;
+					}
+					tokenIndex++;
+					int paramCount = 1;
+					while (tokenIndex < tokenCount && paramCount > 0)
+					{
+						if (tokens[tokenIndex].IsValue(')'))
+						{
+							paramCount--;
+						}
+						else if (tokens[tokenIndex].IsValue('('))
+						{
+							paramCount++;
+						}
+						tokenIndex++;
+					}
+					if (paramCount > 0)
+					{
+						return false;
+					}
+				}
+
 				// Find the ')'
 				typeIndex = tokenIndex;
 				while (true)
