Fix issue with UFUNCTION validation with const objects.
Updated output test due to changes in NoExportTypes.h

#rnx

diff --git a/Engine/Source/Programs/Shared/EpicGames.UHT/Types/UhtClass.cs b/Engine/Source/Programs/Shared/EpicGames.UHT/Types/UhtClass.cs
--- a/Engine/Source/Programs/Shared/EpicGames.UHT/Types/UhtClass.cs
+++ b/Engine/Source/Programs/Shared/EpicGames.UHT/Types/UhtClass.cs
@@ -829,6 +829,10 @@ namespace EpicGames.UHT.Types
 					}
 					break;
 
+				case UhtResolvePhase.Properties:
+					UhtPropertyParser.ResolveChildren(this, UhtPropertyParseOptions.AddModuleRelativePath);
+					break;
+
 				case UhtResolvePhase.Final:
 					Dictionary<string, List<GetterSetterToResolve>>? gsToResolve = null;
 					foreach (UhtType child in this.Children)
diff --git a/Engine/Source/Programs/Shared/EpicGames.UHT/Types/UhtScriptStruct.cs b/Engine/Source/Programs/Shared/EpicGames.UHT/Types/UhtScriptStruct.cs
--- a/Engine/Source/Programs/Shared/EpicGames.UHT/Types/UhtScriptStruct.cs
+++ b/Engine/Source/Programs/Shared/EpicGames.UHT/Types/UhtScriptStruct.cs
@@ -4,6 +4,7 @@ using System;
 using System.Collections.Generic;
 using System.Text.Json.Serialization;
 using EpicGames.Core;
+using EpicGames.UHT.Parsers;
 using EpicGames.UHT.Tables;
 using EpicGames.UHT.Utils;
 
@@ -642,6 +643,21 @@ namespace EpicGames.UHT.Types
 					break;
 			}
 		}
+
+		/// <inheritdoc/>
+		protected override bool ResolveSelf(UhtResolvePhase resolvePhase)
+		{
+			bool result = base.ResolveSelf(resolvePhase);
+
+			switch (resolvePhase)
+			{
+				case UhtResolvePhase.Properties:
+					UhtPropertyParser.ResolveChildren(this, UhtPropertyParseOptions.AddModuleRelativePath);
+					break;
+			}
+			return result;
+		}
+
 		/// <inheritdoc/>
 		protected override void ResolveChildren(UhtResolvePhase phase)
 		{
diff --git a/Engine/Source/Programs/Shared/EpicGames.UHT/Types/UhtStruct.cs b/Engine/Source/Programs/Shared/EpicGames.UHT/Types/UhtStruct.cs
--- a/Engine/Source/Programs/Shared/EpicGames.UHT/Types/UhtStruct.cs
+++ b/Engine/Source/Programs/Shared/EpicGames.UHT/Types/UhtStruct.cs
@@ -159,20 +159,6 @@ namespace EpicGames.UHT.Types
 			base.ResolveSuper(resolvePhase);
 		}
 
-		/// <inheritdoc/>
-		protected override bool ResolveSelf(UhtResolvePhase resolvePhase)
-		{
-			bool result = base.ResolveSelf(resolvePhase);
-
-			switch (resolvePhase)
-			{
-				case UhtResolvePhase.Properties:
-					UhtPropertyParser.ResolveChildren(this, UhtPropertyParseOptions.AddModuleRelativePath);
-					break;
-			}
-			return result;
-		}
-
 		/// <summary>
 		/// Check properties to see if any instances are referenced.
 		/// This method does NOT cache the result.
