Remove extension methods to slice IList<T> types. Many types which are not naturally lists implement this interface.

New AsSegment() method allows creating a segment, while explicit extension methods now exist for explicitly slicing List<T>.

diff --git a/Engine/Source/Programs/Shared/EpicGames.Core/ListSegment.cs b/Engine/Source/Programs/Shared/EpicGames.Core/ListSegment.cs
--- a/Engine/Source/Programs/Shared/EpicGames.Core/ListSegment.cs
+++ b/Engine/Source/Programs/Shared/EpicGames.Core/ListSegment.cs
@@ -42,8 +42,10 @@ namespace EpicGames.Core
 		/// all the elements in the specified array.
 		/// </summary>
 		public ListSegment(IList<T> list)
-			: this(list, 0, list.Count)
 		{
+			List = list;
+			Offset = 0;
+			Count = list.Count;
 		}
 
 		/// <summary>
@@ -51,21 +53,64 @@ namespace EpicGames.Core
 		/// </summary>
 		/// <param name="list"></param>
 		/// <param name="offset"></param>
-		/// <param name="count"></param>
-		public ListSegment(IList<T> list, int offset, int count)
+		/// <param name="length"></param>
+		public ListSegment(IList<T> list, int offset, int length)
 		{
 			if (offset < 0 || offset > list.Count)
 			{
 				throw new ArgumentOutOfRangeException(nameof(offset));
 			}
-			if (count < 0 || offset + count > list.Count)
+			if (length < 0 || offset + length > list.Count)
 			{
-				throw new ArgumentOutOfRangeException(nameof(count));
+				throw new ArgumentOutOfRangeException(nameof(length));
 			}
 
 			List = list;
 			Offset = offset;
-			Count = count;
+			Count = length;
+		}
+
+		/// <summary>
+		/// Create a new slice of this list segment
+		/// </summary>
+		/// <param name="offset">Offset for the start of the segment</param>
+		public ListSegment<T> Slice(int offset)
+		{
+			if(offset < 0)
+			{
+				throw new ArgumentException("Offset may not be negative", nameof(offset));
+			}
+			if(offset > Count)
+			{
+				throw new ArgumentException("Offset may not exceed the length of the existing segment", nameof(offset));
+			}
+			return new ListSegment<T>(List, Offset + offset, Count - offset);
+		}
+
+		/// <summary>
+		/// Create a new slice of this list segment
+		/// </summary>
+		/// <param name="offset">Offset for the start of the segment</param>
+		/// <param name="count">Number of elements to include</param>
+		public ListSegment<T> Slice(int offset, int count)
+		{
+			if(offset < 0)
+			{
+				throw new ArgumentException("Offset may not be negative", nameof(offset));
+			}
+			if(offset > Count)
+			{
+				throw new ArgumentException("Offset may not exceed the length of the existing segment", nameof(offset));
+			}
+			if(count < 0)
+			{
+				throw new ArgumentException("Count may not be negative", nameof(count));
+			}
+			if(offset + count > Count)
+			{
+				throw new ArgumentException("Offset plus count may not exceed the length of the segment", nameof(count));
+			}
+			return new ListSegment<T>(List, Offset + offset, count);
 		}
 
 		/// <summary>
@@ -167,28 +212,77 @@ namespace EpicGames.Core
 	public static class ListSegmentExtensions
 	{
 		/// <summary>
-		/// Create a list segment
+		/// Create a fixed size list segment
 		/// </summary>
-		/// <typeparam name="T"></typeparam>
-		/// <param name="list"></param>
-		/// <param name="offset"></param>
-		/// <returns></returns>
-		public static ListSegment<T> Slice<T>(this IList<T> list, int offset)
+		/// <typeparam name="T">List element type</typeparam>
+		/// <param name="list">Base list to turn into a segment</param>
+		/// <returns>New segment for the list</returns>
+		public static ListSegment<T> AsSegment<T>(this IList<T> list)
+		{
+			return new ListSegment<T>(list);
+		}
+
+		/// <summary>
+		/// Create a fixed size list segment
+		/// </summary>
+		/// <typeparam name="T">List element type</typeparam>
+		/// <param name="list">Base list to turn into a segment</param>
+		/// <param name="offset">Starting position within the list</param>
+		/// <returns>New segment for the list</returns>
+		public static ListSegment<T> AsSegment<T>(this IList<T> list, int offset)
 		{
 			return new ListSegment<T>(list, offset, list.Count - offset);
 		}
 
 		/// <summary>
-		/// Create a list segment
+		/// Create a fixed size list segment
 		/// </summary>
-		/// <typeparam name="T"></typeparam>
-		/// <param name="list"></param>
-		/// <param name="offset"></param>
-		/// <param name="count"></param>
-		/// <returns></returns>
-		public static ListSegment<T> Slice<T>(this IList<T> list, int offset, int count)
+		/// <typeparam name="T">List element type</typeparam>
+		/// <param name="list">Base list to turn into a segment</param>
+		/// <param name="offset">Starting position within the list</param>
+		/// <param name="length">Length of the list</param>
+		/// <returns>New segment for the list</returns>
+		public static ListSegment<T> AsSegment<T>(this IList<T> list, int offset, int length)
+		{
+			return new ListSegment<T>(list, offset, length);
+		}
+
+		/// <summary>
+		/// Create a fixed size list segment
+		/// </summary>
+		/// <typeparam name="T">List element type</typeparam>
+		/// <param name="list">Base list to turn into a segment</param>
+		/// <param name="range">Range to create the segment from</param>
+		/// <returns>New segment for the list</returns>
+		public static ListSegment<T> AsSegment<T>(this IList<T> list, Range range)
+		{
+			(int offset, int length) = range.GetOffsetAndLength(list.Count);
+			return new ListSegment<T>(list, offset, length);
+		}
+
+		/// <summary>
+		/// Take a slice from a list
+		/// </summary>
+		/// <typeparam name="T">List element type</typeparam>
+		/// <param name="list">Base list to turn into a segment</param>
+		/// <param name="offset">Starting position within the list</param>
+		/// <returns>New segment for the list</returns>
+		public static ListSegment<T> Slice<T>(this List<T> list, int offset)
+		{
+			return new ListSegment<T>(list, offset, list.Count - offset);
+		}
+
+		/// <summary>
+		/// Take a slice from a list
+		/// </summary>
+		/// <typeparam name="T">List element type</typeparam>
+		/// <param name="list">Base list to turn into a segment</param>
+		/// <param name="offset">Starting position within the list</param>
+		/// <param name="length"></param>
+		/// <returns>New segment for the list</returns>
+		public static ListSegment<T> Slice<T>(this List<T> list, int offset, int length)
 		{
-			return new ListSegment<T>(list, offset, count);
+			return new ListSegment<T>(list, offset, length);
 		}
 	}
 }
diff --git a/Engine/Source/Programs/Shared/EpicGames.Horde/Bundles/Bundle.cs b/Engine/Source/Programs/Shared/EpicGames.Horde/Bundles/Bundle.cs
--- a/Engine/Source/Programs/Shared/EpicGames.Horde/Bundles/Bundle.cs
+++ b/Engine/Source/Programs/Shared/EpicGames.Horde/Bundles/Bundle.cs
@@ -850,7 +850,7 @@ namespace EpicGames.Horde.Bundles
 				nextBlobCost += writeNodes[maxIdx].Length;
 				if (nextBlobCost > Options.MaxBlobSize)
 				{
-					await WriteObjectAsync(writeNodes.Slice(minIdx, maxIdx - minIdx));
+					await WriteObjectAsync(writeNodes.AsSegment(minIdx, maxIdx - minIdx));
 					minIdx = maxIdx;
 					nextBlobCost = writeNodes[maxIdx].Length;
 				}
@@ -858,12 +858,12 @@ namespace EpicGames.Horde.Bundles
 
 			if (nextBlobCost > Options.MaxInlineBlobSize && minIdx + 1 < writeNodes.Length)
 			{
-				await WriteObjectAsync(writeNodes.Slice(minIdx, writeNodes.Length - 1 - minIdx));
+				await WriteObjectAsync(writeNodes.AsSegment(minIdx, writeNodes.Length - 1 - minIdx));
 				minIdx = writeNodes.Length - 1;
 			}
 
 			// Write the final ref
-			await WriteRefAsync(writeNodes.Slice(minIdx), bucketId, refId, metadata);
+			await WriteRefAsync(writeNodes.AsSegment(minIdx), bucketId, refId, metadata);
 
 			// Copy the stats over
 			Stats = _nextStats;
