Updated error handling as per change in C++ UHT.

#rnx

diff --git a/Engine/Source/Programs/Shared/EpicGames.UHT/Types/UhtClass.cs b/Engine/Source/Programs/Shared/EpicGames.UHT/Types/UhtClass.cs
--- a/Engine/Source/Programs/Shared/EpicGames.UHT/Types/UhtClass.cs
+++ b/Engine/Source/Programs/Shared/EpicGames.UHT/Types/UhtClass.cs
@@ -1516,9 +1516,11 @@ namespace EpicGames.UHT.Types
 								continue;
 							}
 
+							// Get the actual interface
+							UhtClass interfaceClass = baseClass.AlternateObject != null ? (UhtClass)baseClass.AlternateObject : baseClass;
+
 							// Loop through the function to make sure they are implemented
-							List<UhtType> baseChildren = baseClass.AlternateObject != null ? baseClass.AlternateObject.Children : baseClass.Children;
-							foreach (UhtType baseChild in baseChildren)
+							foreach (UhtType baseChild in interfaceClass.Children)
 							{
 								UhtFunction? baseFunction = baseChild as UhtFunction;
 								if (baseFunction == null)
@@ -1578,8 +1580,7 @@ namespace EpicGames.UHT.Types
 								if (!implemented
 									&& baseFunction.FunctionFlags.HasAnyFlags(EFunctionFlags.BlueprintCallable)
 									&& !baseFunction.FunctionFlags.HasAnyFlags(EFunctionFlags.BlueprintEvent)
-									&& !baseClass.MetaData.ContainsKey(UhtNames.CannotImplementInterfaceInBlueprint)
-									&& (baseClass.AlternateObject == null || !baseClass.AlternateObject.MetaData.ContainsKey(UhtNames.CannotImplementInterfaceInBlueprint)))
+									&& !interfaceClass.CannotImplementInBlueprints())
 								{
 									this.LogError($"Missing UFunction implementation of function '{baseFunction.SourceName}' from interface '{baseClass.SourceName}'.  This function needs a UFUNCTION() declaration.");
 								}
@@ -1602,11 +1603,8 @@ namespace EpicGames.UHT.Types
 					{
 						// Interface with blueprint data should declare explicitly Blueprintable or NotBlueprintable to be clear
 						// In the backward compatible case where they declare neither, both of these bools are false
-						bool canImplementInBlueprints = this.MetaData.GetBoolean(UhtNames.IsBlueprintBase);
-						bool cannotImplementInBlueprints = (!canImplementInBlueprints && this.MetaData.ContainsKey(UhtNames.IsBlueprintBase))
-							|| this.MetaData.ContainsKey(UhtNames.CannotImplementInterfaceInBlueprint);
-
-						//bool canImplementInterfaceInBlueprint = !this.MetaData.ContainsKey(UhtNames.CannotImplementInterfaceInBlueprint);
+						bool canImplementInBlueprints = CanImplementInBlueprints();
+						bool cannotImplementInBlueprints = CannotImplementInBlueprints(canImplementInBlueprints);
 						foreach (UhtType child in this.Children)
 						{
 							if (child is UhtFunction function)
@@ -1829,6 +1827,38 @@ namespace EpicGames.UHT.Types
 				}
 			}
 		}
+
+		/// <summary>
+		/// Interface with blueprint data should declare explicitly Blueprintable or NotBlueprintable to be clear
+		/// In the backward compatible case where they declare neither, both of these bools are false
+		/// </summary>
+		/// <returns>True if the interface is marked as Blueprintable</returns>
+		private bool CanImplementInBlueprints()
+		{
+			return this.MetaData.GetBoolean(UhtNames.IsBlueprintBase);
+		}
+
+		/// <summary>
+		/// Interface with blueprint data should declare explicitly Blueprintable or NotBlueprintable to be clear
+		/// In the backward compatible case where they declare neither, both of these bools are false
+		/// </summary>
+		/// <returns>True if the interface is marked as NotBlueprintable</returns>
+		private bool CannotImplementInBlueprints()
+		{
+			return CannotImplementInBlueprints(CanImplementInBlueprints());
+		}
+
+		/// <summary>
+		/// Interface with blueprint data should declare explicitly Blueprintable or NotBlueprintable to be clear
+		/// In the backward compatible case where they declare neither, both of these bools are false
+		/// </summary>
+		/// <param name="canImplementInBlueprints">If true, class has already been marked as Blueprintable</param>
+		/// <returns>True if the interface is marked as NotBlueprintable</returns>
+		private bool CannotImplementInBlueprints(bool canImplementInBlueprints)
+		{
+			return (!canImplementInBlueprints && this.MetaData.ContainsKey(UhtNames.IsBlueprintBase))
+				|| this.MetaData.ContainsKey(UhtNames.CannotImplementInterfaceInBlueprint);
+		}
 		#endregion
 
 		/// <inheritdoc/>
