Minor optimizations by reducing allocations (removed over 300k allocations)

#rnx

diff --git a/Engine/Source/Programs/Shared/EpicGames.UHT/Parsers/UhtClassParser.cs b/Engine/Source/Programs/Shared/EpicGames.UHT/Parsers/UhtClassParser.cs
--- a/Engine/Source/Programs/Shared/EpicGames.UHT/Parsers/UhtClassParser.cs
+++ b/Engine/Source/Programs/Shared/EpicGames.UHT/Parsers/UhtClassParser.cs
@@ -64,7 +64,7 @@ namespace EpicGames.UHT.Parsers
 					using UhtMessageContext tokenContext = new(ScopeName);
 					// Parse the specifiers
 					UhtSpecifierContext specifierContext = new(topScope, topScope.TokenReader, classObj.MetaData);
-					UhtSpecifierParser specifiers = topScope.HeaderParser.GetCachedSpecifierParser(specifierContext, ScopeName, parentScope.Session.GetSpecifierTable(UhtTableNames.Class));
+					UhtSpecifierParser specifiers = UhtSpecifierParser.GetThreadInstance(specifierContext, ScopeName, parentScope.Session.GetSpecifierTable(UhtTableNames.Class));
 					specifiers.ParseSpecifiers();
 					classObj.PrologLineNumber = topScope.TokenReader.InputLine;
 					classObj.ClassFlags |= EClassFlags.Native;
diff --git a/Engine/Source/Programs/Shared/EpicGames.UHT/Parsers/UhtEnumParser.cs b/Engine/Source/Programs/Shared/EpicGames.UHT/Parsers/UhtEnumParser.cs
--- a/Engine/Source/Programs/Shared/EpicGames.UHT/Parsers/UhtEnumParser.cs
+++ b/Engine/Source/Programs/Shared/EpicGames.UHT/Parsers/UhtEnumParser.cs
@@ -37,7 +37,7 @@ namespace EpicGames.UHT.Parsers
 
 					// Parse the specifiers
 					UhtSpecifierContext specifierContext = new(topScope, topScope.TokenReader, enumObject.MetaData);
-					UhtSpecifierParser specifiers = topScope.HeaderParser.GetCachedSpecifierParser(specifierContext, ScopeName, parentScope.Session.GetSpecifierTable(UhtTableNames.Enum));
+					UhtSpecifierParser specifiers = UhtSpecifierParser.GetThreadInstance(specifierContext, ScopeName, parentScope.Session.GetSpecifierTable(UhtTableNames.Enum));
 					specifiers.ParseSpecifiers();
 
 					// Read the name and the CPP type
diff --git a/Engine/Source/Programs/Shared/EpicGames.UHT/Parsers/UhtFunctionParser.cs b/Engine/Source/Programs/Shared/EpicGames.UHT/Parsers/UhtFunctionParser.cs
--- a/Engine/Source/Programs/Shared/EpicGames.UHT/Parsers/UhtFunctionParser.cs
+++ b/Engine/Source/Programs/Shared/EpicGames.UHT/Parsers/UhtFunctionParser.cs
@@ -140,7 +140,7 @@ namespace EpicGames.UHT.Parsers
 					topScope.AddModuleRelativePathToMetaData();
 
 					UhtSpecifierContext specifierContext = new(topScope, topScope.TokenReader, function.MetaData);
-					UhtSpecifierParser specifiers = topScope.HeaderParser.GetCachedSpecifierParser(specifierContext, ScopeName, parentScope.Session.GetSpecifierTable(UhtTableNames.Function));
+					UhtSpecifierParser specifiers = UhtSpecifierParser.GetThreadInstance(specifierContext, ScopeName, parentScope.Session.GetSpecifierTable(UhtTableNames.Function));
 
 					// If this is a UDELEGATE, parse the specifiers
 					StringView delegateMacro = new();
@@ -312,7 +312,7 @@ namespace EpicGames.UHT.Parsers
 					topScope.AddModuleRelativePathToMetaData();
 
 					UhtSpecifierContext specifierContext = new(topScope, topScope.TokenReader, function.MetaData);
-					UhtSpecifierParser specifierParser = topScope.HeaderParser.GetCachedSpecifierParser(specifierContext, scopeName, parentScope.Session.GetSpecifierTable(UhtTableNames.Function));
+					UhtSpecifierParser specifierParser = UhtSpecifierParser.GetThreadInstance(specifierContext, scopeName, parentScope.Session.GetSpecifierTable(UhtTableNames.Function));
 					specifierParser.ParseSpecifiers();
 
 					if (!outerClass.ClassFlags.HasAnyFlags(EClassFlags.Native))
diff --git a/Engine/Source/Programs/Shared/EpicGames.UHT/Parsers/UhtHeaderFileParser.cs b/Engine/Source/Programs/Shared/EpicGames.UHT/Parsers/UhtHeaderFileParser.cs
--- a/Engine/Source/Programs/Shared/EpicGames.UHT/Parsers/UhtHeaderFileParser.cs
+++ b/Engine/Source/Programs/Shared/EpicGames.UHT/Parsers/UhtHeaderFileParser.cs
@@ -208,12 +208,6 @@ namespace EpicGames.UHT.Parsers
 		/// </summary>
 		public UhtPropertyParser? PropertyParser { get; set; } = null;
 
-		/// <summary>
-		/// For a given header file, we share a common specifier parser to reduce the number of allocations.
-		/// Before the parser can be reused, the ParseDeferred method must be called to dispatch that list.
-		/// </summary>
-		private UhtSpecifierParser? _specifierParser = null;
-
 		/// <summary>
 		/// Stack of current #if states
 		/// </summary>
@@ -298,26 +292,6 @@ namespace EpicGames.UHT.Parsers
 			this._topScope = scope.ParentScope;
 		}
 
-		/// <summary>
-		/// Get the cached specifier parser
-		/// </summary>
-		/// <param name="specifierContext">Specifier context</param>
-		/// <param name="context">User facing context</param>
-		/// <param name="table">Specifier table</param>
-		/// <returns>Specifier parser</returns>
-		public UhtSpecifierParser GetCachedSpecifierParser(UhtSpecifierContext specifierContext, StringView context, UhtSpecifierTable table)
-		{
-			if (this._specifierParser == null)
-			{
-				this._specifierParser = new UhtSpecifierParser(specifierContext, context, table);
-			}
-			else
-			{
-				this._specifierParser.Reset(specifierContext, context, table);
-			}
-			return this._specifierParser;
-		}
-
 		/// <summary>
 		/// Get the cached property parser
 		/// </summary>
diff --git a/Engine/Source/Programs/Shared/EpicGames.UHT/Parsers/UhtInterfaceClassParser.cs b/Engine/Source/Programs/Shared/EpicGames.UHT/Parsers/UhtInterfaceClassParser.cs
--- a/Engine/Source/Programs/Shared/EpicGames.UHT/Parsers/UhtInterfaceClassParser.cs
+++ b/Engine/Source/Programs/Shared/EpicGames.UHT/Parsers/UhtInterfaceClassParser.cs
@@ -61,7 +61,7 @@ namespace EpicGames.UHT.Parsers
 
 					// Parse the specifiers
 					UhtSpecifierContext specifierContext = new(topScope, topScope.TokenReader, classObj.MetaData);
-					UhtSpecifierParser specifiers = topScope.HeaderParser.GetCachedSpecifierParser(specifierContext, ScopeName, parentScope.Session.GetSpecifierTable(UhtTableNames.Interface));
+					UhtSpecifierParser specifiers = UhtSpecifierParser.GetThreadInstance(specifierContext, ScopeName, parentScope.Session.GetSpecifierTable(UhtTableNames.Interface));
 					specifiers.ParseSpecifiers();
 					classObj.PrologLineNumber = topScope.TokenReader.InputLine;
 					classObj.ClassFlags |= EClassFlags.Native | EClassFlags.Interface | EClassFlags.Abstract;
diff --git a/Engine/Source/Programs/Shared/EpicGames.UHT/Parsers/UhtParsingScope.cs b/Engine/Source/Programs/Shared/EpicGames.UHT/Parsers/UhtParsingScope.cs
--- a/Engine/Source/Programs/Shared/EpicGames.UHT/Parsers/UhtParsingScope.cs
+++ b/Engine/Source/Programs/Shared/EpicGames.UHT/Parsers/UhtParsingScope.cs
@@ -1,6 +1,7 @@
 // Copyright Epic Games, Inc. All Rights Reserved.
 
 using System;
+using System.Buffers;
 using System.Text;
 using EpicGames.Core;
 using EpicGames.UHT.Tables;
@@ -177,17 +178,17 @@ namespace EpicGames.UHT.Parsers
 			ReadOnlySpan<StringView> comments = this.TokenReader.Comments;
 
 			// If we don't have any comments, just return
-			string mergedString = String.Empty;
 			if (comments.Length == 0)
 			{
 				return;
 			}
 
 			// Set the comment as just a simple concatenation of all the strings
+			string mergedString = String.Empty;
 			if (comments.Length == 1)
 			{
 				mergedString = comments[0].ToString();
-				type.MetaData.Add(UhtNames.Comment, metaNameIndex, comments[0].ToString());
+				type.MetaData.Add(UhtNames.Comment, metaNameIndex, mergedString);
 			}
 			else
 			{
@@ -202,10 +203,10 @@ namespace EpicGames.UHT.Parsers
 			}
 
 			// Format the tooltip and set the metadata
-			StringView toolTip = FormatCommentForToolTip(mergedString);
-			if (toolTip.Span.Length > 0)
+			string toolTip = FormatCommentForToolTip(mergedString);
+			if (!String.IsNullOrEmpty(toolTip))
 			{
-				type.MetaData.Add(UhtNames.ToolTip, metaNameIndex, toolTip.ToString());
+				type.MetaData.Add(UhtNames.ToolTip, metaNameIndex, toolTip);
 
 				//COMPATIBILITY-TODO - Old UHT would only clear the comments if there was some form of a tooltip
 				this.TokenReader.ClearComments();
@@ -252,97 +253,96 @@ namespace EpicGames.UHT.Parsers
 		/// </summary>
 		/// <param name="comments">Comments to be parsed</param>
 		/// <returns>The generated tooltip</returns>
-		private static StringView FormatCommentForToolTip(string comments)
+		private static string FormatCommentForToolTip(string comments)
 		{
 			if (!HasValidCommentChar(comments))
 			{
-				return new StringView();
+				return String.Empty;
 			}
 
-			string result = comments;
+			// Use the scratch characters to store the string as we process it in MANY passes
+			char[] scratchChars = ArrayPool<char>.Shared.Rent(comments.Length);
+			comments.CopyTo(0, scratchChars, 0, comments.Length);
+
+			// Remove ignore comments and strip the comment markers:
+			// These are all issues with how the old UHT worked.
+			// 1) Must be done in order
+			// 2) Block comment markers are removed first so that '///**' is process by removing '/**' first and then '//' second
+			// 3) We only remove block comments if we find the start of a comment.  This means that if there is a '*/' in a line comment, it won't be removed.
+			// 4) We must check to see if we have cppStyle prior to removing block style comments. 
+			ReadOnlySpan<char> span = comments.AsSpan(0, comments.Length);
+			int commentsLength = RemoveIgnoreComments(scratchChars, comments.Length);
+			bool javaDocStyle = span.Contains("/**", StringComparison.Ordinal);
+			bool cStyle = javaDocStyle || span.Contains("/*", StringComparison.Ordinal);
+			bool cppStyle = span.StartsWith("//", StringComparison.Ordinal);
+			commentsLength = javaDocStyle || cStyle ? RemoveBlockCommentMarkers(scratchChars, commentsLength) : commentsLength;
+			commentsLength = cppStyle ? RemoveLineCommentMarkers(scratchChars, commentsLength) : commentsLength;
 
-			// Sweep out comments marked to be ignored.
+			//wx widgets has a hard coded tab size of 8
 			{
-				int commentStart, commentEnd;
+				const int SpacesPerTab = 8;
 
-				// Block comments go first
-				while ((commentStart = result.IndexOf("/*~", StringComparison.Ordinal)) != -1)
+				// If we have any tab characters, then we need to convert them to spaces
+				span = scratchChars.AsSpan(0, commentsLength);
+				int tabIndex = span.IndexOf('\t');
+				if (tabIndex != -1)
 				{
-					commentEnd = result.IndexOf("*/", commentStart, StringComparison.Ordinal);
-					if (commentEnd != -1)
+					using BorrowStringBuilder tabsBorrower = new(StringBuilderCache.Small);
+					StringBuilder tabsBuilder = tabsBorrower.StringBuilder;
+					UhtFCString.TabsToSpaces(span, SpacesPerTab, true, tabIndex, tabsBuilder);
+					commentsLength = tabsBuilder.Length;
+					if (commentsLength > scratchChars.Length)
 					{
-						result = result.Remove(commentStart, commentEnd - commentStart + 2);
-					}
-					else
-					{
-						// This looks like an error - an unclosed block comment.
-						break;
+						ArrayPool<char>.Shared.Return(scratchChars);
+						scratchChars = ArrayPool<char>.Shared.Rent(commentsLength);
 					}
+					tabsBuilder.CopyTo(0, scratchChars, 0, commentsLength);
 				}
+			}
 
-				// Leftover line comments go next
-				while ((commentStart = result.IndexOf("//~", StringComparison.Ordinal)) != -1)
+			static bool IsAllSameChar(ReadOnlySpan<char> line, int startIndex, char testChar)
+			{
+				for (int index = startIndex, end = line.Length; index < end; ++index)
 				{
-					commentEnd = result.IndexOf("\n", commentStart, StringComparison.Ordinal);
-					if (commentEnd != -1)
-					{
-						result = result.Remove(commentStart, commentEnd - commentStart + 1);
-					}
-					else
+					if (line[index] != testChar)
 					{
-						result = result.Remove(commentStart);
-						break;
+						return false;
 					}
 				}
+				return true;
 			}
 
-			// Check for known commenting styles.
-			bool javaDocStyle = result.Contains("/**", StringComparison.Ordinal);
-			bool cStyle = result.Contains("/*", StringComparison.Ordinal);
-			bool cppStyle = result.StartsWith("//", StringComparison.Ordinal);
-
-			if (javaDocStyle || cStyle)
+			static bool IsWhitespaceOrLineSeparator(ReadOnlySpan<char> line)
 			{
-				// Remove beginning and end markers.
-				if (javaDocStyle)
+				// Skip any leading spaces
+				int index = 0;
+				int endPos = line.Length;
+				for (; index < endPos && UhtFCString.IsWhitespace(line[index]); ++index)
 				{
-					result = result.Replace("/**", "", StringComparison.Ordinal);
 				}
-				if (cStyle)
+				if (index == endPos)
 				{
-					result = result.Replace("/*", "", StringComparison.Ordinal);
+					return true;
 				}
-				result = result.Replace("*/", "", StringComparison.Ordinal);
-			}
 
-			if (cppStyle)
-			{
-				// Remove c++-style comment markers.  Also handle javadoc-style comments 
-				result = result.Replace("///", "", StringComparison.Ordinal);
-				result = result.Replace("//", "", StringComparison.Ordinal);
-
-				// Parser strips cpptext and replaces it with "// (cpptext)" -- prevent
-				// this from being treated as a comment on variables declared below the
-				// cpptext section
-				result = result.Replace("(cpptext)", "", StringComparison.Ordinal);
+				// Check for the same character
+				return IsAllSameChar(line, index, '-') || IsAllSameChar(line, index, '=') || IsAllSameChar(line, index, '*');
 			}
 
-			// Get rid of carriage return or tab characters, which mess up tooltips.
-			result = result.Replace("\r", "", StringComparison.Ordinal);
-
-			//wx widgets has a hard coded tab size of 8
+			// Loop while we have data
+			span = scratchChars.AsSpan(0, commentsLength);
+			bool firstLine = true;
+			int maxNumWhitespaceToRemove = 0;
+			int lastNonWhitespaceLength = 0;
+			int outEndPos = 0;
+			while (span.Length > 0)
 			{
-				const int SpacesPerTab = 8;
-				result = UhtFCString.TabsToSpaces(result, SpacesPerTab, true);
-			}
 
-			// get rid of uniform leading whitespace and all trailing whitespace, on each line
-			string[] lines = result.Split('\n');
-
-			for (int index = 0; index < lines.Length; ++index)
-			{
-				// Remove trailing whitespace
-				string line = lines[index].TrimEnd();
+				// Extract the next line to process
+				int eolIndex = span.IndexOf('\n');
+				ReadOnlySpan<char> line = eolIndex != -1 ? span[..eolIndex] : span;
+				span = eolIndex != -1 ? span[(eolIndex + 1)..] : new ReadOnlySpan<char>();
+				line = line.TrimEnd();
 
 				// Remove leading "*" and "* " in javadoc comments.
 				if (javaDocStyle)
@@ -366,113 +366,228 @@ namespace EpicGames.UHT.Parsers
 						line = line[(pos + 1)..];
 					}
 				}
-				lines[index] = line;
-			}
 
-			static bool IsAllSameChar(string line, int startIndex, char testChar)
-			{
-				for (int index = startIndex, end = line.Length; index < end; ++index)
+				// Test to see if this is whitespace or line separator.  If also first line, then just skip
+				bool isWhitespaceOrLineSeparator = IsWhitespaceOrLineSeparator(line);
+				if (firstLine && isWhitespaceOrLineSeparator)
 				{
-					if (line[index] != testChar)
-					{
-						return false;
-					}
+					continue;
 				}
-				return true;
-			}
 
-			static bool IsWhitespaceOrLineSeparator(string line)
-			{
-				// Skip any leading spaces
-				int index = 0;
-				int endPos = line.Length;
-				for (; index < endPos && UhtFCString.IsWhitespace(line[index]); ++index)
+				// Figure out how much whitespace is on the first line
+				if (firstLine)
 				{
+					for (; maxNumWhitespaceToRemove < line.Length; maxNumWhitespaceToRemove++)
+					{
+						if (!UhtFCString.IsWhitespace(line[maxNumWhitespaceToRemove]))
+						{
+							break;
+						}
+					}
+					line = line[maxNumWhitespaceToRemove..];
 				}
-				if (index == endPos)
+				else
 				{
-					return true;
-				}
 
-				// Check for the same character
-				return IsAllSameChar(line, index, '-') || IsAllSameChar(line, index, '=') || IsAllSameChar(line, index, '*');
-			}
-
-			// Find first meaningful line
-			int firstIndex = 0;
-			for (; firstIndex < lines.Length && IsWhitespaceOrLineSeparator(lines[firstIndex]); ++firstIndex)
-			{ }
-
-			int lastIndex = lines.Length;
-			for (; lastIndex > firstIndex && IsWhitespaceOrLineSeparator(lines[lastIndex - 1]); --lastIndex)
-			{ }
+					// Trim any leading whitespace
+					for (int i = 0; i < maxNumWhitespaceToRemove && line.Length > 0; i++)
+					{
+						if (!UhtFCString.IsWhitespace(line[0]))
+						{
+							break;
+						}
+						line = line[1..];
+					}
 
-			result = String.Empty;
-			if (firstIndex != lastIndex)
-			{
-				string firstLine = lines[firstIndex];
+					scratchChars[outEndPos++] = '\n';
+				}
 
-				// Figure out how much whitespace is on the first line
-				int maxNumWhitespaceToRemove = 0;
-				for (; maxNumWhitespaceToRemove < firstLine.Length; maxNumWhitespaceToRemove++)
+				if (line.Length > 0 && !IsAllSameChar(line, 0, '='))
 				{
-					if (/*!UhtFCString.IsLinebreak(FirstLine[MaxNumWhitespaceToRemove]) && */!UhtFCString.IsWhitespace(firstLine[maxNumWhitespaceToRemove]))
+					for (int i = 0; i < line.Length; i++)
 					{
-						break;
+						scratchChars[outEndPos++] = line[i];
 					}
 				}
 
-				for (int index = firstIndex; index != lastIndex; ++index)
+				if (!isWhitespaceOrLineSeparator)
 				{
-					string line = lines[index];
+					lastNonWhitespaceLength = outEndPos;
+				}
+				firstLine = false;
+			}
 
-					int temporaryMaxWhitespace = maxNumWhitespaceToRemove;
+			outEndPos = lastNonWhitespaceLength;
 
-					// Allow eating an extra tab on subsequent lines if it's present
-					if ((index > firstIndex) && (line.Length > 0) && (line[0] == '\t'))
-					{
-						temporaryMaxWhitespace++;
-					}
+			//@TODO: UCREMOVAL: Really want to trim an arbitrary number of newlines above and below, but keep multiple newlines internally
+			// Make sure it doesn't start with a newline
+			int outStartPos = 0;
+			if (outStartPos < outEndPos && scratchChars[outStartPos] == '\n')
+			{
+				outStartPos++;
+			}
 
-					// Advance past whitespace
-					int pos = 0;
-					while (pos < temporaryMaxWhitespace && pos < line.Length && UhtFCString.IsWhitespace(line[pos]))
-					{
-						++pos;
-					}
+			// Make sure it doesn't end with a dead newline
+			if (outStartPos < outEndPos && scratchChars[outEndPos - 1] == '\n')
+			{
+				outEndPos--;
+			}
 
-					if (pos > 0)
-					{
-						line = line[pos..];
-					}
+			string results = scratchChars.AsSpan(outStartPos, outEndPos - outStartPos).ToString();
+			ArrayPool<char>.Shared.Return(scratchChars);
+			return results;
+		}
 
-					if (index > firstIndex)
-					{
-						result += "\n";
-					}
+		/// <summary>
+		/// Remove any comments marked to be ignored
+		/// </summary>
+		/// <param name="comments">Buffer containing comments to be processed.  Comments are removed inline</param>
+		/// <param name="inLength">Length of the comments</param>
+		/// <returns>New length of the comments</returns>
+		private static int RemoveIgnoreComments(char[] comments, int inLength)
+		{
+			ReadOnlySpan<char> span = comments.AsSpan(0, inLength);
+			int commentStart, commentEnd;
 
-					if (line.Length > 0 && !IsAllSameChar(line, 0, '='))
-					{
-						result += line;
-					}
+			// Block comments go first
+			while ((commentStart = span.IndexOf("/*~", StringComparison.Ordinal)) != -1)
+			{
+				commentEnd = span[commentStart..].IndexOf("*/", StringComparison.Ordinal);
+				if (commentEnd != -1)
+				{
+					commentEnd += 2;
+					Array.Copy(comments, commentStart + commentEnd, comments, commentStart, span.Length - (commentStart + commentEnd));
+					span = span[..(span.Length - commentEnd)];
+				}
+				else
+				{
+					// This looks like an error - an unclosed block comment.
+					break;
 				}
 			}
 
-			//@TODO: UCREMOVAL: Really want to trim an arbitrary number of newlines above and below, but keep multiple newlines internally
-			// Make sure it doesn't start with a newline
-			if (result.Length != 0 && UhtFCString.IsLinebreak(result[0]))
+			// Leftover line comments go next
+			while ((commentStart = span.IndexOf("//~", StringComparison.Ordinal)) != -1)
 			{
-				result = result[1..];
+				commentEnd = span[commentStart..].IndexOf("\n", StringComparison.Ordinal);
+				if (commentEnd != -1)
+				{
+					commentEnd++;
+					Array.Copy(comments, commentStart + commentEnd, comments, commentStart, span.Length - (commentStart + commentEnd));
+					span = span[..(span.Length - commentEnd)];
+				}
+				else
+				{
+					span = span[..commentStart];
+					break;
+				}
 			}
 
-			// Make sure it doesn't end with a dead newline
-			if (result.Length != 0 && UhtFCString.IsLinebreak(result[^1]))
+			return span.Length;
+		}
+
+		/// <summary>
+		/// Remove any block comment markers
+		/// </summary>
+		/// <param name="comments">Buffer containing comments to be processed.  Comments are removed inline</param>
+		/// <param name="inLength">Length of the comments</param>
+		/// <returns>New length of the comments</returns>
+		private static int RemoveBlockCommentMarkers(char[] comments, int inLength)
+		{
+			int outPos = 0;
+			int inPos = 0;
+			while (inPos < inLength)
 			{
-				result = result[0..^1];
+				switch (comments[inPos])
+				{
+					case '\r':
+						inPos++;
+						break;
+					case '/':
+						if (inPos + 2 < inLength && comments[inPos + 1] == '*' && comments[inPos + 2] == '*')
+						{
+							inPos += 3;
+						}
+						else if (inPos + 1 < inLength && comments[inPos + 1] == '*')
+						{
+							inPos += 2;
+						}
+						else
+						{
+							comments[outPos++] = comments[inPos++];
+						}
+						break;
+					case '*':
+						if (inPos + 1 < inLength && comments[inPos + 1] == '/')
+						{
+							inPos += 2;
+						}
+						else
+						{
+							comments[outPos++] = comments[inPos++];
+						}
+						break;
+					default:
+						comments[outPos++] = comments[inPos++];
+						break;
+				}
 			}
+			return outPos;
+		}
 
-			// Done.
-			return result;
+		/// <summary>
+		/// Remove any line comment markers
+		/// </summary>
+		/// <param name="comments">Buffer containing comments to be processed.  Comments are removed inline</param>
+		/// <param name="inLength">Length of the comments</param>
+		/// <returns>New length of the comments</returns>
+		private static int RemoveLineCommentMarkers(char[] comments, int inLength)
+		{
+			ReadOnlySpan<char> span = comments.AsSpan(0, inLength);
+			int outPos = 0;
+			int inPos = 0;
+			while (inPos < inLength)
+			{
+				switch (comments[inPos])
+				{
+					case '\r':
+						inPos++;
+						break;
+					case '/':
+						if (inPos + 1 < inLength && comments[inPos + 1] == '/')
+						{
+							if (inPos + 2 < inLength && comments[inPos + 2] == '/')
+							{
+								inPos += 3;
+							}
+							else
+							{
+								inPos += 2;
+							}
+						}
+						else
+						{
+							comments[outPos++] = comments[inPos++];
+						}
+						break;
+					case '(':
+						{
+							if (span[inPos..].StartsWith("(cpptext)", StringComparison.Ordinal))
+							{
+								inPos += 9;
+							}
+							else
+							{
+								comments[outPos++] = comments[inPos++];
+							}
+						}
+						break;
+					default:
+						comments[outPos++] = comments[inPos++];
+						break;
+				}
+			}
+			return outPos;
 		}
 	}
 
diff --git a/Engine/Source/Programs/Shared/EpicGames.UHT/Parsers/UhtPropertyParser.cs b/Engine/Source/Programs/Shared/EpicGames.UHT/Parsers/UhtPropertyParser.cs
--- a/Engine/Source/Programs/Shared/EpicGames.UHT/Parsers/UhtPropertyParser.cs
+++ b/Engine/Source/Programs/Shared/EpicGames.UHT/Parsers/UhtPropertyParser.cs
@@ -764,7 +764,7 @@ namespace EpicGames.UHT.Parsers
 
 			bool isParamList = this._specifierContext.PropertySettings.PropertyCategory != UhtPropertyCategory.Member && this.TokenReader.TryOptional("UPARAM");
 
-			UhtSpecifierParser specifiers = this.TopScope.HeaderParser.GetCachedSpecifierParser(this._specifierContext, "Variable",
+			UhtSpecifierParser specifiers = UhtSpecifierParser.GetThreadInstance(this._specifierContext, "Variable",
 				isParamList ? this.TopScope.Session.GetSpecifierTable(UhtTableNames.PropertyArgument) : this.TopScope.Session.GetSpecifierTable(UhtTableNames.PropertyMember));
 			if (_declarationStyle == UhtParsePropertyDeclarationStyle.UPROPERTY || isParamList)
 			{
@@ -1102,7 +1102,7 @@ namespace EpicGames.UHT.Parsers
 			// Try gathering metadata for member fields
 			if (newProperty.PropertyCategory == UhtPropertyCategory.Member)
 			{
-				UhtSpecifierParser specifiers = this.TopScope.HeaderParser.GetCachedSpecifierParser(this._specifierContext, newProperty.SourceName, this.TopScope.Session.GetSpecifierTable(UhtTableNames.PropertyMember));
+				UhtSpecifierParser specifiers = UhtSpecifierParser.GetThreadInstance(this._specifierContext, newProperty.SourceName, this.TopScope.Session.GetSpecifierTable(UhtTableNames.PropertyMember));
 				specifiers.ParseFieldMetaData();
 				this.TopScope.AddFormattedCommentsAsTooltipMetaData(newProperty);
 			}
diff --git a/Engine/Source/Programs/Shared/EpicGames.UHT/Parsers/UhtScriptStructParser.cs b/Engine/Source/Programs/Shared/EpicGames.UHT/Parsers/UhtScriptStructParser.cs
--- a/Engine/Source/Programs/Shared/EpicGames.UHT/Parsers/UhtScriptStructParser.cs
+++ b/Engine/Source/Programs/Shared/EpicGames.UHT/Parsers/UhtScriptStructParser.cs
@@ -75,7 +75,7 @@ namespace EpicGames.UHT.Parsers
 
 					// Parse the specifiers
 					UhtSpecifierContext specifierContext = new(topScope, topScope.TokenReader, scriptStruct.MetaData);
-					UhtSpecifierParser specifiers = topScope.HeaderParser.GetCachedSpecifierParser(specifierContext, ScopeName, parentScope.Session.GetSpecifierTable(UhtTableNames.ScriptStruct));
+					UhtSpecifierParser specifiers = UhtSpecifierParser.GetThreadInstance(specifierContext, ScopeName, parentScope.Session.GetSpecifierTable(UhtTableNames.ScriptStruct));
 					specifiers.ParseSpecifiers();
 
 					// Consume the struct specifier
diff --git a/Engine/Source/Programs/Shared/EpicGames.UHT/Parsers/UhtSpecifierParser.cs b/Engine/Source/Programs/Shared/EpicGames.UHT/Parsers/UhtSpecifierParser.cs
--- a/Engine/Source/Programs/Shared/EpicGames.UHT/Parsers/UhtSpecifierParser.cs
+++ b/Engine/Source/Programs/Shared/EpicGames.UHT/Parsers/UhtSpecifierParser.cs
@@ -2,6 +2,7 @@
 
 using System;
 using System.Collections.Generic;
+using System.Threading;
 using EpicGames.Core;
 using EpicGames.UHT.Tables;
 using EpicGames.UHT.Tokenizer;
@@ -23,6 +24,12 @@ namespace EpicGames.UHT.Parsers
 			public object? _value;
 		}
 
+		/// <summary>
+		/// For a given header file, we share a common specifier parser to reduce the number of allocations.
+		/// Before the parser can be reused, the ParseDeferred method must be called to dispatch that list.
+		/// </summary>
+		private static readonly ThreadLocal<UhtSpecifierParser?> s_tls = new(() => null);
+
 		private static readonly List<KeyValuePair<StringView, StringView>> s_emptyKVPValues = new();
 
 		private UhtSpecifierContext _specifierContext;
@@ -41,6 +48,26 @@ namespace EpicGames.UHT.Parsers
 		private readonly Action _parseKVPValueAction;
 		private readonly Action _parseStringViewListAction;
 
+		/// <summary>
+		/// Get the cached specifier parser
+		/// </summary>
+		/// <param name="specifierContext">Specifier context</param>
+		/// <param name="context">User facing context</param>
+		/// <param name="table">Specifier table</param>
+		/// <returns>Specifier parser</returns>
+		public static UhtSpecifierParser GetThreadInstance(UhtSpecifierContext specifierContext, StringView context, UhtSpecifierTable table)
+		{
+			if (s_tls.Value == null)
+			{
+				s_tls.Value = new UhtSpecifierParser(specifierContext, context, table);
+			}
+			else
+			{
+				s_tls.Value.Reset(specifierContext, context, table);
+			}
+			return s_tls.Value;
+		}
+
 		/// <summary>
 		/// Construct a new specifier parser
 		/// </summary>
@@ -56,22 +83,8 @@ namespace EpicGames.UHT.Parsers
 
 			this._parseAction = ParseInternal;
 			this._parseFieldMetaDataAction = ParseFieldMetaDataInternal;
-			this._parseKVPValueAction = () =>
-			{
-				if (this._currentKVPValues == null)
-				{
-					this._currentKVPValues = new List<KeyValuePair<StringView, StringView>>();
-				}
-				this._currentKVPValues.Add(ReadKVP());
-			};
-			this._parseStringViewListAction = () =>
-			{
-				if (this._currentStringValues == null)
-				{
-					this._currentStringValues = new List<StringView>();
-				}
-				this._currentStringValues.Add(ReadValue());
-			};
+			this._parseKVPValueAction = ParseKVPValueInternal;
+			this._parseStringViewListAction = ParseStringViewListInternal;
 		}
 
 		/// <summary>
@@ -207,6 +220,24 @@ namespace EpicGames.UHT.Parsers
 			this._specifierContext.MetaData.Add(key.Value.ToString(), this._specifierContext.MetaNameIndex, builder.ToString());
 		}
 
+		private void ParseKVPValueInternal()
+		{
+			if (this._currentKVPValues == null)
+			{
+				this._currentKVPValues = new List<KeyValuePair<StringView, StringView>>();
+			}
+			this._currentKVPValues.Add(ReadKVP());
+		}
+
+		private void ParseStringViewListInternal()
+		{
+			if (this._currentStringValues == null)
+			{
+				this._currentStringValues = new List<StringView>();
+			}
+			this._currentStringValues.Add(ReadValue());
+		}
+
 		private void Dispatch(UhtSpecifier specifier, object? value)
 		{
 			UhtSpecifierDispatchResults results = specifier.Dispatch(this._specifierContext, value);
diff --git a/Engine/Source/Programs/Shared/EpicGames.UHT/Utils/UhtFCString.cs b/Engine/Source/Programs/Shared/EpicGames.UHT/Utils/UhtFCString.cs
--- a/Engine/Source/Programs/Shared/EpicGames.UHT/Utils/UhtFCString.cs
+++ b/Engine/Source/Programs/Shared/EpicGames.UHT/Utils/UhtFCString.cs
@@ -452,25 +452,10 @@ namespace EpicGames.UHT.Utils
 			{
 				return input;
 			}
-			return TabsToSpacesInternal(input.Span, tabSpacing, emulateCrBug, tabIndex).AsMemory();
-		}
-
-		/// <summary>
-		/// Replace tabs to spaces in a string containing zero or more lines.
-		/// </summary>
-		/// <param name="input">Input string</param>
-		/// <param name="tabSpacing">Number of spaces to exchange for tabs</param>
-		/// <param name="emulateCrBug">Due to a bug in UE ConvertTabsToSpacesInline, any \n is considered part of the line length.</param>
-		/// <returns>Resulting string or the original string if the string didn't contain any spaces.</returns>
-		public static string TabsToSpaces(string input, int tabSpacing, bool emulateCrBug)
-		{
-			// If we have any tab characters, then we need to convert them to spaces
-			int tabIndex = input.IndexOf('\t', StringComparison.Ordinal);
-			if (tabIndex == -1)
-			{
-				return input;
-			}
-			return TabsToSpacesInternal(input.AsSpan(), tabSpacing, emulateCrBug, tabIndex);
+			using BorrowStringBuilder borrower = new(StringBuilderCache.Small);
+			StringBuilder builder = borrower.StringBuilder;
+			TabsToSpaces(input.Span, tabSpacing, emulateCrBug, tabIndex, builder);
+			return builder.ToString().AsMemory();
 		}
 
 		/// <summary>
@@ -480,14 +465,12 @@ namespace EpicGames.UHT.Utils
 		/// <param name="tabSpacing">Number of spaces to exchange for tabs</param>
 		/// <param name="emulateCrBug">Due to a bug in UE ConvertTabsToSpacesInline, any \n is considered part of the line length.</param>
 		/// <param name="tabIndex">Initial tab index</param>
-		/// <returns>Resulting string or the original string if the string didn't contain any spaces.</returns>
-		private static string TabsToSpacesInternal(ReadOnlySpan<char> span, int tabSpacing, bool emulateCrBug, int tabIndex)
+		/// <param name="builder">Destination string builder</param>
+		public static void TabsToSpaces(ReadOnlySpan<char> span, int tabSpacing, bool emulateCrBug, int tabIndex, StringBuilder builder)
 		{
 			// Locate the last \n since all tabs have to be computed relative to this.
 			int crPos = span[..tabIndex].LastIndexOf('\n') + 1;
 
-			using BorrowStringBuilder borrower = new(StringBuilderCache.Small);
-			StringBuilder builder = borrower.StringBuilder;
 			int committed = 0;
 			do
 			{
@@ -522,7 +505,6 @@ namespace EpicGames.UHT.Utils
 
 			// Commit the remaining data
 			builder.Append(span);
-			return builder.ToString();
 		}
 	}
 }
diff --git a/Engine/Source/Programs/Shared/EpicGames.UHT/Utils/UhtSession.cs b/Engine/Source/Programs/Shared/EpicGames.UHT/Utils/UhtSession.cs
--- a/Engine/Source/Programs/Shared/EpicGames.UHT/Utils/UhtSession.cs
+++ b/Engine/Source/Programs/Shared/EpicGames.UHT/Utils/UhtSession.cs
@@ -1052,24 +1052,6 @@ namespace EpicGames.UHT.Utils
 			StepExport();
 		}
 
-		/// <summary>
-		/// Try the given action regardless of any prior errors.  If an exception occurs that doesn't have the required
-		/// context, use the supplied context to generate the message.
-		/// </summary>
-		/// <param name="messageSource">Message context for when the exception doesn't contain a context.</param>
-		/// <param name="action">The lambda to be invoked</param>
-		public void TryAlways(IUhtMessageSource? messageSource, Action action)
-		{
-			try
-			{
-				action();
-			}
-			catch (Exception e)
-			{
-				HandleException(messageSource, e);
-			}
-		}
-
 		/// <summary>
 		/// Try the given action.  If an exception occurs that doesn't have the required
 		/// context, use the supplied context to generate the message.
@@ -1723,11 +1705,15 @@ namespace EpicGames.UHT.Utils
 		private UhtHeaderFileParser? ParseHeaderFile(UhtHeaderFile headerFile)
 		{
 			UhtHeaderFileParser? parser = null;
-			TryAlways(headerFile.MessageSource, () =>
+			try
 			{
 				headerFile.Read();
 				parser = UhtHeaderFileParser.Parse(headerFile);
-			});
+			}
+			catch (Exception e)
+			{
+				HandleException(headerFile.MessageSource, e);
+			}
 			return parser;
 		}
 
@@ -1987,10 +1973,14 @@ namespace EpicGames.UHT.Utils
 
 		private void Resolve(UhtHeaderFile headerFile, UhtResolvePhase resolvePhase)
 		{
-			TryAlways(headerFile.MessageSource, () =>
+			try
 			{
 				headerFile.Resolve(resolvePhase);
-			});
+			}
+			catch (Exception e)
+			{
+				HandleException(headerFile.MessageSource, e);
+			}
 		}
 
 		private delegate void StepDelegate(UhtHeaderFile headerFile);
