Updated C# UHT with recent C++ UHT changes

#rnx

diff --git a/Engine/Source/Programs/Shared/EpicGames.UHT/Exporters/CodeGen/UhtHeaderCodeGeneratorHFile.cs b/Engine/Source/Programs/Shared/EpicGames.UHT/Exporters/CodeGen/UhtHeaderCodeGeneratorHFile.cs
--- a/Engine/Source/Programs/Shared/EpicGames.UHT/Exporters/CodeGen/UhtHeaderCodeGeneratorHFile.cs
+++ b/Engine/Source/Programs/Shared/EpicGames.UHT/Exporters/CodeGen/UhtHeaderCodeGeneratorHFile.cs
@@ -475,17 +475,17 @@ namespace EpicGames.UHT.Exporters.CodeGen
 			return builder;
 		}
 
-		private static StringBuilder AppendFieldNotify(StringBuilder builder, UhtClass classObj,
+		private StringBuilder AppendFieldNotify(StringBuilder builder, UhtClass classObj,
 			bool hasProperties, bool hasFunctions, bool hasEditorFields, bool allEditorFields,
 			bool includeEditorOnlyFields)
 		{
-			builder.Append("\tUE_FIELD_NOTIFICATION_DECLARE_CLASS_DESCRIPTOR_BEGIN() \\\r\n");
+			builder.Append("\tUE_FIELD_NOTIFICATION_DECLARE_CLASS_DESCRIPTOR_BEGIN(").Append(this.PackageApi).Append(") \\\r\n");
 
 			//UE_FIELD_NOTIFICATION_DECLARE_FIELD
 			AppendFieldNotify(builder, classObj, hasProperties, hasFunctions, hasEditorFields, allEditorFields,
 				includeEditorOnlyFields, false, (StringBuilder builder, UhtClass classObj, string name) =>
 			{
-				builder.Append($"\tUE_FIELD_NOTIFICATION_DECLARE_FIELD({name}) \\\r\n");
+				builder.Append("\tUE_FIELD_NOTIFICATION_DECLARE_FIELD(").Append(name).Append(") \\\r\n");
 			});
 
 			//UE_FIELD_NOTIFICATION_DECLARE_ENUM_FIELD
@@ -496,16 +496,16 @@ namespace EpicGames.UHT.Exporters.CodeGen
 				if (isFirst)
 				{
 					isFirst = false;
-					builder.Append($"\tUE_FIELD_NOTIFICATION_DECLARE_ENUM_FIELD_BEGIN({name}) \\\r\n");
+					builder.Append("\tUE_FIELD_NOTIFICATION_DECLARE_ENUM_FIELD_BEGIN(").Append(name).Append(") \\\r\n");
 				}
 				else
 				{
-					builder.Append($"\tUE_FIELD_NOTIFICATION_DECLARE_ENUM_FIELD({name}) \\\r\n");
+					builder.Append("\tUE_FIELD_NOTIFICATION_DECLARE_ENUM_FIELD(").Append(name).Append(") \\\r\n");
 				}
 			});
 
 			builder.Append("\tUE_FIELD_NOTIFICATION_DECLARE_ENUM_FIELD_END() \\\r\n");
-			builder.Append($"\tUE_FIELD_NOTIFICATION_DECLARE_CLASS_DESCRIPTOR_END(); \\\r\n");
+			builder.Append("\tUE_FIELD_NOTIFICATION_DECLARE_CLASS_DESCRIPTOR_END(); \\\r\n");
 			return builder;
 		}
 
