BuildGraph: Allow use of absolute paths in Include elements.

diff --git a/Engine/Source/Programs/AutomationTool/BuildGraph/BuildGraph.cs b/Engine/Source/Programs/AutomationTool/BuildGraph/BuildGraph.cs
--- a/Engine/Source/Programs/AutomationTool/BuildGraph/BuildGraph.cs
+++ b/Engine/Source/Programs/AutomationTool/BuildGraph/BuildGraph.cs
@@ -158,6 +158,42 @@ namespace AutomationTool
 			Files.Sort(StringComparer.OrdinalIgnoreCase);
 			return Task.FromResult(Files.ToArray());
 		}
+
+		/// <inheritdoc/>
+		public string CombinePaths(string basePath, string nextPath)
+		{
+			if (Path.IsPathRooted(nextPath))
+			{
+				return nextPath;
+			}
+
+			List<string> fragments = new List<string>(basePath.Split('/'));
+			fragments.RemoveAt(fragments.Count - 1);
+
+			foreach (string appendFragment in nextPath.Split('/'))
+			{
+				if (appendFragment.Equals(".", StringComparison.Ordinal))
+				{
+					continue;
+				}
+				else if (appendFragment.Equals("..", StringComparison.Ordinal))
+				{
+					if (fragments.Count > 0)
+					{
+						fragments.RemoveAt(fragments.Count - 1);
+					}
+					else
+					{
+						throw new Exception($"Path '{nextPath}' cannot be combined with '{basePath}'");
+					}
+				}
+				else
+				{
+					fragments.Add(appendFragment);
+				}
+			}
+			return String.Join('/', fragments);
+		}
 	}
 
 	/// <summary>
diff --git a/Engine/Source/Programs/Shared/EpicGames.BuildGraph/BgScriptReader.cs b/Engine/Source/Programs/Shared/EpicGames.BuildGraph/BgScriptReader.cs
--- a/Engine/Source/Programs/Shared/EpicGames.BuildGraph/BgScriptReader.cs
+++ b/Engine/Source/Programs/Shared/EpicGames.BuildGraph/BgScriptReader.cs
@@ -286,6 +286,14 @@ namespace EpicGames.BuildGraph
 		/// <param name="path">Path to format</param>
 		/// <returns></returns>
 		object GetNativePath(string path);
+
+		/// <summary>
+		/// Combines one path with another
+		/// </summary>
+		/// <param name="path">The base file path</param>
+		/// <param name="next">Relative path to the next file</param>
+		/// <returns></returns>
+		string CombinePaths(string path, string next);
 	}
 
 	/// <summary>
@@ -578,36 +586,6 @@ namespace EpicGames.BuildGraph
 			return false;
 		}
 
-		static string CombinePaths(string basePath, string nextPath)
-		{
-			List<string> fragments = new List<string>(basePath.Split('/'));
-			fragments.RemoveAt(fragments.Count - 1);
-
-			foreach (string appendFragment in nextPath.Split('/'))
-			{
-				if (appendFragment.Equals(".", StringComparison.Ordinal))
-				{
-					continue;
-				}
-				else if (appendFragment.Equals("..", StringComparison.Ordinal))
-				{
-					if (fragments.Count > 0)
-					{
-						fragments.RemoveAt(fragments.Count - 1);
-					}
-					else
-					{
-						throw new Exception($"Path '{nextPath}' cannot be combined with '{basePath}'");
-					}
-				}
-				else
-				{
-					fragments.Add(appendFragment);
-				}
-			}
-			return String.Join('/', fragments);
-		}
-
 		/// <summary>
 		/// Read an include directive, and the contents of the target file
 		/// </summary>
@@ -619,7 +597,7 @@ namespace EpicGames.BuildGraph
 				HashSet<string> files = new HashSet<string>();
 				foreach (string script in ReadListAttribute(element, "Script"))
 				{
-					string includePath = CombinePaths(element.Location.File, script);
+					string includePath = Context.CombinePaths(element.Location.File, script);
 					if (Regex.IsMatch(includePath, @"\*|\?|\.\.\."))
 					{
 						files.UnionWith(await Context.FindAsync(includePath));
