Horde: Remove Perforce information from Log parser context and event matchers, and add it in a post-process from the agent instead. This allows having to propagate context into UBT/UAT when adding structured logging context there.

diff --git a/Engine/Source/Programs/Horde/Horde.Agent.Tests/LogParserTests.cs b/Engine/Source/Programs/Horde/Horde.Agent.Tests/LogParserTests.cs
--- a/Engine/Source/Programs/Horde/Horde.Agent.Tests/LogParserTests.cs
+++ b/Engine/Source/Programs/Horde/Horde.Agent.Tests/LogParserTests.cs
@@ -5,6 +5,7 @@ using System.Collections.Generic;
 using System.IO;
 using System.Text;
 using EpicGames.Core;
+using EpicGames.Perforce;
 using Horde.Agent.Parser;
 using Horde.Agent.Utility;
 using Microsoft.Extensions.Logging;
@@ -30,32 +31,34 @@ namespace Horde.Agent.Tests
 
 			public void Log<TState>(LogLevel logLevel, EventId eventId, TState state, Exception exception, Func<TState, Exception?, string> formatter)
 			{
-				if (state is LogEvent logEvent)
+				LogEvent logEvent = LogEvent.Read(JsonLogEvent.FromLoggerState(logLevel, eventId, state, exception, formatter).Data.Span);
+				if (logEvent.Level != LogLevel.Information || logEvent.Id != default || logEvent.Properties != null)
 				{
 					logEvent.Properties ??= new Dictionary<string, object>();
 					logEvent.Properties.Add(LogLine, _logLineIndex);
 					_events.Add(logEvent);
 				}
-				else if(logLevel != LogLevel.Information || eventId.Id != 0)
-				{
-					_events.Add(LogEvent.FromState(logLevel, eventId, state, exception, formatter));
-				}
+
 				_logLineIndex++;
 			}
 		}
 
-		class JsonLoggerImpl : JsonLogger
+		class JsonLoggerImpl : ILogger
 		{
 			public List<LogEvent> _lines = new List<LogEvent>();
 
 			public JsonLoggerImpl(bool? warnings, ILogger inner) 
-				: base(warnings, inner)
 			{
 			}
 
-			protected override void WriteFormattedEvent(LogLevel level, int lineIndex, int lineCount, byte[] line)
+			public IDisposable BeginScope<TState>(TState state) => null!;
+
+			public bool IsEnabled(LogLevel logLevel) => true;
+
+			public void Log<TState>(LogLevel logLevel, EventId eventId, TState state, Exception exception, Func<TState, Exception?, string> formatter)
 			{
-				_lines.Add(LogEvent.Read(line));
+				JsonLogEvent logEvent = JsonLogEvent.FromLoggerState(logLevel, eventId, state, exception, formatter);
+				_lines.Add(LogEvent.Read(logEvent.Data.Span));
 			}
 		}
 
@@ -78,10 +81,10 @@ namespace Horde.Agent.Tests
 		{
 			string[] lines =
 			{
-				"{\"Timestamp\":\"2021-10-28T09:14:48.8561456-04:00\",\"Level\":\"Information\",\"MessageTemplate\":\"Hello {World}\",\"RenderedMessage\":\"Hello 123\",\"Properties\":{\"World\":123,\"EventId\":{\"Id\":123},\"SourceContext\":\"HordeAgent\",\"dd.env\":\"default\",\"dd.service\":\"hordeagent\",\"dd.version\":\"1.0.0\"}}",
-				"{\"Timestamp\":\"2021-10-28T05:36:08\",\"Level\":\"Information\",\"RenderedMessage\":\"Building 43 projects (see Log \u0027Engine/Programs/AutomationTool/Saved/Logs/Log.txt\u0027 for more details)\"}",
-				"{\"Timestamp\":\"2021-10-28T05:36:08\",\"Level\":\"Warning\",\"RenderedMessage\":\" Restore...\"}",
-				"{\"Timestamp\":\"2021-10-28T05:36:15\",\"Level\":\"Error\",\"RenderedMessage\":\" Build...\"}",
+				new LogEvent(DateTime.UtcNow, LogLevel.Information, new EventId(123), "Hello 123", null, null, null).ToJson(),
+				new LogEvent(DateTime.UtcNow, LogLevel.Information, default, "Building 43 projects (see Log \u0027Engine/Programs/AutomationTool/Saved/Logs/Log.txt\u0027 for more details)", null, null, null).ToJson(),
+				new LogEvent(DateTime.UtcNow, LogLevel.Warning, default, " Restore...", null, null, null).ToJson(),
+				new LogEvent(DateTime.UtcNow, LogLevel.Error, default, " Build...", null, null, null).ToJson(),
 			};
 
 			List<LogEvent> logEvents = Parse(lines);
@@ -414,7 +417,7 @@ namespace Horde.Agent.Tests
 		{
 			string[] lines =
 			{
-				@"LogBlueprint: Warning: [AssetLog] D:\\build\\++UE5\\Sync\\QAGame\\Plugins\\NiagaraFluids\\Content\\Blueprints\\Phsyarum_BP.uasset: [Compiler] Fill Texture 2D : Usage of 'Fill Texture 2D' has been deprecated. This function has been replaced by object user variables on the emitter to specify render targets to fill with data.",
+				@"LogBlueprint: Warning: [AssetLog] D:\build\++UE5\Sync\QAGame\Plugins\NiagaraFluids\Content\Blueprints\Phsyarum_BP.uasset: [Compiler] Fill Texture 2D : Usage of 'Fill Texture 2D' has been deprecated. This function has been replaced by object user variables on the emitter to specify render targets to fill with data.",
 			};
 
 			List<LogEvent> events = Parse(String.Join("\n", lines), new DirectoryReference("D:\\build\\++UE5\\Sync"));
@@ -1048,6 +1051,9 @@ namespace Horde.Agent.Tests
 
 		List<LogEvent> Parse(string text, DirectoryReference workspaceDir)
 		{
+			PerforceViewMap viewMap = new PerforceViewMap();
+			viewMap.Entries.Add(new PerforceViewMapEntry(true, "...", "//UE4/Main/..."));
+
 			LogParserContext context = new LogParserContext();
 			context.WorkspaceDir = workspaceDir;
 			context.PerforceStream = "//UE4/Main";
@@ -1060,7 +1066,7 @@ namespace Horde.Agent.Tests
 			Random generator = new Random(0);
 
 			LoggerCapture logger = new LoggerCapture();
-			using (LogParser parser = new LogParser(logger, context, ignorePatterns))
+			using (LogParser parser = new LogParser(new JsonPerforceLogger(logger, workspaceDir, viewMap, 12345), context, ignorePatterns))
 			{
 				int pos = 0;
 				while(pos < textBytes.Length)
diff --git a/Engine/Source/Programs/Horde/Horde.Agent/Commands/Parse/ParseCommand.cs b/Engine/Source/Programs/Horde/Horde.Agent/Commands/Parse/ParseCommand.cs
--- a/Engine/Source/Programs/Horde/Horde.Agent/Commands/Parse/ParseCommand.cs
+++ b/Engine/Source/Programs/Horde/Horde.Agent/Commands/Parse/ParseCommand.cs
@@ -6,7 +6,9 @@ using System.ComponentModel;
 using System.IO;
 using System.Threading.Tasks;
 using EpicGames.Core;
+using EpicGames.Perforce;
 using Horde.Agent.Parser;
+using Horde.Agent.Utility;
 using Microsoft.Extensions.Logging;
 
 namespace Horde.Agent.Commands.Parse
@@ -63,12 +65,11 @@ namespace Horde.Agent.Commands.Parse
 			// Read the file and pipe it through the event parser
 			using (FileStream inputStream = FileReference.Open(InputFile, FileMode.Open, FileAccess.Read))
 			{
-				LogParserContext context = new LogParserContext();
-				context.WorkspaceDir = WorkspaceDir;
-				context.PerforceStream = Stream;
-				context.PerforceChange = Change;
+				PerforceViewMap viewMap = new PerforceViewMap();
+				viewMap.Entries.Add(new PerforceViewMapEntry(true, "...", $"{Stream}/..."));
 
-				using (LogParser parser = new LogParser(logger, context, ignorePatterns))
+				JsonPerforceLogger jsonLogger = new JsonPerforceLogger(logger, WorkspaceDir ?? DirectoryReference.GetCurrentDirectory(), viewMap, Change);
+				using (LogParser parser = new LogParser(jsonLogger, ignorePatterns))
 				{
 					byte[] data = new byte[1024];
 					for (; ; )
diff --git a/Engine/Source/Programs/Horde/Horde.Agent/Execution/BuildGraphExecutor.cs b/Engine/Source/Programs/Horde/Horde.Agent/Execution/BuildGraphExecutor.cs
--- a/Engine/Source/Programs/Horde/Horde.Agent/Execution/BuildGraphExecutor.cs
+++ b/Engine/Source/Programs/Horde/Horde.Agent/Execution/BuildGraphExecutor.cs
@@ -13,6 +13,7 @@ using System.Text.Json.Serialization;
 using System.Threading;
 using System.Threading.Tasks;
 using EpicGames.Core;
+using EpicGames.Perforce;
 using Grpc.Core;
 using Horde.Agent.Execution.Interfaces;
 using Horde.Agent.Parser;
@@ -809,16 +810,16 @@ namespace Horde.Agent.Execution
 			{
 				FileUtils.ForceDeleteDirectoryContents(telemetryDir);
 			}
-			
-			LogParserContext context = new LogParserContext();
-			context.WorkspaceDir = workspaceDir;
-			context.PerforceStream = _stream.Name;
-			context.PerforceChange = _job.Change;
+
+			PerforceViewMap viewMap = new PerforceViewMap();
+			viewMap.Entries.Add(new PerforceViewMapEntry(true, "...", $"{_stream.Name}/..."));
+
+			PerforceLogger perforceLogger = new PerforceLogger(logger, workspaceDir, viewMap, _job.Change);
 
 			List<string> ignorePatterns = await ReadIgnorePatternsAsync(workspaceDir, logger);
 
 			int exitCode;
-			using (LogParser filter = new LogParser(logger, context, ignorePatterns))
+			using (LogParser filter = new LogParser(perforceLogger, ignorePatterns))
 			{
 				await ExecuteCleanupScriptAsync(cleanupScript, filter, logger);
 				try
diff --git a/Engine/Source/Programs/Horde/Horde.Agent/Execution/TestExecutor.cs b/Engine/Source/Programs/Horde/Horde.Agent/Execution/TestExecutor.cs
--- a/Engine/Source/Programs/Horde/Horde.Agent/Execution/TestExecutor.cs
+++ b/Engine/Source/Programs/Horde/Horde.Agent/Execution/TestExecutor.cs
@@ -7,6 +7,7 @@ using System.Reflection;
 using System.Threading;
 using System.Threading.Tasks;
 using EpicGames.Core;
+using EpicGames.Perforce;
 using Horde.Agent.Parser;
 using Horde.Agent.Utility;
 using HordeCommon;
@@ -146,12 +147,11 @@ namespace Horde.Agent.Execution
 				logger.LogInformation("Credential: {CredentialName}={CredentialValue}", credential.Key, credential.Value);
 			}
 
-			LogParserContext context = new LogParserContext();
-			context.WorkspaceDir = new DirectoryReference("D:\\Test");
-			context.PerforceStream = "//UE4/Main";
-			context.PerforceChange = 12345;
-
-			using (LogParser filter = new LogParser(logger, context, new List<string>()))
+			PerforceViewMap viewMap = new PerforceViewMap();
+			viewMap.Entries.Add(new PerforceViewMapEntry(true, "...", "//UE4/Main/..."));
+			
+			JsonPerforceLogger perforceLogger = new JsonPerforceLogger(logger, new DirectoryReference("D:\\Test"), viewMap, 12345);
+			using (LogParser filter = new LogParser(perforceLogger, new List<string>()))
 			{
 				if(outcome == JobStepOutcome.Warnings)
 				{
diff --git a/Engine/Source/Programs/Horde/Horde.Agent/Parser/Interfaces/ILogContext.cs b/Engine/Source/Programs/Horde/Horde.Agent/Parser/Interfaces/ILogContext.cs
deleted file mode 100644
diff --git a/Engine/Source/Programs/Horde/Horde.Agent/Parser/LogEventMarkup.cs b/Engine/Source/Programs/Horde/Horde.Agent/Parser/LogEventMarkup.cs
--- a/Engine/Source/Programs/Horde/Horde.Agent/Parser/LogEventMarkup.cs
+++ b/Engine/Source/Programs/Horde/Horde.Agent/Parser/LogEventMarkup.cs
@@ -1,10 +1,10 @@
 // Copyright Epic Games, Inc. All Rights Reserved.
 
+using System;
 using System.Collections.Generic;
 using System.IO;
 using System.Text.RegularExpressions;
 using EpicGames.Core;
-using Horde.Agent.Parser.Interfaces;
 
 namespace Horde.Agent.Parser
 {
@@ -13,6 +13,9 @@ namespace Horde.Agent.Parser
 	/// </summary>
 	static class LogEventMarkup
 	{
+		public static readonly Utf8String Asset = new Utf8String("Asset");
+		public static readonly Utf8String SourceFile = new Utf8String("SourceFile");
+
 		public static LogValue Channel => new LogValue("Channel", "");
 		public static LogValue Severity => new LogValue("Severity", "");
 		public static LogValue Message => new LogValue("Message", "");
@@ -26,45 +29,27 @@ namespace Horde.Agent.Parser
 		/// <summary>
 		/// Marks a span of text as a source file
 		/// </summary>
-		public static void AnnotateSourceFile(this LogEventBuilder builder, Group group, ILogContext context, string baseDir)
+		public static void AnnotateSourceFile(this LogEventBuilder builder, Group group, string? baseDir)
 		{
-			LogValue? value = null;
-			if (context.WorkspaceDir != null && context.PerforceStream != null && context.PerforceChange != null)
+			Dictionary<Utf8String, object>? properties = null;
+			if (!String.IsNullOrEmpty(baseDir))
 			{
-				FileReference location = FileReference.Combine(context.WorkspaceDir, baseDir.Replace('\\', Path.DirectorySeparatorChar), group.Value.Replace('\\', Path.DirectorySeparatorChar));
-				if (location.IsUnderDirectory(context.WorkspaceDir) && !location.ContainsName("Intermediate", context.WorkspaceDir))
+				string file = group.Value;
+				if (!Path.IsPathRooted(file))
 				{
-					string relativePath = location.MakeRelativeTo(context.WorkspaceDir).Replace('\\', '/');
-					string depotPath = $"{context.PerforceStream.TrimEnd('/')}/{relativePath.Replace(Path.DirectorySeparatorChar, '/')}@{context.PerforceChange.Value}";
-
-					Dictionary<string, object> properties = new Dictionary<string, object>();
-					properties["relativePath"] = relativePath;
-					properties["depotPath"] = depotPath;
-					value = new LogValue("SourceFile", "", properties);
+					properties = new Dictionary<Utf8String, object>();
+					properties["file"] = Path.Combine(baseDir, file);
 				}
 			}
-			builder.Annotate(group, value);
+			builder.Annotate(group, new LogValue(SourceFile, group.Value, properties));
 		}
 
 		/// <summary>
 		/// Marks a span of text as a source file
 		/// </summary>
-		public static void AnnotateAsset(this LogEventBuilder builder, Group group, ILogContext context)
+		public static void AnnotateAsset(this LogEventBuilder builder, Group group)
 		{
-			if (context.WorkspaceDir != null && context.PerforceStream != null && context.PerforceChange != null)
-			{
-				FileReference location = FileReference.Combine(DirectoryReference.Combine(context.WorkspaceDir, "Engine", "Binaries", "Win64"), group.Value);
-				if (location.IsUnderDirectory(context.WorkspaceDir) && !location.ContainsName("Intermediate", context.WorkspaceDir))
-				{
-					string relativePath = location.MakeRelativeTo(context.WorkspaceDir).Replace('\\', '/');
-					string depotPath = $"{context.PerforceStream.TrimEnd('/')}/{relativePath.Replace(Path.DirectorySeparatorChar, '/')}@{context.PerforceChange.Value}";
-
-					Dictionary<string, object> properties = new Dictionary<string, object>();
-					properties["relativePath"] = relativePath;
-					properties["depotPath"] = depotPath;
-					builder.Annotate(group, new LogValue("Asset", "", properties));
-				}
-			}
+			builder.Annotate(group, new LogValue(Asset, group.Value));
 		}
 
 		/// <summary>
@@ -84,7 +69,7 @@ namespace Horde.Agent.Parser
 			identifier = Regex.Replace(identifier, "^.* ", "");
 
 			// Add it to the list
-			Dictionary<string, object> properties = new Dictionary<string, object>();
+			Dictionary<Utf8String, object> properties = new Dictionary<Utf8String, object>();
 			properties["identifier"] = identifier;
 			builder.Annotate(group, new LogValue("symbol", "", properties));
 		}
diff --git a/Engine/Source/Programs/Horde/Horde.Agent/Parser/LogParser.cs b/Engine/Source/Programs/Horde/Horde.Agent/Parser/LogParser.cs
--- a/Engine/Source/Programs/Horde/Horde.Agent/Parser/LogParser.cs
+++ b/Engine/Source/Programs/Horde/Horde.Agent/Parser/LogParser.cs
@@ -2,85 +2,41 @@
 
 using System;
 using System.Collections.Generic;
-using System.Linq;
 using System.Reflection;
 using System.Text.RegularExpressions;
 using EpicGames.Core;
-using Horde.Agent.Parser.Interfaces;
 using Microsoft.Extensions.Logging;
 
 namespace Horde.Agent.Parser
 {
-	/// <summary>
-	/// Concrete implementation of <see cref="ILogContext"/>
-	/// </summary>
-	public class LogParserContext : ILogContext
-	{
-		/// <inheritdoc/>
-		public DirectoryReference? WorkspaceDir { get; set; }
-
-		/// <inheritdoc/>
-		public string? PerforceStream { get; set; }
-
-		/// <inheritdoc/>
-		public int? PerforceChange { get; set; }
-
-		/// <inheritdoc/>
-		public bool HasLoggedErrors { get; set; }
-	}
-
 	/// <summary>
 	/// Turns raw text output into structured logging events
 	/// </summary>
 	public class LogParser : LogEventParser
 	{
-		readonly LogParserContext _context;
-
 		/// <summary>
 		/// Constructor
 		/// </summary>
 		/// <param name="logger">The logger to receive parsed output messages</param>
-		/// <param name="context">Context for parsing this log</param>
 		/// <param name="ignorePatterns">List of patterns to ignore</param>
-		public LogParser(ILogger logger, LogParserContext context, List<string> ignorePatterns)
+		public LogParser(ILogger logger, List<string> ignorePatterns)
 			: base(logger)
 		{
-			_context = context;
 			foreach (string ignorePattern in ignorePatterns)
 			{
-				IgnorePatterns.Add(new Regex(ignorePattern));
+				IgnorePatterns.Add(new Regex(ignorePattern, RegexOptions.Compiled));
 			}
 
 			foreach (Type type in Assembly.GetExecutingAssembly().GetTypes())
 			{
 				if (type.IsClass && typeof(ILogEventMatcher).IsAssignableFrom(type))
 				{
-					ILogEventMatcher matcher;
-					if (type.GetConstructor(new[] { typeof(ILogContext) }) != null)
-					{
-						matcher = (ILogEventMatcher)Activator.CreateInstance(type, (ILogContext)context)!;
-					}
-					else
-					{
-						matcher = (ILogEventMatcher)Activator.CreateInstance(type)!;
-					}
+					ILogEventMatcher matcher = (ILogEventMatcher)Activator.CreateInstance(type)!;
 					Matchers.Add(matcher);
 				}
 			}
 		}
 
-		/// <inheritdoc/>
-		protected override void WriteEvents(List<LogEvent> events)
-		{
-			base.WriteEvents(events);
-
-			// Also flag that an error has occurred for future add the error to the log context, so that future errors can examine it
-			if (events.Any(x => x.Level >= LogLevel.Error))
-			{
-				_context.HasLoggedErrors = true;
-			}
-		}
-
 		/// <summary>
 		/// Static constructor
 		/// </summary>
diff --git a/Engine/Source/Programs/Horde/Horde.Agent/Parser/Matchers/CompileEventMatcher.cs b/Engine/Source/Programs/Horde/Horde.Agent/Parser/Matchers/CompileEventMatcher.cs
--- a/Engine/Source/Programs/Horde/Horde.Agent/Parser/Matchers/CompileEventMatcher.cs
+++ b/Engine/Source/Programs/Horde/Horde.Agent/Parser/Matchers/CompileEventMatcher.cs
@@ -5,7 +5,6 @@ using System.Diagnostics.CodeAnalysis;
 using System.Linq;
 using System.Text.RegularExpressions;
 using EpicGames.Core;
-using Horde.Agent.Parser.Interfaces;
 using Microsoft.Extensions.Logging;
 
 namespace Horde.Agent.Parser.Matchers
@@ -54,13 +53,6 @@ namespace Horde.Agent.Parser.Matchers
 
 		const string DefaultSourceFileBaseDir = "Engine/Source";
 
-		private readonly ILogContext _context;
-
-		public CompileEventMatcher(ILogContext context)
-		{
-			_context = context;
-		}
-
 		/// <inheritdoc/>
 		public LogEventMatch? Match(ILogCursor input)
 		{
@@ -115,7 +107,7 @@ namespace Horde.Agent.Parser.Matchers
 				{
 					LogLevel level = GetLogLevelFromSeverity(match);
 
-					builder.AnnotateSourceFile(match.Groups["file"], _context, "Engine/Source");
+					builder.AnnotateSourceFile(match.Groups["file"], "Engine/Source");
 					builder.Annotate(match.Groups["severity"], LogEventMarkup.Severity);
 					builder.TryAnnotate(match.Groups["line"], LogEventMarkup.LineNumber);
 					builder.TryAnnotate(match.Groups["column"], LogEventMarkup.ColumnNumber);
@@ -129,7 +121,7 @@ namespace Horde.Agent.Parser.Matchers
 						Group fileGroup = match.Groups["file"];
 						if (fileGroup.Success)
 						{
-							builder.AnnotateSourceFile(fileGroup, _context, DefaultSourceFileBaseDir);
+							builder.AnnotateSourceFile(fileGroup, DefaultSourceFileBaseDir);
 							builder.TryAnnotate(match.Groups["line"], LogEventMarkup.LineNumber);
 						}
 					}
@@ -164,7 +156,7 @@ namespace Horde.Agent.Parser.Matchers
 					Match? projectMatch;
 					if (builder.Current.TryMatch(@"\[(?<project>[^[\]]+)]\s*$", out projectMatch))
 					{
-						builder.AnnotateSourceFile(projectMatch.Groups[1], _context, "");
+						builder.AnnotateSourceFile(projectMatch.Groups[1], "");
 						sourceFileBaseDir = GetPlatformAgnosticDirectoryName(projectMatch.Groups[1].Value) ?? sourceFileBaseDir;
 					}
 				}
@@ -179,14 +171,14 @@ namespace Horde.Agent.Parser.Matchers
 					Match? projectMatch;
 					if (builder.Current.TryMatch(@"\[(?<file>[^[\]]+)]\s*$", out projectMatch))
 					{
-						builder.AnnotateSourceFile(projectMatch.Groups[1], _context, "");
+						builder.AnnotateSourceFile(projectMatch.Groups[1], "");
 						outEvent = builder.ToMatch(LogEventPriority.High, level, KnownLogEvents.MSBuild);
 						return true;
 					}
 				}
 			}
 
-			builder.AnnotateSourceFile(match.Groups["file"], _context, sourceFileBaseDir);
+			builder.AnnotateSourceFile(match.Groups["file"], sourceFileBaseDir);
 			builder.TryAnnotate(match.Groups["line"], LogEventMarkup.LineNumber);
 			builder.TryAnnotate(match.Groups["column"], LogEventMarkup.ColumnNumber);
 
@@ -199,7 +191,7 @@ namespace Horde.Agent.Parser.Matchers
 				Group group = match.Groups["file"];
 				if (group.Success)
 				{
-					builder.AnnotateSourceFile(group, _context, DefaultSourceFileBaseDir);
+					builder.AnnotateSourceFile(group, DefaultSourceFileBaseDir);
 					builder.TryAnnotate(match.Groups["line"], LogEventMarkup.LineNumber);
 					builder.AddProperty("note", true);
 				}
diff --git a/Engine/Source/Programs/Horde/Horde.Agent/Parser/Matchers/ContentEventMatcher.cs b/Engine/Source/Programs/Horde/Horde.Agent/Parser/Matchers/ContentEventMatcher.cs
--- a/Engine/Source/Programs/Horde/Horde.Agent/Parser/Matchers/ContentEventMatcher.cs
+++ b/Engine/Source/Programs/Horde/Horde.Agent/Parser/Matchers/ContentEventMatcher.cs
@@ -3,7 +3,6 @@
 using System;
 using System.Text.RegularExpressions;
 using EpicGames.Core;
-using Horde.Agent.Parser.Interfaces;
 using Microsoft.Extensions.Logging;
 
 namespace Horde.Agent.Parser.Matchers
@@ -21,13 +20,6 @@ namespace Horde.Agent.Parser.Matchers
 			@"(?<asset>(?:[a-zA-Z]:)?[^:]+(?:.uasset|.umap)):\s*" + 
 			@"(?<message>.*)";
 
-		private readonly ILogContext _context;
-
-		public ContentEventMatcher(ILogContext context)
-		{
-			_context = context;
-		}
-
 		/// <inheritdoc/>
 		public LogEventMatch? Match(ILogCursor input)
 		{
@@ -42,7 +34,7 @@ namespace Horde.Agent.Parser.Matchers
 
 					builder.Annotate(match.Groups["channel"], LogEventMarkup.Channel);
 					builder.Annotate(match.Groups["severity"], LogEventMarkup.Severity);
-					builder.AnnotateAsset(match.Groups["asset"], _context);
+					builder.AnnotateAsset(match.Groups["asset"]);
 					builder.Annotate(match.Groups["message"], LogEventMarkup.Message);
 
 					return builder.ToMatch(LogEventPriority.AboveNormal, GetLogLevelFromSeverity(match), KnownLogEvents.Engine_AssetLog);
diff --git a/Engine/Source/Programs/Horde/Horde.Agent/Parser/Matchers/MicrosoftEventMatcher.cs b/Engine/Source/Programs/Horde/Horde.Agent/Parser/Matchers/MicrosoftEventMatcher.cs
--- a/Engine/Source/Programs/Horde/Horde.Agent/Parser/Matchers/MicrosoftEventMatcher.cs
+++ b/Engine/Source/Programs/Horde/Horde.Agent/Parser/Matchers/MicrosoftEventMatcher.cs
@@ -2,7 +2,6 @@
 
 using System.Text.RegularExpressions;
 using EpicGames.Core;
-using Horde.Agent.Parser.Interfaces;
 using Microsoft.Extensions.Logging;
 
 namespace Horde.Agent.Parser.Matchers
@@ -12,13 +11,6 @@ namespace Horde.Agent.Parser.Matchers
 	/// </summary>
 	class MicrosoftEventMatcher : ILogEventMatcher
 	{
-		private readonly ILogContext _context;
-
-		public MicrosoftEventMatcher(ILogContext context)
-		{
-			_context = context;
-		}
-
 		public LogEventMatch? Match(ILogCursor cursor)
 		{
 			// filename(line# [, column#]) | toolname} : [ any text ] {error | warning} code+number:localizable string [ any text ]
@@ -36,7 +28,7 @@ namespace Horde.Agent.Parser.Matchers
 					Match fileMatch = Regex.Match(fileOrToolMatch.Value, @"^\s*(?<file>.*)\((?<line>\d+)(?:, (?<column>\d+))?\)\s*:$");
 					if (fileMatch.Success)
 					{
-						builder.AnnotateSourceFile(fileMatch.Groups["file"], _context, "");
+						builder.AnnotateSourceFile(fileMatch.Groups["file"], null);
 						builder.Annotate(fileMatch.Groups["line"], LogEventMarkup.LineNumber);
 						builder.TryAnnotate(fileMatch.Groups["column"]);
 					}
diff --git a/Engine/Source/Programs/Horde/Horde.Agent/Parser/Matchers/SourceFileLineEventMatcher.cs b/Engine/Source/Programs/Horde/Horde.Agent/Parser/Matchers/SourceFileLineEventMatcher.cs
--- a/Engine/Source/Programs/Horde/Horde.Agent/Parser/Matchers/SourceFileLineEventMatcher.cs
+++ b/Engine/Source/Programs/Horde/Horde.Agent/Parser/Matchers/SourceFileLineEventMatcher.cs
@@ -3,7 +3,6 @@
 using System;
 using System.Text.RegularExpressions;
 using EpicGames.Core;
-using Horde.Agent.Parser.Interfaces;
 using Microsoft.Extensions.Logging;
 
 namespace Horde.Agent.Parser.Matchers
@@ -29,13 +28,6 @@ namespace Horde.Agent.Parser.Matchers
 		const string LinePattern =
 			@"(?<line>\d+)";
 
-		readonly ILogContext _context;
-
-		public SourceFileLineEventMatcher(ILogContext context)
-		{
-			_context = context;
-		}
-
 		/// <inheritdoc/>
 		public LogEventMatch? Match(ILogCursor input)
 		{
@@ -46,7 +38,7 @@ namespace Horde.Agent.Parser.Matchers
 
 				LogEventBuilder builder = new LogEventBuilder(input);
 
-				builder.AnnotateSourceFile(match.Groups["file"], _context, "");
+				builder.AnnotateSourceFile(match.Groups["file"], "");
 				builder.Annotate(match.Groups["severity"], LogEventMarkup.Severity);
 				builder.TryAnnotate(match.Groups["line"], LogEventMarkup.LineNumber);
 
diff --git a/Engine/Source/Programs/Horde/Horde.Agent/Utility/JsonLogger.cs b/Engine/Source/Programs/Horde/Horde.Agent/Utility/JsonLogger.cs
deleted file mode 100644
diff --git a/Engine/Source/Programs/Horde/Horde.Agent/Utility/JsonRpcLogger.cs b/Engine/Source/Programs/Horde/Horde.Agent/Utility/JsonRpcLogger.cs
--- a/Engine/Source/Programs/Horde/Horde.Agent/Utility/JsonRpcLogger.cs
+++ b/Engine/Source/Programs/Horde/Horde.Agent/Utility/JsonRpcLogger.cs
@@ -8,6 +8,7 @@ using System.Text.Json;
 using System.Threading;
 using System.Threading.Channels;
 using System.Threading.Tasks;
+using EpicGames.Core;
 using Horde.Agent.Utility;
 using HordeCommon;
 using HordeCommon.Rpc;
@@ -18,7 +19,7 @@ namespace Horde.Agent.Parser
 	/// <summary>
 	/// Class to handle uploading log data to the server in the background
 	/// </summary>
-	class JsonRpcLogger : JsonLogger, IAsyncDisposable
+	class JsonRpcLogger : ILogger, IAsyncDisposable
 	{
 		class QueueItem
 		{
@@ -42,6 +43,8 @@ namespace Horde.Agent.Parser
 		readonly string? _jobId;
 		readonly string? _jobBatchId;
 		readonly string? _jobStepId;
+		readonly bool _warnings;
+		readonly ILogger _inner;
 		readonly Channel<QueueItem> _dataChannel;
 		Task? _dataWriter;
 
@@ -65,19 +68,39 @@ namespace Horde.Agent.Parser
 		/// <param name="warnings">Whether to include warnings in the output</param>
 		/// <param name="inner">Additional logger to write to</param>
 		public JsonRpcLogger(IRpcConnection rpcClient, string logId, string? jobId, string? jobBatchId, string? jobStepId, bool? warnings, ILogger inner)
-			: base(warnings, inner)
 		{
 			_rpcClient = rpcClient;
 			_logId = logId;
 			_jobId = jobId;
 			_jobBatchId = jobBatchId;
 			_jobStepId = jobStepId;
+			_warnings = warnings ?? true;
+			_inner = inner;
 			_dataChannel = Channel.CreateUnbounded<QueueItem>();
 			_dataWriter = Task.Run(() => RunDataWriter());
 			Outcome = JobStepOutcome.Success;
 		}
 
-		protected override void WriteFormattedEvent(LogLevel level, int lineIndex, int lineCount, byte[] line)
+		/// <inheritdoc/>
+		public void Log<TState>(LogLevel logLevel, EventId eventId, TState state, Exception exception, Func<TState, Exception?, string> formatter)
+		{
+			// Downgrade warnings to information if not required
+			if (logLevel == LogLevel.Warning && !_warnings)
+			{
+				logLevel = LogLevel.Information;
+			}
+
+			LogEvent logEvent = LogEvent.FromState(logLevel, eventId, state, exception, formatter);
+			WriteFormattedEvent(logLevel, logEvent.ToJsonBytes());
+		}
+
+		/// <inheritdoc/>
+		public bool IsEnabled(LogLevel logLevel) => _inner.IsEnabled(logLevel);
+
+		/// <inheritdoc/>
+		public IDisposable BeginScope<TState>(TState state) => _inner.BeginScope(state);
+
+		protected void WriteFormattedEvent(LogLevel level, byte[] line)
 		{
 			// Update the state of this job if this is an error status
 			if (level == LogLevel.Error || level == LogLevel.Critical)
@@ -91,9 +114,10 @@ namespace Horde.Agent.Parser
 
 			// If we want an event for this log event, create one now
 			CreateEventRequest? eventRequest = null;
-			if (lineIndex == 0)
+			if (level == LogLevel.Warning || level == LogLevel.Error || level == LogLevel.Critical)
 			{
-				if (level == LogLevel.Warning || level == LogLevel.Error || level == LogLevel.Critical)
+				int lineCount = GetEventLineCount(line);
+				if (lineCount > 0)
 				{
 					eventRequest = CreateEvent(level, lineCount);
 				}
@@ -107,29 +131,29 @@ namespace Horde.Agent.Parser
 			}
 		}
 
-		/// <summary>
-		/// Callback to write a systemic event
-		/// </summary>
-		/// <param name="eventId">The event id</param>
-		/// <param name="text">The event text</param>
-		protected override void WriteSystemicEvent(EventId eventId, string text)
+		static int GetEventLineCount(byte[] line)
 		{
-			if (_jobId == null)
+			Utf8JsonReader reader = new Utf8JsonReader(line);
+			if(!reader.Read() || reader.TokenType != JsonTokenType.StartObject)
 			{
-				Inner.LogWarning("Systemic event {KnownLogEventId} in log {LogId}: {Text}", eventId.Id, text);
+				return 0;
 			}
-			else if (_jobBatchId == null)
-			{
-				Inner.LogWarning("Systemic event {KnownLogEventId} in log {LogId} (job {JobId}): {Text})", eventId.Id, _jobId, text);
-			}
-			else if (_jobStepId == null)
-			{
-				Inner.LogWarning("Systemic event {KnownLogEventId} in log {LogId} (job batch {JobId}:{BatchId}): {Text})", eventId.Id, _jobId, _jobStepId, text);
-			}
-			else
+
+			int lineCount = 1;
+			while(reader.Read() && reader.TokenType == JsonTokenType.PropertyName)
 			{
-				Inner.LogWarning("Systemic event {KnownLogEventId} in log {LogId} (job step {JobId}:{BatchId}:{StepId}): {Text}", eventId.Id, _logId, _jobId, _jobBatchId, _jobStepId, text);
+				if(reader.ValueTextEquals(LogEventPropertyName.Line) && reader.TryGetInt32(out int lineIndex) && lineIndex > 0)
+				{
+					return 0;
+				}
+				else if(reader.ValueTextEquals(LogEventPropertyName.LineCount) && reader.TryGetInt32(out int lineCountValue))
+				{
+					lineCount = lineCountValue;
+				}
+				reader.Skip();
 			}
+
+			return lineCount;
 		}
 
 		/// <summary>
@@ -246,7 +270,7 @@ namespace Horde.Agent.Parser
 					}
 					catch (Exception ex)
 					{
-						Inner.LogWarning(ex, "Unable to write data to server (log {LogId}, offset {Offset})", _logId, offset);
+						_inner.LogWarning(ex, "Unable to write data to server (log {LogId}, offset {Offset})", _logId, offset);
 					}
 					offset += data.Length;
 				}
@@ -260,7 +284,7 @@ namespace Horde.Agent.Parser
 					}
 					catch (Exception ex)
 					{
-						Inner.LogWarning(ex, "Unable to create events");
+						_inner.LogWarning(ex, "Unable to create events");
 					}
 				}
 
@@ -273,7 +297,7 @@ namespace Horde.Agent.Parser
 					}
 					catch (Exception ex)
 					{
-						Inner.LogWarning(ex, "Unable to update step outcome to {NewOutcome}", Outcome);
+						_inner.LogWarning(ex, "Unable to update step outcome to {NewOutcome}", Outcome);
 					}
 					postedOutcome = Outcome;
 				}
@@ -287,7 +311,7 @@ namespace Horde.Agent.Parser
 					}
 					catch (Exception ex)
 					{
-						Inner.LogWarning(ex, "Unable to flush data to server (log {LogId}, offset {Offset})", _logId, offset);
+						_inner.LogWarning(ex, "Unable to flush data to server (log {LogId}, offset {Offset})", _logId, offset);
 					}
 					break;
 				}
diff --git a/Engine/Source/Programs/Horde/Horde.Agent/Utility/PerforceLogger.cs b/Engine/Source/Programs/Horde/Horde.Agent/Utility/PerforceLogger.cs
new file mode 100644
--- /dev/null
+++ b/Engine/Source/Programs/Horde/Horde.Agent/Utility/PerforceLogger.cs
@@ -0,0 +1,196 @@
+// Copyright Epic Games, Inc. All Rights Reserved.
+
+using System;
+using System.IO;
+using System.Text;
+using System.Text.Json;
+using EpicGames.Core;
+using EpicGames.Perforce;
+using Microsoft.Extensions.Logging;
+
+namespace Horde.Agent.Utility
+{
+	class PerforceLogger : ILogger
+	{
+		readonly ILogger _inner;
+		readonly DirectoryReference _baseDir;
+		readonly PerforceViewMap? _viewMap;
+		readonly int? _viewChange;
+
+		public PerforceLogger(ILogger inner, DirectoryReference baseDir, PerforceViewMap? viewMap, int? viewChange)
+		{
+			_inner = inner;
+			_baseDir = baseDir;
+			_viewMap = viewMap;
+			_viewChange = viewChange;
+		}
+
+		/// <inheritdoc/>
+		public void Log<TState>(LogLevel logLevel, EventId eventId, TState state, Exception exception, Func<TState, Exception?, string> formatter)
+		{
+			JsonLogEvent logEvent = JsonLogEvent.FromLoggerState(logLevel, eventId, state, exception, formatter);
+			logEvent = new JsonLogEvent(logLevel, eventId, Annotate(logEvent.Data));
+			_inner.Log(logLevel, eventId, logEvent, null, JsonLogEvent.Format);
+		}
+
+		/// <inheritdoc/>
+		public bool IsEnabled(LogLevel logLevel) => _inner.IsEnabled(logLevel);
+
+		/// <inheritdoc/>
+		public IDisposable BeginScope<TState>(TState state) => _inner.BeginScope(state);
+
+		static bool ReadFirstLogProperty(ref Utf8JsonReader reader)
+		{
+			// Enter the main object
+			if (!reader.Read() || reader.TokenType != JsonTokenType.StartObject)
+			{
+				return false;
+			}
+
+			// Find the 'properties' property
+			for (; ; )
+			{
+				if (!reader.Read() || reader.TokenType != JsonTokenType.PropertyName)
+				{
+					return false;
+				}
+				if (reader.ValueTextEquals(LogEventPropertyName.Properties))
+				{
+					break;
+				}
+				reader.Skip();
+			}
+
+			// Enter the properties object
+			if (!reader.Read() || reader.TokenType != JsonTokenType.StartObject)
+			{
+				return false;
+			}
+
+			// Find the first structured property object
+			return ReadNextLogProperty(ref reader);
+		}
+
+		static bool ReadNextLogProperty(ref Utf8JsonReader reader)
+		{
+			for (; ; )
+			{
+				// Move to the next property name
+				if (!reader.Read() || reader.TokenType != JsonTokenType.PropertyName)
+				{
+					return false;
+				}
+
+				// Move to the next property value
+				if (!reader.Read())
+				{
+					return false;
+				}
+
+				// If it's an object, enter it
+				if (reader.TokenType == JsonTokenType.StartObject)
+				{
+					return true;
+				}
+
+				// Otherwise skip this value
+				reader.Skip();
+			}
+		}
+
+		static readonly Utf8String s_sourceFileType = "SourceFile";
+		static readonly Utf8String s_assetType = "Asset";
+
+		static readonly Utf8String s_file = "file";
+
+		ReadOnlyMemory<byte> Annotate(ReadOnlyMemory<byte> data)
+		{
+			ReadOnlyMemory<byte> output = data;
+			int shift = 0;
+
+			Utf8JsonReader reader = new Utf8JsonReader(data.Span);
+			for (bool valid = ReadFirstLogProperty(ref reader); valid; valid = ReadNextLogProperty(ref reader))
+			{
+				ReadOnlySpan<byte> type = ReadOnlySpan<byte>.Empty;
+				ReadOnlySpan<byte> text = ReadOnlySpan<byte>.Empty;
+				string? file = null;
+
+				// Get the $type and $text properties
+				while (reader.Read() && reader.TokenType == JsonTokenType.PropertyName)
+				{
+					if (type.Length == 0 && reader.ValueTextEquals(LogEventPropertyName.Type))
+					{
+						if (reader.Read() && reader.TokenType == JsonTokenType.String)
+						{
+							type = reader.GetUtf8String();
+							continue;
+						}
+					}
+					if (text.Length == 0 && reader.ValueTextEquals(LogEventPropertyName.Text))
+					{
+						if (reader.Read() && reader.TokenType == JsonTokenType.String)
+						{
+							text = reader.GetUtf8String();
+							continue;
+						}
+					}
+					if (file == null && reader.ValueTextEquals(s_file))
+					{
+						if (reader.Read() && reader.TokenType == JsonTokenType.String)
+						{
+							file = reader.GetString();
+							continue;
+						}
+					}
+					reader.Skip();
+				}
+
+				// If we're at the end of the object, append any additional data
+				byte[]? annotationBytes = GetAnnotations(type, text, file);
+				if (annotationBytes != null)
+				{
+					int position = (int)reader.TokenStartIndex + shift;
+
+					byte[] newOutput = new byte[output.Length + annotationBytes.Length];
+					output.Span.Slice(0, position).CopyTo(newOutput);
+					annotationBytes.CopyTo(newOutput.AsSpan(position));
+					output.Span.Slice(position).CopyTo(newOutput.AsSpan(position + annotationBytes.Length));
+
+					output = newOutput;
+					shift += annotationBytes.Length;
+				}
+			}
+
+			return output;
+		}
+
+		byte[]? GetAnnotations(ReadOnlySpan<byte> type, ReadOnlySpan<byte> text, string? file)
+		{
+			if (type.SequenceEqual(s_sourceFileType) || type.SequenceEqual(s_assetType))
+			{
+				StringBuilder annotations = new StringBuilder();
+
+				file ??= Encoding.UTF8.GetString(text);
+
+				FileReference location = FileReference.Combine(_baseDir, file.Replace('\\', Path.DirectorySeparatorChar));
+				if (location.IsUnderDirectory(_baseDir))
+				{
+					string relativePath = location.MakeRelativeTo(_baseDir).Replace('\\', '/');
+					annotations.Append($",\"relativePath\":\"{JsonEncodedText.Encode(relativePath)}\"");
+
+					if (_viewMap != null && _viewMap.TryMapFile(relativePath, StringComparison.OrdinalIgnoreCase, out string depotFile))
+					{
+						if (_viewChange != null)
+						{
+							depotFile = $"{depotFile}@{_viewChange}";
+						}
+						annotations.Append($",\"depotPath\":\"{JsonEncodedText.Encode(depotFile)}\"");
+					}
+				}
+
+				return Encoding.UTF8.GetBytes(annotations.ToString());
+			}
+			return null;
+		}
+	}
+}
diff --git a/Engine/Source/Programs/Shared/EpicGames.Core/JsonExtensions.cs b/Engine/Source/Programs/Shared/EpicGames.Core/JsonExtensions.cs
--- a/Engine/Source/Programs/Shared/EpicGames.Core/JsonExtensions.cs
+++ b/Engine/Source/Programs/Shared/EpicGames.Core/JsonExtensions.cs
@@ -1,6 +1,7 @@
 // Copyright Epic Games, Inc. All Rights Reserved.
 
 using System;
+using System.Diagnostics;
 using System.Text.Json;
 
 namespace EpicGames.Core
@@ -35,9 +36,170 @@ namespace EpicGames.Core
 		/// </summary>
 		/// <param name="reader"></param>
 		/// <returns></returns>
-		public static Utf8String GetUtf8String(this Utf8JsonReader reader)
+		public static ReadOnlySpan<byte> GetUtf8String(this Utf8JsonReader reader)
 		{
-			throw new NotImplementedException();
+			ReadOnlySpan<byte> span = reader.ValueSpan;
+
+			// Scan for the first backslash
+			for (int firstIdx = 0; firstIdx < span.Length; firstIdx++)
+			{
+				if (span[firstIdx] == '\\')
+				{
+					// Compute the length of the unescaped string
+					int length = span.Length;
+					for(int idx = firstIdx; idx < span.Length; idx++)
+					{
+						if (span[idx] == '\\')
+						{
+							switch (span[idx + 1])
+							{
+								case (byte)'"':
+								case (byte)'\\':
+								case (byte)'/':
+								case (byte)'b':
+								case (byte)'f':
+								case (byte)'n':
+								case (byte)'r':
+								case (byte)'t':
+									idx += 2 - 1;
+									length--;
+									break;
+								case (byte)'u':
+									int c = (StringUtils.GetHexDigit(span[idx + 2]) << 12) + (StringUtils.GetHexDigit(span[idx + 3]) << 8) + (StringUtils.GetHexDigit(span[idx + 4]) << 4) + StringUtils.GetHexDigit(span[idx + 5]);
+									idx += 6 - 1;
+									length += GetUtf8ByteCount(c) - 6;
+									break;
+							}
+						}
+					}
+
+					// Write the new string
+					byte[] result = new byte[length];
+					span.Slice(0, firstIdx).CopyTo(result);
+
+					// Trim the output spans
+					ReadOnlySpan<byte> input = span.Slice(firstIdx);
+					Span<byte> output = result.AsSpan(firstIdx);
+
+					// Compute the length of the unescaped string
+					for(; ;)
+					{
+						switch (input[1])
+						{
+							case (byte)'"':
+								input = input.Slice(2);
+								output[0] = (byte)'\"';
+								output = output.Slice(1);
+								break;
+							case (byte)'\\':
+								input = input.Slice(2);
+								output[0] = (byte)'\\';
+								output = output.Slice(1);
+								break;
+							case (byte)'/':
+								input = input.Slice(2);
+								output[0] = (byte)'/';
+								output = output.Slice(1);
+								break;
+							case (byte)'b':
+								input = input.Slice(2);
+								output[0] = (byte)'\b';
+								output = output.Slice(1);
+								break;
+							case (byte)'f':
+								input = input.Slice(2);
+								output[0] = (byte)'\f';
+								output = output.Slice(1);
+								break;
+							case (byte)'n':
+								input = input.Slice(2);
+								output[0] = (byte)'\n';
+								output = output.Slice(1);
+								break;
+							case (byte)'r':
+								input = input.Slice(2);
+								output[0] = (byte)'\r';
+								output = output.Slice(1);
+								break;
+							case (byte)'t':
+								input = input.Slice(2);
+								output[0] = (byte)'\t';
+								output = output.Slice(1);
+								break;
+							case (byte)'u':
+								int c = (StringUtils.GetHexDigit(input[2]) << 12) + (StringUtils.GetHexDigit(input[3]) << 8) + (StringUtils.GetHexDigit(input[4]) << 4) + StringUtils.GetHexDigit(input[5]);
+								input = input.Slice(6);
+								int len = GetUtf8Bytes(c, output);
+								output = output.Slice(len);
+								break;
+						}
+
+						int escapeIdx = input.IndexOf((byte)'\\');
+						if (escapeIdx == -1)
+						{
+							Debug.Assert(input.Length == output.Length);
+							input.CopyTo(output);
+							return result;
+						}
+
+						input.Slice(0, escapeIdx).CopyTo(output);
+						input = input.Slice(escapeIdx);
+						output = output.Slice(escapeIdx);
+					}
+				}
+			}
+
+			return span;
+		}
+
+		static int GetUtf8Bytes(int c, Span<byte> output)
+		{
+			if (c <= 0x7f)
+			{
+				output[0] = (byte)c;
+				return 1;
+			}
+			else if (c <= 0x7ff)
+			{
+				output[0] = (byte)(0b11000000 | (c >> 6));
+				output[1] = (byte)(0b10000000 | (c & 0x1f));
+				return 2;
+			}
+			else if (c <= 0x7ff)
+			{
+				output[0] = (byte)(0b11100000 | (c >> 12));
+				output[1] = (byte)(0b10000000 | ((c >> 6) & 0x1f));
+				output[2] = (byte)(0b10000000 | (c & 0x1f));
+				return 3;
+			}
+			else
+			{
+				output[0] = (byte)(0b11110000 | (c >> 18));
+				output[1] = (byte)(0b10000000 | ((c >> 12) & 0x1f));
+				output[2] = (byte)(0b10000000 | ((c >> 6) & 0x1f));
+				output[3] = (byte)(0b10000000 | (c & 0x1f));
+				return 4;
+			}
+		}
+
+		static int GetUtf8ByteCount(int c)
+		{
+			if (c <= 0x7f)
+			{
+				return 1;
+			}
+			else if (c <= 0x7ff)
+			{
+				return 2;
+			}
+			else if (c <= 0xffff)
+			{
+				return 3;
+			}
+			else
+			{
+				return 4;
+			}
 		}
 	}
 }
diff --git a/Engine/Source/Programs/Shared/EpicGames.Core/JsonLogEvent.cs b/Engine/Source/Programs/Shared/EpicGames.Core/JsonLogEvent.cs
--- a/Engine/Source/Programs/Shared/EpicGames.Core/JsonLogEvent.cs
+++ b/Engine/Source/Programs/Shared/EpicGames.Core/JsonLogEvent.cs
@@ -12,16 +12,8 @@ namespace EpicGames.Core
 	/// </summary>
 	public struct JsonLogEvent
 	{
-		static Utf8String LevelString { get; } = new Utf8String("Level");
-		static Utf8String PropertiesString { get; } = new Utf8String("Properties");
-		static Utf8String LineIndexString { get; } = new Utf8String("Line");
-		static Utf8String LineCountString { get; } = new Utf8String("LineCount");
-		static Utf8String EventIdString { get; } = new Utf8String("EventId");
-		static Utf8String IdString { get; } = new Utf8String("Id");
-		static Utf8String RenderedMessageString { get; } = new Utf8String("RenderedMessage");
-
 		/// <summary>
-		/// The 
+		/// The log level
 		/// </summary>
 		public LogLevel Level { get; }
 
@@ -30,16 +22,6 @@ namespace EpicGames.Core
 		/// </summary>
 		public EventId EventId { get; }
 
-		/// <summary>
-		/// Index of this line within the current event
-		/// </summary>
-		public int LineIndex { get; }
-
-		/// <summary>
-		/// Number of lines in the current event
-		/// </summary>
-		public int LineCount { get; }
-
 		/// <summary>
 		/// The utf-8 encoded JSON event
 		/// </summary>
@@ -48,15 +30,33 @@ namespace EpicGames.Core
 		/// <summary>
 		/// Constructor
 		/// </summary>
-		public JsonLogEvent(LogLevel level, EventId eventId, int lineIndex, int lineCount, ReadOnlyMemory<byte> data)
+		public JsonLogEvent(LogLevel level, EventId eventId, ReadOnlyMemory<byte> data)
 		{
 			Level = level;
 			EventId = eventId;
-			LineIndex = lineIndex;
-			LineCount = lineCount;
 			Data = data;
 		}
 
+		/// <summary>
+		/// Creates a json log event from the given logger paramters
+		/// </summary>
+		/// <inheritdoc cref="ILogger.Log{TState}(LogLevel, EventId, TState, Exception, Func{TState, Exception, String})"/>
+		public static JsonLogEvent FromLoggerState<TState>(LogLevel logLevel, EventId eventId, TState state, Exception? exception, Func<TState, Exception?, string> formatter)
+		{
+			if (state is JsonLogEvent jsonLogEvent)
+			{
+				return jsonLogEvent;
+			}
+
+			LogEvent? logEvent = state as LogEvent;
+			if (logEvent == null)
+			{
+				logEvent = LogEvent.FromState(logLevel, eventId, state, exception, formatter);
+			}
+
+			return new JsonLogEvent(logLevel, eventId, logEvent.ToJsonBytes());
+		}
+
 		/// <summary>
 		/// Tries to parse a Json log event from the given 
 		/// </summary>
@@ -65,57 +65,36 @@ namespace EpicGames.Core
 		/// <returns></returns>
 		public static bool TryParse(ReadOnlyMemory<byte> data, out JsonLogEvent logEvent)
 		{
-			Utf8StringComparer stringComparer = Utf8StringComparer.OrdinalIgnoreCase;
 			try
 			{
 				LogLevel level = LogLevel.None;
 				int eventId = 0;
-				int lineIndex = 0;
-				int lineCount = 1;
 
 				Utf8JsonReader reader = new Utf8JsonReader(data.Span);
 				if (reader.Read() && reader.TokenType == JsonTokenType.StartObject)
 				{
-					ReadOnlySpan<byte> propertyName;
-					for (; ReadNextPropertyName(ref reader, out propertyName); reader.Skip())
+					while (reader.Read() && reader.TokenType == JsonTokenType.PropertyName)
 					{
-						if (stringComparer.Equals(propertyName, LevelString.Span))
+						ReadOnlySpan<byte> propertyName = reader.ValueSpan;
+						if (!reader.Read())
 						{
-							if (!Enum.TryParse(reader.GetString(), out level))
-							{
-								break;
-							}
+							break;
 						}
-						else if (stringComparer.Equals(propertyName, LineIndexString.Span))
+						else if (propertyName.SequenceEqual(LogEventPropertyName.Level) && reader.TokenType == JsonTokenType.String)
 						{
-							lineIndex = reader.GetInt32();
+							level = ParseLevel(reader.ValueSpan);
 						}
-						else if (stringComparer.Equals(propertyName, LineCountString.Span))
+						else if (propertyName.SequenceEqual(LogEventPropertyName.Id) && reader.TokenType == JsonTokenType.Number)
 						{
-							lineCount = reader.GetInt32();
-						}
-						else if (stringComparer.Equals(propertyName, PropertiesString.Span) && reader.TokenType == JsonTokenType.StartObject)
-						{
-							for (; ReadNextPropertyName(ref reader, out propertyName); reader.Skip())
-							{
-								if (stringComparer.Equals(EventIdString.Span, propertyName) && reader.TokenType == JsonTokenType.StartObject)
-								{
-									for (; ReadNextPropertyName(ref reader, out propertyName); reader.Skip())
-									{
-										if (stringComparer.Equals(IdString.Span, propertyName) && reader.TryGetInt32(out int newEventId))
-										{
-											eventId = newEventId;
-										}
-									}
-								}
-							}
+							eventId = reader.GetInt32();
 						}
+						reader.Skip();
 					}
 				}
 
 				if (reader.TokenType == JsonTokenType.EndObject && level != LogLevel.None && reader.BytesConsumed == data.Length)
 				{
-					logEvent = new JsonLogEvent(level, new EventId(eventId), lineIndex, lineCount, data);
+					logEvent = new JsonLogEvent(level, new EventId(eventId), data);
 					return true;
 				}
 			}
@@ -127,59 +106,44 @@ namespace EpicGames.Core
 			return false;
 		}
 
-		static bool ReadNextPropertyName(ref Utf8JsonReader reader, out ReadOnlySpan<byte> propertyName)
+		static readonly sbyte[] s_firstCharToLogLevel;
+		static readonly byte[][] s_logLevelNames;
+
+		static JsonLogEvent()
 		{
-			if (reader.Read() && reader.TokenType == JsonTokenType.PropertyName)
+			const int LogLevelCount = (int)LogLevel.None;
+
+			s_firstCharToLogLevel = new sbyte[256];
+			Array.Fill(s_firstCharToLogLevel, (sbyte)-1);
+
+			s_logLevelNames = new byte[LogLevelCount][];
+			for (int idx = 0; idx < (int)LogLevel.None; idx++)
 			{
-				propertyName = reader.ValueSpan;
-				return reader.Read();
+				byte[] name = Encoding.UTF8.GetBytes(Enum.GetName(typeof(LogLevel), (LogLevel)idx)!);
+				s_logLevelNames[idx] = name;
+				s_firstCharToLogLevel[name[0]] = (sbyte)idx;
 			}
-			else
+		}
+
+		static LogLevel ParseLevel(ReadOnlySpan<byte> level)
+		{
+			int result = s_firstCharToLogLevel[level[0]];
+			if (!level.SequenceEqual(s_logLevelNames[result]))
 			{
-				propertyName = ReadOnlySpan<byte>.Empty;
-				return false;
+				throw new InvalidOperationException();
 			}
+			return (LogLevel)result;
 		}
 
 		/// <summary>
 		/// Formats an event as a string
 		/// </summary>
 		/// <param name="state"></param>
-		/// <param name="ex"></param>
 		/// <returns></returns>
 		/// <exception cref="NotImplementedException"></exception>
-#pragma warning disable IDE0060 // Remove unused parameter
-		public static string Format(JsonLogEvent state, Exception? ex)
-#pragma warning restore IDE0060 // Remove unused parameter
-		{
-			Utf8JsonReader reader = new Utf8JsonReader(state.Data.Span);
-			if (reader.Read() && reader.TokenType == JsonTokenType.StartObject)
-			{
-				for (; ; )
-				{
-					if (!reader.Read() || reader.TokenType != JsonTokenType.PropertyName)
-					{
-						break;
-					}
-
-					bool isMessage = Utf8StringComparer.OrdinalIgnoreCase.Equals(reader.ValueSpan, RenderedMessageString.Span);
-					if (!reader.TrySkip())
-					{
-						break;
-					}
-					if (isMessage && reader.TokenType == JsonTokenType.String)
-					{
-						return reader.GetString();
-					}
-				}
-			}
-			return String.Empty;
-		}
+		public static string Format(JsonLogEvent state, Exception? _) => LogEvent.Read(state.Data.Span).ToString();
 
 		/// <inheritdoc/>
-		public override string ToString()
-		{
-			return Encoding.UTF8.GetString(Data.ToArray());
-		}
+		public override string ToString() => Encoding.UTF8.GetString(Data.ToArray());
 	}
 }
diff --git a/Engine/Source/Programs/Shared/EpicGames.Core/LogEvent.cs b/Engine/Source/Programs/Shared/EpicGames.Core/LogEvent.cs
--- a/Engine/Source/Programs/Shared/EpicGames.Core/LogEvent.cs
+++ b/Engine/Source/Programs/Shared/EpicGames.Core/LogEvent.cs
@@ -13,6 +13,29 @@ using Microsoft.Extensions.Logging;
 
 namespace EpicGames.Core
 {
+	/// <summary>
+	/// Static read-only utf8 strings for parsing log events
+	/// </summary>
+	public static class LogEventPropertyName
+	{
+		public static readonly Utf8String Time = new Utf8String("time");
+		public static readonly Utf8String Level = new Utf8String("level");
+		public static readonly Utf8String Id = new Utf8String("id");
+		public static readonly Utf8String Line = new Utf8String("line");
+		public static readonly Utf8String LineCount = new Utf8String("lineCount");
+		public static readonly Utf8String Message = new Utf8String("message");
+		public static readonly Utf8String Format = new Utf8String("format");
+		public static readonly Utf8String Properties = new Utf8String("properties");
+
+		public static readonly Utf8String Type = new Utf8String("$type");
+		public static readonly Utf8String Text = new Utf8String("$text");
+
+		public static readonly Utf8String Exception = new Utf8String("exception");
+		public static readonly Utf8String Trace = new Utf8String("trace");
+		public static readonly Utf8String InnerException = new Utf8String("innerException");
+		public static readonly Utf8String InnerExceptions = new Utf8String("innerExceptions");
+	}
+
 	/// <summary>
 	/// Epic representation of a log event. Can be serialized to/from Json for the Horde dashboard, and passed directly through ILogger interfaces.
 	/// </summary>
@@ -25,26 +48,6 @@ namespace EpicGames.Core
 			public static string LineCount = "$lineCount";
 		}
 
-		static class JsonPropertyNames
-		{
-			public static Utf8String Time { get; } = new Utf8String("time");
-			public static Utf8String Level { get; } = new Utf8String("level");
-			public static Utf8String Id { get; } = new Utf8String("id");
-			public static Utf8String Line { get; } = new Utf8String("line");
-			public static Utf8String LineCount { get; } = new Utf8String("lineCount");
-			public static Utf8String Message { get; } = new Utf8String("message");
-			public static Utf8String Format { get; } = new Utf8String("format");
-			public static Utf8String Properties { get; } = new Utf8String("properties");
-
-			public static Utf8String Type { get; } = new Utf8String("$type");
-			public static Utf8String Text { get; } = new Utf8String("$text");
-
-			public static Utf8String Exception { get; } = new Utf8String("exception");
-			public static Utf8String Trace { get; } = new Utf8String("trace");
-			public static Utf8String InnerException { get; } = new Utf8String("innerException");
-			public static Utf8String InnerExceptions { get; } = new Utf8String("innerExceptions");
-		}
-
 		/// <summary>
 		/// Time that the event was emitted
 		/// </summary>
@@ -146,35 +149,35 @@ namespace EpicGames.Core
 			ReadOnlySpan<byte> propertyName;
 			for (; JsonExtensions.TryReadNextPropertyName(ref reader, out propertyName); reader.Skip())
 			{
-				if (Utf8StringComparer.OrdinalIgnoreCase.Equals(propertyName, JsonPropertyNames.Time.Span))
+				if (Utf8StringComparer.OrdinalIgnoreCase.Equals(propertyName, LogEventPropertyName.Time.Span))
 				{
 					time = reader.GetDateTime();
 				}
-				else if (Utf8StringComparer.OrdinalIgnoreCase.Equals(propertyName, JsonPropertyNames.Level.Span))
+				else if (Utf8StringComparer.OrdinalIgnoreCase.Equals(propertyName, LogEventPropertyName.Level.Span))
 				{
 					level = Enum.Parse<LogLevel>(reader.GetString());
 				}
-				else if (Utf8StringComparer.OrdinalIgnoreCase.Equals(propertyName, JsonPropertyNames.Id.Span))
+				else if (Utf8StringComparer.OrdinalIgnoreCase.Equals(propertyName, LogEventPropertyName.Id.Span))
 				{
 					eventId = reader.GetInt32();
 				}
-				else if (Utf8StringComparer.OrdinalIgnoreCase.Equals(propertyName, JsonPropertyNames.Line.Span))
+				else if (Utf8StringComparer.OrdinalIgnoreCase.Equals(propertyName, LogEventPropertyName.Line.Span))
 				{
 					line = reader.GetInt32();
 				}
-				else if (Utf8StringComparer.OrdinalIgnoreCase.Equals(propertyName, JsonPropertyNames.LineCount.Span))
+				else if (Utf8StringComparer.OrdinalIgnoreCase.Equals(propertyName, LogEventPropertyName.LineCount.Span))
 				{
 					lineCount = reader.GetInt32();
 				}
-				else if (Utf8StringComparer.OrdinalIgnoreCase.Equals(propertyName, JsonPropertyNames.Message.Span))
+				else if (Utf8StringComparer.OrdinalIgnoreCase.Equals(propertyName, LogEventPropertyName.Message.Span))
 				{
 					message = reader.GetString();
 				}
-				else if (Utf8StringComparer.OrdinalIgnoreCase.Equals(propertyName, JsonPropertyNames.Format.Span))
+				else if (Utf8StringComparer.OrdinalIgnoreCase.Equals(propertyName, LogEventPropertyName.Format.Span))
 				{
 					format = reader.GetString();
 				}
-				else if (Utf8StringComparer.OrdinalIgnoreCase.Equals(propertyName, JsonPropertyNames.Properties.Span))
+				else if (Utf8StringComparer.OrdinalIgnoreCase.Equals(propertyName, LogEventPropertyName.Properties.Span))
 				{
 					properties = ReadProperties(ref reader);
 				}
@@ -221,23 +224,23 @@ namespace EpicGames.Core
 		{
 			string type = String.Empty;
 			string text = String.Empty;
-			Dictionary<string, object>? properties = null;
+			Dictionary<Utf8String, object>? properties = null;
 
 			ReadOnlySpan<byte> propertyName;
 			for (; JsonExtensions.TryReadNextPropertyName(ref reader, out propertyName); reader.Skip())
 			{
-				if (Utf8StringComparer.OrdinalIgnoreCase.Equals(propertyName, JsonPropertyNames.Type.Span))
+				if (Utf8StringComparer.OrdinalIgnoreCase.Equals(propertyName, LogEventPropertyName.Type.Span))
 				{
 					type = reader.GetString();
 				}
-				else if (Utf8StringComparer.OrdinalIgnoreCase.Equals(propertyName, JsonPropertyNames.Text.Span))
+				else if (Utf8StringComparer.OrdinalIgnoreCase.Equals(propertyName, LogEventPropertyName.Text.Span))
 				{
 					text = reader.GetString();
 				}
 				else
 				{
-					properties ??= new Dictionary<string, object>();
-					properties.Add(Encoding.UTF8.GetString(propertyName), ReadPropertyValue(ref reader));
+					properties ??= new Dictionary<Utf8String, object>();
+					properties.Add(new Utf8String(propertyName.ToArray()), ReadPropertyValue(ref reader));
 				}
 			}
 
@@ -251,33 +254,33 @@ namespace EpicGames.Core
 		public void Write(Utf8JsonWriter writer)
 		{
 			writer.WriteStartObject();
-			writer.WriteString(JsonPropertyNames.Time.Span, Time.ToString("s", CultureInfo.InvariantCulture));
-			writer.WriteString(JsonPropertyNames.Level.Span, Level.ToString());
-			writer.WriteString(JsonPropertyNames.Message.Span, Message);
+			writer.WriteString(LogEventPropertyName.Time.Span, Time.ToString("s", CultureInfo.InvariantCulture));
+			writer.WriteString(LogEventPropertyName.Level.Span, Level.ToString());
+			writer.WriteString(LogEventPropertyName.Message.Span, Message);
 
 			if (Id.Id != 0)
 			{
-				writer.WriteNumber(JsonPropertyNames.Id.Span, Id.Id);
+				writer.WriteNumber(LogEventPropertyName.Id.Span, Id.Id);
 			}
 
 			if (LineIndex > 0)
 			{
-				writer.WriteNumber(JsonPropertyNames.Line.Span, LineIndex);
+				writer.WriteNumber(LogEventPropertyName.Line.Span, LineIndex);
 			}
 
 			if (LineCount > 1)
 			{
-				writer.WriteNumber(JsonPropertyNames.LineCount.Span, LineCount);
+				writer.WriteNumber(LogEventPropertyName.LineCount.Span, LineCount);
 			}
 
 			if (Format != null)
 			{
-				writer.WriteString(JsonPropertyNames.Format.Span, Format);
+				writer.WriteString(LogEventPropertyName.Format.Span, Format);
 			}
 
 			if (Properties != null && Properties.Any())
 			{
-				writer.WriteStartObject(JsonPropertyNames.Properties.Span);
+				writer.WriteStartObject(LogEventPropertyName.Properties.Span);
 				foreach ((string name, object? value) in Properties!)
 				{
 					writer.WritePropertyName(name);
@@ -288,7 +291,7 @@ namespace EpicGames.Core
 
 			if (Exception != null)
 			{
-				writer.WriteStartObject(JsonPropertyNames.Exception.Span);
+				writer.WriteStartObject(LogEventPropertyName.Exception.Span);
 				WriteException(ref writer, Exception);
 				writer.WriteEndObject();
 			}
@@ -323,11 +326,11 @@ namespace EpicGames.Core
 		static void WriteStructuredPropertyValue(ref Utf8JsonWriter writer, LogValue value)
 		{
 			writer.WriteStartObject();
-			writer.WriteString(JsonPropertyNames.Type.Span, value.Type);
-			writer.WriteString(JsonPropertyNames.Text.Span, value.Text);
+			writer.WriteString(LogEventPropertyName.Type.Span, value.Type);
+			writer.WriteString(LogEventPropertyName.Text.Span, value.Text);
 			if (value.Properties != null)
 			{
-				foreach ((string propertyName, object? propertyValue) in value.Properties)
+				foreach ((Utf8String propertyName, object? propertyValue) in value.Properties)
 				{
 					writer.WritePropertyName(propertyName);
 					WritePropertyValue(ref writer, propertyValue);
@@ -546,7 +549,7 @@ namespace EpicGames.Core
 		/// <summary>
 		/// Type of the event
 		/// </summary>
-		public string Type { get; set; }
+		public Utf8String Type { get; set; }
 
 		/// <summary>
 		/// Rendering of the value
@@ -556,7 +559,7 @@ namespace EpicGames.Core
 		/// <summary>
 		/// Properties associated with the value
 		/// </summary>
-		public Dictionary<string, object>? Properties { get; }
+		public Dictionary<Utf8String, object>? Properties { get; }
 
 		/// <summary>
 		/// Constructor
@@ -564,7 +567,7 @@ namespace EpicGames.Core
 		/// <param name="type">Type of the value</param>
 		/// <param name="text">Rendering of the value as text</param>
 		/// <param name="properties">Additional properties for this value</param>
-		public LogValue(string type, string text, Dictionary<string, object>? properties = null)
+		public LogValue(Utf8String type, string text, Dictionary<Utf8String, object>? properties = null)
 		{
 			Type = type;
 			Text = text;
diff --git a/Engine/Source/Programs/Shared/EpicGames.Core/Utf8String.cs b/Engine/Source/Programs/Shared/EpicGames.Core/Utf8String.cs
--- a/Engine/Source/Programs/Shared/EpicGames.Core/Utf8String.cs
+++ b/Engine/Source/Programs/Shared/EpicGames.Core/Utf8String.cs
@@ -335,6 +335,12 @@ namespace EpicGames.Core
 		{
 			return new Utf8String(new ReadOnlyMemory<byte>(Encoding.UTF8.GetBytes(text)));
 		}
+
+		/// <summary>
+		/// Implict conversion to a span of bytes
+		/// </summary>
+		/// <param name="str"></param>
+		public static implicit operator ReadOnlySpan<byte>(Utf8String str) => str.Span;
 	}
 
 	/// <summary>
