Horde: Add a log event parser into the standard EpicGames.Core.Log class.

diff --git a/Engine/Source/Programs/Horde/Horde.Agent/Parser/LogEventMarkup.cs b/Engine/Source/Programs/Horde/Horde.Agent/Parser/LogEventMarkup.cs
deleted file mode 100644
diff --git a/Engine/Source/Programs/Shared/EpicGames.Core/Log.cs b/Engine/Source/Programs/Shared/EpicGames.Core/Log.cs
--- a/Engine/Source/Programs/Shared/EpicGames.Core/Log.cs
+++ b/Engine/Source/Programs/Shared/EpicGames.Core/Log.cs
@@ -94,14 +94,24 @@ namespace EpicGames.Core
 	public static class Log
 	{
 		/// <summary>
-		/// Singleton instance of the default logger
+		/// Singleton instance of the default output logger
 		/// </summary>
-		private static DefaultLogger DefaultLogger { get; } = new DefaultLogger();
+		private static readonly DefaultLogger DefaultLogger = new DefaultLogger();
 
 		/// <summary>
-		/// The public logging interface
+		/// Logger instance which parses events and forwards them to the main logger.
 		/// </summary>
-		public static ILogger Logger { get; set; } = DefaultLogger;
+		private static readonly LegacyEventLogger LegacyLogger = new LegacyEventLogger(DefaultLogger);
+
+		/// <summary>
+		/// Accessor for the global event parser from legacy events
+		/// </summary>
+		public static LogEventParser EventParser => LegacyLogger.Parser;
+
+		/// <summary>
+		/// Logger instance
+		/// </summary>
+		public static ILogger Logger => LegacyLogger;
 
 		/// <summary>
 		/// When true, verbose logging is enabled.
@@ -172,6 +182,15 @@ namespace EpicGames.Core
 		/// </summary>
 		private static readonly HashSet<string> s_writeOnceSet = new HashSet<string>();
 
+		/// <summary>
+		/// Overrides the logger used for formatting output, after event parsing
+		/// </summary>
+		/// <param name="logger"></param>
+		public static void SetInnerLogger(ILogger logger)
+		{
+			LegacyLogger.SetInnerLogger(logger);
+		}
+
 		/// <summary>
 		/// Adds a trace listener that writes to a log file.
 		/// If Log.DuplicateLogFiles is true, two files will be created - one with the requested name,
@@ -362,7 +381,14 @@ namespace EpicGames.Core
 				}
 
 				// Forward it on to the internal logger
-				Logger.Log((LogLevel)verbosity, new EventId(), message, null, (message, ex) => message.ToString());
+				if (verbosity < LogEventType.Console)
+				{
+					Logger.Log((LogLevel)verbosity, message.ToString());
+				}
+				else
+				{
+					EventParser.WriteLine(message.ToString());
+				}
 			}
 		}
 
@@ -749,6 +775,41 @@ namespace EpicGames.Core
 		}
 	}
 
+	/// <summary>
+	/// Wrapper around a custom logger interface which flushes the event parser when switching between legacy
+	/// and native structured logging
+	/// </summary>
+	class LegacyEventLogger : ILogger
+	{
+		private ILogger _inner;
+		private readonly LogEventParser _parser;
+
+		public LogEventParser Parser => _parser;
+
+		public LegacyEventLogger(ILogger inner)
+		{
+			_inner = inner;
+			_parser = new LogEventParser(inner);
+		}
+
+		public void SetInnerLogger(ILogger inner)
+		{
+			_parser.Flush();
+			_inner = inner;
+			_parser.Logger = inner;
+		}
+
+		public IDisposable BeginScope<TState>(TState state) => _inner.BeginScope(state);
+
+		public bool IsEnabled(LogLevel logLevel) => _inner.IsEnabled(logLevel);
+
+		public void Log<TState>(LogLevel logLevel, EventId eventId, TState state, Exception exception, Func<TState, Exception, string> formatter)
+		{
+			_parser.Flush();
+			_inner.Log(logLevel, eventId, state, exception, formatter);
+		}
+	}
+
 	/// <summary>
 	/// Default log output device
 	/// </summary>
@@ -859,6 +920,11 @@ namespace EpicGames.Core
 		/// </summary>
 		private string _statusText = "";
 
+		/// <summary>
+		/// Parser for transforming legacy log output into structured events
+		/// </summary>
+		public LogEventParser EventParser { get; }
+
 		/// <summary>
 		/// Last time a status message was pushed to the stack
 		/// </summary>
@@ -874,6 +940,7 @@ namespace EpicGames.Core
 			OutputLevel = LogLevel.Debug;
 			ColorConsoleOutput = true;
 			IncludeStartingTimestamp = true;
+			EventParser = new LogEventParser(this);
 
 			string? envVar = Environment.GetEnvironmentVariable("UE_STDOUT_JSON");
 			if(envVar != null && Int32.TryParse(envVar, out int value) && value != 0)
@@ -981,7 +1048,7 @@ namespace EpicGames.Core
 			return logLevel >= OutputLevel;
 		}
 
-		public void Log<TState>(LogLevel logLevel, EventId eventId, TState state, Exception exception, Func<TState, Exception, string> formatter)
+		public void Log<TState>(LogLevel logLevel, EventId eventId, TState state, Exception exception, Func<TState, Exception?, string> formatter)
 		{
 			string[] lines = formatter(state, exception).Split('\n');
 			lock (_syncObject)
@@ -1033,7 +1100,7 @@ namespace EpicGames.Core
 							_jsonBufferWriter.Clear();
 							_jsonWriter.Reset();
 
-							LogEvent logEvent = new LogEvent(DateTime.UtcNow, logLevel, eventId, formatter(state, exception), null, null, LogException.FromException(exception));
+							LogEvent logEvent = LogEvent.FromState<TState>(logLevel, eventId, state, exception, formatter);
 							logEvent.Write(_jsonWriter);
 
 							_jsonWriter.Flush();
diff --git a/Engine/Source/Programs/Shared/EpicGames.Core/LogEvent.cs b/Engine/Source/Programs/Shared/EpicGames.Core/LogEvent.cs
--- a/Engine/Source/Programs/Shared/EpicGames.Core/LogEvent.cs
+++ b/Engine/Source/Programs/Shared/EpicGames.Core/LogEvent.cs
@@ -439,6 +439,11 @@ namespace EpicGames.Core
 		/// </summary>
 		public static LogEvent FromState<TState>(LogLevel level, EventId eventId, TState state, Exception? exception, Func<TState, Exception?, string> formatter)
 		{
+			if(state is LogEvent logEvent)
+			{
+				return logEvent;
+			}
+
 			DateTime time = DateTime.UtcNow;
 
 			// Render the message
diff --git a/Engine/Source/Programs/Shared/EpicGames.Core/LogEventMarkup.cs b/Engine/Source/Programs/Shared/EpicGames.Core/LogEventMarkup.cs
new file mode 100644
--- /dev/null
+++ b/Engine/Source/Programs/Shared/EpicGames.Core/LogEventMarkup.cs
@@ -0,0 +1,115 @@
+// Copyright Epic Games, Inc. All Rights Reserved.
+
+using System;
+using System.Collections.Generic;
+using System.IO;
+using System.Text.RegularExpressions;
+using EpicGames.Core;
+
+namespace EpicGames.Core
+{
+	/// <summary>
+	/// Extension methods to allow adding markup to log event spans
+	/// </summary>
+	public static class LogEventMarkup
+	{
+		public static readonly Utf8String Asset = new Utf8String("Asset");
+		public static readonly Utf8String SourceFile = new Utf8String("SourceFile");
+
+		public static LogValue Channel => new LogValue("Channel", "");
+		public static LogValue Severity => new LogValue("Severity", "");
+		public static LogValue Message => new LogValue("Message", "");
+		public static LogValue LineNumber => new LogValue("Line", "");
+		public static LogValue ColumnNumber => new LogValue("Column", "");
+		public static LogValue Symbol => new LogValue("Symbol", "");
+		public static LogValue ErrorCode => new LogValue("ErrorCode", "");
+		public static LogValue ToolName => new LogValue("ToolName", "");
+		public static LogValue ScreenshotTest => new LogValue("ScreenshotTest", "");
+
+		static readonly char[] s_pathChars = { '/', '\\' };
+
+		static string RemoveRelativeDirs(string path)
+		{
+			int idx0 = path.IndexOfAny(s_pathChars);
+			if (idx0 != -1)
+			{
+				int idx1 = path.IndexOfAny(s_pathChars, idx0 + 1);
+				while (idx1 != -1)
+				{
+					int idx2 = path.IndexOfAny(s_pathChars, idx1 + 1);
+					if (idx2 == -1)
+					{
+						break;
+					}
+
+					if (idx2 == idx1 + 3 && path[idx1 + 1] == '.' && path[idx1 + 2] == '.')
+					{
+						path = path.Remove(idx0, idx2 - idx0);
+						idx1 = path.IndexOfAny(s_pathChars, idx0 + 1);
+					}
+					else
+					{
+						idx0 = idx1;
+						idx1 = idx2;
+					}
+				}
+			}
+			return path;
+		}
+
+		/// <summary>
+		/// Marks a span of text as a source file
+		/// </summary>
+		public static void AnnotateSourceFile(this LogEventBuilder builder, Group group, string? baseDir)
+		{
+			Dictionary<Utf8String, object>? properties = null;
+			if (!String.IsNullOrEmpty(baseDir))
+			{
+				string file = group.Value;
+				if (!Path.IsPathRooted(file))
+				{
+					try
+					{
+						string combinedPath = RemoveRelativeDirs(Path.Combine(baseDir, file));
+						properties = new Dictionary<Utf8String, object>();
+						properties["file"] = combinedPath;
+					}
+					catch
+					{
+					}
+				}
+			}
+			builder.Annotate(group, new LogValue(SourceFile, group.Value, properties));
+		}
+
+		/// <summary>
+		/// Marks a span of text as a source file
+		/// </summary>
+		public static void AnnotateAsset(this LogEventBuilder builder, Group group)
+		{
+			builder.Annotate(group, new LogValue(Asset, group.Value));
+		}
+
+		/// <summary>
+		/// Marks a span of text as a symbol
+		/// </summary>
+		public static void AnnotateSymbol(this LogEventBuilder builder, Group group)
+		{
+			string identifier = group.Value;
+
+			// Remove any __declspec qualifiers
+			identifier = Regex.Replace(identifier, "(?<![^a-zA-Z_])__declspec\\([^\\)]+\\)", "");
+
+			// Remove any argument lists for functions (anything after the first paren)
+			identifier = Regex.Replace(identifier, "\\(.*$", "");
+
+			// Remove any decorators and type information (greedy match up to the last space)
+			identifier = Regex.Replace(identifier, "^.* ", "");
+
+			// Add it to the list
+			Dictionary<Utf8String, object> properties = new Dictionary<Utf8String, object>();
+			properties["identifier"] = identifier;
+			builder.Annotate(group, new LogValue("symbol", "", properties));
+		}
+	}
+}
diff --git a/Engine/Source/Programs/Shared/EpicGames.Core/LogEventParser.cs b/Engine/Source/Programs/Shared/EpicGames.Core/LogEventParser.cs
--- a/Engine/Source/Programs/Shared/EpicGames.Core/LogEventParser.cs
+++ b/Engine/Source/Programs/Shared/EpicGames.Core/LogEventParser.cs
@@ -4,6 +4,7 @@ using System;
 using System.Collections.Generic;
 using System.IO;
 using System.Linq;
+using System.Reflection;
 using System.Text;
 using System.Text.RegularExpressions;
 using Microsoft.Extensions.Logging;
@@ -89,7 +90,16 @@ namespace EpicGames.Core
 		/// <summary>
 		/// The inner logger
 		/// </summary>
-		readonly ILogger _logger;
+		ILogger _logger;
+
+		/// <summary>
+		/// Public accessor for the logger
+		/// </summary>
+		public ILogger Logger
+		{
+			get => _logger;
+			set => _logger = value;
+		}
 
 		/// <summary>
 		/// Constructor
@@ -104,13 +114,42 @@ namespace EpicGames.Core
 		/// <inheritdoc/>
 		public void Dispose() => Flush();
 
+		/// <summary>
+		/// Enumerate all the types that implement <see cref="ILogEventMatcher"/> in the given assembly, and create instances of them
+		/// </summary>
+		/// <param name="assembly">The assembly to enumerate matchers from</param>
+		public void AddMatchersFromAssembly(Assembly assembly)
+		{
+			foreach (Type type in assembly.GetTypes())
+			{
+				if (type.IsClass && typeof(ILogEventMatcher).IsAssignableFrom(type))
+				{
+					ILogEventMatcher matcher = (ILogEventMatcher)Activator.CreateInstance(type)!;
+					Matchers.Add(matcher);
+				}
+			}
+		}
+
 		/// <summary>
 		/// Writes a line to the event filter
 		/// </summary>
 		/// <param name="line">The line to output</param>
 		public void WriteLine(string line)
 		{
-			_buffer.AddLine(line);
+			if (line.Length > 0 && line[0] == '{')
+			{
+				byte[] data = Encoding.UTF8.GetBytes(line);
+
+				JsonLogEvent jsonEvent;
+				if (JsonLogEvent.TryParse(data, out jsonEvent))
+				{
+					ProcessData(true);
+					_logger.Log(jsonEvent.Level, jsonEvent.EventId, jsonEvent, null, JsonLogEvent.Format);
+					return;
+				}
+			}
+
+			_buffer.AddLine(StringUtils.ParseEscapeCodes(line));
 			ProcessData(false);
 		}
 
