EpicGames.Core: Speculative fix for exceptions when writing log output from background threads.

diff --git a/Engine/Source/Programs/Shared/EpicGames.Core/Log.cs b/Engine/Source/Programs/Shared/EpicGames.Core/Log.cs
--- a/Engine/Source/Programs/Shared/EpicGames.Core/Log.cs
+++ b/Engine/Source/Programs/Shared/EpicGames.Core/Log.cs
@@ -387,7 +387,10 @@ namespace EpicGames.Core
 				}
 				else
 				{
-					EventParser.WriteLine(message.ToString());
+					lock (EventParser)
+					{
+						EventParser.WriteLine(message.ToString());
+					}
 				}
 			}
 		}
@@ -794,9 +797,12 @@ namespace EpicGames.Core
 
 		public void SetInnerLogger(ILogger inner)
 		{
-			_parser.Flush();
-			_inner = inner;
-			_parser.Logger = inner;
+			lock (_parser)
+			{
+				_parser.Flush();
+				_inner = inner;
+				_parser.Logger = inner;
+			}
 		}
 
 		public IDisposable BeginScope<TState>(TState state) => _inner.BeginScope(state);
@@ -805,7 +811,10 @@ namespace EpicGames.Core
 
 		public void Log<TState>(LogLevel logLevel, EventId eventId, TState state, Exception exception, Func<TState, Exception, string> formatter)
 		{
-			_parser.Flush();
+			lock (_parser)
+			{
+				_parser.Flush();
+			}
 			_inner.Log(logLevel, eventId, state, exception, formatter);
 		}
 	}
