EpicGames.Serialization: Support deserializing into getter-only collection properties.

diff --git a/Engine/Source/Programs/Shared/EpicGames.Serialization.Tests/SerializerTests.cs b/Engine/Source/Programs/Shared/EpicGames.Serialization.Tests/SerializerTests.cs
--- a/Engine/Source/Programs/Shared/EpicGames.Serialization.Tests/SerializerTests.cs
+++ b/Engine/Source/Programs/Shared/EpicGames.Serialization.Tests/SerializerTests.cs
@@ -129,5 +129,37 @@ namespace EpicGames.Serialization.Tests
 			Assert.AreEqual(output.Value, 123);
 			Assert.AreEqual(output.Items, null);
 		}
+
+		class ObjectWithReadOnlyCollection
+		{
+			[CbField("s")]
+			public List<string> Strings { get; } = new List<string>();
+
+			[CbField("i")]
+			public List<int> Integers { get; } = new List<int>();
+		}
+
+		[TestMethod]
+		public void ObjectWithReadOnlyCollectionTest()
+		{
+			ObjectWithReadOnlyCollection input = new ObjectWithReadOnlyCollection();
+			input.Strings.Add("Hello");
+			input.Strings.Add("World");
+			input.Integers.Add(1);
+			input.Integers.Add(2);
+			input.Integers.Add(3);
+			input.Integers.Add(4);
+			CbObject obj = CbSerializer.Serialize(input);
+
+			ObjectWithReadOnlyCollection output = CbSerializer.Deserialize<ObjectWithReadOnlyCollection>(obj.AsField());
+			Assert.AreEqual(2, output.Strings.Count);
+			Assert.AreEqual("Hello", output.Strings[0]);
+			Assert.AreEqual("World", output.Strings[1]);
+			Assert.AreEqual(4, output.Integers.Count);
+			Assert.AreEqual(1, output.Integers[0]);
+			Assert.AreEqual(2, output.Integers[1]);
+			Assert.AreEqual(3, output.Integers[2]);
+			Assert.AreEqual(4, output.Integers[3]);
+		}
 	}
 }
diff --git a/Engine/Source/Programs/Shared/EpicGames.Serialization/Converters/CbClassConverter.cs b/Engine/Source/Programs/Shared/EpicGames.Serialization/Converters/CbClassConverter.cs
--- a/Engine/Source/Programs/Shared/EpicGames.Serialization/Converters/CbClassConverter.cs
+++ b/Engine/Source/Programs/Shared/EpicGames.Serialization/Converters/CbClassConverter.cs
@@ -3,6 +3,7 @@
 using EpicGames.Core;
 using System;
 using System.Collections.Generic;
+using System.Diagnostics.CodeAnalysis;
 using System.Linq;
 using System.Linq.Expressions;
 using System.Reflection;
@@ -383,11 +384,28 @@ namespace EpicGames.Serialization.Converters
 				// FieldLocal = Stack.Pop()
 				generator.Emit(OpCodes.Stloc, fieldLocal);
 
-				// Property.SetMethod(NewObjectLocal, ReadMethod(FieldLocal))
-				generator.Emit(OpCodes.Ldloc, newObjectLocal);
-				generator.Emit(OpCodes.Ldloc, fieldLocal);
-				generator.EmitCall(OpCodes.Call, readMethod, null);
-				generator.EmitCall(OpCodes.Call, property.SetMethod!, null);
+				// Copy the collection over if necessary
+				if (property.SetMethod != null)
+				{
+					// Property.SetMethod(NewObjectLocal, ReadMethod(FieldLocal))
+					generator.Emit(OpCodes.Ldloc, newObjectLocal);
+					generator.Emit(OpCodes.Ldloc, fieldLocal);
+					generator.EmitCall(OpCodes.Call, readMethod, null);
+					generator.EmitCall(OpCodes.Call, property.SetMethod!, null);
+				}
+				else if (TryGetCollectionCopyMethod(property.PropertyType, out MethodInfo? copyMethod))
+				{
+					// CopyMethod(ReadMethod(FieldLocal), Property.GetMethod(NewObjectLocal))
+					generator.Emit(OpCodes.Ldloc, fieldLocal);
+					generator.EmitCall(OpCodes.Call, readMethod, null);
+					generator.Emit(OpCodes.Ldloc, newObjectLocal);
+					generator.EmitCall(OpCodes.Call, property.GetMethod!, null);
+					generator.EmitCall(OpCodes.Call, copyMethod, null);
+				}
+				else
+				{
+					throw new CbException($"Unable to write to property {property.Name}");
+				}
 
 				// if(!IteratorLocal.MoveNext()) goto ReturnLabel
 				generator.Emit(OpCodes.Ldloc, iteratorLocal);
@@ -424,6 +442,30 @@ namespace EpicGames.Serialization.Converters
 			generator.Emit(OpCodes.Ret);
 		}
 
+		static bool TryGetCollectionCopyMethod(Type type, [NotNullWhen(true)] out MethodInfo? method)
+		{
+			foreach (Type interfaceType in type.GetInterfaces())
+			{
+				if (interfaceType.IsGenericType && interfaceType.GetGenericTypeDefinition() == typeof(ICollection<>))
+				{
+					MethodInfo genericMethod = typeof(CbClassConverterMethods).GetMethod(nameof(CopyCollection), BindingFlags.Static | BindingFlags.NonPublic)!;
+					method = genericMethod.MakeGenericMethod(interfaceType, interfaceType.GetGenericArguments()[0]);
+					return true;
+				}
+			}
+
+			method = null;
+			return false;
+		}
+
+		static void CopyCollection<TCollection, TElement>(TCollection source, TCollection target) where TCollection : ICollection<TElement>
+		{
+			foreach(TElement element in source)
+			{
+				target.Add(element);
+			}
+		}
+
 		static (Utf8String, PropertyInfo)[] GetProperties(Type type)
 		{
 			List<(Utf8String, PropertyInfo)> propertyList = new List<(Utf8String, PropertyInfo)>();
