EpicGames.Serialization: Add tests for deserializing dictionaries into read-only properties.

diff --git a/Engine/Source/Programs/Shared/EpicGames.Serialization.Tests/SerializerTests.cs b/Engine/Source/Programs/Shared/EpicGames.Serialization.Tests/SerializerTests.cs
--- a/Engine/Source/Programs/Shared/EpicGames.Serialization.Tests/SerializerTests.cs
+++ b/Engine/Source/Programs/Shared/EpicGames.Serialization.Tests/SerializerTests.cs
@@ -137,6 +137,9 @@ namespace EpicGames.Serialization.Tests
 
 			[CbField("i")]
 			public List<int> Integers { get; } = new List<int>();
+
+			[CbField("d")]
+			public Dictionary<string, int> Dictionary { get; } = new Dictionary<string, int>(StringComparer.Ordinal);
 		}
 
 		[TestMethod]
@@ -149,6 +152,8 @@ namespace EpicGames.Serialization.Tests
 			input.Integers.Add(2);
 			input.Integers.Add(3);
 			input.Integers.Add(4);
+			input.Dictionary.Add("hello", 123);
+			input.Dictionary.Add("world", 456);
 			CbObject obj = CbSerializer.Serialize(input);
 
 			ObjectWithReadOnlyCollection output = CbSerializer.Deserialize<ObjectWithReadOnlyCollection>(obj.AsField());
@@ -160,6 +165,9 @@ namespace EpicGames.Serialization.Tests
 			Assert.AreEqual(2, output.Integers[1]);
 			Assert.AreEqual(3, output.Integers[2]);
 			Assert.AreEqual(4, output.Integers[3]);
+			Assert.AreEqual(2, output.Dictionary.Count);
+			Assert.AreEqual(123, output.Dictionary["hello"]);
+			Assert.AreEqual(456, output.Dictionary["world"]);
 		}
 	}
 }
