Upgrade EpicGames.Horde to NET 6.

diff --git a/Engine/Source/Programs/Horde/Horde.Agent/Commands/Compute/AwsLambdaCommand.cs b/Engine/Source/Programs/Horde/Horde.Agent/Commands/Compute/AwsLambdaCommand.cs
--- a/Engine/Source/Programs/Horde/Horde.Agent/Commands/Compute/AwsLambdaCommand.cs
+++ b/Engine/Source/Programs/Horde/Horde.Agent/Commands/Compute/AwsLambdaCommand.cs
@@ -195,7 +195,7 @@ namespace Horde.Agent.Commands
 				OAuthHandlerFactory oAuthHandlerFactory = new OAuthHandlerFactory(httpClient);
 				OAuthHandler<HttpStorageClient> oAuthHandler = oAuthHandlerFactory.Create<HttpStorageClient>(clientOptions);
 				using HttpClient client = new HttpClient(oAuthHandler);
-				client.BaseAddress = new Uri(clientOptions.Url);
+				client.BaseAddress = clientOptions.Url;
 				IStorageClient storageClientManual = new HttpStorageClient(client);
 
 				AppDomain.CurrentDomain.ProcessExit += (s, e) => 
@@ -238,8 +238,8 @@ namespace Horde.Agent.Commands
 		{
 			return new HttpServiceClientOptions
 			{
-				Url = GetEnvVar("UE_HORDE_STORAGE_URL"),
-				AuthUrl = GetEnvVar("UE_HORDE_STORAGE_OAUTH_URL"),
+				Url = new Uri(GetEnvVar("UE_HORDE_STORAGE_URL")),
+				AuthUrl = new Uri(GetEnvVar("UE_HORDE_STORAGE_OAUTH_URL")),
 				GrantType = GetEnvVar("UE_HORDE_STORAGE_OAUTH_GRANT_TYPE"),
 				ClientId = GetEnvVar("UE_HORDE_STORAGE_OAUTH_CLIENT_ID"),
 				ClientSecret = clientSecret,
diff --git a/Engine/Source/Programs/Horde/Horde.Agent/Commands/Compute/ComputeCommand.cs b/Engine/Source/Programs/Horde/Horde.Agent/Commands/Compute/ComputeCommand.cs
--- a/Engine/Source/Programs/Horde/Horde.Agent/Commands/Compute/ComputeCommand.cs
+++ b/Engine/Source/Programs/Horde/Horde.Agent/Commands/Compute/ComputeCommand.cs
@@ -196,12 +196,18 @@ namespace Horde.Agent.Commands
 
 				JsonRequirements jsonRequirements = jsonComputeTask.Requirements;
 				Requirements requirements = new Requirements(jsonRequirements.Condition ?? String.Empty);
-				requirements.Resources = jsonRequirements.Resources.ToDictionary(x => x.Key, x => x.Value);
+				foreach ((string name, int count) in jsonRequirements.Resources)
+				{
+					requirements.Resources.Add(name, count);
+				}
 				requirements.Exclusive = jsonRequirements.Exclusive;
 				IoHash requirementsHash = AddCbObject(blobs, requirements);
 
 				ComputeTask task = new ComputeTask(jsonComputeTask.Executable, jsonComputeTask.Arguments.ConvertAll<Utf8String>(x => x), jsonComputeTask.WorkingDirectory, sandboxHash);
-				task.EnvVars = jsonComputeTask.EnvVars.ToDictionary(x => (Utf8String)x.Key, x => (Utf8String)x.Value);
+				foreach ((string name, string value) in jsonComputeTask.EnvVars)
+				{
+					task.EnvVars.Add(name, value);
+				}
 				task.OutputPaths.AddRange(jsonComputeTask.OutputPaths.Select(x => (Utf8String)x));
 				task.RequirementsHash = AddCbObject(blobs, requirements);
 
diff --git a/Engine/Source/Programs/Shared/EpicGames.Horde.Tests/BundleTests.cs b/Engine/Source/Programs/Shared/EpicGames.Horde.Tests/BundleTests.cs
--- a/Engine/Source/Programs/Shared/EpicGames.Horde.Tests/BundleTests.cs
+++ b/Engine/Source/Programs/Shared/EpicGames.Horde.Tests/BundleTests.cs
@@ -17,7 +17,7 @@ using System.Threading.Tasks;
 namespace EpicGames.Horde.Tests
 {
 	[TestClass]
-	public class BundleTests : IDisposable
+	public sealed class BundleTests : IDisposable
 	{
 		readonly MemoryCache _cache;
 		readonly MemoryStorageClient _storageClient = new MemoryStorageClient();
@@ -147,7 +147,7 @@ namespace EpicGames.Horde.Tests
 				Assert.AreEqual(oldExport.Packet.DecodedLength, newExport.Packet.DecodedLength);
 				Assert.AreEqual(oldExport.Offset, newExport.Offset);
 				Assert.AreEqual(oldExport.Length, newExport.Length);
-				Assert.IsTrue(oldExport.References.AsSpan().SequenceEqual(newExport.References.AsSpan()));
+				Assert.IsTrue(oldExport.References.SequenceEqual(newExport.References));
 			}
 		}
 
@@ -366,7 +366,7 @@ namespace EpicGames.Horde.Tests
 				data[idx] = (byte)idx;
 			}
 
-			Bundle<FileNode> newBundle = Bundle.Create<FileNode>(_storageClient, _namespaceId, new FileNode(), new BundleOptions(), _cache);
+			using Bundle<FileNode> newBundle = Bundle.Create<FileNode>(_storageClient, _namespaceId, new FileNode(), new BundleOptions(), _cache);
 			FileNode root = newBundle.Root;
 
 			const int NumIterations = 100;
diff --git a/Engine/Source/Programs/Shared/EpicGames.Horde.Tests/EpicGames.Horde.Tests.csproj b/Engine/Source/Programs/Shared/EpicGames.Horde.Tests/EpicGames.Horde.Tests.csproj
--- a/Engine/Source/Programs/Shared/EpicGames.Horde.Tests/EpicGames.Horde.Tests.csproj
+++ b/Engine/Source/Programs/Shared/EpicGames.Horde.Tests/EpicGames.Horde.Tests.csproj
@@ -1,7 +1,7 @@
 <Project Sdk="Microsoft.NET.Sdk">
 
   <PropertyGroup>
-    <TargetFramework>netcoreapp3.1</TargetFramework>
+    <TargetFramework>net6.0</TargetFramework>
     <Configurations>Debug;Release;Development</Configurations>
     <Nullable>enable</Nullable>
     <IsPackable>false</IsPackable>
diff --git a/Engine/Source/Programs/Shared/EpicGames.Horde/Auth/AuthExtensions.cs b/Engine/Source/Programs/Shared/EpicGames.Horde/Auth/AuthExtensions.cs
--- a/Engine/Source/Programs/Shared/EpicGames.Horde/Auth/AuthExtensions.cs
+++ b/Engine/Source/Programs/Shared/EpicGames.Horde/Auth/AuthExtensions.cs
@@ -16,12 +16,12 @@ namespace EpicGames.Horde.Auth
 		/// Base address for http requests
 		/// </summary>
 		[Required]
-		public string Url { get; set; } = String.Empty;
+		public Uri Url { get; set; } = null!;
 
 		#region OAuth2
 
 		/// <inheritdoc/>
-		public string AuthUrl { get; set; } = String.Empty;
+		public Uri? AuthUrl { get; set; }
 
 		/// <inheritdoc/>
 		public string GrantType { get; set; } = String.Empty;
@@ -58,7 +58,7 @@ namespace EpicGames.Horde.Auth
 			services.AddHttpClient<TClient, TImplementation>((serviceProvider, client) =>
 				{
 					HttpServiceClientOptions options = getOptions(serviceProvider);
-					client.BaseAddress = new Uri(options.Url);
+					client.BaseAddress = options.Url;
 				})
 				.ConfigurePrimaryHttpMessageHandler(serviceProvider =>
 				{
@@ -69,7 +69,7 @@ namespace EpicGames.Horde.Auth
 
 		static HttpMessageHandler CreateMessageHandler<TImplementation>(IServiceProvider serviceProvider, HttpServiceClientOptions options)
 		{
-			if (!String.IsNullOrEmpty(options.AuthUrl))
+			if (options.AuthUrl != null)
 			{
 				return serviceProvider.GetRequiredService<OAuthHandler<TImplementation>>();
 			}
diff --git a/Engine/Source/Programs/Shared/EpicGames.Horde/Auth/OAuthHandler.cs b/Engine/Source/Programs/Shared/EpicGames.Horde/Auth/OAuthHandler.cs
--- a/Engine/Source/Programs/Shared/EpicGames.Horde/Auth/OAuthHandler.cs
+++ b/Engine/Source/Programs/Shared/EpicGames.Horde/Auth/OAuthHandler.cs
@@ -34,7 +34,7 @@ namespace EpicGames.Horde.Auth
 		/// <summary>
 		/// Url of the auth server
 		/// </summary>
-		string AuthUrl { get; }
+		Uri? AuthUrl { get; }
 
 		/// <summary>
 		/// Type of grant
@@ -124,7 +124,7 @@ namespace EpicGames.Horde.Auth
 					throw new AuthenticationException($"Authentication failed. Response: {response.Content}", null);
 				}
 
-				byte[] responseData = await response.Content.ReadAsByteArrayAsync();
+				byte[] responseData = await response.Content.ReadAsByteArrayAsync(cancellationToken);
 				ClientCredentialsResponse result = JsonSerializer.Deserialize<ClientCredentialsResponse>(responseData)!;
 
 				string? accessToken = result?.access_token;
diff --git a/Engine/Source/Programs/Shared/EpicGames.Horde/Bundles/Bundle.cs b/Engine/Source/Programs/Shared/EpicGames.Horde/Bundles/Bundle.cs
--- a/Engine/Source/Programs/Shared/EpicGames.Horde/Bundles/Bundle.cs
+++ b/Engine/Source/Programs/Shared/EpicGames.Horde/Bundles/Bundle.cs
@@ -285,6 +285,16 @@ namespace EpicGames.Horde.Bundles
 
 		/// <inheritdoc/>
 		public void Dispose()
+		{
+			Dispose(true);
+			GC.SuppressFinalize(this);
+		}
+
+		/// <summary>
+		/// Dispose method
+		/// </summary>
+		/// <param name="disposing"></param>
+		protected virtual void Dispose(bool disposing)
 		{
 			_readSema.Dispose();
 		}
@@ -488,10 +498,9 @@ namespace EpicGames.Horde.Bundles
 					nodeState = node.State;
 				}
 
-				Debug.Assert(nodeState.Blob != null);
 				Debug.Assert(nodeState.Export != null);
 
-				BundleObject bundleObject = await GetObjectAsync(nodeState.Blob.Hash);
+				BundleObject bundleObject = await GetObjectAsync(nodeState.Blob!.Hash);
 				ReadOnlyMemory<byte> packetData = DecodePacket(nodeState.Blob.Hash, bundleObject.Data, nodeState.Export.Packet);
 				ReadOnlyMemory<byte> nodeData = packetData.Slice(nodeState.Export.Offset, nodeState.Export.Length);
 
@@ -804,7 +813,7 @@ namespace EpicGames.Horde.Bundles
 				// Figure out which blobs need to be repacked
 				foreach (BlobInfo liveBlob in liveBlobs)
 				{
-					int liveCost = liveBlob.LiveNodes.Sum(x => x.Length);
+					int liveCost = liveBlob.LiveNodes!.Sum(x => x.Length);
 					int minLiveCost = (int)(liveBlob.TotalCost * Options.RepackRatio);
 
 					if (liveCost < minLiveCost)
@@ -1017,7 +1026,7 @@ namespace EpicGames.Horde.Bundles
 				}
 
 				// Create the export for this node
-				int[] references = node.State.References.Select(x => nodeToIndex[x]).ToArray();
+				int[] references = node.State.References!.Select(x => nodeToIndex[x]).ToArray();
 				BundleExport export = new BundleExport(node.Hash, node.Rank, packet, blockSize, node.Length, references);
 				newObject.Exports.Add(export);
 
@@ -1063,7 +1072,7 @@ namespace EpicGames.Horde.Bundles
 			}
 		}
 
-		void CreateFreeSpace(ref byte[] buffer, int usedSize, int requiredSize)
+		static void CreateFreeSpace(ref byte[] buffer, int usedSize, int requiredSize)
 		{
 			if (requiredSize > buffer.Length)
 			{
diff --git a/Engine/Source/Programs/Shared/EpicGames.Horde/Bundles/BundleNode.cs b/Engine/Source/Programs/Shared/EpicGames.Horde/Bundles/BundleNode.cs
--- a/Engine/Source/Programs/Shared/EpicGames.Horde/Bundles/BundleNode.cs
+++ b/Engine/Source/Programs/Shared/EpicGames.Horde/Bundles/BundleNode.cs
@@ -258,7 +258,8 @@ namespace EpicGames.Horde.Bundles.Nodes
 	/// <summary>
 	/// Attribute used to define a factory for a particular node type
 	/// </summary>
-	public class BundleNodeDeserializerAttribute : Attribute
+	[AttributeUsage(AttributeTargets.Class)]
+	public sealed class BundleNodeDeserializerAttribute : Attribute
 	{
 		/// <summary>
 		/// The factory type. Should be derived from <see cref="BundleNodeDeserializer{T}"/>
diff --git a/Engine/Source/Programs/Shared/EpicGames.Horde/Bundles/BundleObject.cs b/Engine/Source/Programs/Shared/EpicGames.Horde/Bundles/BundleObject.cs
--- a/Engine/Source/Programs/Shared/EpicGames.Horde/Bundles/BundleObject.cs
+++ b/Engine/Source/Programs/Shared/EpicGames.Horde/Bundles/BundleObject.cs
@@ -47,14 +47,14 @@ namespace EpicGames.Horde.Bundles
 		/// Other objects that are referenced.
 		/// </summary>
 		[CbField("imports")]
-		public List<BundleImportObject> ImportObjects { get; set; } = new List<BundleImportObject>();
+		public List<BundleImportObject> ImportObjects { get; } = new List<BundleImportObject>();
 
 		/// <summary>
 		/// Exported items in the data 
 		/// </summary>
 		[CbField("exports")]
 		[CbConverter(typeof(BundleExportListConverter))]
-		public List<BundleExport> Exports { get; set; } = new List<BundleExport>();
+		public List<BundleExport> Exports { get; } = new List<BundleExport>();
 
 		/// <summary>
 		/// Raw data for this exported entries in this blob.
@@ -86,7 +86,7 @@ namespace EpicGames.Horde.Bundles
 		/// </summary>
 		[CbField("imports")]
 		[CbConverter(typeof(BundleImportListConverter))]
-		public List<BundleImport> Imports { get; set; } = new List<BundleImport>();
+		public List<BundleImport> Imports { get; } = new List<BundleImport>();
 
 		/// <summary>
 		/// Default constructor, for serialization
@@ -236,12 +236,12 @@ namespace EpicGames.Horde.Bundles
 		/// <summary>
 		/// References to other nodes.
 		/// </summary>
-		public int[] References { get; }
+		public IReadOnlyList<int> References { get; }
 
 		/// <summary>
 		/// Constructor
 		/// </summary>
-		public BundleExport(IoHash hash, int rank, BundleCompressionPacket packet, int offset, int length, int[] references)
+		public BundleExport(IoHash hash, int rank, BundleCompressionPacket packet, int offset, int length, IReadOnlyList<int> references)
 			: base(hash, rank, length)
 		{
 			Packet = packet;
@@ -338,7 +338,7 @@ namespace EpicGames.Horde.Bundles
 				int lengthBytes = VarInt.Write(span, export.Length);
 				span = span.Slice(lengthBytes);
 
-				int numReferencesBytes = VarInt.Write(span, export.References.Length);
+				int numReferencesBytes = VarInt.Write(span, export.References.Count);
 				span = span.Slice(numReferencesBytes);
 
 				foreach (int reference in export.References)
@@ -367,7 +367,7 @@ namespace EpicGames.Horde.Bundles
 					prevPacket = packet;
 				}
 
-				length += IoHash.NumBytes + VarInt.Measure(export.Rank) + VarInt.Measure(export.Length) + VarInt.Measure(export.References.Length) + export.References.Sum(x => VarInt.Measure(x));
+				length += IoHash.NumBytes + VarInt.Measure(export.Rank) + VarInt.Measure(export.Length) + VarInt.Measure(export.References.Count) + export.References.Sum(x => VarInt.Measure(x));
 			}
 
 			return length;
diff --git a/Engine/Source/Programs/Shared/EpicGames.Horde/Bundles/BundleTools.cs b/Engine/Source/Programs/Shared/EpicGames.Horde/Bundles/BundleTools.cs
--- a/Engine/Source/Programs/Shared/EpicGames.Horde/Bundles/BundleTools.cs
+++ b/Engine/Source/Programs/Shared/EpicGames.Horde/Bundles/BundleTools.cs
@@ -296,7 +296,7 @@ namespace EpicGames.Horde.Bundles
 					writer.WriteLine($"- hash: {export.Hash}");
 					writer.WriteLine($"  packet-offset: {export.Packet.Offset}");
 					writer.WriteLine($"  size: {export.Length}");
-					if (export.References.Length > 0)
+					if (export.References.Count > 0)
 					{
 						writer.WriteLine($"  references:");
 						foreach (int reference in export.References)
diff --git a/Engine/Source/Programs/Shared/EpicGames.Horde/Bundles/Nodes/DirectoryNode.cs b/Engine/Source/Programs/Shared/EpicGames.Horde/Bundles/Nodes/DirectoryNode.cs
--- a/Engine/Source/Programs/Shared/EpicGames.Horde/Bundles/Nodes/DirectoryNode.cs
+++ b/Engine/Source/Programs/Shared/EpicGames.Horde/Bundles/Nodes/DirectoryNode.cs
@@ -17,8 +17,14 @@ namespace EpicGames.Horde.Bundles.Nodes
 	/// <summary>
 	/// Flags for a directory entry
 	/// </summary>
-	public enum FileEntryFlags : byte
+	[Flags]
+	public enum FileEntryFlags
 	{
+		/// <summary>
+		/// No flags set
+		/// </summary>
+		None = 0,
+
 		/// <summary>
 		/// Indicates that the referenced file is executable
 		/// </summary>
diff --git a/Engine/Source/Programs/Shared/EpicGames.Horde/Common/Condition.cs b/Engine/Source/Programs/Shared/EpicGames.Horde/Common/Condition.cs
--- a/Engine/Source/Programs/Shared/EpicGames.Horde/Common/Condition.cs
+++ b/Engine/Source/Programs/Shared/EpicGames.Horde/Common/Condition.cs
@@ -647,16 +647,16 @@ namespace EpicGames.Horde.Common
 	sealed class ConditionTypeConverter : TypeConverter
 	{
 		/// <inheritdoc/>
-		public override bool CanConvertFrom(ITypeDescriptorContext context, Type sourceType) => sourceType == typeof(string);
+		public override bool CanConvertFrom(ITypeDescriptorContext? context, Type sourceType) => sourceType == typeof(string);
 
 		/// <inheritdoc/>
-		public override object ConvertFrom(ITypeDescriptorContext context, CultureInfo culture, object value) => Condition.TryParse((string)value);
+		public override object ConvertFrom(ITypeDescriptorContext? context, CultureInfo? culture, object value) => Condition.TryParse((string)value);
 
 		/// <inheritdoc/>
-		public override bool CanConvertTo(ITypeDescriptorContext context, Type destinationType) => destinationType == typeof(string);
+		public override bool CanConvertTo(ITypeDescriptorContext? context, Type? destinationType) => destinationType == typeof(string);
 
 		/// <inheritdoc/>
-		public override object? ConvertTo(ITypeDescriptorContext context, CultureInfo culture, object value, Type destinationType) => ((Condition)value).Text;
+		public override object? ConvertTo(ITypeDescriptorContext? context, CultureInfo? culture, object? value, Type? destinationType) => ((Condition)value!).Text;
 	}
 
 	/// <summary>
diff --git a/Engine/Source/Programs/Shared/EpicGames.Horde/Compute/ChannelId.cs b/Engine/Source/Programs/Shared/EpicGames.Horde/Compute/ChannelId.cs
--- a/Engine/Source/Programs/Shared/EpicGames.Horde/Compute/ChannelId.cs
+++ b/Engine/Source/Programs/Shared/EpicGames.Horde/Compute/ChannelId.cs
@@ -84,9 +84,9 @@ namespace EpicGames.Horde.Compute
 	sealed class ChannelIdTypeConverter : TypeConverter
 	{
 		/// <inheritdoc/>
-		public override bool CanConvertFrom(ITypeDescriptorContext context, Type sourceType) => sourceType == typeof(string);
+		public override bool CanConvertFrom(ITypeDescriptorContext? context, Type? sourceType) => sourceType == typeof(string);
 
 		/// <inheritdoc/>
-		public override object ConvertFrom(ITypeDescriptorContext context, CultureInfo culture, object value) => new ChannelId((string)value);
+		public override object ConvertFrom(ITypeDescriptorContext? context, CultureInfo? culture, object? value) => new ChannelId((string)value!);
 	}
 }
diff --git a/Engine/Source/Programs/Shared/EpicGames.Horde/Compute/ClusterId.cs b/Engine/Source/Programs/Shared/EpicGames.Horde/Compute/ClusterId.cs
--- a/Engine/Source/Programs/Shared/EpicGames.Horde/Compute/ClusterId.cs
+++ b/Engine/Source/Programs/Shared/EpicGames.Horde/Compute/ClusterId.cs
@@ -84,10 +84,10 @@ namespace EpicGames.Horde.Compute
 	sealed class ClusterIdTypeConverter : TypeConverter
 	{
 		/// <inheritdoc/>
-		public override bool CanConvertFrom(ITypeDescriptorContext context, Type sourceType) => sourceType == typeof(string);
+		public override bool CanConvertFrom(ITypeDescriptorContext? context, Type? sourceType) => sourceType == typeof(string);
 
 		/// <inheritdoc/>
-		public override object ConvertFrom(ITypeDescriptorContext context, CultureInfo culture, object value) => new ClusterId((string)value);
+		public override object ConvertFrom(ITypeDescriptorContext? context, CultureInfo? culture, object? value) => new ClusterId((string)value!);
 	}
 }
 
diff --git a/Engine/Source/Programs/Shared/EpicGames.Horde/Compute/ComputeTask.cs b/Engine/Source/Programs/Shared/EpicGames.Horde/Compute/ComputeTask.cs
--- a/Engine/Source/Programs/Shared/EpicGames.Horde/Compute/ComputeTask.cs
+++ b/Engine/Source/Programs/Shared/EpicGames.Horde/Compute/ComputeTask.cs
@@ -21,13 +21,13 @@ namespace EpicGames.Horde.Compute
 		/// List of command line arguments for the process to run.
 		/// </summary>
 		[CbField("a")]
-		public List<Utf8String> Arguments { get; set; } = new List<Utf8String>();
+		public List<Utf8String> Arguments { get; } = new List<Utf8String>();
 
 		/// <summary>
 		/// Environment variables to set for the child process
 		/// </summary>
 		[CbField("v")]
-		public Dictionary<Utf8String, Utf8String> EnvVars { get; set; } = new Dictionary<Utf8String, Utf8String>();
+		public Dictionary<Utf8String, Utf8String> EnvVars { get; } = new Dictionary<Utf8String, Utf8String>();
 
 		/// <summary>
 		/// Path to the working directory within the workspace
@@ -51,7 +51,7 @@ namespace EpicGames.Horde.Compute
 		/// List of output paths to be captured on completion of the action. These may be files or directories.
 		/// </summary>
 		[CbField("o")]
-		public List<Utf8String> OutputPaths { get; set; } = new List<Utf8String>();
+		public List<Utf8String> OutputPaths { get; } = new List<Utf8String>();
 
 		/// <summary>
 		/// Default constructor for serialization
diff --git a/Engine/Source/Programs/Shared/EpicGames.Horde/Compute/DirectoryTree.cs b/Engine/Source/Programs/Shared/EpicGames.Horde/Compute/DirectoryTree.cs
--- a/Engine/Source/Programs/Shared/EpicGames.Horde/Compute/DirectoryTree.cs
+++ b/Engine/Source/Programs/Shared/EpicGames.Horde/Compute/DirectoryTree.cs
@@ -15,13 +15,13 @@ namespace EpicGames.Horde.Compute
 		/// List of files within the directory, sorted by name
 		/// </summary>
 		[CbField("f")]
-		public List<FileNode> Files { get; set; } = new List<FileNode>();
+		public List<FileNode> Files { get; } = new List<FileNode>();
 
 		/// <summary>
 		/// Sub-directories within this directory, sorted by name
 		/// </summary>
 		[CbField("d")]
-		public List<DirectoryNode> Directories { get; set; } = new List<DirectoryNode>();
+		public List<DirectoryNode> Directories { get; } = new List<DirectoryNode>();
 	}
 
 	/// <summary>
diff --git a/Engine/Source/Programs/Shared/EpicGames.Horde/Compute/IComputeClient.cs b/Engine/Source/Programs/Shared/EpicGames.Horde/Compute/IComputeClient.cs
--- a/Engine/Source/Programs/Shared/EpicGames.Horde/Compute/IComputeClient.cs
+++ b/Engine/Source/Programs/Shared/EpicGames.Horde/Compute/IComputeClient.cs
@@ -191,7 +191,9 @@ namespace EpicGames.Horde.Compute
 		/// If the order or meaning of fields in this array are changed, the field name should be changed.
 		/// </summary>
 		[CbField("s")]
+#pragma warning disable CA1819 // Properties should not return arrays
 		public int[] Scopes
+#pragma warning restore CA1819 // Properties should not return arrays
 		{
 			get => _scopes;
 			set => value.AsSpan(0, Math.Max(_scopes.Length, value.Length)).CopyTo(_scopes);
@@ -263,7 +265,9 @@ namespace EpicGames.Horde.Compute
 		/// Timing values measured sequentially starting from StartTime, in milliseconds. If the order or meaning of fields in this array are changed, the field name should be changed.
 		/// </summary>
 		[CbField("s")]
+#pragma warning disable CA1819 // Properties should not return arrays
 		public int[] Scopes
+#pragma warning restore CA1819 // Properties should not return arrays
 		{
 			get => _scopes;
 			set => value.AsSpan(0, Math.Max(_scopes.Length, value.Length)).CopyTo(_scopes);
diff --git a/Engine/Source/Programs/Shared/EpicGames.Horde/Compute/Impl/HttpComputeClient.cs b/Engine/Source/Programs/Shared/EpicGames.Horde/Compute/Impl/HttpComputeClient.cs
--- a/Engine/Source/Programs/Shared/EpicGames.Horde/Compute/Impl/HttpComputeClient.cs
+++ b/Engine/Source/Programs/Shared/EpicGames.Horde/Compute/Impl/HttpComputeClient.cs
@@ -48,7 +48,7 @@ namespace EpicGames.Horde.Compute.Impl
 		/// Refs for tasks to be executed
 		/// </summary>
 		[Required, CbField("t")]
-		public List<RefId> TaskRefIds { get; set; } = new List<RefId>();
+		public List<RefId> TaskRefIds { get; } = new List<RefId>();
 
 		/// <summary>
 		/// Requirements for agents executing the tasks
@@ -72,7 +72,7 @@ namespace EpicGames.Horde.Compute.Impl
 		/// Task updates
 		/// </summary>
 		[CbField("u")]
-		public List<ComputeTaskStatus> Updates { get; set; } = new List<ComputeTaskStatus>();
+		public List<ComputeTaskStatus> Updates { get; } = new List<ComputeTaskStatus>();
 	}
 
 	#endregion
@@ -111,7 +111,7 @@ namespace EpicGames.Horde.Compute.Impl
 			ReadOnlyMemoryContent content = new ReadOnlyMemoryContent(CbSerializer.Serialize(addTasks).GetView());
 			content.Headers.Add("Content-Type", "application/x-ue-cb");
 
-			HttpRequestMessage request = new HttpRequestMessage(HttpMethod.Post, $"api/v1/compute/{clusterId}");
+			using HttpRequestMessage request = new HttpRequestMessage(HttpMethod.Post, $"api/v1/compute/{clusterId}");
 			request.Content = content;
 
 			HttpResponseMessage response = await _httpClient.SendAsync(request, cancellationToken);
@@ -129,7 +129,7 @@ namespace EpicGames.Horde.Compute.Impl
 				using HttpResponseMessage response = await _httpClient.SendAsync(request, HttpCompletionOption.ResponseContentRead, cancellationToken);
 				response.EnsureSuccessStatusCode();
 
-				byte[] data = await response.Content.ReadAsByteArrayAsync();
+				byte[] data = await response.Content.ReadAsByteArrayAsync(cancellationToken);
 				GetTaskUpdatesResponse parsedResponse = CbSerializer.Deserialize<GetTaskUpdatesResponse>(new CbField(data));
 
 				foreach (ComputeTaskStatus update in parsedResponse.Updates)
diff --git a/Engine/Source/Programs/Shared/EpicGames.Horde/Compute/Requirements.cs b/Engine/Source/Programs/Shared/EpicGames.Horde/Compute/Requirements.cs
--- a/Engine/Source/Programs/Shared/EpicGames.Horde/Compute/Requirements.cs
+++ b/Engine/Source/Programs/Shared/EpicGames.Horde/Compute/Requirements.cs
@@ -22,7 +22,7 @@ namespace EpicGames.Horde.Compute
 		/// Resources used by the process
 		/// </summary>
 		[CbField("r")]
-		public Dictionary<string, int> Resources { get; set; } = new Dictionary<string, int>();
+		public Dictionary<string, int> Resources { get; } = new Dictionary<string, int>();
 
 		/// <summary>
 		/// Whether we require exclusive access to the device
diff --git a/Engine/Source/Programs/Shared/EpicGames.Horde/EpicGames.Horde.csproj b/Engine/Source/Programs/Shared/EpicGames.Horde/EpicGames.Horde.csproj
--- a/Engine/Source/Programs/Shared/EpicGames.Horde/EpicGames.Horde.csproj
+++ b/Engine/Source/Programs/Shared/EpicGames.Horde/EpicGames.Horde.csproj
@@ -1,6 +1,6 @@
 <Project Sdk="Microsoft.NET.Sdk">
   <PropertyGroup>
-    <TargetFramework>netcoreapp3.1</TargetFramework>
+    <TargetFramework>net6.0</TargetFramework>
     <Configuration Condition=" '$(Configuration)' == '' ">Development</Configuration>
     <OutputType>Library</OutputType>
     <Deterministic>true</Deterministic>
diff --git a/Engine/Source/Programs/Shared/EpicGames.Horde/Storage/BucketId.cs b/Engine/Source/Programs/Shared/EpicGames.Horde/Storage/BucketId.cs
--- a/Engine/Source/Programs/Shared/EpicGames.Horde/Storage/BucketId.cs
+++ b/Engine/Source/Programs/Shared/EpicGames.Horde/Storage/BucketId.cs
@@ -69,10 +69,10 @@ namespace EpicGames.Horde.Storage
 	sealed class BucketIdTypeConverter : TypeConverter
 	{
 		/// <inheritdoc/>
-		public override bool CanConvertFrom(ITypeDescriptorContext context, Type sourceType) => sourceType == typeof(string);
+		public override bool CanConvertFrom(ITypeDescriptorContext? context, Type? sourceType) => sourceType == typeof(string);
 
 		/// <inheritdoc/>
-		public override object ConvertFrom(ITypeDescriptorContext context, CultureInfo culture, object value) => new BucketId((string)value);
+		public override object ConvertFrom(ITypeDescriptorContext? context, CultureInfo? culture, object? value) => new BucketId((string)value!);
 	}
 
 	sealed class BucketIdCbConverter : CbConverterBase<BucketId>
diff --git a/Engine/Source/Programs/Shared/EpicGames.Horde/Storage/Impl/FileStorageClient.cs b/Engine/Source/Programs/Shared/EpicGames.Horde/Storage/Impl/FileStorageClient.cs
--- a/Engine/Source/Programs/Shared/EpicGames.Horde/Storage/Impl/FileStorageClient.cs
+++ b/Engine/Source/Programs/Shared/EpicGames.Horde/Storage/Impl/FileStorageClient.cs
@@ -75,7 +75,7 @@ namespace EpicGames.Horde.Storage.Impl
 			_logger.LogInformation("Writing {File} ({Size:n0} bytes)", file, stream.Length);
 			using (Stream outputStream = FileReference.Open(file, FileMode.Create, FileAccess.Write, FileShare.Read))
 			{
-				await stream.CopyToAsync(outputStream);
+				await stream.CopyToAsync(outputStream, cancellationToken);
 			}
 		}
 
@@ -85,7 +85,7 @@ namespace EpicGames.Horde.Storage.Impl
 			byte[] data;
 			using (MemoryStream memoryStream = new MemoryStream())
 			{
-				await stream.CopyToAsync(memoryStream);
+				await stream.CopyToAsync(memoryStream, cancellationToken);
 				data = memoryStream.ToArray();
 			}
 			return await this.WriteBlobFromMemoryAsync(namespaceId, data, cancellationToken);
@@ -131,7 +131,7 @@ namespace EpicGames.Horde.Storage.Impl
 		{
 			FileReference file = GetRefFile(namespaceId, bucketId, refId);
 			_logger.LogInformation("Reading {File} ({Size:n0} bytes)", file, new FileInfo(file.FullName).Length);
-			byte[] data = await FileReference.ReadAllBytesAsync(file);
+			byte[] data = await FileReference.ReadAllBytesAsync(file, cancellationToken);
 			return new Ref(namespaceId, bucketId, refId, new CbObject(data));
 		}
 
diff --git a/Engine/Source/Programs/Shared/EpicGames.Horde/Storage/Impl/HttpStorageClient.cs b/Engine/Source/Programs/Shared/EpicGames.Horde/Storage/Impl/HttpStorageClient.cs
--- a/Engine/Source/Programs/Shared/EpicGames.Horde/Storage/Impl/HttpStorageClient.cs
+++ b/Engine/Source/Programs/Shared/EpicGames.Horde/Storage/Impl/HttpStorageClient.cs
@@ -24,7 +24,7 @@ namespace EpicGames.Horde.Storage.Impl
 		/// <summary>
 		/// Set of missing hashes
 		/// </summary>
-		public HashSet<IoHash> Needs { get; set; } = new HashSet<IoHash>();
+		public HashSet<IoHash> Needs { get; } = new HashSet<IoHash>();
 	}
 
 	/// <summary>
@@ -35,7 +35,7 @@ namespace EpicGames.Horde.Storage.Impl
 		/// <summary>
 		/// List of missing hashes
 		/// </summary>
-		public List<IoHash> Needs { get; set; } = new List<IoHash>();
+		public List<IoHash> Needs { get; } = new List<IoHash>();
 	}
 
 	/// <summary>
@@ -46,7 +46,7 @@ namespace EpicGames.Horde.Storage.Impl
 		/// <summary>
 		/// List of hashes that the blob store needs.
 		/// </summary>
-		public List<RefId> Needs { get; set; } = new List<RefId>();
+		public List<RefId> Needs { get; } = new List<RefId>();
 	}
 
 	/// <summary>
@@ -74,7 +74,7 @@ namespace EpicGames.Horde.Storage.Impl
 		}
 
 		const string HashHeaderName = "X-Jupiter-IoHash";
-		const string LastAccessHeaderName = "X-Jupiter-LastAccess";
+//		const string LastAccessHeaderName = "X-Jupiter-LastAccess";
 
 		const string CompactBinaryMimeType = "application/x-ue-cb";
 
@@ -105,7 +105,7 @@ namespace EpicGames.Horde.Storage.Impl
 				}
 
 				response.EnsureSuccessStatusCode();
-				return await response.Content.ReadAsStreamAsync();
+				return await response.Content.ReadAsStreamAsync(cancellationToken);
 			}
 			catch
 			{
@@ -147,10 +147,10 @@ namespace EpicGames.Horde.Storage.Impl
 			HttpResponseMessage response = await _httpClient.SendAsync(request, cancellationToken);
 			response.EnsureSuccessStatusCode();
 
-			using (Stream responseStream = await response.Content.ReadAsStreamAsync())
+			using (Stream responseStream = await response.Content.ReadAsStreamAsync(cancellationToken))
 			{
-				WriteBlobResponse responseMessage = await JsonSerializer.DeserializeAsync<WriteBlobResponse>(responseStream, cancellationToken: cancellationToken);
-				return responseMessage.Identifier;
+				WriteBlobResponse? responseMessage = await JsonSerializer.DeserializeAsync<WriteBlobResponse>(responseStream, cancellationToken: cancellationToken);
+				return responseMessage!.Identifier;
 			}
 		}
 
@@ -210,7 +210,7 @@ namespace EpicGames.Horde.Storage.Impl
 			request.Headers.Accept.Clear();
 			request.Headers.Accept.Add(new MediaTypeWithQualityHeaderValue(CompactBinaryMimeType));
 
-			using HttpResponseMessage response = await _httpClient.SendAsync(request);
+			using HttpResponseMessage response = await _httpClient.SendAsync(request, cancellationToken);
 			if (!response.IsSuccessStatusCode)
 			{
 				if (response.StatusCode == System.Net.HttpStatusCode.NotFound)
@@ -223,7 +223,7 @@ namespace EpicGames.Horde.Storage.Impl
 				}
 			}
 
-			byte[] data = await response.Content.ReadAsByteArrayAsync();
+			byte[] data = await response.Content.ReadAsByteArrayAsync(cancellationToken);
 			return new RefImpl(namespaceId, bucketId, refId, new CbObject(data));
 		}
 
@@ -242,7 +242,7 @@ namespace EpicGames.Horde.Storage.Impl
 			request.Headers.Accept.Add(new MediaTypeWithQualityHeaderValue("application/json"));
 			request.Content = content;
 
-			using HttpResponseMessage response = await _httpClient.SendAsync(request);
+			using HttpResponseMessage response = await _httpClient.SendAsync(request, cancellationToken);
 			if (!response.IsSuccessStatusCode)
 			{
 				throw new RefException(namespaceId, bucketId, refId, await GetMessageFromResponse(response));
@@ -265,7 +265,7 @@ namespace EpicGames.Horde.Storage.Impl
 			request.Headers.Accept.Clear();
 			request.Headers.Accept.Add(new MediaTypeWithQualityHeaderValue("application/json"));
 
-			using HttpResponseMessage response = await _httpClient.SendAsync(request);
+			using HttpResponseMessage response = await _httpClient.SendAsync(request, cancellationToken);
 			if (!response.IsSuccessStatusCode)
 			{
 				throw new RefException(namespaceId, bucketId, refId, await GetMessageFromResponse(response));
@@ -357,7 +357,7 @@ namespace EpicGames.Horde.Storage.Impl
 		static async Task<T> ReadJsonResponse<T>(HttpContent content)
 		{
 			byte[] data = await content.ReadAsByteArrayAsync();
-			return JsonSerializer.Deserialize<T>(data, new JsonSerializerOptions { PropertyNameCaseInsensitive = true });
+			return JsonSerializer.Deserialize<T>(data, new JsonSerializerOptions { PropertyNameCaseInsensitive = true })!;
 		}
 	}
 }
diff --git a/Engine/Source/Programs/Shared/EpicGames.Horde/Storage/Impl/MemoryStorageClient.cs b/Engine/Source/Programs/Shared/EpicGames.Horde/Storage/Impl/MemoryStorageClient.cs
--- a/Engine/Source/Programs/Shared/EpicGames.Horde/Storage/Impl/MemoryStorageClient.cs
+++ b/Engine/Source/Programs/Shared/EpicGames.Horde/Storage/Impl/MemoryStorageClient.cs
@@ -50,7 +50,7 @@ namespace EpicGames.Horde.Storage.Impl
 		{
 			using (MemoryStream memoryStream = new MemoryStream())
 			{
-				await stream.CopyToAsync(memoryStream);
+				await stream.CopyToAsync(memoryStream, cancellationToken);
 				Blobs[(namespaceId, hash)] = memoryStream.ToArray();
 			}
 		}
@@ -60,7 +60,7 @@ namespace EpicGames.Horde.Storage.Impl
 		{
 			using (MemoryStream memoryStream = new MemoryStream())
 			{
-				await stream.CopyToAsync(memoryStream);
+				await stream.CopyToAsync(memoryStream, cancellationToken);
 
 				byte[] data = memoryStream.ToArray();
 				IoHash hash = IoHash.Compute(data);
diff --git a/Engine/Source/Programs/Shared/EpicGames.Horde/Storage/LocalBlobStore.cs b/Engine/Source/Programs/Shared/EpicGames.Horde/Storage/LocalBlobStore.cs
--- a/Engine/Source/Programs/Shared/EpicGames.Horde/Storage/LocalBlobStore.cs
+++ b/Engine/Source/Programs/Shared/EpicGames.Horde/Storage/LocalBlobStore.cs
@@ -31,7 +31,7 @@ namespace EpicGames.Horde.Storage
 		/// <summary>
 		/// Handle to the first piece in this block
 		/// </summary>
-		public readonly CellHandle FirstCell;
+		public CellHandle FirstCell { get; }
 
 		/// <summary>
 		/// Accessor for this handle as a long value
@@ -1040,7 +1040,7 @@ namespace EpicGames.Horde.Storage
 			}
 		}
 
-		private void LinkPage(ref int firstPageIdx, Page? lastPage, int newPageIdx)
+		private static void LinkPage(ref int firstPageIdx, Page? lastPage, int newPageIdx)
 		{
 			if (lastPage == null)
 			{
diff --git a/Engine/Source/Programs/Shared/EpicGames.Horde/Storage/NamespaceId.cs b/Engine/Source/Programs/Shared/EpicGames.Horde/Storage/NamespaceId.cs
--- a/Engine/Source/Programs/Shared/EpicGames.Horde/Storage/NamespaceId.cs
+++ b/Engine/Source/Programs/Shared/EpicGames.Horde/Storage/NamespaceId.cs
@@ -66,9 +66,9 @@ namespace EpicGames.Horde.Storage
 	sealed class NamespaceIdTypeConverter : TypeConverter
 	{
 		/// <inheritdoc/>
-		public override bool CanConvertFrom(ITypeDescriptorContext context, Type sourceType) => sourceType == typeof(string);
+		public override bool CanConvertFrom(ITypeDescriptorContext? context, Type sourceType) => sourceType == typeof(string);
 
 		/// <inheritdoc/>
-		public override object ConvertFrom(ITypeDescriptorContext context, CultureInfo culture, object value) => new NamespaceId((string)value);
+		public override object ConvertFrom(ITypeDescriptorContext? context, CultureInfo? culture, object value) => new NamespaceId((string)value);
 	}
 }
diff --git a/Engine/Source/Programs/Shared/EpicGames.Horde/Storage/RefId.cs b/Engine/Source/Programs/Shared/EpicGames.Horde/Storage/RefId.cs
--- a/Engine/Source/Programs/Shared/EpicGames.Horde/Storage/RefId.cs
+++ b/Engine/Source/Programs/Shared/EpicGames.Horde/Storage/RefId.cs
@@ -79,22 +79,22 @@ namespace EpicGames.Horde.Storage
 	sealed class RefIdTypeConverter : TypeConverter
 	{
 		/// <inheritdoc/>
-		public override bool CanConvertFrom(ITypeDescriptorContext context, Type sourceType)
+		public override bool CanConvertFrom(ITypeDescriptorContext? context, Type sourceType)
 		{
 			return sourceType == typeof(string);
 		}
 
 		/// <inheritdoc/>
-		public override object ConvertFrom(ITypeDescriptorContext context, CultureInfo culture, object value)
+		public override object ConvertFrom(ITypeDescriptorContext? context, CultureInfo? culture, object? value)
 		{
 			IoHash hash;
-			if (IoHash.TryParse((string)value, out hash))
+			if (IoHash.TryParse((string)value!, out hash))
 			{
 				return new RefId(hash);
 			}
 			else
 			{
-				return new RefId((string)value);
+				return new RefId((string)value!);
 			}
 		}
 	}
diff --git a/Engine/Source/Programs/Shared/EpicGames.Horde/StringId.cs b/Engine/Source/Programs/Shared/EpicGames.Horde/StringId.cs
--- a/Engine/Source/Programs/Shared/EpicGames.Horde/StringId.cs
+++ b/Engine/Source/Programs/Shared/EpicGames.Horde/StringId.cs
@@ -40,7 +40,9 @@ namespace EpicGames.Horde
 				{
 					if (character >= 'A' && character <= 'Z')
 					{
+#pragma warning disable CA1308 // Normalize strings to uppercase
 						Text = Text.ToLowerInvariant();
+#pragma warning restore CA1308 // Normalize strings to uppercase
 					}
 					else
 					{
