EpicGames.Perforce: Do not default to ignoring 'Empty' error responses. These are unpleasant to deal with when querying file collections, but are legitimate indications of errors when (for example) trying to submit empty changelists.

diff --git a/Engine/Source/Programs/Shared/EpicGames.Perforce/IPerforceOutput.cs b/Engine/Source/Programs/Shared/EpicGames.Perforce/IPerforceOutput.cs
--- a/Engine/Source/Programs/Shared/EpicGames.Perforce/IPerforceOutput.cs
+++ b/Engine/Source/Programs/Shared/EpicGames.Perforce/IPerforceOutput.cs
@@ -191,10 +191,8 @@ namespace EpicGames.Perforce
 						maxParsedLen = parsedLen + newBufferPos;
 						break;
 					}
-					if (response.Error == null || response.Error.Generic != PerforceGenericCode.Empty)
-					{
-						yield return response;
-					}
+
+					yield return response;
 					bufferPos = newBufferPos;
 				}
 
@@ -248,10 +246,8 @@ namespace EpicGames.Perforce
 						maxParsedLen = newBufferPos;
 						break;
 					}
-					if (response.Error == null || response.Error.Generic != PerforceGenericCode.Empty)
-					{
-						responses.Add(response);
-					}
+
+					responses.Add(response);
 					bufferPos = newBufferPos;
 				}
 
diff --git a/Engine/Source/Programs/Shared/EpicGames.Perforce/PerforceConnection.cs b/Engine/Source/Programs/Shared/EpicGames.Perforce/PerforceConnection.cs
--- a/Engine/Source/Programs/Shared/EpicGames.Perforce/PerforceConnection.cs
+++ b/Engine/Source/Programs/Shared/EpicGames.Perforce/PerforceConnection.cs
@@ -1528,7 +1528,7 @@ namespace EpicGames.Perforce
 		/// <param name="fileSpecs">List of file specifications to query</param>
 		/// <param name="cancellationToken">Token used to cancel the operation</param>
 		/// <returns>List of responses from the server</returns>
-		public static Task<PerforceResponseList<DeleteRecord>> TryDeleteAsync(this IPerforceConnection connection, int changeNumber, DeleteOptions options, FileSpecList fileSpecs, CancellationToken cancellationToken = default)
+		public static async Task<PerforceResponseList<DeleteRecord>> TryDeleteAsync(this IPerforceConnection connection, int changeNumber, DeleteOptions options, FileSpecList fileSpecs, CancellationToken cancellationToken = default)
 		{
 			List<string> arguments = new List<string>();
 			if (changeNumber != -1)
@@ -1548,7 +1548,9 @@ namespace EpicGames.Perforce
 				arguments.Add("-v");
 			}
 
-			return BatchedCommandAsync<DeleteRecord>(connection, "delete", arguments, fileSpecs.List, cancellationToken);
+			PerforceResponseList<DeleteRecord> records = await BatchedCommandAsync<DeleteRecord>(connection, "delete", arguments, fileSpecs.List, cancellationToken);
+			records.RemoveAll(x => x.Error != null && x.Error.Generic == PerforceGenericCode.Empty);
+			return records;
 		}
 
 		#endregion
@@ -1911,7 +1913,7 @@ namespace EpicGames.Perforce
 		/// <param name="fileSpecs">List of file specifications to query</param>
 		/// <param name="cancellationToken">Token used to cancel the operation</param>
 		/// <returns>List of responses from the server</returns>
-		public static Task<PerforceResponseList<FileLogRecord>> TryFileLogAsync(this IPerforceConnection connection, int changeNumber, int maxChanges, FileLogOptions options, FileSpecList fileSpecs, CancellationToken cancellationToken = default)
+		public static async Task<PerforceResponseList<FileLogRecord>> TryFileLogAsync(this IPerforceConnection connection, int changeNumber, int maxChanges, FileLogOptions options, FileSpecList fileSpecs, CancellationToken cancellationToken = default)
 		{
 			// Build the argument list
 			List<string> arguments = new List<string>();
@@ -1952,7 +1954,9 @@ namespace EpicGames.Perforce
 			arguments.Add("-t");
 
 			// Append all the arguments
-			return BatchedCommandAsync<FileLogRecord>(connection, "filelog", arguments, fileSpecs.List, cancellationToken);
+			PerforceResponseList<FileLogRecord> records = await BatchedCommandAsync<FileLogRecord>(connection, "filelog", arguments, fileSpecs.List, cancellationToken);
+			records.RemoveAll(x => x.Error != null && x.Error.Generic == PerforceGenericCode.Empty);
+			return records;
 		}
 
 		#endregion
@@ -2008,7 +2012,7 @@ namespace EpicGames.Perforce
 		/// <param name="fileSpecs">List of file specifications to query</param>
 		/// <param name="cancellationToken">Token used to cancel the operation</param>
 		/// <returns>List of response objects</returns>
-		public static Task<PerforceResponseList<FilesRecord>> TryFilesAsync(this IPerforceConnection connection, FilesOptions options, int maxFiles, FileSpecList fileSpecs, CancellationToken cancellationToken = default)
+		public static async Task<PerforceResponseList<FilesRecord>> TryFilesAsync(this IPerforceConnection connection, FilesOptions options, int maxFiles, FileSpecList fileSpecs, CancellationToken cancellationToken = default)
 		{
 			List<string> arguments = new List<string>();
 			if ((options & FilesOptions.AllRevisions) != 0)
@@ -2032,7 +2036,9 @@ namespace EpicGames.Perforce
 				arguments.Add($"-m{maxFiles}");
 			}
 
-			return BatchedCommandAsync<FilesRecord>(connection, "files", arguments, fileSpecs.List, cancellationToken);
+			PerforceResponseList<FilesRecord> records = await BatchedCommandAsync<FilesRecord>(connection, "files", arguments, fileSpecs.List, cancellationToken);
+			records.RemoveAll(x => x.Error != null && x.Error.Generic == PerforceGenericCode.Empty);
+			return records;
 		}
 
 		#endregion
@@ -2245,6 +2251,8 @@ namespace EpicGames.Perforce
 
 			// Execute the command
 			IAsyncEnumerable<PerforceResponse<FStatRecord>> records = StreamCommandAsync<FStatRecord>(connection, "fstat", arguments, fileSpecs.List, null, cancellationToken);
+			records = records.Where(x => x.Error == null || x.Error.Generic != PerforceGenericCode.Empty);
+
 			if (onlyChangeNumber != -1)
 			{
 				records = records.Where(x => !x.Succeeded || x.Data.Description == null);
@@ -2277,7 +2285,9 @@ namespace EpicGames.Perforce
 		/// <returns>List of file records</returns>
 		public static IAsyncEnumerable<PerforceResponse<HaveRecord>> TryHaveAsync(this IPerforceConnection connection, FileSpecList fileSpec, CancellationToken cancellationToken = default)
 		{
-			return StreamCommandAsync<HaveRecord>(connection, "have", new List<string>(), fileSpec.List, null, cancellationToken);
+			IAsyncEnumerable<PerforceResponse<HaveRecord>> records = StreamCommandAsync<HaveRecord>(connection, "have", new List<string>(), fileSpec.List, null, cancellationToken);
+			records = records.Where(x => x.Error == null || x.Error.Generic != PerforceGenericCode.Empty);
+			return records;
 		}
 
 		#endregion
@@ -2414,7 +2424,7 @@ namespace EpicGames.Perforce
 		/// <param name="targetFileSpec">The target filespec</param>
 		/// <param name="cancellationToken">Cancellation token</param>
 		/// <returns>List of records</returns>
-		public static Task<PerforceResponseList<MergeRecord>> MergeAsync(this IPerforceConnection connection, MergeOptions options, int change, int maxFiles, string sourceFileSpec, string targetFileSpec, CancellationToken cancellationToken = default)
+		public static async Task<PerforceResponseList<MergeRecord>> MergeAsync(this IPerforceConnection connection, MergeOptions options, int change, int maxFiles, string sourceFileSpec, string targetFileSpec, CancellationToken cancellationToken = default)
 		{
 			List<string> arguments = new List<string>();
 			if ((options & MergeOptions.Preview) != 0)
@@ -2432,7 +2442,9 @@ namespace EpicGames.Perforce
 			arguments.Add(sourceFileSpec);
 			arguments.Add(targetFileSpec);
 
-			return CommandAsync<MergeRecord>(connection, "merge", arguments, null, cancellationToken);
+			PerforceResponseList<MergeRecord> records = await CommandAsync<MergeRecord>(connection, "merge", arguments, null, cancellationToken);
+			records.RemoveAll(x => x.Error != null && x.Error.Generic == PerforceGenericCode.Empty);
+			return records;
 		}
 
 		#endregion
@@ -2466,7 +2478,7 @@ namespace EpicGames.Perforce
 		/// <param name="targetFileSpec">The target file(s)</param>
 		/// <param name="cancellationToken">Token used to cancel the operation</param>
 		/// <returns>Response from the server</returns>
-		public static Task<PerforceResponseList<MoveRecord>> TryMoveAsync(this IPerforceConnection connection, int changeNumber, string? fileType, MoveOptions options, string sourceFileSpec, string targetFileSpec, CancellationToken cancellationToken = default)
+		public static async Task<PerforceResponseList<MoveRecord>> TryMoveAsync(this IPerforceConnection connection, int changeNumber, string? fileType, MoveOptions options, string sourceFileSpec, string targetFileSpec, CancellationToken cancellationToken = default)
 		{
 			List<string> arguments = new List<string>();
 			if (changeNumber != -1)
@@ -2492,7 +2504,9 @@ namespace EpicGames.Perforce
 			arguments.Add(sourceFileSpec);
 			arguments.Add(targetFileSpec);
 
-			return CommandAsync<MoveRecord>(connection, "move", arguments, null, cancellationToken);
+			PerforceResponseList<MoveRecord> records = await CommandAsync<MoveRecord>(connection, "move", arguments, null, cancellationToken);
+			records.RemoveAll(x => x.Error != null && x.Error.Generic == PerforceGenericCode.Empty);
+			return records;
 		}
 
 		#endregion
@@ -2587,7 +2601,8 @@ namespace EpicGames.Perforce
 				arguments.Add($"-s");
 			}
 
-			return StreamCommandAsync<OpenedRecord>(connection, "opened", arguments, fileSpecs.List, null, cancellationToken);
+			IAsyncEnumerable<PerforceResponse<OpenedRecord>> records = StreamCommandAsync<OpenedRecord>(connection, "opened", arguments, fileSpecs.List, null, cancellationToken);
+			return records.Where(x => x.Error == null || x.Error.Generic != PerforceGenericCode.Empty);
 		}
 
 		#endregion
@@ -2724,7 +2739,7 @@ namespace EpicGames.Perforce
 		/// <param name="fileSpecs">Files to be reverted</param>
 		/// <param name="cancellationToken">Token used to cancel the operation</param>
 		/// <returns>Response from the server</returns>
-		public static Task<PerforceResponseList<ReconcileRecord>> TryReconcileAsync(this IPerforceConnection connection, int changeNumber, ReconcileOptions options, FileSpecList fileSpecs, CancellationToken cancellationToken = default)
+		public static async Task<PerforceResponseList<ReconcileRecord>> TryReconcileAsync(this IPerforceConnection connection, int changeNumber, ReconcileOptions options, FileSpecList fileSpecs, CancellationToken cancellationToken = default)
 		{
 			List<string> arguments = new List<string>();
 			if (changeNumber != -1)
@@ -2765,7 +2780,9 @@ namespace EpicGames.Perforce
 			}
 			arguments.AddRange(fileSpecs.List);
 
-			return CommandAsync<ReconcileRecord>(connection, "reconcile", arguments, null, cancellationToken);
+			PerforceResponseList<ReconcileRecord> records = await CommandAsync<ReconcileRecord>(connection, "reconcile", arguments, null, cancellationToken);
+			records.RemoveAll(x => x.Error != null && x.Error.Generic == PerforceGenericCode.Empty);
+			return records;
 		}
 
 		#endregion
@@ -2795,7 +2812,7 @@ namespace EpicGames.Perforce
 		/// <param name="fileSpec">Files to be reverted</param>
 		/// <param name="cancellationToken">Token used to cancel the operation</param>
 		/// <returns>Response from the server</returns>
-		public static Task<PerforceResponseList<ReopenRecord>> TryReopenAsync(this IPerforceConnection connection, int? changeNumber, string? fileType, string fileSpec, CancellationToken cancellationToken = default)
+		public static async Task<PerforceResponseList<ReopenRecord>> TryReopenAsync(this IPerforceConnection connection, int? changeNumber, string? fileType, string fileSpec, CancellationToken cancellationToken = default)
 		{
 			List<string> arguments = new List<string>();
 			if (changeNumber != null)
@@ -2815,7 +2832,9 @@ namespace EpicGames.Perforce
 			}
 			arguments.Add(fileSpec);
 
-			return CommandAsync<ReopenRecord>(connection, "reopen", arguments, null, cancellationToken);
+			PerforceResponseList<ReopenRecord> records = await CommandAsync<ReopenRecord>(connection, "reopen", arguments, null, cancellationToken);
+			records.RemoveAll(x => x.Error != null && x.Error.Generic == PerforceGenericCode.Empty);
+			return records;
 		}
 
 		#endregion
@@ -3157,7 +3176,7 @@ namespace EpicGames.Perforce
 		/// <param name="bUnloaded">Whether to enumerate unloaded workspaces</param>
 		/// <param name="cancellationToken">Token used to cancel the operation</param>
 		/// <returns>List of streams matching the given criteria</returns>
-		public static Task<PerforceResponseList<StreamsRecord>> TryGetStreamsAsync(this IPerforceConnection connection, string? streamPath, int maxResults, string? filter, bool bUnloaded, CancellationToken cancellationToken = default)
+		public static async Task<PerforceResponseList<StreamsRecord>> TryGetStreamsAsync(this IPerforceConnection connection, string? streamPath, int maxResults, string? filter, bool bUnloaded, CancellationToken cancellationToken = default)
 		{
 			// Build the command line
 			List<string> arguments = new List<string>();
@@ -3180,7 +3199,9 @@ namespace EpicGames.Perforce
 			}
 
 			// Execute the command
-			return CommandAsync<StreamsRecord>(connection, "streams", arguments, null, cancellationToken);
+			PerforceResponseList<StreamsRecord> records = await CommandAsync<StreamsRecord>(connection, "streams", arguments, null, cancellationToken);
+			records.RemoveAll(x => x.Error != null && x.Error.Generic == PerforceGenericCode.Empty);
+			return records;
 		}
 
 		#endregion
@@ -3557,7 +3578,7 @@ namespace EpicGames.Perforce
 		/// <param name="fileSpecs">Files to unshelve</param>
 		/// <param name="cancellationToken">Token used to cancel the operation</param>
 		/// <returns>Response from the server</returns>
-		public static Task<PerforceResponseList<UnshelveRecord>> TryUnshelveAsync(this IPerforceConnection connection, int changeNumber, int intoChangeNumber, string? usingBranchSpec, string? usingStream, string? forceParentStream, UnshelveOptions options, FileSpecList fileSpecs, CancellationToken cancellationToken = default)
+		public static async Task<PerforceResponseList<UnshelveRecord>> TryUnshelveAsync(this IPerforceConnection connection, int changeNumber, int intoChangeNumber, string? usingBranchSpec, string? usingStream, string? forceParentStream, UnshelveOptions options, FileSpecList fileSpecs, CancellationToken cancellationToken = default)
 		{
 			List<string> arguments = new List<string>();
 			arguments.Add($"-s{changeNumber}");
@@ -3587,7 +3608,9 @@ namespace EpicGames.Perforce
 			}
 			arguments.AddRange(fileSpecs.List);
 
-			return CommandAsync<UnshelveRecord>(connection, "unshelve", arguments, null, cancellationToken);
+			PerforceResponseList<UnshelveRecord> records = await CommandAsync<UnshelveRecord>(connection, "unshelve", arguments, null, cancellationToken);
+			records.RemoveAll(x => x.Error != null && x.Error.Generic == PerforceGenericCode.Empty);
+			return records;
 		}
 
 		#endregion
@@ -3646,7 +3669,7 @@ namespace EpicGames.Perforce
 		/// <param name="maxResults">Maximum number of results to return</param>
 		/// <param name="cancellationToken">Token used to cancel the operation</param>
 		/// <returns>Response from the server</returns>
-		public static Task<PerforceResponseList<UsersRecord>> TryGetUsersAsync(this IPerforceConnection connection, UsersOptions options, int maxResults, CancellationToken cancellationToken = default)
+		public static async Task<PerforceResponseList<UsersRecord>> TryGetUsersAsync(this IPerforceConnection connection, UsersOptions options, int maxResults, CancellationToken cancellationToken = default)
 		{
 			List<string> arguments = new List<string>();
 			if ((options & UsersOptions.IncludeServiceUsers) != 0)
@@ -3670,7 +3693,9 @@ namespace EpicGames.Perforce
 				arguments.Add($"-m{maxResults}");
 			}
 
-			return CommandAsync<UsersRecord>(connection, "users", arguments, null, cancellationToken);
+			PerforceResponseList<UsersRecord> records = await CommandAsync<UsersRecord>(connection, "users", arguments, null, cancellationToken);
+			records.RemoveAll(x => x.Error != null && x.Error.Generic == PerforceGenericCode.Empty);
+			return records;
 		}
 
 		#endregion
@@ -3698,7 +3723,8 @@ namespace EpicGames.Perforce
 		/// <returns>List of responses from the server</returns>
 		public static IAsyncEnumerable<PerforceResponse<WhereRecord>> TryWhereAsync(this IPerforceConnection connection, FileSpecList fileSpecs, CancellationToken cancellationToken = default)
 		{
-			return StreamCommandAsync<WhereRecord>(connection, "where", new List<string>(), fileSpecs.List, null, cancellationToken);
+			IAsyncEnumerable<PerforceResponse<WhereRecord>> records = StreamCommandAsync<WhereRecord>(connection, "where", new List<string>(), fileSpecs.List, null, cancellationToken);
+			return records.Where(x => x.Error == null || x.Error.Generic != PerforceGenericCode.Empty);
 		}
 
 		#endregion
@@ -3726,13 +3752,17 @@ namespace EpicGames.Perforce
 		/// <param name="changeNumber">Changelist number to receive the changes</param>
 		/// <param name="cancellationToken">Token used to cancel the operation</param>
 		/// <returns>Response from the server</returns>
-		public static Task<PerforceResponseList<UndoRecord>> TryUndoChangeAsync(this IPerforceConnection connection, int changeNumberToUndo, int changeNumber, CancellationToken cancellationToken = default)
+		public static async Task<PerforceResponseList<UndoRecord>> TryUndoChangeAsync(this IPerforceConnection connection, int changeNumberToUndo, int changeNumber, CancellationToken cancellationToken = default)
 		{
 			List<string> arguments = new List<string>();
 			arguments.Add($"-c{changeNumber}");
 			arguments.Add($"//...@{changeNumberToUndo}");
-			return CommandAsync<UndoRecord>(connection, "undo", arguments, null, cancellationToken);
+
+			PerforceResponseList<UndoRecord> records = await CommandAsync<UndoRecord>(connection, "undo", arguments, null, cancellationToken);
+			records.RemoveAll(x => x.Error != null && x.Error.Generic == PerforceGenericCode.Empty);
+			return records;
 		}
+
 		#endregion
 
 		#region p4 annotate
