Horde: Fix race condition parsing output from child processes, caused by Json log events referencing memory directly in the output buffer.

diff --git a/Engine/Source/Programs/Shared/EpicGames.Core/JsonLogEvent.cs b/Engine/Source/Programs/Shared/EpicGames.Core/JsonLogEvent.cs
--- a/Engine/Source/Programs/Shared/EpicGames.Core/JsonLogEvent.cs
+++ b/Engine/Source/Programs/Shared/EpicGames.Core/JsonLogEvent.cs
@@ -94,7 +94,7 @@ namespace EpicGames.Core
 
 				if (reader.TokenType == JsonTokenType.EndObject && level != LogLevel.None && reader.BytesConsumed == data.Length)
 				{
-					logEvent = new JsonLogEvent(level, new EventId(eventId), data);
+					logEvent = new JsonLogEvent(level, new EventId(eventId), data.ToArray());
 					return true;
 				}
 			}
