EpicGames.Core: Fix CommandLineArguments.MarkAsUnused() not correctly setting the flag value.

diff --git a/Engine/Source/Programs/Shared/EpicGames.Core/CommandLineArguments.cs b/Engine/Source/Programs/Shared/EpicGames.Core/CommandLineArguments.cs
--- a/Engine/Source/Programs/Shared/EpicGames.Core/CommandLineArguments.cs
+++ b/Engine/Source/Programs/Shared/EpicGames.Core/CommandLineArguments.cs
@@ -251,7 +251,7 @@ namespace EpicGames.Core
 		/// <param name="index">Index of the argument to mark as being unused</param>
 		public void MarkAsUnused(int index)
 		{
-			_usedArguments.Set(index, true);
+			_usedArguments.Set(index, false);
 		}
 
 		/// <summary>
