Revert change to use new DotNetPerforceLib. Causing problems with minimal workspaces on the farm.

diff --git a/Engine/Source/Programs/Shared/EpicGames.Perforce/EpicGames.Perforce.csproj b/Engine/Source/Programs/Shared/EpicGames.Perforce/EpicGames.Perforce.csproj
--- a/Engine/Source/Programs/Shared/EpicGames.Perforce/EpicGames.Perforce.csproj
+++ b/Engine/Source/Programs/Shared/EpicGames.Perforce/EpicGames.Perforce.csproj
@@ -1,4 +1,4 @@
-﻿<Project Sdk="Microsoft.NET.Sdk">
+<Project Sdk="Microsoft.NET.Sdk">
 
   <PropertyGroup>
     <TargetFramework>netcoreapp3.1</TargetFramework>
@@ -37,29 +37,25 @@
   </ItemGroup>
 
   <PropertyGroup>
-    <NativeConfig>Shipping</NativeConfig>
-    <NativeConfig Condition=" '$(Configuration)|$(Platform)' == 'Debug|AnyCPU' ">Debug</NativeConfig>
+    <WindowsNativeConfig>Release</WindowsNativeConfig>
+    <WindowsNativeConfig Condition=" '$(Configuration)|$(Platform)' == 'Debug|AnyCPU' ">Debug</WindowsNativeConfig>
     <DebugType>full</DebugType>
     <EnableNETAnalyzers>True</EnableNETAnalyzers>
   </PropertyGroup>
 
   <ItemGroup>
-    <ContentWithTargetPath Include="..\..\..\..\Binaries\Win64\DotNetPerforceLib-Win64-$(NativeConfig).dll" Visible="false">
+    <None Include="..\..\..\..\Binaries\DotNET\EpicGames.Perforce.Native\win-x64\$(WindowsNativeConfig)\EpicGames.Perforce.Native.dll" Visible="false">
       <CopyToOutputDirectory>PreserveNewest</CopyToOutputDirectory>
-      <TargetPath>EpicGames.Perforce.Native.dll</TargetPath>
-    </ContentWithTargetPath>
-    <ContentWithTargetPath Include="..\..\..\..\Binaries\Win64\DotNetPerforceLib-Win64-$(NativeConfig).pdb" Visible="false" Condition="Exists('..\..\..\..\Binaries\Win64\DotNetPerforceLib-Win64-$(NativeConfig).pdb')">
+    </None>
+    <None Include="..\..\..\..\Binaries\DotNET\EpicGames.Perforce.Native\win-x64\$(WindowsNativeConfig)\EpicGames.Perforce.Native.pdb" Visible="false" Condition="Exists('..\..\..\..\Binaries\DotNET\EpicGames.Perforce.Native\bin\win-x64\$(WindowsNativeConfig)\EpicGames.Perforce.Native.pdb')">
       <CopyToOutputDirectory>PreserveNewest</CopyToOutputDirectory>
-      <TargetPath>EpicGames.Perforce.Native.pdb</TargetPath>
-    </ContentWithTargetPath>
-    <ContentWithTargetPath Include="..\..\..\..\Binaries\Linux\libDotNetPerforceLib-Linux-$(NativeConfig).so" Visible="false">
+    </None>
+    <None Include="..\..\..\..\Binaries\DotNET\EpicGames.Perforce.Native\linux-x64\Release\EpicGames.Perforce.Native.so" Visible="false">
       <CopyToOutputDirectory>PreserveNewest</CopyToOutputDirectory>
-      <TargetPath>EpicGames.Perforce.Native.so</TargetPath>
-    </ContentWithTargetPath>
-    <ContentWithTargetPath Include="..\..\..\..\Binaries\Mac\DotNetPerforceLib-Mac-$(NativeConfig).dylib" Visible="false">
+    </None>
+    <None Include="..\..\..\..\Binaries\DotNET\EpicGames.Perforce.Native\mac-x64\Release\EpicGames.Perforce.Native.dylib" Visible="false">
       <CopyToOutputDirectory>PreserveNewest</CopyToOutputDirectory>
-      <TargetPath>EpicGames.Perforce.Native.dylib</TargetPath>
-    </ContentWithTargetPath>
+    </None>
   </ItemGroup>
 
   <ItemGroup>
