Output EpicGames.Perforce.Native files to Engine/Binaries/DotNET rather than Engine/Binaries/{Platform}.

diff --git a/Engine/Source/Programs/Shared/EpicGames.Perforce.Native/BuildLibrary.xml b/Engine/Source/Programs/Shared/EpicGames.Perforce.Native/BuildLibrary.xml
--- a/Engine/Source/Programs/Shared/EpicGames.Perforce.Native/BuildLibrary.xml
+++ b/Engine/Source/Programs/Shared/EpicGames.Perforce.Native/BuildLibrary.xml
@@ -4,6 +4,7 @@
 	<Agent Name="DotNetPerforceLib Win64" Type="Win64">
 		<Node Name="Compile DotNetPerforceLib Win64" Produces="#BinariesWin64">
 			<SetVersion Change="$(Change)" Branch="$(EscapedBranch)" If="$(IsBuildMachine)"/>
+			<Delete Files="$(RootDir)/Engine/Binaries/DotNET/EpicGames.Perforce.Native/win-x64/..."/>
 			<Compile Target="DotNetPerforceLib" Configuration="Debug" Platform="Win64" Tag="#FilesToSign"/>
 			<Compile Target="DotNetPerforceLib" Configuration="Shipping" Platform="Win64" Tag="#FilesToSign"/>
 			<Sign Files="#FilesToSign" Tag="#BinariesWin64"/>
@@ -14,6 +15,7 @@
 	<Agent Name="DotNetPerforceLib Mac" Type="Mac">
 		<Node Name="Compile DotNetPerforceLib Mac" Produces="#BinariesMac">
 			<SetVersion Change="$(Change)" Branch="$(EscapedBranch)" If="$(IsBuildMachine)"/>
+			<Delete Files="$(RootDir)/Engine/Binaries/DotNET/EpicGames.Perforce.Native/mac-x64/..."/>
 			<Compile Target="DotNetPerforceLib" Configuration="Debug" Platform="Mac" Tag="#BinariesMac"/>
 			<Compile Target="DotNetPerforceLib" Configuration="Shipping" Platform="Mac" Tag="#BinariesMac"/>
 			<Log Message="Files:" Files="#BinariesMac"/>
@@ -23,6 +25,7 @@
 	<Agent Name="DotNetPerforceLib Linux" Type="Linux">
 		<Node Name="Compile DotNetPerforceLib Linux" Produces="#BinariesLinux">
 			<SetVersion Change="$(Change)" Branch="$(EscapedBranch)" If="$(IsBuildMachine)"/>
+			<Delete Files="$(RootDir)/Engine/Binaries/DotNET/EpicGames.Perforce.Native/linux-x64/..."/>
 			<Compile Target="DotNetPerforceLib" Configuration="Debug" Platform="Linux" Tag="#BinariesLinux"/>
 			<Compile Target="DotNetPerforceLib" Configuration="Shipping" Platform="Linux" Tag="#BinariesLinux"/>
 			<Log Message="Files:" Files="#BinariesLinux"/>
@@ -31,7 +34,8 @@
 	
 	<Agent Name="DotNetPerforceLib Submit" Type="Win64">
 		<Node Name="Submit DotNetPerforceLib Binaries" Requires="#BinariesWin64;#BinariesMac;#BinariesLinux">
-			<Submit Files="#BinariesWin64;#BinariesMac;#BinariesLinux" Description="Rebuilt DotNetPerforceLib binaries for CL $(Change)"/>
+			<Tag Files="#BinariesWin64;#BinariesMac;#BinariesLinux" Filter="Engine/Binaries/DotNET/EpicGames.Perforce.Native/..." With="#FilesToSubmit"/>
+			<Submit Files="#FilesToSubmit" Description="Rebuilt DotNetPerforceLib binaries for CL $(Change)"/>
 		</Node>
 	</Agent>
 
diff --git a/Engine/Source/Programs/Shared/EpicGames.Perforce.Native/DotNetPerforceLib.Target.cs b/Engine/Source/Programs/Shared/EpicGames.Perforce.Native/DotNetPerforceLib.Target.cs
--- a/Engine/Source/Programs/Shared/EpicGames.Perforce.Native/DotNetPerforceLib.Target.cs
+++ b/Engine/Source/Programs/Shared/EpicGames.Perforce.Native/DotNetPerforceLib.Target.cs
@@ -18,5 +18,19 @@ public class DotNetPerforceLibTarget : TargetRules
 		bBuildWithEditorOnlyData = false;
 		bCompileAgainstEngine = false;
 		bCompileAgainstCoreUObject = false;
+
+		string ConfigFolder = (Target.Configuration == UnrealTargetConfiguration.Shipping) ? "Release" : "Debug";
+		if (Target.Platform == UnrealTargetPlatform.Win64)
+		{
+			OutputFile = $"Binaries/DotNET/EpicGames.Perforce.Native/win-x64/{ConfigFolder}/EpicGames.Perforce.Native.dll";
+		}
+		else if (Target.Platform == UnrealTargetPlatform.Mac)
+		{
+			OutputFile = $"Binaries/DotNET/EpicGames.Perforce.Native/mac-x64/{ConfigFolder}/EpicGames.Perforce.Native.dylib";
+		}
+		else if (Target.Platform == UnrealTargetPlatform.Linux)
+		{
+			OutputFile = $"Binaries/DotNET/EpicGames.Perforce.Native/linux-x64/{ConfigFolder}/EpicGames.Perforce.Native.so";
+		}
 	}
 }
