Horde: Add support for compressed files in compute tasks

When FileNode is marked as compressed, the appropriate method for download is used as compressed blobs are references by their uncompressed hash.
Files are still downloaded as-is and are not decompressed by the agent executor.

Flagging compression for files during upload is still left to be implemented.

diff --git a/Engine/Source/Programs/Horde/Horde.Agent.Tests/ComputeTaskExecutorTest.cs b/Engine/Source/Programs/Horde/Horde.Agent.Tests/ComputeTaskExecutorTest.cs
new file mode 100644
--- /dev/null
+++ b/Engine/Source/Programs/Horde/Horde.Agent.Tests/ComputeTaskExecutorTest.cs
@@ -0,0 +1,135 @@
+// Copyright Epic Games, Inc. All Rights Reserved.
+
+using System;
+using System.IO;
+using System.Text;
+using System.Threading.Tasks;
+using EpicGames.Core;
+using EpicGames.Horde.Compute;
+using EpicGames.Horde.Storage;
+using EpicGames.Horde.Storage.Impl;
+using Horde.Agent.Execution;
+using Microsoft.Extensions.Logging.Abstractions;
+using Microsoft.VisualStudio.TestTools.UnitTesting;
+
+namespace Horde.Agent.Tests
+{
+	public static class FakeCompressor
+	{
+		private static readonly byte[] s_prefix = Encoding.UTF8.GetBytes("FAKE-COMPRESSED-");
+		
+		/// <summary>
+		/// Decompress data that was fake compressed
+		/// <see cref="Compress(byte[])" />
+		/// </summary>
+		/// <param name="input">Data to decompress</param>
+		/// <returns>Decompressed data</returns>
+		public static byte[] Decompress(byte[] input)
+		{
+			bool isPrefixCorrect = input.AsSpan(0, input.Length).SequenceEqual(s_prefix.AsSpan());
+			if (isPrefixCorrect)
+			{
+				throw new ArgumentException("Compressed data is missing magic prefix");
+			}
+
+			return input.AsSpan(s_prefix.Length).ToArray();
+		}
+		
+		/// <summary>
+		/// Decompress data that was fake compressed
+		/// <see cref="Compress(byte[])" />
+		/// </summary>
+		/// <param name="input">Data to decompress</param>
+		/// <returns>Decompressed data</returns>
+		public static string Decompress(string input)
+		{
+			return Encoding.UTF8.GetString(Decompress(Encoding.UTF8.GetBytes(input)));
+		}
+		
+		/// <summary>
+		/// Fake compression of input data by appending a prefix
+		/// Won't make the data smaller, but it will modify it leading to a changed hash
+		/// </summary>
+		/// <param name="input">Data to compress</param>
+		/// <returns>Compressed data</returns>
+		public static byte[] Compress(byte[] input)
+		{
+			byte[] compressedInput = new byte[input.Length + s_prefix.Length];
+			Array.Copy(s_prefix, 0, compressedInput, 0, s_prefix.Length);
+			Array.Copy(input, 0, compressedInput, s_prefix.Length, input.Length);
+			return compressedInput;
+		}
+		
+		/// <summary>
+		/// Fake compression of input data by appending a prefix
+		/// <see cref="Compress(byte[])"/>
+		/// </summary>
+		/// <param name="input">Data to compress</param>
+		/// <returns>Compressed data</returns>
+		public static string Compress(string input)
+		{
+			return Encoding.UTF8.GetString(Compress(Encoding.UTF8.GetBytes(input)));
+		}
+	}
+	
+	[TestClass]
+	public class ComputeTaskExecutorTest
+	{
+		private readonly MemoryStorageClient _storageClient = new ();
+		private readonly ComputeTaskExecutor _executor;
+		private readonly NamespaceId _namespaceId = new ("my-namespace");
+		private readonly DirectoryReference _sandboxDir;
+		
+		public ComputeTaskExecutorTest()
+		{
+			_executor = new (_storageClient, NullLogger.Instance);
+			_sandboxDir = new (Path.Join(Path.GetTempPath(), "horde-agent", Guid.NewGuid().ToString()));
+			DirectoryReference.CreateDirectory(_sandboxDir);
+		}
+
+		[TestCleanup]
+		public void TestCleanup()
+		{
+			if (Directory.Exists(_sandboxDir.FullName))
+			{
+				Directory.Delete(_sandboxDir.FullName, true);
+			}
+		}
+		
+		[TestMethod]
+		public async Task SetupSandboxWithCompressedFiles()
+		{
+			DirectoryTree rootDir = new ();
+			FileNode file1 = await CreateFileInStorage("file1", "file1Data", false);
+			FileNode file2 = await CreateFileInStorage("file2", "file2Data", true);
+			rootDir.Files.Add(file1);
+			rootDir.Files.Add(file2);
+
+			await _executor.SetupSandboxAsync(_namespaceId, rootDir, _sandboxDir);
+
+			Assert.AreEqual("file1Data", File.ReadAllText(Path.Join(_sandboxDir.FullName, file1.Name.ToString())));
+			Assert.AreEqual(FakeCompressor.Compress("file2Data"), File.ReadAllText(Path.Join(_sandboxDir.FullName, file2.Name.ToString())));
+		}
+
+		private async Task<FileNode> CreateFileInStorage(string name, string data, bool shouldCompress)
+		{
+			byte[] uncompressedData = Encoding.UTF8.GetBytes(data);
+			IoHash uncompressedHash = IoHash.Compute(uncompressedData);
+			byte[] compressedData = FakeCompressor.Compress(uncompressedData);
+
+			if (shouldCompress)
+			{
+				FileNode fileNode = new (name, uncompressedHash, compressedData.Length, 0, true);
+				using MemoryStream stream = new(compressedData);
+				await _storageClient.WriteCompressedBlobAsync(_namespaceId, uncompressedHash, stream);
+				return fileNode;
+			}
+			else
+			{
+				FileNode fileNode = new (name, uncompressedHash, uncompressedData.Length, 0, false);
+				await _storageClient.WriteBlobFromMemoryAsync(_namespaceId, uncompressedData);
+				return fileNode;				
+			}
+		}
+	}
+}
diff --git a/Engine/Source/Programs/Horde/Horde.Agent/Commands/Compute/ComputeCommand.cs b/Engine/Source/Programs/Horde/Horde.Agent/Commands/Compute/ComputeCommand.cs
--- a/Engine/Source/Programs/Horde/Horde.Agent/Commands/Compute/ComputeCommand.cs
+++ b/Engine/Source/Programs/Horde/Horde.Agent/Commands/Compute/ComputeCommand.cs
@@ -107,7 +107,7 @@ namespace Horde.Agent.Commands
 				byte[] data = File.ReadAllBytes(fileInfo.FullName);
 				IoHash hash = IoHash.Compute(data);
 				uploadList[hash] = data;
-				tree.Files.Add(new FileNode(fileInfo.Name, hash, fileInfo.Length, (int)fileInfo.Attributes));
+				tree.Files.Add(new FileNode(fileInfo.Name, hash, fileInfo.Length, (int)fileInfo.Attributes, false));
 			}
 			tree.Files.SortBy(x => x.Name, Utf8StringComparer.Ordinal);
 
diff --git a/Engine/Source/Programs/Horde/Horde.Agent/Execution/ComputeTaskExecutor.cs b/Engine/Source/Programs/Horde/Horde.Agent/Execution/ComputeTaskExecutor.cs
--- a/Engine/Source/Programs/Horde/Horde.Agent/Execution/ComputeTaskExecutor.cs
+++ b/Engine/Source/Programs/Horde/Horde.Agent/Execution/ComputeTaskExecutor.cs
@@ -181,7 +181,7 @@ namespace Horde.Agent.Execution
 		/// <param name="inputDirectory">The directory spec</param>
 		/// <param name="outputDir">Output directory on disk</param>
 		/// <returns>Async task</returns>
-		async Task SetupSandboxAsync(NamespaceId namespaceId, DirectoryTree inputDirectory, DirectoryReference outputDir)
+		internal async Task SetupSandboxAsync(NamespaceId namespaceId, DirectoryTree inputDirectory, DirectoryReference outputDir)
 		{
 			DirectoryReference.CreateDirectory(outputDir);
 
@@ -189,7 +189,9 @@ namespace Horde.Agent.Execution
 			{
 				FileReference file = FileReference.Combine(outputDir, fileNode.Name.ToString());
 				_logger.LogInformation("Downloading {File} (digest: {Digest})", file, fileNode.Hash);
-				byte[] data = await _storageClient.ReadBlobToMemoryAsync(namespaceId, fileNode.Hash);
+				byte[] data = fileNode.IsCompressed
+					? await _storageClient.ReadCompressedBlobToMemoryAsync(namespaceId, fileNode.Hash)
+					: await _storageClient.ReadBlobToMemoryAsync(namespaceId, fileNode.Hash);
 				_logger.LogInformation("Writing {File} (digest: {Digest})", file, fileNode.Hash);
 				await FileReference.WriteAllBytesAsync(file, data);
 			}
@@ -276,7 +278,8 @@ namespace Horde.Agent.Execution
 		{
 			byte[] data = await FileReference.ReadAllBytesAsync(file);
 			IoHash hash = await _storageClient.WriteBlobFromMemoryAsync(namespaceId, data);
-			return new FileNode(name, hash, data.Length, (int)FileReference.GetAttributes(file));
+			// TODO: Figure out how file can be marked as compressed
+			return new FileNode(name, hash, data.Length, (int)FileReference.GetAttributes(file), false);
 		}
 
 		/// <summary>
diff --git a/Engine/Source/Programs/Horde/Horde.Build/Storage/Services/BasicStorageClient.cs b/Engine/Source/Programs/Horde/Horde.Build/Storage/Services/BasicStorageClient.cs
--- a/Engine/Source/Programs/Horde/Horde.Build/Storage/Services/BasicStorageClient.cs
+++ b/Engine/Source/Programs/Horde/Horde.Build/Storage/Services/BasicStorageClient.cs
@@ -70,6 +70,12 @@ namespace Horde.Build.Storage.Services
 
 			return stream;
 		}
+		
+		/// <inheritdoc/>
+		public Task<Stream> ReadCompressedBlobAsync(NamespaceId namespaceId, IoHash uncompressedHash, CancellationToken cancellationToken)
+		{
+			throw new NotImplementedException();
+		}
 
 		/// <inheritdoc/>
 		public async Task WriteBlobAsync(NamespaceId namespaceId, IoHash hash, Stream stream, CancellationToken cancellationToken)
@@ -109,6 +115,12 @@ namespace Horde.Build.Storage.Services
 			return dataHash;
 		}
 
+		/// <inheritdoc/>
+		public Task WriteCompressedBlobAsync(NamespaceId namespaceId, IoHash uncompressedHash, Stream stream, CancellationToken cancellationToken = default)
+		{
+			throw new NotImplementedException();
+		}
+
 		/// <inheritdoc/>
 		public async Task<bool> HasBlobAsync(NamespaceId namespaceId, IoHash hash, CancellationToken cancellationToken)
 		{
diff --git a/Engine/Source/Programs/Shared/EpicGames.Horde/Compute/DirectoryTree.cs b/Engine/Source/Programs/Shared/EpicGames.Horde/Compute/DirectoryTree.cs
--- a/Engine/Source/Programs/Shared/EpicGames.Horde/Compute/DirectoryTree.cs
+++ b/Engine/Source/Programs/Shared/EpicGames.Horde/Compute/DirectoryTree.cs
@@ -52,6 +52,12 @@ namespace EpicGames.Horde.Compute
 		/// </summary>
 		[CbField("a")]
 		public int Attributes { get; set; }
+		
+		/// <summary>
+		/// Whether the file is compressed or not.
+		/// </summary>
+		[CbField("c")]
+		public bool IsCompressed { get; set; }
 
 		/// <summary>
 		/// Private constructor for serialization
@@ -63,12 +69,13 @@ namespace EpicGames.Horde.Compute
 		/// <summary>
 		/// Constructor
 		/// </summary>
-		public FileNode(Utf8String name, CbBinaryAttachment hash, long size, int attributes)
+		public FileNode(Utf8String name, CbBinaryAttachment hash, long size, int attributes, bool isCompressed)
 		{
 			Name = name;
 			Hash = hash;
 			Size = size;
 			Attributes = attributes;
+			IsCompressed = isCompressed;
 		}
 	}
 
diff --git a/Engine/Source/Programs/Shared/EpicGames.Horde/Storage/IStorageClient.cs b/Engine/Source/Programs/Shared/EpicGames.Horde/Storage/IStorageClient.cs
--- a/Engine/Source/Programs/Shared/EpicGames.Horde/Storage/IStorageClient.cs
+++ b/Engine/Source/Programs/Shared/EpicGames.Horde/Storage/IStorageClient.cs
@@ -165,6 +165,15 @@ namespace EpicGames.Horde.Storage
 		/// <param name="cancellationToken">Cancellation token for the operation</param>
 		/// <returns>Stream for the blob, or null if it does not exist</returns>
 		Task<Stream> ReadBlobAsync(NamespaceId namespaceId, IoHash hash, CancellationToken cancellationToken = default);
+		
+		/// <summary>
+		/// Opens a blob read stream for compressed blobs
+		/// </summary>
+		/// <param name="namespaceId">Namespace to operate on</param>
+		/// <param name="uncompressedHash">Hash of the uncompressed blob</param>
+		/// <param name="cancellationToken">Cancellation token for the operation</param>
+		/// <returns>Stream for the blob, or null if it does not exist</returns>
+		Task<Stream> ReadCompressedBlobAsync(NamespaceId namespaceId, IoHash uncompressedHash, CancellationToken cancellationToken = default);
 
 		/// <summary>
 		/// Writes a blob to storage
@@ -176,13 +185,23 @@ namespace EpicGames.Horde.Storage
 		Task WriteBlobAsync(NamespaceId namespaceId, IoHash hash, Stream stream, CancellationToken cancellationToken = default);
 
 		/// <summary>
-		/// Writes a blob to storage
+		/// Writes a blob to storage and calculates hash
 		/// </summary>
 		/// <param name="namespaceId">Namespace to operate on</param>
 		/// <param name="stream">The stream to write</param>
 		/// <param name="cancellationToken">Cancellation token for the operation</param>
+		/// <returns>Hash of the blob written</returns>
 		Task<IoHash> WriteBlobAsync(NamespaceId namespaceId, Stream stream, CancellationToken cancellationToken = default);
 
+		/// <summary>
+		/// Writes a compressed blob to storage
+		/// </summary>
+		/// <param name="namespaceId">Namespace to operate on</param>
+		/// <param name="uncompressedHash">Hash of the blob</param>
+		/// <param name="compressedStream">Compressed stream to write</param>
+		/// <param name="cancellationToken">Cancellation token for the operation</param>
+		Task WriteCompressedBlobAsync(NamespaceId namespaceId, IoHash uncompressedHash, Stream compressedStream, CancellationToken cancellationToken = default);
+		
 		/// <summary>
 		/// Checks if the given blob exists
 		/// </summary>
@@ -285,9 +304,30 @@ namespace EpicGames.Horde.Storage
 		/// <param name="maxInMemoryBlobLength">Maximum allowed memory allocation to store the blob</param>
 		/// <param name="cancellationToken">Cancellation token for the operation</param>
 		/// <returns>Data for the blob that was read. Throws an exception if the blob was not found.</returns>
-		public static async Task<byte[]> ReadBlobToMemoryAsync(this IStorageClient storageClient, NamespaceId namespaceId, IoHash hash, int maxInMemoryBlobLength = DefaultMaxInMemoryBlobLength, CancellationToken cancellationToken = default)
+		public static Task<byte[]> ReadBlobToMemoryAsync(this IStorageClient storageClient, NamespaceId namespaceId, IoHash hash, int maxInMemoryBlobLength = DefaultMaxInMemoryBlobLength, CancellationToken cancellationToken = default)
+		{
+			return ReadBlobToMemoryAsync(storageClient, false, namespaceId, hash, maxInMemoryBlobLength, cancellationToken);
+		}
+		
+		/// <summary>
+		/// Gets a compressed blob as a byte array
+		/// </summary>
+		/// <param name="storageClient">The storage interface</param>
+		/// <param name="namespaceId">Namespace containing the blob</param>
+		/// <param name="hash">Hash of the blob to read</param>
+		/// <param name="maxInMemoryBlobLength">Maximum allowed memory allocation to store the blob</param>
+		/// <param name="cancellationToken">Cancellation token for the operation</param>
+		/// <returns>Data for the blob that was read. Throws an exception if the blob was not found.</returns>
+		public static Task<byte[]> ReadCompressedBlobToMemoryAsync(this IStorageClient storageClient, NamespaceId namespaceId, IoHash hash, int maxInMemoryBlobLength = DefaultMaxInMemoryBlobLength, CancellationToken cancellationToken = default)
+		{
+			return ReadBlobToMemoryAsync(storageClient, true, namespaceId, hash, maxInMemoryBlobLength, cancellationToken);
+		}
+		
+		private static async Task<byte[]> ReadBlobToMemoryAsync(IStorageClient storageClient, bool isBlobCompressed, NamespaceId namespaceId, IoHash hash, int maxInMemoryBlobLength = DefaultMaxInMemoryBlobLength, CancellationToken cancellationToken = default)
 		{
-			using Stream stream = await storageClient.ReadBlobAsync(namespaceId, hash, cancellationToken);
+			using Stream stream = isBlobCompressed
+				? await storageClient.ReadCompressedBlobAsync(namespaceId, hash, cancellationToken)
+				: await storageClient.ReadBlobAsync(namespaceId, hash, cancellationToken);
 
 			long length = stream.Length;
 			if (length > maxInMemoryBlobLength)
diff --git a/Engine/Source/Programs/Shared/EpicGames.Horde/Storage/Impl/FileStorageClient.cs b/Engine/Source/Programs/Shared/EpicGames.Horde/Storage/Impl/FileStorageClient.cs
--- a/Engine/Source/Programs/Shared/EpicGames.Horde/Storage/Impl/FileStorageClient.cs
+++ b/Engine/Source/Programs/Shared/EpicGames.Horde/Storage/Impl/FileStorageClient.cs
@@ -65,6 +65,12 @@ namespace EpicGames.Horde.Storage.Impl
 			_logger.LogInformation("Reading {File} ({Size:n0} bytes)", file, new FileInfo(file.FullName).Length);
 			return Task.FromResult<Stream>(FileReference.Open(file, FileMode.Open, FileAccess.Read, FileShare.Read));
 		}
+		
+		/// <inheritdoc/>
+		public Task<Stream> ReadCompressedBlobAsync(NamespaceId namespaceId, IoHash uncompressedHash, CancellationToken cancellationToken = default)
+		{
+			throw new NotImplementedException();
+		}
 
 		/// <inheritdoc/>
 		public async Task WriteBlobAsync(NamespaceId namespaceId, IoHash hash, Stream stream, CancellationToken cancellationToken = default)
@@ -91,6 +97,12 @@ namespace EpicGames.Horde.Storage.Impl
 			return await this.WriteBlobFromMemoryAsync(namespaceId, data, cancellationToken);
 		}
 
+		/// <inheritdoc/>
+		public Task WriteCompressedBlobAsync(NamespaceId namespaceId, IoHash uncompressedHash, Stream stream, CancellationToken cancellationToken = default)
+		{
+			throw new NotImplementedException();
+		}
+
 		/// <inheritdoc/>
 		public Task<bool> HasBlobAsync(NamespaceId namespaceId, IoHash hash, CancellationToken cancellationToken = default)
 		{
diff --git a/Engine/Source/Programs/Shared/EpicGames.Horde/Storage/Impl/HttpStorageClient.cs b/Engine/Source/Programs/Shared/EpicGames.Horde/Storage/Impl/HttpStorageClient.cs
--- a/Engine/Source/Programs/Shared/EpicGames.Horde/Storage/Impl/HttpStorageClient.cs
+++ b/Engine/Source/Programs/Shared/EpicGames.Horde/Storage/Impl/HttpStorageClient.cs
@@ -92,9 +92,20 @@ namespace EpicGames.Horde.Storage.Impl
 		#region Blobs
 
 		/// <inheritdoc/>
-		public async Task<Stream> ReadBlobAsync(NamespaceId namespaceId, IoHash hash, CancellationToken cancellationToken)
+		public Task<Stream> ReadBlobAsync(NamespaceId namespaceId, IoHash hash, CancellationToken cancellationToken)
 		{
-			using HttpRequestMessage request = new HttpRequestMessage(HttpMethod.Get, $"api/v1/blobs/{namespaceId}/{hash}");
+			return ReadBlobFromUrlAsync($"api/v1/blobs/{namespaceId}/{hash}", namespaceId, hash, cancellationToken);
+		}
+		
+		/// <inheritdoc/>
+		public Task<Stream> ReadCompressedBlobAsync(NamespaceId namespaceId, IoHash hash, CancellationToken cancellationToken = default)
+		{
+			return ReadBlobFromUrlAsync($"api/v1/compressed-blobs/{namespaceId}/{hash}", namespaceId, hash, cancellationToken);
+		}
+		
+		private async Task<Stream> ReadBlobFromUrlAsync(string url, NamespaceId namespaceId, IoHash hash, CancellationToken cancellationToken)
+		{
+			using HttpRequestMessage request = new HttpRequestMessage(HttpMethod.Get, url);
 
 			HttpResponseMessage response = await _httpClient.SendAsync(request, cancellationToken);
 			try
@@ -154,6 +165,20 @@ namespace EpicGames.Horde.Storage.Impl
 			}
 		}
 
+		/// <inheritdoc/>
+		public async Task WriteCompressedBlobAsync(NamespaceId namespaceId, IoHash uncompressedHash, Stream stream, CancellationToken cancellationToken = default)
+		{
+			StreamContent content = new StreamContent(stream);
+			content.Headers.ContentType = new MediaTypeHeaderValue("application/x-ue-comp");
+			content.Headers.ContentLength = stream.Length;
+
+			using HttpRequestMessage request = new HttpRequestMessage(HttpMethod.Put, $"api/v1/compressed-blobs/{namespaceId}/{uncompressedHash}");
+			request.Content = content;
+
+			HttpResponseMessage response = await _httpClient.SendAsync(request, cancellationToken);
+			response.EnsureSuccessStatusCode();
+		}
+
 		/// <inheritdoc/>
 		public async Task<bool> HasBlobAsync(NamespaceId namespaceId, IoHash hash, CancellationToken cancellationToken = default)
 		{
diff --git a/Engine/Source/Programs/Shared/EpicGames.Horde/Storage/Impl/MemoryStorageClient.cs b/Engine/Source/Programs/Shared/EpicGames.Horde/Storage/Impl/MemoryStorageClient.cs
--- a/Engine/Source/Programs/Shared/EpicGames.Horde/Storage/Impl/MemoryStorageClient.cs
+++ b/Engine/Source/Programs/Shared/EpicGames.Horde/Storage/Impl/MemoryStorageClient.cs
@@ -34,6 +34,11 @@ namespace EpicGames.Horde.Storage.Impl
 
 		/// <inheritdoc/>
 		public Dictionary<(NamespaceId, IoHash), ReadOnlyMemory<byte>> Blobs { get; } = new Dictionary<(NamespaceId, IoHash), ReadOnlyMemory<byte>>();
+		
+		/// <summary>
+		/// Mapping from uncompressed hash to compressed hash.
+		/// </summary>
+		public Dictionary<(NamespaceId, IoHash), IoHash> UncompressedToCompressedHash { get; } = new Dictionary<(NamespaceId, IoHash), IoHash>();
 
 		/// <inheritdoc/>
 		public Dictionary<(NamespaceId, BucketId, RefId), IRef> Refs { get; } = new Dictionary<(NamespaceId, BucketId, RefId), IRef>();
@@ -44,6 +49,14 @@ namespace EpicGames.Horde.Storage.Impl
 			ReadOnlyMemory<byte> data = Blobs[(namespaceId, hash)];
 			return Task.FromResult<Stream>(new ReadOnlyMemoryStream(data));
 		}
+		
+		/// <inheritdoc/>
+		public Task<Stream> ReadCompressedBlobAsync(NamespaceId namespaceId, IoHash uncompressedHash, CancellationToken cancellationToken = default)
+		{
+			IoHash compressedHash = UncompressedToCompressedHash[(namespaceId, uncompressedHash)];
+			ReadOnlyMemory<byte> data = Blobs[(namespaceId, compressedHash)];
+			return Task.FromResult<Stream>(new ReadOnlyMemoryStream(data));
+		}
 
 		/// <inheritdoc/>
 		public async Task WriteBlobAsync(NamespaceId namespaceId, IoHash hash, Stream stream, CancellationToken cancellationToken = default)
@@ -70,6 +83,17 @@ namespace EpicGames.Horde.Storage.Impl
 			}
 		}
 
+		/// <inheritdoc/>
+		public async Task WriteCompressedBlobAsync(NamespaceId namespaceId, IoHash uncompressedHash, Stream compressedStream, CancellationToken cancellationToken = default)
+		{
+			using MemoryStream memoryStream = new MemoryStream();
+			await compressedStream.CopyToAsync(memoryStream, cancellationToken);
+			byte[] compressedData = memoryStream.ToArray();
+			IoHash compressedHash = IoHash.Compute(compressedData);
+			Blobs[(namespaceId, compressedHash)] = compressedData;
+			UncompressedToCompressedHash[(namespaceId, uncompressedHash)] = compressedHash;
+		}
+
 		/// <inheritdoc/>
 		public Task<bool> HasBlobAsync(NamespaceId namespaceId, IoHash hash, CancellationToken cancellationToken = default)
 		{
diff --git a/Engine/Source/Programs/Shared/EpicGames.Horde/Storage/Impl/TracingStorageClient.cs b/Engine/Source/Programs/Shared/EpicGames.Horde/Storage/Impl/TracingStorageClient.cs
--- a/Engine/Source/Programs/Shared/EpicGames.Horde/Storage/Impl/TracingStorageClient.cs
+++ b/Engine/Source/Programs/Shared/EpicGames.Horde/Storage/Impl/TracingStorageClient.cs
@@ -35,6 +35,13 @@ namespace EpicGames.Horde.Storage.Impl
 			_logger.LogDebug("Reading blob {NamespaceId}/{Hash}", namespaceId, hash);
 			return _inner.ReadBlobAsync(namespaceId, hash, cancellationToken);
 		}
+		
+		/// <inheritdoc/>
+		public Task<Stream> ReadCompressedBlobAsync(NamespaceId namespaceId, IoHash hash, CancellationToken cancellationToken = default)
+		{
+			_logger.LogDebug("Reading compressed blob {NamespaceId}/{Hash}", namespaceId, hash);
+			return _inner.ReadCompressedBlobAsync(namespaceId, hash, cancellationToken);
+		}
 
 		/// <inheritdoc/>
 		public Task WriteBlobAsync(NamespaceId namespaceId, IoHash hash, Stream stream, CancellationToken cancellationToken = default)
@@ -51,6 +58,13 @@ namespace EpicGames.Horde.Storage.Impl
 			return hash;
 		}
 
+		/// <inheritdoc/>
+		public Task WriteCompressedBlobAsync(NamespaceId namespaceId, IoHash uncompressedHash, Stream stream, CancellationToken cancellationToken = default)
+		{
+			_logger.LogDebug("Writing compressed blob {NamespaceId}/{UncompressedHash}", namespaceId, uncompressedHash);
+			return _inner.WriteCompressedBlobAsync(namespaceId, uncompressedHash, stream, cancellationToken);
+		}
+
 		/// <inheritdoc/>
 		public Task<bool> HasBlobAsync(NamespaceId namespaceId, IoHash hash, CancellationToken cancellationToken = default)
 		{
