Horde: Guard against exception when log value is null.

diff --git a/Engine/Source/Programs/Shared/EpicGames.Core/LogValueFormatter.cs b/Engine/Source/Programs/Shared/EpicGames.Core/LogValueFormatter.cs
--- a/Engine/Source/Programs/Shared/EpicGames.Core/LogValueFormatter.cs
+++ b/Engine/Source/Programs/Shared/EpicGames.Core/LogValueFormatter.cs
@@ -282,8 +282,14 @@ namespace EpicGames.Core
 		/// </summary>
 		/// <param name="value">Value to write</param>
 		/// <param name="writer">Json writer</param>
-		public static void Format(object value, Utf8JsonWriter writer)
+		public static void Format(object? value, Utf8JsonWriter writer)
 		{
+			if (value == null)
+			{
+				writer.WriteNullValue();
+				return;
+			}
+
 			ILogValueFormatter formatter = GetFormatter(value.GetType());
 			formatter.Format(value, writer);
 		}
