Fixed issue with parsing and formatting of numbers which ended up using current locale instead of invariant.

Fixed issue where specifiers from a previous header file which errored out could be processed by a different header.

#rnx

diff --git a/Engine/Source/Programs/Shared/EpicGames.UHT/Parsers/UhtSpecifierParser.cs b/Engine/Source/Programs/Shared/EpicGames.UHT/Parsers/UhtSpecifierParser.cs
--- a/Engine/Source/Programs/Shared/EpicGames.UHT/Parsers/UhtSpecifierParser.cs
+++ b/Engine/Source/Programs/Shared/EpicGames.UHT/Parsers/UhtSpecifierParser.cs
@@ -99,6 +99,7 @@ namespace EpicGames.UHT.Parsers
 			this._tokenReader = specifierContext.Scope.TokenReader;
 			this._context = context;
 			this._table = table;
+			this._deferredSpecifiers.Clear();
 		}
 
 		/// <summary>
diff --git a/Engine/Source/Programs/Shared/EpicGames.UHT/Tokenizer/UhtToken.cs b/Engine/Source/Programs/Shared/EpicGames.UHT/Tokenizer/UhtToken.cs
--- a/Engine/Source/Programs/Shared/EpicGames.UHT/Tokenizer/UhtToken.cs
+++ b/Engine/Source/Programs/Shared/EpicGames.UHT/Tokenizer/UhtToken.cs
@@ -684,28 +684,31 @@ namespace EpicGames.UHT.Tokenizer
 			}
 		}
 
-		float GetFloatValue()
+		private const NumberStyles s_defaultNumberStyles = NumberStyles.AllowLeadingWhite | NumberStyles.AllowTrailingWhite |
+			NumberStyles.AllowLeadingSign | NumberStyles.AllowDecimalPoint | NumberStyles.AllowThousands | NumberStyles.AllowExponent;
+
+		private float GetFloatValue()
 		{
 			if (Value.Span.Length > 0)
 			{
 				if (UhtFCString.IsFloatMarker(Value.Span[^1]))
 				{
-					return Single.Parse(Value.Span[0..^1]);
+					return Single.Parse(Value.Span[0..^1], s_defaultNumberStyles, CultureInfo.InvariantCulture);
 				}
 			}
-			return Single.Parse(Value.Span);
+			return Single.Parse(Value.Span, s_defaultNumberStyles, CultureInfo.InvariantCulture);
 		}
 
-		double GetDoubleValue()
+		private double GetDoubleValue()
 		{
 			if (Value.Span.Length > 0)
 			{
 				if (UhtFCString.IsFloatMarker(Value.Span[^1]))
 				{
-					return Double.Parse(Value.Span[0..^1]);
+					return Double.Parse(Value.Span[0..^1], s_defaultNumberStyles, CultureInfo.InvariantCulture);
 				}
 			}
-			return Double.Parse(Value.Span);
+			return Double.Parse(Value.Span, s_defaultNumberStyles, CultureInfo.InvariantCulture);
 		}
 
 		long GetDecimalValue()
diff --git a/Engine/Source/Programs/Shared/EpicGames.UHT/Types/Properties/UhtDoubleProperty.cs b/Engine/Source/Programs/Shared/EpicGames.UHT/Types/Properties/UhtDoubleProperty.cs
--- a/Engine/Source/Programs/Shared/EpicGames.UHT/Types/Properties/UhtDoubleProperty.cs
+++ b/Engine/Source/Programs/Shared/EpicGames.UHT/Types/Properties/UhtDoubleProperty.cs
@@ -1,6 +1,7 @@
 // Copyright Epic Games, Inc. All Rights Reserved.
 
 using System.Diagnostics.CodeAnalysis;
+using System.Globalization;
 using System.Text;
 using EpicGames.UHT.Tables;
 using EpicGames.UHT.Tokenizer;
@@ -48,7 +49,7 @@ namespace EpicGames.UHT.Types
 		/// <inheritdoc/>
 		public override bool SanitizeDefaultValue(IUhtTokenReader defaultValueReader, StringBuilder innerDefaultValue)
 		{
-			innerDefaultValue.AppendFormat("{0:F6}", defaultValueReader.GetConstFloatExpression());
+			innerDefaultValue.AppendFormat(CultureInfo.InvariantCulture, "{0:F6}", defaultValueReader.GetConstFloatExpression());
 			return true;
 		}
 
diff --git a/Engine/Source/Programs/Shared/EpicGames.UHT/Types/Properties/UhtFloatProperty.cs b/Engine/Source/Programs/Shared/EpicGames.UHT/Types/Properties/UhtFloatProperty.cs
--- a/Engine/Source/Programs/Shared/EpicGames.UHT/Types/Properties/UhtFloatProperty.cs
+++ b/Engine/Source/Programs/Shared/EpicGames.UHT/Types/Properties/UhtFloatProperty.cs
@@ -1,6 +1,7 @@
 // Copyright Epic Games, Inc. All Rights Reserved.
 
 using System.Diagnostics.CodeAnalysis;
+using System.Globalization;
 using System.Text;
 using EpicGames.UHT.Tables;
 using EpicGames.UHT.Tokenizer;
@@ -47,7 +48,7 @@ namespace EpicGames.UHT.Types
 		/// <inheritdoc/>
 		public override bool SanitizeDefaultValue(IUhtTokenReader defaultValueReader, StringBuilder innerDefaultValue)
 		{
-			innerDefaultValue.AppendFormat("{0:F6}", defaultValueReader.GetConstFloatExpression());
+			innerDefaultValue.AppendFormat(CultureInfo.InvariantCulture, "{0:F6}", defaultValueReader.GetConstFloatExpression());
 			return true;
 		}
 
diff --git a/Engine/Source/Programs/Shared/EpicGames.UHT/Types/Properties/UhtLargeWorldCoordinatesRealProperty.cs b/Engine/Source/Programs/Shared/EpicGames.UHT/Types/Properties/UhtLargeWorldCoordinatesRealProperty.cs
--- a/Engine/Source/Programs/Shared/EpicGames.UHT/Types/Properties/UhtLargeWorldCoordinatesRealProperty.cs
+++ b/Engine/Source/Programs/Shared/EpicGames.UHT/Types/Properties/UhtLargeWorldCoordinatesRealProperty.cs
@@ -1,6 +1,7 @@
 // Copyright Epic Games, Inc. All Rights Reserved.
 
 using System.Diagnostics.CodeAnalysis;
+using System.Globalization;
 using System.Text;
 using EpicGames.UHT.Tables;
 using EpicGames.UHT.Tokenizer;
@@ -49,7 +50,7 @@ namespace EpicGames.UHT.Types
 		/// <inheritdoc/>
 		public override bool SanitizeDefaultValue(IUhtTokenReader defaultValueReader, StringBuilder innerDefaultValue)
 		{
-			innerDefaultValue.AppendFormat("{0:F6}", defaultValueReader.GetConstFloatExpression());
+			innerDefaultValue.AppendFormat(CultureInfo.InvariantCulture, "{0:F6}", defaultValueReader.GetConstFloatExpression());
 			return true;
 		}
 
diff --git a/Engine/Source/Programs/Shared/EpicGames.UHT/Types/Properties/UhtStructProperty.cs b/Engine/Source/Programs/Shared/EpicGames.UHT/Types/Properties/UhtStructProperty.cs
--- a/Engine/Source/Programs/Shared/EpicGames.UHT/Types/Properties/UhtStructProperty.cs
+++ b/Engine/Source/Programs/Shared/EpicGames.UHT/Types/Properties/UhtStructProperty.cs
@@ -2,6 +2,7 @@
 
 using System.Collections.Generic;
 using System.Diagnostics.CodeAnalysis;
+using System.Globalization;
 using System.Text;
 using System.Text.Json.Serialization;
 using EpicGames.Core;
@@ -219,9 +220,9 @@ namespace EpicGames.UHT.Types
 				switch (defaultValueReader.GetIdentifier().Value.ToString())
 				{
 					case "ZeroVector": return true;
-					case "UpVector": innerDefaultValue.AppendFormat(Format, 0, 0, 1); return true;
-					case "ForwardVector": innerDefaultValue.AppendFormat(Format, 1, 0, 0); return true;
-					case "RightVector": innerDefaultValue.AppendFormat(Format, 0, 1, 0); return true;
+					case "UpVector": innerDefaultValue.AppendFormat(CultureInfo.InvariantCulture, Format, 0, 0, 1); return true;
+					case "ForwardVector": innerDefaultValue.AppendFormat(CultureInfo.InvariantCulture, Format, 1, 0, 0); return true;
+					case "RightVector": innerDefaultValue.AppendFormat(CultureInfo.InvariantCulture, Format, 0, 1, 0); return true;
 					default: return false;
 				}
 			}
@@ -241,7 +242,7 @@ namespace EpicGames.UHT.Types
 						defaultValueReader.Require(',');
 						z = defaultValueReader.GetConstDoubleExpression();
 					}
-					innerDefaultValue.AppendFormat(Format, x, y, z);
+					innerDefaultValue.AppendFormat(CultureInfo.InvariantCulture, Format, x, y, z);
 				}
 				defaultValueReader.Require(")");
 				return true;
@@ -275,7 +276,7 @@ namespace EpicGames.UHT.Types
 					double y = defaultValueReader.GetConstDoubleExpression();
 					defaultValueReader.Require(',');
 					double z = defaultValueReader.GetConstDoubleExpression();
-					innerDefaultValue.AppendFormat("{0:F6},{1:F6},{2:F6}", x, y, z);
+					innerDefaultValue.AppendFormat(CultureInfo.InvariantCulture, "{0:F6},{1:F6},{2:F6}", x, y, z);
 				}
 				defaultValueReader.Require(")");
 				return true;
@@ -295,7 +296,7 @@ namespace EpicGames.UHT.Types
 				switch (defaultValueReader.GetIdentifier().Value.ToString())
 				{
 					case "ZeroVector": return true;
-					case "UnitVector": innerDefaultValue.AppendFormat(Format, 1.0, 1.0); return true;
+					case "UnitVector": innerDefaultValue.AppendFormat(CultureInfo.InvariantCulture, Format, 1.0, 1.0); return true;
 					default: return false;
 				}
 			}
@@ -310,7 +311,7 @@ namespace EpicGames.UHT.Types
 					double x = defaultValueReader.GetConstDoubleExpression();
 					defaultValueReader.Require(',');
 					double y = defaultValueReader.GetConstDoubleExpression();
-					innerDefaultValue.AppendFormat(Format, x, y);
+					innerDefaultValue.AppendFormat(CultureInfo.InvariantCulture, Format, x, y);
 				}
 				defaultValueReader.Require(")");
 				return true;
@@ -329,14 +330,14 @@ namespace EpicGames.UHT.Types
 			{
 				switch (defaultValueReader.GetIdentifier().Value.ToString())
 				{
-					case "White": innerDefaultValue.AppendFormat(Format, 1.0, 1.0, 1.0, 1.0); ; return true;
-					case "Gray": innerDefaultValue.AppendFormat(Format, 0.5, 0.5, 0.5, 1.0); ; return true;
-					case "Black": innerDefaultValue.AppendFormat(Format, 0.0, 0.0, 0.0, 1.0); ; return true;
-					case "Transparent": innerDefaultValue.AppendFormat(Format, 0.0, 0.0, 0.0, 0.0); ; return true;
-					case "Red": innerDefaultValue.AppendFormat(Format, 1.0, 0.0, 0.0, 1.0); ; return true;
-					case "Green": innerDefaultValue.AppendFormat(Format, 0.0, 1.0, 0.0, 1.0); ; return true;
-					case "Blue": innerDefaultValue.AppendFormat(Format, 0.0, 0.0, 1.0, 1.0); ; return true;
-					case "Yellow": innerDefaultValue.AppendFormat(Format, 1.0, 1.0, 0.0, 1.0); ; return true;
+					case "White": innerDefaultValue.AppendFormat(CultureInfo.InvariantCulture, Format, 1.0, 1.0, 1.0, 1.0); return true;
+					case "Gray": innerDefaultValue.AppendFormat(CultureInfo.InvariantCulture, Format, 0.5, 0.5, 0.5, 1.0); return true;
+					case "Black": innerDefaultValue.AppendFormat(CultureInfo.InvariantCulture, Format, 0.0, 0.0, 0.0, 1.0); return true;
+					case "Transparent": innerDefaultValue.AppendFormat(CultureInfo.InvariantCulture, Format, 0.0, 0.0, 0.0, 0.0); return true;
+					case "Red": innerDefaultValue.AppendFormat(CultureInfo.InvariantCulture, Format, 1.0, 0.0, 0.0, 1.0); return true;
+					case "Green": innerDefaultValue.AppendFormat(CultureInfo.InvariantCulture, Format, 0.0, 1.0, 0.0, 1.0); return true;
+					case "Blue": innerDefaultValue.AppendFormat(CultureInfo.InvariantCulture, Format, 0.0, 0.0, 1.0, 1.0); return true;
+					case "Yellow": innerDefaultValue.AppendFormat(CultureInfo.InvariantCulture, Format, 1.0, 1.0, 0.0, 1.0); return true;
 					default: return false;
 				}
 			}
@@ -358,7 +359,7 @@ namespace EpicGames.UHT.Types
 					{
 						a = defaultValueReader.GetConstDoubleExpression();
 					}
-					innerDefaultValue.AppendFormat(Format, r, g, b, a);
+					innerDefaultValue.AppendFormat(CultureInfo.InvariantCulture, Format, r, g, b, a);
 				}
 				defaultValueReader.Require(")");
 				return true;
@@ -377,14 +378,14 @@ namespace EpicGames.UHT.Types
 			{
 				switch (defaultValueReader.GetIdentifier().Value.ToString())
 				{
-					case "White": innerDefaultValue.AppendFormat(Format, 255, 255, 255, 255); ; return true;
-					case "Black": innerDefaultValue.AppendFormat(Format, 0, 0, 0, 255); ; return true;
-					case "Red": innerDefaultValue.AppendFormat(Format, 255, 0, 0, 255); ; return true;
-					case "Green": innerDefaultValue.AppendFormat(Format, 0, 255, 0, 255); ; return true;
-					case "Blue": innerDefaultValue.AppendFormat(Format, 0, 0, 255, 255); ; return true;
-					case "Yellow": innerDefaultValue.AppendFormat(Format, 255, 255, 0, 255); ; return true;
-					case "Cyan": innerDefaultValue.AppendFormat(Format, 0, 255, 255, 255); ; return true;
-					case "Magenta": innerDefaultValue.AppendFormat(Format, 255, 0, 255, 255); ; return true;
+					case "White": innerDefaultValue.AppendFormat(CultureInfo.InvariantCulture, Format, 255, 255, 255, 255); return true;
+					case "Black": innerDefaultValue.AppendFormat(CultureInfo.InvariantCulture, Format, 0, 0, 0, 255); return true;
+					case "Red": innerDefaultValue.AppendFormat(CultureInfo.InvariantCulture, Format, 255, 0, 0, 255); return true;
+					case "Green": innerDefaultValue.AppendFormat(CultureInfo.InvariantCulture, Format, 0, 255, 0, 255); return true;
+					case "Blue": innerDefaultValue.AppendFormat(CultureInfo.InvariantCulture, Format, 0, 0, 255, 255); return true;
+					case "Yellow": innerDefaultValue.AppendFormat(CultureInfo.InvariantCulture, Format, 255, 255, 0, 255); return true;
+					case "Cyan": innerDefaultValue.AppendFormat(CultureInfo.InvariantCulture, Format, 0, 255, 255, 255); return true;
+					case "Magenta": innerDefaultValue.AppendFormat(CultureInfo.InvariantCulture, Format, 255, 0, 255, 255); return true;
 					default: return false;
 				}
 			}
@@ -406,7 +407,7 @@ namespace EpicGames.UHT.Types
 					{
 						a = defaultValueReader.GetConstIntExpression();
 					}
-					innerDefaultValue.AppendFormat(Format, r, g, b, a);
+					innerDefaultValue.AppendFormat(CultureInfo.InvariantCulture, Format, r, g, b, a);
 				}
 				defaultValueReader.Require(")");
 				return true;
diff --git a/Engine/Source/Programs/Shared/EpicGames.UHT/Types/Properties/UhtTextProperty.cs b/Engine/Source/Programs/Shared/EpicGames.UHT/Types/Properties/UhtTextProperty.cs
--- a/Engine/Source/Programs/Shared/EpicGames.UHT/Types/Properties/UhtTextProperty.cs
+++ b/Engine/Source/Programs/Shared/EpicGames.UHT/Types/Properties/UhtTextProperty.cs
@@ -2,6 +2,7 @@
 
 using System;
 using System.Diagnostics.CodeAnalysis;
+using System.Globalization;
 using System.Text;
 using EpicGames.Core;
 using EpicGames.UHT.Tables;
@@ -517,7 +518,7 @@ namespace EpicGames.UHT.Types
 			bool isFloat = UhtFCString.IsFloatMarker(c);
 			if (token.GetConstDouble(out double value))
 			{
-				innerDefaultValue.AppendFormat("{0:F6}", value);
+				innerDefaultValue.AppendFormat(CultureInfo.InvariantCulture, "{0:F6}", value);
 				if (isFloat)
 				{
 					innerDefaultValue.Append(c);
diff --git a/Engine/Source/Programs/Shared/EpicGames.UHT/Types/UhtProperty.cs b/Engine/Source/Programs/Shared/EpicGames.UHT/Types/UhtProperty.cs
--- a/Engine/Source/Programs/Shared/EpicGames.UHT/Types/UhtProperty.cs
+++ b/Engine/Source/Programs/Shared/EpicGames.UHT/Types/UhtProperty.cs
@@ -3,6 +3,7 @@
 using System;
 using System.Collections.Generic;
 using System.Diagnostics.CodeAnalysis;
+using System.Globalization;
 using System.Linq;
 using System.Text;
 using System.Text.Json.Serialization;
@@ -2135,7 +2136,7 @@ namespace EpicGames.UHT.Types
 		public static StringBuilder AppendFlags(this StringBuilder builder, EPropertyFlags propertyFlags)
 		{
 			propertyFlags &= ~EPropertyFlags.ComputedFlags;
-			return builder.Append("(EPropertyFlags)0x").AppendFormat("{0:x16}", (ulong)propertyFlags);
+			return builder.Append("(EPropertyFlags)0x").AppendFormat(CultureInfo.InvariantCulture, "{0:x16}", (ulong)propertyFlags);
 		}
 
 		/// <summary>
diff --git a/Engine/Source/Programs/Shared/EpicGames.UHT/Utils/UhtStringBuilderExtensions.cs b/Engine/Source/Programs/Shared/EpicGames.UHT/Utils/UhtStringBuilderExtensions.cs
--- a/Engine/Source/Programs/Shared/EpicGames.UHT/Utils/UhtStringBuilderExtensions.cs
+++ b/Engine/Source/Programs/Shared/EpicGames.UHT/Utils/UhtStringBuilderExtensions.cs
@@ -2,6 +2,7 @@
 
 using System;
 using System.Collections.Generic;
+using System.Globalization;
 using System.Text;
 using EpicGames.Core;
 using EpicGames.UHT.Types;
@@ -380,7 +381,7 @@ namespace EpicGames.UHT.Utils
 							{
 								trailingHex = true;
 								//Builder.Append($"\\x{(uint)c:x2}");
-								builder.Append("\\x").AppendFormat("{0:x2}", (uint)c);
+								builder.Append("\\x").AppendFormat(CultureInfo.InvariantCulture, "{0:x2}", (uint)c);
 							}
 							else
 							{
@@ -435,28 +436,28 @@ namespace EpicGames.UHT.Utils
 			if (c < 0x80)
 			{
 				builder
-					.Append("\\x").AppendFormat("{0:x2}", c);
+					.Append("\\x").AppendFormat(CultureInfo.InvariantCulture, "{0:x2}", c);
 			}
 			else if (c < 0x800)
 			{
 				builder
-					.Append("\\x").AppendFormat("{0:x2}", 0xC0 + (c >> 6))
-					.Append("\\x").AppendFormat("{0:x2}", 0x80 + (c & 0x3f));
+					.Append("\\x").AppendFormat(CultureInfo.InvariantCulture, "{0:x2}", 0xC0 + (c >> 6))
+					.Append("\\x").AppendFormat(CultureInfo.InvariantCulture, "{0:x2}", 0x80 + (c & 0x3f));
 			}
 			else if (c < 0x10000)
 			{
 				builder
-					.Append("\\x").AppendFormat("{0:x2}", 0xE0 + (c >> 12))
-					.Append("\\x").AppendFormat("{0:x2}", 0x80 + ((c >> 6) & 0x3f))
-					.Append("\\x").AppendFormat("{0:x2}", 0x80 + (c & 0x3f));
+					.Append("\\x").AppendFormat(CultureInfo.InvariantCulture, "{0:x2}", 0xE0 + (c >> 12))
+					.Append("\\x").AppendFormat(CultureInfo.InvariantCulture, "{0:x2}", 0x80 + ((c >> 6) & 0x3f))
+					.Append("\\x").AppendFormat(CultureInfo.InvariantCulture, "{0:x2}", 0x80 + (c & 0x3f));
 			}
 			else
 			{
 				builder
-					.Append("\\x").AppendFormat("{0:x2}", 0xF0 + (c >> 18))
-					.Append("\\x").AppendFormat("{0:x2}", 0x80 + ((c >> 12) & 0x3f))
-					.Append("\\x").AppendFormat("{0:x2}", 0x80 + ((c >> 6) & 0x3f))
-					.Append("\\x").AppendFormat("{0:x2}", 0x80 + (c & 0x3f));
+					.Append("\\x").AppendFormat(CultureInfo.InvariantCulture, "{0:x2}", 0xF0 + (c >> 18))
+					.Append("\\x").AppendFormat(CultureInfo.InvariantCulture, "{0:x2}", 0x80 + ((c >> 12) & 0x3f))
+					.Append("\\x").AppendFormat(CultureInfo.InvariantCulture, "{0:x2}", 0x80 + ((c >> 6) & 0x3f))
+					.Append("\\x").AppendFormat(CultureInfo.InvariantCulture, "{0:x2}", 0x80 + (c & 0x3f));
 			}
 			return builder;
 		}
