Horde: Mark compressed blob reading as not supported in certain clients

Instead of NotImplementedException.

diff --git a/Engine/Source/Programs/Horde/Horde.Build/Storage/Services/BasicStorageClient.cs b/Engine/Source/Programs/Horde/Horde.Build/Storage/Services/BasicStorageClient.cs
--- a/Engine/Source/Programs/Horde/Horde.Build/Storage/Services/BasicStorageClient.cs
+++ b/Engine/Source/Programs/Horde/Horde.Build/Storage/Services/BasicStorageClient.cs
@@ -74,7 +74,7 @@ namespace Horde.Build.Storage.Services
 		/// <inheritdoc/>
 		public Task<Stream> ReadCompressedBlobAsync(NamespaceId namespaceId, IoHash uncompressedHash, CancellationToken cancellationToken)
 		{
-			throw new NotImplementedException();
+			throw new NotSupportedException();
 		}
 
 		/// <inheritdoc/>
diff --git a/Engine/Source/Programs/Shared/EpicGames.Horde/Storage/Impl/FileStorageClient.cs b/Engine/Source/Programs/Shared/EpicGames.Horde/Storage/Impl/FileStorageClient.cs
--- a/Engine/Source/Programs/Shared/EpicGames.Horde/Storage/Impl/FileStorageClient.cs
+++ b/Engine/Source/Programs/Shared/EpicGames.Horde/Storage/Impl/FileStorageClient.cs
@@ -69,7 +69,7 @@ namespace EpicGames.Horde.Storage.Impl
 		/// <inheritdoc/>
 		public Task<Stream> ReadCompressedBlobAsync(NamespaceId namespaceId, IoHash uncompressedHash, CancellationToken cancellationToken = default)
 		{
-			throw new NotImplementedException();
+			throw new NotSupportedException();
 		}
 
 		/// <inheritdoc/>
