EpicGames.Perforce.Native: Attempt to fix P4 symbol visibility.

diff --git a/Engine/Source/Programs/Shared/EpicGames.Perforce.Native/Client.cpp b/Engine/Source/Programs/Shared/EpicGames.Perforce.Native/Client.cpp
--- a/Engine/Source/Programs/Shared/EpicGames.Perforce.Native/Client.cpp
+++ b/Engine/Source/Programs/Shared/EpicGames.Perforce.Native/Client.cpp
@@ -19,7 +19,7 @@ THIRD_PARTY_INCLUDES_END
 	#define NATIVE_API __declspec(dllexport)
 	#define strcasecmp _stricmp
 #else
-	#define NATIVE_API
+	#define NATIVE_API __attribute__((visibility("default")))
 #endif
 
 #pragma warning(disable:4100)
@@ -616,7 +616,8 @@ public:
 		{
 			if (GzipInst->OutputFull())
 			{
-				Flush(e);
+				FFileSys::Write(Buffer, (int)(GzipInst->os - Buffer), e);
+				GzipInst->os = Buffer;
 			}
 			if (e->Test() || !GzipInst->Uncompress(e) || GzipInst->InputEmpty())
 			{
@@ -627,17 +628,12 @@ public:
 
 	virtual void Close(Error* e) override
 	{
-		Flush(e);
-		FFileSys::Close(e);
-	}
-
-	void Flush(Error* e)
-	{
-		if (GzipInst && GzipInst->os > Buffer)
+		if (GzipInst && mode == FOM_WRITE && GzipInst->os > Buffer)
 		{
 			FFileSys::Write(Buffer, (int)(GzipInst->os - Buffer), e);
-			GzipInst->os = Buffer;
 		}
+
+		FFileSys::Close(e);
 	}
 };
 
