UAT: Support multi-line log messages correctly in log parser, and guard against exceptions when parsing JSON log events.

diff --git a/Engine/Source/Programs/Shared/EpicGames.Core/Log.cs b/Engine/Source/Programs/Shared/EpicGames.Core/Log.cs
--- a/Engine/Source/Programs/Shared/EpicGames.Core/Log.cs
+++ b/Engine/Source/Programs/Shared/EpicGames.Core/Log.cs
@@ -389,7 +389,16 @@ namespace EpicGames.Core
 				{
 					lock (EventParser)
 					{
-						EventParser.WriteLine(message.ToString());
+						int baseIdx = 0;
+						for (int idx = 0; idx < message.Length; idx++)
+						{
+							if (message[idx] == '\n')
+							{
+								EventParser.WriteLine(message.ToString(baseIdx, idx - baseIdx));
+								baseIdx = idx + 1;
+							}
+						}
+						EventParser.WriteLine(message.ToString(baseIdx, message.Length - baseIdx));
 					}
 				}
 			}
diff --git a/Engine/Source/Programs/Shared/EpicGames.Core/LogEvent.cs b/Engine/Source/Programs/Shared/EpicGames.Core/LogEvent.cs
--- a/Engine/Source/Programs/Shared/EpicGames.Core/LogEvent.cs
+++ b/Engine/Source/Programs/Shared/EpicGames.Core/LogEvent.cs
@@ -278,7 +278,14 @@ namespace EpicGames.Core
 				case JsonTokenType.String:
 					return reader.GetString();
 				case JsonTokenType.Number:
-					return reader.GetInt32();
+					if (reader.TryGetInt32(out int intValue))
+					{
+						return intValue;
+					}
+					else
+					{
+						return reader.GetString();
+					}
 				default:
 					throw new InvalidOperationException("Unhandled property type");
 			}
diff --git a/Engine/Source/Programs/Shared/EpicGames.Core/LogEventParser.cs b/Engine/Source/Programs/Shared/EpicGames.Core/LogEventParser.cs
--- a/Engine/Source/Programs/Shared/EpicGames.Core/LogEventParser.cs
+++ b/Engine/Source/Programs/Shared/EpicGames.Core/LogEventParser.cs
@@ -139,13 +139,19 @@ namespace EpicGames.Core
 			if (line.Length > 0 && line[0] == '{')
 			{
 				byte[] data = Encoding.UTF8.GetBytes(line);
-
-				JsonLogEvent jsonEvent;
-				if (JsonLogEvent.TryParse(data, out jsonEvent))
+				try
 				{
-					ProcessData(true);
-					_logger.Log(jsonEvent.Level, jsonEvent.EventId, jsonEvent, null, JsonLogEvent.Format);
-					return;
+					JsonLogEvent jsonEvent;
+					if (JsonLogEvent.TryParse(data, out jsonEvent))
+					{
+						ProcessData(true);
+						_logger.Log(jsonEvent.Level, jsonEvent.EventId, jsonEvent, null, JsonLogEvent.Format);
+						return;
+					}
+				}
+				catch (Exception ex)
+				{
+					_logger.LogError(ex, "Exception while parsing log event");
 				}
 			}
 
