EpicGames.Core: Support parsing floating point numbers into LogEvent structures.

diff --git a/Engine/Source/Programs/Shared/EpicGames.Core/LogEvent.cs b/Engine/Source/Programs/Shared/EpicGames.Core/LogEvent.cs
--- a/Engine/Source/Programs/Shared/EpicGames.Core/LogEvent.cs
+++ b/Engine/Source/Programs/Shared/EpicGames.Core/LogEvent.cs
@@ -282,9 +282,13 @@ namespace EpicGames.Core
 					{
 						return intValue;
 					}
+					else if (reader.TryGetDouble(out double doubleValue))
+					{
+						return doubleValue;
+					}
 					else
 					{
-						return reader.GetString();
+						return Encoding.UTF8.GetString(reader.ValueSpan);
 					}
 				default:
 					throw new InvalidOperationException("Unhandled property type");
