Added CanEditChange structure flag to the C# UnrealEngineTypes.cs

#rnx

diff --git a/Engine/Source/Programs/Shared/EpicGames.Core/UnrealEngineTypes.cs b/Engine/Source/Programs/Shared/EpicGames.Core/UnrealEngineTypes.cs
--- a/Engine/Source/Programs/Shared/EpicGames.Core/UnrealEngineTypes.cs
+++ b/Engine/Source/Programs/Shared/EpicGames.Core/UnrealEngineTypes.cs
@@ -1422,6 +1422,11 @@ namespace EpicGames.Core
 		/// </summary>
 		NewerVersionExists = 0x01000000,
 
+		/// <summary>
+		/// If set, this struct will have CanEditChange on it in the editor to determine if a child property can be edited
+		/// </summary>
+		CanEditChange = 0x02000000,
+
 		/// <summary>
 		/// Struct flags that are automatically inherited
 		/// </summary>
