Horde: Rename environment variable used for structured logging.

diff --git a/Engine/Source/Programs/Shared/EpicGames.Core/Log.cs b/Engine/Source/Programs/Shared/EpicGames.Core/Log.cs
--- a/Engine/Source/Programs/Shared/EpicGames.Core/Log.cs
+++ b/Engine/Source/Programs/Shared/EpicGames.Core/Log.cs
@@ -958,7 +958,7 @@ namespace EpicGames.Core
 			IncludeStartingTimestamp = true;
 			EventParser = new LogEventParser(this);
 
-			string? envVar = Environment.GetEnvironmentVariable("UE_STDOUT_JSON");
+			string? envVar = Environment.GetEnvironmentVariable("UE_LOG_JSON_TO_STDOUT");
 			if(envVar != null && Int32.TryParse(envVar, out int value) && value != 0)
 			{
 				WriteJsonToStdOut = true;
