UnrealBuildTool: Ensure that working directory is updated when the new directory only differs by text case

#rnx

diff --git a/Engine/Source/Programs/Shared/EpicGames.Core/DirectoryReference.cs b/Engine/Source/Programs/Shared/EpicGames.Core/DirectoryReference.cs
--- a/Engine/Source/Programs/Shared/EpicGames.Core/DirectoryReference.cs
+++ b/Engine/Source/Programs/Shared/EpicGames.Core/DirectoryReference.cs
@@ -412,6 +412,12 @@
 		/// <param name="location">Location of the new current directory</param>
 		public static void SetCurrentDirectory(DirectoryReference location)
 		{
+			// If the new working directory only differs by text case on Windows, Directory.SetCurrentDirectory() will not actually update the working directory.
+			// To work around this, set the working directory to an entirely different path first.
+			if (RuntimePlatform.Current == RuntimePlatform.Type.Windows && Directory.GetCurrentDirectory().Equals(location.FullName, StringComparison.OrdinalIgnoreCase))
+			{
+				Directory.SetCurrentDirectory(Path.GetTempPath());
+			}
 			Directory.SetCurrentDirectory(location.FullName);
 		}
 
