Horde: Add support for compressed blob uploading in remote exec test client

diff --git a/Engine/Source/Programs/Horde/Horde.Agent.Tests/ComputeTaskExecutorTest.cs b/Engine/Source/Programs/Horde/Horde.Agent.Tests/ComputeTaskExecutorTest.cs
--- a/Engine/Source/Programs/Horde/Horde.Agent.Tests/ComputeTaskExecutorTest.cs
+++ b/Engine/Source/Programs/Horde/Horde.Agent.Tests/ComputeTaskExecutorTest.cs
@@ -14,68 +14,11 @@
 
 namespace Horde.Agent.Tests
 {
-	public static class FakeCompressor
-	{
-		private static readonly byte[] s_prefix = Encoding.UTF8.GetBytes("FAKE-COMPRESSED-");
-		
-		/// <summary>
-		/// Decompress data that was fake compressed
-		/// <see cref="Compress(byte[])" />
-		/// </summary>
-		/// <param name="input">Data to decompress</param>
-		/// <returns>Decompressed data</returns>
-		public static byte[] Decompress(byte[] input)
-		{
-			bool isPrefixCorrect = input.AsSpan(0, input.Length).SequenceEqual(s_prefix.AsSpan());
-			if (isPrefixCorrect)
-			{
-				throw new ArgumentException("Compressed data is missing magic prefix");
-			}
-
-			return input.AsSpan(s_prefix.Length).ToArray();
-		}
-		
-		/// <summary>
-		/// Decompress data that was fake compressed
-		/// <see cref="Compress(byte[])" />
-		/// </summary>
-		/// <param name="input">Data to decompress</param>
-		/// <returns>Decompressed data</returns>
-		public static string Decompress(string input)
-		{
-			return Encoding.UTF8.GetString(Decompress(Encoding.UTF8.GetBytes(input)));
-		}
-		
-		/// <summary>
-		/// Fake compression of input data by appending a prefix
-		/// Won't make the data smaller, but it will modify it leading to a changed hash
-		/// </summary>
-		/// <param name="input">Data to compress</param>
-		/// <returns>Compressed data</returns>
-		public static byte[] Compress(byte[] input)
-		{
-			byte[] compressedInput = new byte[input.Length + s_prefix.Length];
-			Array.Copy(s_prefix, 0, compressedInput, 0, s_prefix.Length);
-			Array.Copy(input, 0, compressedInput, s_prefix.Length, input.Length);
-			return compressedInput;
-		}
-		
-		/// <summary>
-		/// Fake compression of input data by appending a prefix
-		/// <see cref="Compress(byte[])"/>
-		/// </summary>
-		/// <param name="input">Data to compress</param>
-		/// <returns>Compressed data</returns>
-		public static string Compress(string input)
-		{
-			return Encoding.UTF8.GetString(Compress(Encoding.UTF8.GetBytes(input)));
-		}
-	}
-	
 	[TestClass]
 	public class ComputeTaskExecutorTest
 	{
 		private readonly MemoryStorageClient _storageClient = new ();
+		private readonly FakeCompressor _compressor = new ();
 		private readonly ComputeTaskExecutor _executor;
 		private readonly NamespaceId _namespaceId = new ("my-namespace");
 		private readonly DirectoryReference _sandboxDir;
@@ -108,14 +51,14 @@
 			await _executor.SetupSandboxAsync(_namespaceId, rootDir, _sandboxDir);
 
 			Assert.AreEqual("file1Data", File.ReadAllText(Path.Join(_sandboxDir.FullName, file1.Name.ToString())));
-			Assert.AreEqual(FakeCompressor.Compress("file2Data"), File.ReadAllText(Path.Join(_sandboxDir.FullName, file2.Name.ToString())));
+			Assert.AreEqual(await _compressor.CompressToMemoryAsync("file2Data"), File.ReadAllText(Path.Join(_sandboxDir.FullName, file2.Name.ToString())));
 		}
 
 		private async Task<FileNode> CreateFileInStorage(string name, string data, bool shouldCompress)
 		{
 			byte[] uncompressedData = Encoding.UTF8.GetBytes(data);
 			IoHash uncompressedHash = IoHash.Compute(uncompressedData);
-			byte[] compressedData = FakeCompressor.Compress(uncompressedData);
+			byte[] compressedData = await _compressor.CompressToMemoryAsync(uncompressedData);
 
 			if (shouldCompress)
 			{
diff --git a/Engine/Source/Programs/Horde/Horde.Agent/Commands/Compute/ComputeCommand.cs b/Engine/Source/Programs/Horde/Horde.Agent/Commands/Compute/ComputeCommand.cs
--- a/Engine/Source/Programs/Horde/Horde.Agent/Commands/Compute/ComputeCommand.cs
+++ b/Engine/Source/Programs/Horde/Horde.Agent/Commands/Compute/ComputeCommand.cs
@@ -1,6 +1,7 @@
 // Copyright Epic Games, Inc. All Rights Reserved.
 
 using System;
+using System.Buffers.Binary;
 using System.Collections.Generic;
 using System.Diagnostics;
 using System.IO;
@@ -91,7 +92,7 @@
 
 		readonly Stopwatch _timer = Stopwatch.StartNew();
 
-		static (DirectoryTree, IoHash) CreateSandbox(DirectoryInfo baseDirInfo, Dictionary<IoHash, byte[]> uploadList)
+		static (DirectoryTree, IoHash) CreateSandbox(DirectoryInfo baseDirInfo, Dictionary<IoHash, (byte[] data, bool isCompressed)> uploadList)
 		{
 			DirectoryTree tree = new DirectoryTree();
 
@@ -105,9 +106,10 @@
 			foreach (FileInfo fileInfo in baseDirInfo.EnumerateFiles())
 			{
 				byte[] data = File.ReadAllBytes(fileInfo.FullName);
+				bool isCompressedBuffer = IsDataCompressedBuffer(data);
 				IoHash hash = IoHash.Compute(data);
-				uploadList[hash] = data;
-				tree.Files.Add(new FileNode(fileInfo.Name, hash, fileInfo.Length, (int)fileInfo.Attributes, false));
+				uploadList[hash] = (data, isCompressedBuffer);
+				tree.Files.Add(new FileNode(fileInfo.Name, hash, fileInfo.Length, (int)fileInfo.Attributes, isCompressedBuffer));
 			}
 			tree.Files.SortBy(x => x.Name, Utf8StringComparer.Ordinal);
 
@@ -183,7 +185,7 @@
 
 				AddLocalScope("setup");
 
-				Dictionary<IoHash, byte[]> blobs = new Dictionary<IoHash, byte[]>();
+				Dictionary<IoHash, (byte[] data, bool isCompressed)> blobs = new ();
 				(_, IoHash sandboxHash) = CreateSandbox(InputDir.ToDirectoryInfo(), blobs);
 
 				AddLocalScope("scan");
@@ -234,23 +236,57 @@
 			return 0;
 		}
 
-		static IoHash AddCbObject<T>(Dictionary<IoHash, byte[]> hashToData, T source)
+		static IoHash AddCbObject<T>(Dictionary<IoHash, (byte[] data, bool isCompressed)> hashToData, T source)
 		{
 			CbObject obj = CbSerializer.Serialize<T>(source);
 			IoHash hash = obj.GetHash();
-			hashToData[hash] = obj.GetView().ToArray();
+			hashToData[hash] = (obj.GetView().ToArray(), false);
 			return hash;
 		}
 
-		static async Task UploadSandbox(IComputeClusterInfo cluster, IStorageClient storageClient, Dictionary<IoHash, byte[]> uploadList, ILogger logger)
+		/// <summary>
+		/// Rudimentary check if the data is a compressed buffer
+		/// Not intended for production use, just for the ComputeCommand uploading.
+		/// </summary>
+		/// <param name="data"></param>
+		/// <returns>True if compressed buffer</returns>
+		public static bool IsDataCompressedBuffer(byte[] data)
 		{
+			const int HeaderLength = 64;
+			const uint ExpectedMagic = 0xb7756362;
+			const uint CompressionMethodNone = 0x0;
+			
+			if (data.Length < HeaderLength)
+			{
+				return false;
+			}
+
+			using MemoryStream ms = new (data);
+			using BinaryReader reader = new (ms);
+			uint magic = reader.ReadUInt32();
+			uint crc32 = reader.ReadUInt32();
+			uint method = reader.ReadByte();
+			
+			bool needsByteSwap = BitConverter.IsLittleEndian;
+			if (needsByteSwap)
+			{
+				magic = BinaryPrimitives.ReverseEndianness(magic);
+				crc32 = BinaryPrimitives.ReverseEndianness(crc32);
+				method = BinaryPrimitives.ReverseEndianness(method);
+			}
+			
+			return magic == ExpectedMagic && method > CompressionMethodNone;
+		}
+
+		static async Task UploadSandbox(IComputeClusterInfo cluster, IStorageClient storageClient, Dictionary<IoHash, (byte[] data, bool isCompressed)> uploadList, ILogger logger)
+		{
 			int totalUploaded = 0;
 			int totalSkipped = 0;
 			int uploadedBytes = 0;
 			int skippedBytes = 0;
 
 			List<Task> tasks = new List<Task>();
-			foreach ((IoHash hash, byte[] data) in uploadList)
+			foreach ((IoHash hash, (byte[] data, bool isCompressed)) in uploadList)
 			{
 				int index = tasks.Count;
 				async Task UploadItemAsync()
@@ -263,8 +299,18 @@
 					}
 					else
 					{
-						await storageClient.WriteBlobFromMemoryAsync(cluster.NamespaceId, hash, data);
-						logger.LogInformation("Uploaded blob {Idx}/{Count}: {Hash} ({Bytes} bytes)", index, uploadList.Count, hash, data.Length);
+						if (isCompressed)
+						{
+							using MemoryStream ms = new(data);
+							IoHash uncompressedHash = await storageClient.WriteCompressedBlobAsync(cluster.NamespaceId, ms);
+							logger.LogInformation("Uploaded compressed blob {Idx}/{Count}: {Hash} (compressed hash {CompressedHash} - {Bytes} bytes)", index, uploadList.Count, uncompressedHash, hash, data.Length);
+						}
+						else
+						{
+							await storageClient.WriteBlobFromMemoryAsync(cluster.NamespaceId, hash, data);
+							logger.LogInformation("Uploaded blob {Idx}/{Count}: {Hash} ({Bytes} bytes)", index, uploadList.Count, hash, data.Length);
+						}
+						
 						Interlocked.Increment(ref totalUploaded);
 						Interlocked.Add(ref uploadedBytes, data.Length);
 					}
@@ -276,7 +322,7 @@
 			logger.LogInformation("Uploaded {UploadedCount} blobs ({UploadedBytes} bytes), Skipped {SkippedCount} blobs ({SkippedBytes} bytes)", totalUploaded, uploadedBytes, totalSkipped, skippedBytes);
 		}
 
-		private async Task<ComputeTaskStatus> ExecuteAction(ILogger logger, IComputeClient computeClient, IStorageClient storageClient, ClusterId clusterId, ComputeTask task, Dictionary<IoHash, byte[]> uploadList)
+		private async Task<ComputeTaskStatus> ExecuteAction(ILogger logger, IComputeClient computeClient, IStorageClient storageClient, ClusterId clusterId, ComputeTask task, Dictionary<IoHash, (byte[] data, bool isCompressed)> uploadList)
 		{
 			IComputeClusterInfo cluster = await computeClient.GetClusterInfoAsync(clusterId);
 
diff --git a/Engine/Source/Programs/Horde/Horde.Build/Storage/Services/BasicStorageClient.cs b/Engine/Source/Programs/Horde/Horde.Build/Storage/Services/BasicStorageClient.cs
--- a/Engine/Source/Programs/Horde/Horde.Build/Storage/Services/BasicStorageClient.cs
+++ b/Engine/Source/Programs/Horde/Horde.Build/Storage/Services/BasicStorageClient.cs
@@ -118,10 +118,16 @@
 		/// <inheritdoc/>
 		public Task WriteCompressedBlobAsync(NamespaceId namespaceId, IoHash uncompressedHash, Stream stream, CancellationToken cancellationToken = default)
 		{
-			throw new NotImplementedException();
+			throw new NotSupportedException();
 		}
 
 		/// <inheritdoc/>
+		public Task<IoHash> WriteCompressedBlobAsync(NamespaceId namespaceId, Stream compressedStream, CancellationToken cancellationToken = default)
+		{
+			throw new NotSupportedException();
+		}
+
+		/// <inheritdoc/>
 		public async Task<bool> HasBlobAsync(NamespaceId namespaceId, IoHash hash, CancellationToken cancellationToken)
 		{
 			return await _storageBackend.ExistsAsync(GetFullBlobPath(namespaceId, hash));
diff --git a/Engine/Source/Programs/Shared/EpicGames.Horde.Tests/StorageClientTest.cs b/Engine/Source/Programs/Shared/EpicGames.Horde.Tests/StorageClientTest.cs
new file mode 100644
--- /dev/null
+++ b/Engine/Source/Programs/Shared/EpicGames.Horde.Tests/StorageClientTest.cs
@@ -0,0 +1,124 @@
+// Copyright Epic Games, Inc. All Rights Reserved.
+
+using System;
+using EpicGames.Core;
+using EpicGames.Horde.Storage;
+using Microsoft.VisualStudio.TestTools.UnitTesting;
+using System.IO;
+using System.Linq;
+using System.Threading.Tasks;
+using EpicGames.Horde.Storage.Impl;
+
+namespace EpicGames.Horde.Tests
+{
+	/// <summary>
+	/// Base class for storage client tests.
+	/// Allows for contract testing the implementations.
+	/// </summary>
+	[TestClass]
+	public abstract class BaseStorageClientTest
+	{
+		private readonly NamespaceId _ns = new("my-ns");
+		private readonly ReadOnlyMemory<byte> _uncompressedData = new byte[] { 0x41, 0x42, 0x43 }; // "ABC"
+		private readonly ReadOnlyMemory<byte> _compressedData;
+		private readonly IoHash _uncompressedHash;
+		private readonly IoHash _compressedHash;
+		protected FakeCompressor Compressor { get; } = new FakeCompressor();
+
+		protected BaseStorageClientTest()
+		{
+			_compressedData = Compressor.CompressToMemoryAsync(_uncompressedData.ToArray()).Result;
+			_uncompressedHash = IoHash.Compute(_uncompressedData.ToArray());
+			_compressedHash = IoHash.Compute(_compressedData.ToArray());
+		}
+
+		protected abstract IStorageClient GetStorageClient();
+
+		[TestMethod]
+		public async Task WriteBlob()
+		{
+			IStorageClient client = GetStorageClient();
+			IoHash writtenHash = await client.WriteBlobFromMemoryAsync(_ns, _uncompressedData);
+			byte[] actual = await client.ReadBlobToMemoryAsync(_ns, writtenHash);
+			CollectionAssert.AreEqual(_uncompressedData.ToArray(), actual);
+		}
+		
+		[TestMethod]
+		public async Task WriteBlobWithPreCalculatedHash()
+		{
+			IStorageClient client = GetStorageClient();
+			await client.WriteBlobFromMemoryAsync(_ns, _uncompressedHash, _uncompressedData);
+			byte[] actual = await client.ReadBlobToMemoryAsync(_ns, _uncompressedHash);
+			CollectionAssert.AreEqual(_uncompressedData.ToArray(), actual);
+		}
+		
+		[TestMethod]
+		public async Task WriteAndReadCompressedBlob()
+		{
+			IStorageClient client = GetStorageClient();
+			using MemoryStream compressedStream = new(_compressedData.ToArray());
+			IoHash returnedHash = await client.WriteCompressedBlobAsync(_ns, compressedStream);
+			Assert.AreEqual(_uncompressedHash, returnedHash);
+			byte[] returnedData = await client.ReadCompressedBlobToMemoryAsync(_ns, returnedHash);
+			CollectionAssert.AreNotEqual(_uncompressedData.ToArray(), returnedData);
+		}
+		
+		[TestMethod]
+		public async Task WriteAndReadCompressedBlobWithPreCalculatedHash()
+		{
+			IStorageClient client = GetStorageClient();
+			using MemoryStream compressedStream = new(_compressedData.ToArray());
+			await client.WriteCompressedBlobAsync(_ns, _uncompressedHash, compressedStream);
+			byte[] returnedData = await client.ReadCompressedBlobToMemoryAsync(_ns, _uncompressedHash);
+			CollectionAssert.AreNotEqual(_uncompressedData.ToArray(), returnedData);
+		}
+		
+		[TestMethod]
+		public async Task WriteWithBadPreCalculatedHash()
+		{
+			IStorageClient client = GetStorageClient();
+			IoHash badHash = IoHash.Compute(new byte[] {0xBA, 0xDB, 0xAD});
+			using MemoryStream compressedStream = new(_compressedData.ToArray());
+			await Assert.ThrowsExceptionAsync<ArgumentException>(() => client.WriteCompressedBlobAsync(_ns, badHash, compressedStream));
+		}
+		
+		[TestMethod]
+		public async Task WriteWithInvalidCompression()
+		{
+			IStorageClient client = GetStorageClient();
+			using MemoryStream compressedStream = new(new byte[] {0xBA, 0xDB, 0xAD});
+			await Assert.ThrowsExceptionAsync<ArgumentException>(() => client.WriteCompressedBlobAsync(_ns, compressedStream));
+		}
+	}
+	
+	[TestClass]
+	public class MemoryStorageClientTest : BaseStorageClientTest
+	{
+		protected override IStorageClient GetStorageClient()
+		{
+			return new MemoryStorageClient();
+		}
+	}
+	
+	[TestClass]
+	public class FakeCompressorTest
+	{
+		[TestMethod]
+		public async Task CompressAndDecompress()
+		{
+			FakeCompressor compressor = new ();
+			byte[] uncompressedData = { 0x41, 0x42, 0x43 };
+			using MemoryStream input = new(uncompressedData);
+			using MemoryStream output = new();
+			await compressor.CompressAsync(input, output);
+			byte[] compressedData = output.ToArray();
+			Assert.IsFalse(compressedData.SequenceEqual(uncompressedData));
+			
+			using MemoryStream input2 = new(compressedData);
+			using MemoryStream output2 = new();
+			await compressor.DecompressAsync(input2, output2);
+			byte[] decompressedData = output2.ToArray();
+			Assert.IsTrue(decompressedData.SequenceEqual(uncompressedData));
+		}
+	}
+}
diff --git a/Engine/Source/Programs/Shared/EpicGames.Horde.Tests/StorageCompressorTest.cs b/Engine/Source/Programs/Shared/EpicGames.Horde.Tests/StorageCompressorTest.cs
new file mode 100644
--- /dev/null
+++ b/Engine/Source/Programs/Shared/EpicGames.Horde.Tests/StorageCompressorTest.cs
@@ -0,0 +1,62 @@
+// Copyright Epic Games, Inc. All Rights Reserved.
+
+using EpicGames.Horde.Storage;
+using Microsoft.VisualStudio.TestTools.UnitTesting;
+using System;
+using System.IO;
+using System.Linq;
+using System.Text;
+using System.Threading;
+using System.Threading.Tasks;
+
+namespace EpicGames.Horde.Tests
+{
+	/// <summary>
+	/// Fake compressor used for testing. Simply prepends the data with const prefix.
+	/// Allows for easy debugging and is enough to cause a changed hash.
+	/// </summary>
+	public class FakeStorageCompressor : IStorageCompressor
+	{
+		private readonly byte[] _prefix = Encoding.UTF8.GetBytes("FAKECOMP-");
+		
+		public async Task CompressAsync(Stream input, Stream output, CancellationToken cancellationToken)
+		{
+			await output.WriteAsync(_prefix, cancellationToken);
+			await input.CopyToAsync(output, cancellationToken);
+		}
+
+		public async Task DecompressAsync(Stream input, Stream output, CancellationToken cancellationToken)
+		{
+			using MemoryStream ms = new();
+			using BinaryReader reader = new(input);
+			byte[] prefixRead = reader.ReadBytes(_prefix.Length);
+			if (!_prefix.SequenceEqual(prefixRead))
+			{
+				throw new ArgumentException($"Input stream is not compressed. Missing fake prefix!");
+			}
+			await input.CopyToAsync(output, cancellationToken);
+		}
+	}
+
+	[TestClass]
+	public class FakeStorageCompressorTest
+	{
+		[TestMethod]
+		public async Task CompressAndDecompress()
+		{
+			IStorageCompressor compressor = new FakeStorageCompressor();
+			byte[] uncompressedData = { 0x41, 0x42, 0x43 };
+			using MemoryStream input = new(uncompressedData);
+			using MemoryStream output = new();
+			await compressor.CompressAsync(input, output);
+			byte[] compressedData = output.ToArray();
+			Assert.IsFalse(compressedData.SequenceEqual(uncompressedData));
+			
+			using MemoryStream input2 = new(compressedData);
+			using MemoryStream output2 = new();
+			await compressor.DecompressAsync(input2, output2);
+			byte[] decompressedData = output2.ToArray();
+			Assert.IsTrue(decompressedData.SequenceEqual(uncompressedData));
+		}
+	}
+}
diff --git a/Engine/Source/Programs/Shared/EpicGames.Horde/Storage/IStorageClient.cs b/Engine/Source/Programs/Shared/EpicGames.Horde/Storage/IStorageClient.cs
--- a/Engine/Source/Programs/Shared/EpicGames.Horde/Storage/IStorageClient.cs
+++ b/Engine/Source/Programs/Shared/EpicGames.Horde/Storage/IStorageClient.cs
@@ -203,6 +203,15 @@
 		Task WriteCompressedBlobAsync(NamespaceId namespaceId, IoHash uncompressedHash, Stream compressedStream, CancellationToken cancellationToken = default);
 		
 		/// <summary>
+		/// Writes a compressed blob to storage and calculates hash
+		/// </summary>
+		/// <param name="namespaceId">Namespace to operate on</param>
+		/// <param name="compressedStream">Compressed stream to write</param>
+		/// <param name="cancellationToken">Cancellation token for the operation</param>
+		/// <returns>Hash of uncompressed blob</returns>
+		Task<IoHash> WriteCompressedBlobAsync(NamespaceId namespaceId, Stream compressedStream, CancellationToken cancellationToken = default);
+		
+		/// <summary>
 		/// Checks if the given blob exists
 		/// </summary>
 		/// <param name="namespaceId">Namespace to operate on</param>
diff --git a/Engine/Source/Programs/Shared/EpicGames.Horde/Storage/Impl/FileStorageClient.cs b/Engine/Source/Programs/Shared/EpicGames.Horde/Storage/Impl/FileStorageClient.cs
--- a/Engine/Source/Programs/Shared/EpicGames.Horde/Storage/Impl/FileStorageClient.cs
+++ b/Engine/Source/Programs/Shared/EpicGames.Horde/Storage/Impl/FileStorageClient.cs
@@ -104,6 +104,12 @@
 		}
 
 		/// <inheritdoc/>
+		public Task<IoHash> WriteCompressedBlobAsync(NamespaceId namespaceId, Stream compressedStream, CancellationToken cancellationToken = default)
+		{
+			throw new NotSupportedException();
+		}
+
+		/// <inheritdoc/>
 		public Task<bool> HasBlobAsync(NamespaceId namespaceId, IoHash hash, CancellationToken cancellationToken = default)
 		{
 			return Task.FromResult(FileReference.Exists(GetBlobFile(namespaceId, hash)));
diff --git a/Engine/Source/Programs/Shared/EpicGames.Horde/Storage/Impl/HttpStorageClient.cs b/Engine/Source/Programs/Shared/EpicGames.Horde/Storage/Impl/HttpStorageClient.cs
--- a/Engine/Source/Programs/Shared/EpicGames.Horde/Storage/Impl/HttpStorageClient.cs
+++ b/Engine/Source/Programs/Shared/EpicGames.Horde/Storage/Impl/HttpStorageClient.cs
@@ -180,6 +180,26 @@
 		}
 
 		/// <inheritdoc/>
+		public async Task<IoHash> WriteCompressedBlobAsync(NamespaceId namespaceId, Stream compressedStream, CancellationToken cancellationToken = default)
+		{
+			StreamContent content = new (compressedStream);
+			content.Headers.ContentType = new MediaTypeHeaderValue("application/x-ue-comp");
+			content.Headers.ContentLength = compressedStream.Length;
+
+			using HttpRequestMessage request = new (HttpMethod.Post, $"api/v1/compressed-blobs/{namespaceId}");
+			request.Content = content;
+
+			HttpResponseMessage response = await _httpClient.SendAsync(request, cancellationToken);
+			response.EnsureSuccessStatusCode();
+			
+			using (Stream responseStream = await response.Content.ReadAsStreamAsync(cancellationToken))
+			{
+				WriteBlobResponse? responseMessage = await JsonSerializer.DeserializeAsync<WriteBlobResponse>(responseStream, cancellationToken: cancellationToken);
+				return responseMessage!.Identifier;
+			}
+		}
+
+		/// <inheritdoc/>
 		public async Task<bool> HasBlobAsync(NamespaceId namespaceId, IoHash hash, CancellationToken cancellationToken = default)
 		{
 			using HttpRequestMessage request = new HttpRequestMessage(HttpMethod.Head, $"api/v1/blobs/{namespaceId}/{hash}");
diff --git a/Engine/Source/Programs/Shared/EpicGames.Horde/Storage/Impl/MemoryStorageClient.cs b/Engine/Source/Programs/Shared/EpicGames.Horde/Storage/Impl/MemoryStorageClient.cs
--- a/Engine/Source/Programs/Shared/EpicGames.Horde/Storage/Impl/MemoryStorageClient.cs
+++ b/Engine/Source/Programs/Shared/EpicGames.Horde/Storage/Impl/MemoryStorageClient.cs
@@ -6,6 +6,7 @@
 using System.Collections.Generic;
 using System.IO;
 using System.Linq;
+using System.Text;
 using System.Threading;
 using System.Threading.Tasks;
 
@@ -43,6 +44,8 @@
 		/// <inheritdoc/>
 		public Dictionary<(NamespaceId, BucketId, RefId), IRef> Refs { get; } = new Dictionary<(NamespaceId, BucketId, RefId), IRef>();
 
+		private FakeCompressor _compressor = new FakeCompressor();
+
 		/// <inheritdoc/>
 		public Task<Stream> ReadBlobAsync(NamespaceId namespaceId, IoHash hash, CancellationToken cancellationToken = default)
 		{
@@ -86,15 +89,36 @@
 		/// <inheritdoc/>
 		public async Task WriteCompressedBlobAsync(NamespaceId namespaceId, IoHash uncompressedHash, Stream compressedStream, CancellationToken cancellationToken = default)
 		{
-			using MemoryStream memoryStream = new MemoryStream();
-			await compressedStream.CopyToAsync(memoryStream, cancellationToken);
-			byte[] compressedData = memoryStream.ToArray();
+			using MemoryStream compressedStreamMs = new ();
+			await compressedStream.CopyToAsync(compressedStreamMs, cancellationToken);
+			byte[] compressedData = compressedStreamMs.ToArray();
 			IoHash compressedHash = IoHash.Compute(compressedData);
+
+			byte[] decompressedData = await _compressor.DecompressToMemoryAsync(compressedData, cancellationToken);
+			IoHash computedDecompressedHash = IoHash.Compute(decompressedData);
+			if (computedDecompressedHash != uncompressedHash)
+			{
+				throw new ArgumentException("Supplied uncompressed hash does not match hash of uncompressed data");
+			}
+
 			Blobs[(namespaceId, compressedHash)] = compressedData;
 			UncompressedToCompressedHash[(namespaceId, uncompressedHash)] = compressedHash;
 		}
 
 		/// <inheritdoc/>
+		public async Task<IoHash> WriteCompressedBlobAsync(NamespaceId namespaceId, Stream compressedStream, CancellationToken cancellationToken = default)
+		{
+			using MemoryStream compressedStreamMs = new();
+			await compressedStream.CopyToAsync(compressedStreamMs, cancellationToken);
+			byte[] decompressedData = await _compressor.DecompressToMemoryAsync(compressedStreamMs.ToArray(), cancellationToken);
+			IoHash decompressedHash = IoHash.Compute(decompressedData);
+			compressedStreamMs.Seek(0, SeekOrigin.Begin);
+
+			await WriteCompressedBlobAsync(namespaceId, decompressedHash, compressedStreamMs, cancellationToken);
+			return decompressedHash;
+		}
+
+		/// <inheritdoc/>
 		public Task<bool> HasBlobAsync(NamespaceId namespaceId, IoHash hash, CancellationToken cancellationToken = default)
 		{
 			return Task.FromResult(Blobs.ContainsKey((namespaceId, hash)));
@@ -143,4 +167,95 @@
 			return Task.FromResult(new List<IoHash>());
 		}
 	}
+	
+	/// <summary>
+	/// Fake compressor used for testing. Simply prepends the data with const prefix.
+	/// Allows for easy debugging and is enough to cause a changed hash.
+	/// </summary>
+	public class FakeCompressor
+	{
+		private readonly byte[] _prefix = Encoding.UTF8.GetBytes("FAKECOMP-");
+		
+		/// <summary>
+		/// Compress data
+		/// </summary>
+		/// <param name="input">Data to be compressed</param>
+		/// <param name="output">Compressed data</param>
+		/// <param name="cancellationToken">Cancellation token for the operation</param>
+		public async Task CompressAsync(Stream input, Stream output, CancellationToken cancellationToken = default)
+		{
+			await output.WriteAsync(_prefix, cancellationToken);
+			await input.CopyToAsync(output, cancellationToken);
+		}
+
+		/// <summary>
+		/// Decompress data
+		/// </summary>
+		/// <param name="input">Data to be decompressed</param>
+		/// <param name="output">Decompressed data</param>
+		/// <param name="cancellationToken">Cancellation token for the operation</param>
+		public async Task DecompressAsync(Stream input, Stream output, CancellationToken cancellationToken = default)
+		{
+			using MemoryStream ms = new();
+			using BinaryReader reader = new(input);
+			byte[] prefixRead = reader.ReadBytes(_prefix.Length);
+			if (!_prefix.SequenceEqual(prefixRead))
+			{
+				throw new ArgumentException($"Input stream is not compressed. Missing fake prefix!");
+			}
+			await input.CopyToAsync(output, cancellationToken);
+		}
+		
+		/// <summary>
+		/// Compress data to memory
+		/// </summary>
+		/// <param name="input">Data to be compressed</param>
+		/// <param name="cancellationToken">Cancellation token for the operation</param>
+		/// <returns>Compressed data</returns>
+		public async Task<byte[]> CompressToMemoryAsync(byte[] input, CancellationToken cancellationToken = default)
+		{
+			using MemoryStream inputMs = new (input);
+			using MemoryStream outputMs = new ();
+			await CompressAsync(inputMs, outputMs, cancellationToken);
+			return outputMs.ToArray();
+		}
+		
+		/// <summary>
+		/// Decompress data to memory
+		/// </summary>
+		/// <param name="input">Data to be decompressed</param>
+		/// <param name="cancellationToken">Cancellation token for the operation</param>
+		/// <returns>Decompressed data</returns>
+		public async Task<byte[]> DecompressToMemoryAsync(byte[] input, CancellationToken cancellationToken = default)
+		{
+			using MemoryStream inputMs = new (input);
+			using MemoryStream outputMs = new ();
+			await DecompressAsync(inputMs, outputMs, cancellationToken);
+			return outputMs.ToArray();
+		}
+		
+		/// <summary>
+		/// Compress data to memory
+		/// </summary>
+		/// <param name="input">Data to be compressed</param>
+		/// <param name="cancellationToken">Cancellation token for the operation</param>
+		/// <returns>Compressed data</returns>
+		public async Task<string> CompressToMemoryAsync(string input, CancellationToken cancellationToken = default)
+		{
+			byte[] output = await CompressToMemoryAsync(Encoding.UTF8.GetBytes(input), cancellationToken);
+			return Encoding.UTF8.GetString(output);
+		}
+		
+		/// <summary>
+		/// Decompress data to memory
+		/// </summary>
+		/// <param name="input">Data to be decompressed</param>
+		/// <param name="cancellationToken">Cancellation token for the operation</param>
+		/// <returns>Decompressed data</returns>
+		public async Task<string> DecompressToMemoryAsync(string input, CancellationToken cancellationToken = default)
+		{
+			byte[] output = await DecompressToMemoryAsync(Encoding.UTF8.GetBytes(input), cancellationToken);
+			return Encoding.UTF8.GetString(output);
+		}
+	}
 }
diff --git a/Engine/Source/Programs/Shared/EpicGames.Horde/Storage/Impl/TracingStorageClient.cs b/Engine/Source/Programs/Shared/EpicGames.Horde/Storage/Impl/TracingStorageClient.cs
--- a/Engine/Source/Programs/Shared/EpicGames.Horde/Storage/Impl/TracingStorageClient.cs
+++ b/Engine/Source/Programs/Shared/EpicGames.Horde/Storage/Impl/TracingStorageClient.cs
@@ -61,11 +61,19 @@
 		/// <inheritdoc/>
 		public Task WriteCompressedBlobAsync(NamespaceId namespaceId, IoHash uncompressedHash, Stream stream, CancellationToken cancellationToken = default)
 		{
-			_logger.LogDebug("Writing compressed blob {NamespaceId}/{UncompressedHash}", namespaceId, uncompressedHash);
+			_logger.LogDebug("Writing compressed blob {NamespaceId}/{UncompressedHash} (pre-calculated)", namespaceId, uncompressedHash);
 			return _inner.WriteCompressedBlobAsync(namespaceId, uncompressedHash, stream, cancellationToken);
 		}
 
 		/// <inheritdoc/>
+		public async Task<IoHash> WriteCompressedBlobAsync(NamespaceId namespaceId, Stream compressedStream, CancellationToken cancellationToken = default)
+		{
+			IoHash uncompressedHash = await _inner.WriteCompressedBlobAsync(namespaceId, compressedStream, cancellationToken);
+			_logger.LogDebug("Writing compressed blob {NamespaceId}/{UncompressedHash}", namespaceId, uncompressedHash);
+			return uncompressedHash;
+		}
+
+		/// <inheritdoc/>
 		public Task<bool> HasBlobAsync(NamespaceId namespaceId, IoHash hash, CancellationToken cancellationToken = default)
 		{
 			return _inner.HasBlobAsync(namespaceId, hash, cancellationToken);
