UBT: Use ILogger throughout UBT, rather than legacy EpicGames.Core.Log methods. All output going forwards should use structured logging rather than String.Format style logging.

diff --git a/Engine/Source/Programs/Shared/EpicGames.Build/System/CompileScriptModules.cs b/Engine/Source/Programs/Shared/EpicGames.Build/System/CompileScriptModules.cs
--- a/Engine/Source/Programs/Shared/EpicGames.Build/System/CompileScriptModules.cs
+++ b/Engine/Source/Programs/Shared/EpicGames.Build/System/CompileScriptModules.cs
@@ -9,6 +9,7 @@
 using System.Text.Json;
 using Microsoft.Build.Shared;
 using EpicGames.MsBuild;
+using Microsoft.Extensions.Logging;
 
 namespace UnrealBuildBase
 {
@@ -16,8 +17,14 @@
 	{
 		class Hook : CsProjBuildHook
 		{
+			private ILogger Logger;
 			private Dictionary<string, DateTime> WriteTimes = new Dictionary<string, DateTime>();
 
+			public Hook(ILogger InLogger)
+			{
+				Logger = InLogger;
+			}
+
 			public DateTime GetLastWriteTime(DirectoryReference BasePath, string RelativeFilePath)
 			{
 				return GetLastWriteTime(BasePath.FullName, RelativeFilePath);
@@ -39,7 +46,7 @@
 				Dictionary<FileReference, (CsProjBuildRecord, FileReference)> BuildRecords,
 				FileReference ProjectPath)
 			{
-				CompileScriptModule.ValidateBuildRecordRecursively(ValidBuildRecords, InvalidBuildRecords, BuildRecords, ProjectPath, this);
+				CompileScriptModule.ValidateBuildRecordRecursively(ValidBuildRecords, InvalidBuildRecords, BuildRecords, ProjectPath, this, Logger);
 			}
 
 			public bool HasWildcards(string FileSpec)
@@ -88,9 +95,9 @@
 		/// <returns>Collection of all the projects.  They will have been compiled.</returns>
 		public static HashSet<FileReference> InitializeScriptModules(Rules.RulesFileType RulesFileType, 
 			string? ScriptsForProjectFileName, List<string>? AdditionalScriptsFolders, bool bForceCompile, bool bNoCompile, bool bUseBuildRecords, 
-			out bool bBuildSuccess, Action<int> OnBuildingProjects)
+			out bool bBuildSuccess, Action<int> OnBuildingProjects, ILogger Logger)
 		{
-			List<DirectoryReference> GameDirectories = GetGameDirectories(ScriptsForProjectFileName);
+			List<DirectoryReference> GameDirectories = GetGameDirectories(ScriptsForProjectFileName, Logger);
 			List<DirectoryReference> AdditionalDirectories = GetAdditionalDirectories(AdditionalScriptsFolders);
 			List<DirectoryReference> GameBuildDirectories = GetAdditionalBuildDirectories(GameDirectories);
 
@@ -106,7 +113,7 @@
 				GameFolders: Unreal.IsEngineInstalled() ? GameDirectories : new List<DirectoryReference>(), 
 				ForeignPlugins: null, AdditionalSearchPaths: AdditionalDirectories.Concat(GameBuildDirectories).ToList()));
 
-			return GetTargetPaths(Build(RulesFileType, FoundProjects, BaseDirectories, bForceCompile, bNoCompile, bUseBuildRecords, out bBuildSuccess, OnBuildingProjects));
+			return GetTargetPaths(Build(RulesFileType, FoundProjects, BaseDirectories, bForceCompile, bNoCompile, bUseBuildRecords, out bBuildSuccess, OnBuildingProjects, Logger));
 		}
 
 		/// <summary>
@@ -115,19 +122,19 @@
 		/// <param name="FoundProjects">Collection of projects to test</param>
 		/// <param name="BaseDirectories">Base directories of the projects</param>
 		/// <returns>True if all of the projects are up to date</returns>
-		public static bool AreScriptModulesUpToDate(HashSet<FileReference> FoundProjects, List<DirectoryReference> BaseDirectories)
+		public static bool AreScriptModulesUpToDate(HashSet<FileReference> FoundProjects, List<DirectoryReference> BaseDirectories, ILogger Logger)
 		{
-			CsProjBuildHook Hook = new Hook();
+			CsProjBuildHook Hook = new Hook(Logger);
 
 			// Load existing build records, validating them only if (re)compiling script projects is an option
-			Dictionary<FileReference, (CsProjBuildRecord BuildRecord, FileReference BuildRecordPath)> ExistingBuildRecords = LoadExistingBuildRecords(BaseDirectories);
+			Dictionary<FileReference, (CsProjBuildRecord BuildRecord, FileReference BuildRecordPath)> ExistingBuildRecords = LoadExistingBuildRecords(BaseDirectories, Logger);
 
 			Dictionary<FileReference, (CsProjBuildRecord BuildRecord, FileReference BuildRecordPath)> ValidBuildRecords = new Dictionary<FileReference, (CsProjBuildRecord, FileReference)>(ExistingBuildRecords.Count);
 			Dictionary<FileReference, (CsProjBuildRecord BuildRecord, FileReference BuildRecordPath)> InvalidBuildRecords = new Dictionary<FileReference, (CsProjBuildRecord, FileReference)>(ExistingBuildRecords.Count);
 
 			foreach (FileReference Project in FoundProjects)
 			{
-				ValidateBuildRecordRecursively(ValidBuildRecords, InvalidBuildRecords, ExistingBuildRecords, Project, Hook);
+				ValidateBuildRecordRecursively(ValidBuildRecords, InvalidBuildRecords, ExistingBuildRecords, Project, Hook, Logger);
 			}
 
 			// If all found records are valid, we can return their targets directly
@@ -145,17 +152,18 @@
 		/// <param name="bUseBuildRecords"></param>
 		/// <param name="bBuildSuccess"></param>
 		/// <param name="OnBuildingProjects">Action to invoke when projects get built</param>
+		/// <param name="Logger"></param>
 		/// <returns>Collection of all the projects.  They will have been compiled.</returns>
 		public static Dictionary<FileReference, (CsProjBuildRecord, FileReference)> Build(Rules.RulesFileType RulesFileType,
 			HashSet<FileReference> FoundProjects, List<DirectoryReference> BaseDirectories, bool bForceCompile, bool bNoCompile, bool bUseBuildRecords,
-			out bool bBuildSuccess, Action<int> OnBuildingProjects)
+			out bool bBuildSuccess, Action<int> OnBuildingProjects, ILogger Logger)
 		{
-			CsProjBuildHook Hook = new Hook();
+			CsProjBuildHook Hook = new Hook(Logger);
 
 			bool bUseBuildRecordsOnlyForProjectDiscovery = bNoCompile || Unreal.IsEngineInstalled();
 
 			// Load existing build records, validating them only if (re)compiling script projects is an option
-			Dictionary<FileReference, (CsProjBuildRecord BuildRecord, FileReference BuildRecordPath)> ExistingBuildRecords = LoadExistingBuildRecords(BaseDirectories);
+			Dictionary<FileReference, (CsProjBuildRecord BuildRecord, FileReference BuildRecordPath)> ExistingBuildRecords = LoadExistingBuildRecords(BaseDirectories, Logger);
 
 			Dictionary<FileReference, (CsProjBuildRecord BuildRecord, FileReference BuildRecordPath)> ValidBuildRecords = new Dictionary<FileReference, (CsProjBuildRecord, FileReference)>(ExistingBuildRecords.Count);
 			Dictionary<FileReference, (CsProjBuildRecord BuildRecord, FileReference BuildRecordPath)> InvalidBuildRecords = new Dictionary<FileReference, (CsProjBuildRecord, FileReference)>(ExistingBuildRecords.Count);
@@ -164,7 +172,7 @@
 			{
 				foreach (FileReference Project in FoundProjects)
 				{
-					ValidateBuildRecordRecursively(ValidBuildRecords, InvalidBuildRecords, ExistingBuildRecords, Project, Hook);
+					ValidateBuildRecordRecursively(ValidBuildRecords, InvalidBuildRecords, ExistingBuildRecords, Project, Hook, Logger);
 				}
 			}
 
@@ -197,7 +205,7 @@
 						if (bNoCompile)
 						{
 							// when -NoCompile is on the command line, try to run with whatever is available
-							Log.TraceWarning($"Script module \"{TargetPath}\" not found for record \"{Record.Value.BuildRecordPath}\"");
+							Logger.LogWarning("Script module \"{TargetPath}\" not found for record \"{BuildRecordPath}\"", TargetPath, Record.Value.BuildRecordPath);
 						}
 						else
 						{
@@ -217,7 +225,7 @@
 				{
 					if (BuildRecordPath != null)
 					{
-						Log.TraceLog($"Deleting invalid build record \"{BuildRecordPath}\"");
+						Logger.LogDebug("Deleting invalid build record \"{BuildRecordPath}\"", BuildRecordPath);
 						FileReference.Delete(BuildRecordPath);
 					}
 				}
@@ -234,7 +242,7 @@
 			}
 
 			// Fall back to the slower approach: use msbuild to load csproj files & build as necessary
-			return Build(FoundProjects, bForceCompile || !bUseBuildRecords, out bBuildSuccess, Hook, BaseDirectories, OnBuildingProjects);
+			return Build(FoundProjects, bForceCompile || !bUseBuildRecords, out bBuildSuccess, Hook, BaseDirectories, OnBuildingProjects, Logger);
 		}
 
 		/// <summary>
@@ -244,9 +252,9 @@
 		/// </summary>
 		static Dictionary<FileReference, (CsProjBuildRecord, FileReference)> Build(HashSet<FileReference> FoundProjects,
 			bool bForceCompile, out bool bBuildSuccess, CsProjBuildHook Hook, List<DirectoryReference> BaseDirectories,
-			Action<int> OnBuildingProjects)
+			Action<int> OnBuildingProjects, ILogger Logger)
 		{
-			return CsProjBuilder.Build(FoundProjects, bForceCompile, out bBuildSuccess, Hook, BaseDirectories, OnBuildingProjects);
+			return CsProjBuilder.Build(FoundProjects, bForceCompile, out bBuildSuccess, Hook, BaseDirectories, OnBuildingProjects, Logger);
 		}
 
 		/// <summary>
@@ -254,7 +262,7 @@
 		/// </summary>
 		/// <param name="BaseDirectories"></param>
 		/// <returns></returns>
-		static Dictionary<FileReference, (CsProjBuildRecord, FileReference)> LoadExistingBuildRecords(List<DirectoryReference> BaseDirectories)
+		static Dictionary<FileReference, (CsProjBuildRecord, FileReference)> LoadExistingBuildRecords(List<DirectoryReference> BaseDirectories, ILogger Logger)
         {
 			Dictionary<FileReference, (CsProjBuildRecord, FileReference)> LoadedBuildRecords = new Dictionary<FileReference, (CsProjBuildRecord, FileReference)>();
 
@@ -275,11 +283,11 @@
 					try
 					{
 						BuildRecord = JsonSerializer.Deserialize<CsProjBuildRecord>(FileReference.ReadAllText(JsonFile));
-						Log.TraceLog($"Loaded script module build record {JsonFile}");
+						Logger.LogDebug("Loaded script module build record {JsonFile}", JsonFile);
 					}
 					catch(Exception Ex)
 					{
-						Log.TraceWarning($"[{JsonFile}] Failed to load build record: {Ex.Message}");
+						Logger.LogWarning("[{JsonFile}] Failed to load build record: {Message}", JsonFile, Ex.Message);
 					}
 
 					if (BuildRecord != null && BuildRecord.ProjectPath != null)
@@ -289,7 +297,7 @@
 					else
                     {
 						// Delete the invalid build record
-						Log.TraceWarning($"Deleting invalid build record {JsonFile}");
+						Logger.LogWarning("Deleting invalid build record {JsonFile}", JsonFile);
                     }
                 }
 			}
@@ -421,7 +429,7 @@
 			Dictionary<FileReference, (CsProjBuildRecord BuildRecord, FileReference)> ValidBuildRecords,
 			Dictionary<FileReference, (CsProjBuildRecord, FileReference)> InvalidBuildRecords,
 			Dictionary<FileReference, (CsProjBuildRecord, FileReference)> BuildRecords,
-			FileReference ProjectPath, CsProjBuildHook Hook)
+			FileReference ProjectPath, CsProjBuildHook Hook, ILogger Logger)
 		{
 			if (ValidBuildRecords.ContainsKey(ProjectPath) || InvalidBuildRecords.ContainsKey(ProjectPath))
 			{
@@ -432,7 +440,7 @@
 			// Was a build record loaded for this project path? (relevant when considering referenced projects)
 			if (!BuildRecords.TryGetValue(ProjectPath, out (CsProjBuildRecord BuildRecord, FileReference BuildRecordPath) Entry))
 			{
-				Log.TraceLog($"Found project {ProjectPath} with no existing build record");
+				Logger.LogDebug("Found project {ProjectPath} with no existing build record", ProjectPath);
 				return;
 			}
 
@@ -440,7 +448,7 @@
 			if (!ValidateBuildRecord(Entry.BuildRecord, ProjectPath.Directory, out string ValidationFailureMessage, Hook))
 			{
 				string ProjectRelativePath = Path.GetRelativePath(Unreal.EngineDirectory.FullName, ProjectPath.FullName);
-				Log.TraceLog($"[{ProjectRelativePath}] {ValidationFailureMessage}");
+				Logger.LogDebug("[{ProjectRelativePath}] {ValidationFailureMessage}", ProjectRelativePath, ValidationFailureMessage);
 
 				InvalidBuildRecords.Add(ProjectPath, Entry);
 				return;
@@ -450,13 +458,13 @@
 			foreach (string ReferencedProjectPath in Entry.BuildRecord.ProjectReferences)
 			{
 				FileReference FullProjectPath = FileReference.FromString(Path.GetFullPath(ReferencedProjectPath, ProjectPath.Directory.FullName));
-				ValidateBuildRecordRecursively(ValidBuildRecords, InvalidBuildRecords, BuildRecords, FullProjectPath, Hook);
+				ValidateBuildRecordRecursively(ValidBuildRecords, InvalidBuildRecords, BuildRecords, FullProjectPath, Hook, Logger);
 
 				if (!ValidBuildRecords.ContainsKey(FullProjectPath))
 				{
 					string ProjectRelativePath = Path.GetRelativePath(Unreal.EngineDirectory.FullName, ProjectPath.FullName);
 					string DependencyRelativePath = Path.GetRelativePath(Unreal.EngineDirectory.FullName, FullProjectPath.FullName);
-					Log.TraceLog($"[{ProjectRelativePath}] Existing output is not valid because dependency {DependencyRelativePath} is not valid");
+					Logger.LogDebug("[{ProjectRelativePath}] Existing output is not valid because dependency {DependencyRelativePath} is not valid", ProjectRelativePath, DependencyRelativePath);
 					InvalidBuildRecords.Add(ProjectPath, Entry);
 					return;
 				}
@@ -466,7 +474,7 @@
 				{
 					string ProjectRelativePath = Path.GetRelativePath(Unreal.EngineDirectory.FullName, ProjectPath.FullName);
 					string DependencyRelativePath = Path.GetRelativePath(Unreal.EngineDirectory.FullName, FullProjectPath.FullName);
-					Log.TraceLog($"[{ProjectRelativePath}] Existing output is not valid because dependency {DependencyRelativePath} is newer");
+					Logger.LogDebug("[{ProjectRelativePath}] Existing output is not valid because dependency {DependencyRelativePath} is newer", ProjectRelativePath, DependencyRelativePath);
 					InvalidBuildRecords.Add(ProjectPath, Entry);
 					return;
 				}
@@ -475,13 +483,13 @@
 			ValidBuildRecords.Add(ProjectPath, Entry);
 		}
 
-		static List<DirectoryReference> GetGameDirectories(string? ScriptsForProjectFileName)
+		static List<DirectoryReference> GetGameDirectories(string? ScriptsForProjectFileName, ILogger Logger)
         {
 			List<DirectoryReference> GameDirectories = new List<DirectoryReference>();
 
 			if (String.IsNullOrEmpty(ScriptsForProjectFileName))
 			{
-				GameDirectories = NativeProjectsBase.EnumerateProjectFiles().Select(x => x.Directory).ToList();
+				GameDirectories = NativeProjectsBase.EnumerateProjectFiles(Logger).Select(x => x.Directory).ToList();
 			}
 			else
 			{
diff --git a/Engine/Source/Programs/Shared/EpicGames.Build/System/EnumeratePlugins.cs b/Engine/Source/Programs/Shared/EpicGames.Build/System/EnumeratePlugins.cs
--- a/Engine/Source/Programs/Shared/EpicGames.Build/System/EnumeratePlugins.cs
+++ b/Engine/Source/Programs/Shared/EpicGames.Build/System/EnumeratePlugins.cs
@@ -1,6 +1,7 @@
 ﻿// Copyright Epic Games, Inc. All Rights Reserved.
 
 using EpicGames.Core;
+using Microsoft.Extensions.Logging;
 using System;
 using System.Collections.Generic;
 using System.Linq;
@@ -132,9 +133,10 @@
 		/// </summary>
 		/// <param name="ProjectFile">Path to the project file</param>
 		/// <param name="FoundPlugins">Collection of plugins to build</param>
+		/// <param name="Logger">Logger for output</param>
 		/// <param name="Plugins">Collection of built plugins</param>
 		/// <returns>True if the plugins compiled</returns>
-		public static bool BuildUbtPlugins(FileReference? ProjectFile, IEnumerable<FileReference> FoundPlugins,
+		public static bool BuildUbtPlugins(FileReference? ProjectFile, IEnumerable<FileReference> FoundPlugins, ILogger Logger,
 			out (FileReference ProjectFile, FileReference TargetAssembly)[]? Plugins)
 		{
 			bool bBuildSuccess = true;
@@ -156,13 +158,13 @@
 					{
 						if (Log.OutputFile != null)
 						{
-							Log.TraceInformation($"Building {Count} plugins (see Log '{Log.OutputFile}' for more details)");
+							Logger.LogInformation("Building {Count} plugins (see Log '{LogFile}' for more details)", Count, Log.OutputFile);
 						}
 						else
 						{
-							Log.TraceInformation($"Building {Count} plugins");
+							Logger.LogInformation("Building {Count} plugins", Count);
 						}
-					}
+					}, Logger
 				);
 				
 				// Add any built plugins back into the cache
diff --git a/Engine/Source/Programs/Shared/EpicGames.Build/System/FileItem.cs b/Engine/Source/Programs/Shared/EpicGames.Build/System/FileItem.cs
--- a/Engine/Source/Programs/Shared/EpicGames.Build/System/FileItem.cs
+++ b/Engine/Source/Programs/Shared/EpicGames.Build/System/FileItem.cs
@@ -9,6 +9,7 @@
 using System.Runtime.Serialization;
 using EpicGames.Core;
 using System.Collections.Concurrent;
+using Microsoft.Extensions.Logging;
 
 namespace UnrealBuildBase
 {
@@ -210,7 +211,7 @@
 		/// <summary>
 		/// Deletes the file.
 		/// </summary>
-		public void Delete()
+		public void Delete(ILogger Logger)
 		{
 			Debug.Assert(Exists);
 
@@ -239,15 +240,15 @@
 				}
 				catch (Exception Ex)
 				{
-					Log.TraceInformation("Failed to delete file '" + AbsolutePath + "'");
-					Log.TraceInformation("    Exception: " + Ex.Message);
+					Logger.LogInformation(Ex, "Failed to delete file '{Location}'", Location);
+					Logger.LogInformation("    Exception: {Message}", Ex.Message);
 					if (DeleteTryCount < MaxRetryCount)
 					{
-						Log.TraceInformation("Attempting to retry...");
+						Logger.LogInformation("Attempting to retry...");
 					}
 					else
 					{
-						Log.TraceInformation("ERROR: Exhausted all retries!");
+						Logger.LogError("ERROR: Exhausted all retries!");
 					}
 				}
 			}
diff --git a/Engine/Source/Programs/Shared/EpicGames.Build/System/NativeProjectsBase.cs b/Engine/Source/Programs/Shared/EpicGames.Build/System/NativeProjectsBase.cs
--- a/Engine/Source/Programs/Shared/EpicGames.Build/System/NativeProjectsBase.cs
+++ b/Engine/Source/Programs/Shared/EpicGames.Build/System/NativeProjectsBase.cs
@@ -8,6 +8,7 @@
 using System.Text;
 using System.Threading.Tasks;
 using EpicGames.Core;
+using Microsoft.Extensions.Logging;
 
 namespace UnrealBuildBase
 {
@@ -43,7 +44,7 @@
 		/// <summary>
 		/// Retrieve the list of base directories for native projects
 		/// </summary>
-		public static IEnumerable<DirectoryReference> EnumerateBaseDirectories()
+		public static IEnumerable<DirectoryReference> EnumerateBaseDirectories(ILogger Logger)
 		{
 			if(CachedBaseDirectories == null)
 			{
@@ -68,7 +69,7 @@
 										}
 										else
 										{
-											Log.TraceWarning("Project search path '{0}' referenced by '{1}' is not under '{2}', ignoring.", TrimLine, RootFile, Unreal.RootDirectory);
+											Logger.LogWarning("Project search path '{SearchPath}' referenced by '{ProjectDirFile}' is not under '{RootDir}', ignoring.", TrimLine, RootFile, Unreal.RootDirectory);
 										}
 									}
 								}
@@ -85,7 +86,7 @@
 		/// Returns a list of all the projects
 		/// </summary>
 		/// <returns>List of projects</returns>
-		public static IEnumerable<FileReference> EnumerateProjectFiles()
+		public static IEnumerable<FileReference> EnumerateProjectFiles(ILogger Logger)
 		{
 			if(CachedProjectFiles == null)
 			{
@@ -94,7 +95,7 @@
 					if(CachedProjectFiles == null)
 					{
 						HashSet<FileReference> ProjectFiles = new HashSet<FileReference>();
-						foreach(DirectoryReference BaseDirectory in EnumerateBaseDirectories())
+						foreach(DirectoryReference BaseDirectory in EnumerateBaseDirectories(Logger))
 						{
 							if(DirectoryLookupCache.DirectoryExists(BaseDirectory))
 							{
@@ -158,10 +159,11 @@
 		/// Checks if a given project is a native project
 		/// </summary>
 		/// <param name="ProjectFile">The project file to check</param>
+		/// <param name="Logger"></param>
 		/// <returns>True if the given project is a native project</returns>
-		public static bool IsNativeProject(FileReference ProjectFile)
+		public static bool IsNativeProject(FileReference ProjectFile, ILogger Logger)
 		{
-			EnumerateProjectFiles();
+			EnumerateProjectFiles(Logger);
 			return CachedProjectFiles!.Contains(ProjectFile);
 		}
 	}
diff --git a/Engine/Source/Programs/Shared/EpicGames.Build/UEBuildPlatformSDK.cs b/Engine/Source/Programs/Shared/EpicGames.Build/UEBuildPlatformSDK.cs
--- a/Engine/Source/Programs/Shared/EpicGames.Build/UEBuildPlatformSDK.cs
+++ b/Engine/Source/Programs/Shared/EpicGames.Build/UEBuildPlatformSDK.cs
@@ -8,6 +8,7 @@
 using System.Linq;
 using System.Text;
 using System.Text.Json;
+using Microsoft.Extensions.Logging;
 
 namespace EpicGames.Core
 {
@@ -51,6 +52,7 @@
 		{
 			Name = "Unknown";
 		}
+
 		public SDKDescriptor(string InName, string? InMin, string? InMax, string? InCurrent, string? InGroupName, SDKCollection? Collection)
 		{
 			Name = InName;
@@ -286,8 +288,11 @@
 	{
 		// Public SDK handling, not specific to AutoSDK
 
-		public UEBuildPlatformSDK()
+		protected readonly ILogger Logger;
+
+		public UEBuildPlatformSDK(ILogger InLogger)
 		{
+			Logger = InLogger;
 		}
 
 		#region Global SDK Registration
@@ -1302,7 +1307,7 @@
 		{
 			if (!IsAutoSDKSafe())
 			{
-				Log.TraceLog(GetAutoSDKPlatformName() + " attempted to run SDK hook which could have damaged manual SDK install!");
+				Logger.LogDebug("{Platform} attempted to run SDK hook which could have damaged manual SDK install!", GetAutoSDKPlatformName());
 				return false;
 			}
 			if (SDKVersionString != "")
@@ -1312,7 +1317,7 @@
 
 				if (File.Exists(HookExe))
 				{
-					Log.TraceLog("Running {0} hook {1}", Hook, HookExe);
+					Logger.LogDebug("Running {Hook} hook {HookExe}", Hook, HookExe);
 
 					// run it
 					Process HookProcess = new Process();
@@ -1349,7 +1354,7 @@
 
 					if (HookProcess.ExitCode != 0)
 					{
-						Log.TraceLog("Hook exited uncleanly (returned {0}), considering it failed.", HookProcess.ExitCode);
+						Logger.LogDebug("Hook exited uncleanly (returned {ExitCode}), considering it failed.", HookProcess.ExitCode);
 						return false;
 					}
 
@@ -1357,12 +1362,12 @@
 				}
 				else
 				{
-					Log.TraceLog("File {0} does not exist", HookExe);
+					Logger.LogDebug("File {HookExe} does not exist", HookExe);
 				}
 			}
 			else
 			{
-				Log.TraceLog("Version string is blank for {0}. Can't determine {1} hook.", PlatformSDKRoot, Hook.ToString());
+				Logger.LogDebug("Version string is blank for {SdkRoot}. Can't determine {Hook} hook.", PlatformSDKRoot, Hook.ToString());
 			}
 
 			return bHookCanBeNonExistent;
@@ -1401,8 +1406,8 @@
 						string[] Parts = VariableString.Split('=');
 						if (Parts.Length != 2)
 						{
-							Log.TraceLog("Incorrect environment variable declaration:");
-							Log.TraceLog(VariableString);
+							Logger.LogDebug("Incorrect environment variable declaration:");
+							Logger.LogDebug("{VariableString}", VariableString);
 							return false;
 						}
 
@@ -1433,7 +1438,7 @@
 					{
 						string EnvVarName = EnvVarNames[i];
 						string EnvVarValue = EnvVarValues[i];
-						Log.TraceVerbose("Setting variable '{0}' to '{1}'", EnvVarName, EnvVarValue);
+						Logger.LogDebug("Setting variable '{Name}' to '{Value}'", EnvVarName, EnvVarValue);
 						Environment.SetEnvironmentVariable(EnvVarName, EnvVarValue);
 					}
 
@@ -1448,7 +1453,7 @@
                     }
                     else
                     {
-                        Log.TraceVerbose("Path environment variable is null during AutoSDK");
+                        Logger.LogDebug("Path environment variable is null during AutoSDK");
                     }
 
 					List<String> ModifiedPathVars = new List<string>();
@@ -1461,7 +1466,7 @@
 						{
 							if (PathVar.IndexOf(PathRemove, StringComparison.OrdinalIgnoreCase) >= 0)
 							{
-								Log.TraceVerbose("Removing Path: '{0}'", PathVar);
+								Logger.LogDebug("Removing Path: '{Path}'", PathVar);
 								ModifiedPathVars.Remove(PathVar);
 							}
 						}
@@ -1475,7 +1480,7 @@
 						{
 							if (String.Compare(PathAdd, PathVar, true) == 0)
 							{
-								Log.TraceVerbose("Removing Path: '{0}'", PathVar);
+								Logger.LogDebug("Removing Path: '{Path}'", PathVar);
 								ModifiedPathVars.Remove(PathVar);
 							}
 						}
@@ -1486,7 +1491,7 @@
 					{
 						if (!ModifiedPathVars.Contains(PathAdd))
 						{
-							Log.TraceVerbose("Adding Path: '{0}'", PathAdd);
+							Logger.LogDebug("Adding Path: '{Path}'", PathAdd);
 							ModifiedPathVars.Add(PathAdd);
 						}
 					}
@@ -1521,7 +1526,7 @@
 			}
 			else
 			{
-				Log.TraceLog("Cannot set up environment for {1} because command file {2} does not exist.", PlatformSDKRoot, EnvVarFile);
+				Logger.LogDebug("Cannot set up environment for {SdkRoot} because command file {EnvVarFile} does not exist.", PlatformSDKRoot, EnvVarFile);
 			}
 
 			return false;
@@ -1685,7 +1690,7 @@
 					// if the user doesn't want AutoSDK for this platform, then do nothing
 					if (string.Compare(DesiredSDKLevel, "NONE", true) == 0)
 					{
-						Log.TraceLog($"Skipping AutoSDK for {PlatformName} because NONE was specified as the desired AutoSDK level");
+						Logger.LogDebug("Skipping AutoSDK for {PlatformName} because NONE was specified as the desired AutoSDK level", PlatformName);
 						return;
 					}
 
@@ -1699,7 +1704,7 @@
 					// switch over (note that version string can be empty)
 					if (!RunAutoSDKHooks(AutoSDKRoot, CurrentSDKString, CurrentSDKLevel, SDKHookType.Uninstall))
 					{
-						Log.TraceLog("Failed to uninstall currently installed SDK {0}", CurrentSDKString);
+						Logger.LogDebug("Failed to uninstall currently installed SDK {SdkVersion}", CurrentSDKString);
 						InvalidateCurrentlyInstalledAutoSDK();
 						return;
 					}
@@ -1708,7 +1713,7 @@
 
 					if (!RunAutoSDKHooks(AutoSDKRoot, GetAutoSDKDirectoryForMainVersion(), DesiredSDKLevel, SDKHookType.Install, false))
 					{
-						Log.TraceLog("Failed to install required SDK {0}.  Attemping to uninstall", GetAutoSDKDirectoryForMainVersion());
+						Logger.LogDebug("Failed to install required SDK {SdkVersion}.  Attemping to uninstall", GetAutoSDKDirectoryForMainVersion());
 						RunAutoSDKHooks(AutoSDKRoot, GetAutoSDKDirectoryForMainVersion(), DesiredSDKLevel, SDKHookType.Uninstall, false);
 						return;
 					}
@@ -1716,7 +1721,7 @@
 					string EnvVarFile = Path.Combine(AutoSDKRoot, SDKEnvironmentVarsFile);
 					if (!File.Exists(EnvVarFile))
 					{
-						Log.TraceLog("Installation of required SDK {0}.  Did not generate Environment file {1}", GetAutoSDKDirectoryForMainVersion(), EnvVarFile);
+						Logger.LogDebug("Installation of required SDK {SdkVersion}.  Did not generate Environment file {EnvVarFile}", GetAutoSDKDirectoryForMainVersion(), EnvVarFile);
 						RunAutoSDKHooks(AutoSDKRoot, GetAutoSDKDirectoryForMainVersion(), DesiredSDKLevel, SDKHookType.Uninstall, false);
 						return;
 					}
@@ -1762,7 +1767,7 @@
 			// load environment variables from current SDK
 			if (!SetupEnvironmentFromAutoSDK(PlatformSDKRoot))
 			{
-				Log.TraceLog("Failed to load environment from required SDK {0}", GetAutoSDKDirectoryForMainVersion());
+				Logger.LogDebug("Failed to load environment from required SDK {SdkRoot}", GetAutoSDKDirectoryForMainVersion());
 				InvalidateCurrentlyInstalledAutoSDK();
 				return SDKStatus.Invalid;
 			}
diff --git a/Engine/Source/Programs/Shared/EpicGames.Core/ExceptionUtils.cs b/Engine/Source/Programs/Shared/EpicGames.Core/ExceptionUtils.cs
--- a/Engine/Source/Programs/Shared/EpicGames.Core/ExceptionUtils.cs
+++ b/Engine/Source/Programs/Shared/EpicGames.Core/ExceptionUtils.cs
@@ -91,7 +91,7 @@
 			{
 				errorMessage.AppendFormat("\n{0}", line);
 			}
-			return errorMessage.ToString();
+			return errorMessage.Replace("\r\n", "\n").ToString();
 		}
 
 		/// <summary>
@@ -142,7 +142,7 @@
 					}
 				}
 			}
-			return message.ToString();
+			return message.Replace("\r\n", "\n").ToString();
 		}
 	}
 }
diff --git a/Engine/Source/Programs/Shared/EpicGames.MsBuild/CsProjBuilder.cs b/Engine/Source/Programs/Shared/EpicGames.MsBuild/CsProjBuilder.cs
--- a/Engine/Source/Programs/Shared/EpicGames.MsBuild/CsProjBuilder.cs
+++ b/Engine/Source/Programs/Shared/EpicGames.MsBuild/CsProjBuilder.cs
@@ -6,6 +6,7 @@
 using Microsoft.Build.Framework;
 using Microsoft.Build.Graph;
 using Microsoft.Build.Locator;
+using Microsoft.Extensions.Logging;
 using System;
 using System.Collections.Generic;
 using System.Diagnostics;
@@ -15,19 +16,29 @@
 
 namespace EpicGames.MsBuild
 {
+	using ILogger = Microsoft.Extensions.Logging.ILogger;
+	using IBuildLogger = Microsoft.Build.Framework.ILogger;
+
 	public static class CsProjBuilder
 	{
-		class MLogger : ILogger
+		class MLogger : IBuildLogger
 		{
-			LoggerVerbosity ILogger.Verbosity { get => LoggerVerbosity.Normal; set => throw new NotImplementedException(); }
-			string ILogger.Parameters { get => throw new NotImplementedException(); set { } }
+			ILogger Inner;
 
+			LoggerVerbosity IBuildLogger.Verbosity { get => LoggerVerbosity.Normal; set => throw new NotImplementedException(); }
+			string IBuildLogger.Parameters { get => throw new NotImplementedException(); set { } }
+
 			public bool bVeryVerboseLog = false;
 
 			bool bFirstError = true;
 
-			void ILogger.Initialize(IEventSource EventSource)
+			public MLogger(ILogger InInner)
 			{
+				Inner = InInner;
+			}
+
+			void IBuildLogger.Initialize(IEventSource EventSource)
+			{
 				EventSource.ProjectStarted += new ProjectStartedEventHandler(eventSource_ProjectStarted);
 				EventSource.TaskStarted += new TaskStartedEventHandler(eventSource_TaskStarted);
 				EventSource.MessageRaised += new BuildMessageEventHandler(eventSource_MessageRaised);
@@ -41,39 +52,32 @@
 				if (bFirstError)
 				{
 					Trace.WriteLine("");
-					Log.WriteLine(LogEventType.Console, "");
+					Log.Logger.LogInformation("");
 					bFirstError = false;
 				}
-				string Message = $"{e.File}({e.LineNumber},{e.ColumnNumber}): error {e.Code}: {e.Message} ({e.ProjectFile})";
-				Trace.WriteLine(Message); // double-clickable message in VS output
-				Log.WriteLine(LogEventType.Console, Message);
+				Log.Logger.LogError("{File}({Line},{Column}): error {Code}: {Message} ({ProjectFile})", new FileReference(e.File), new LogValue(LogValueType.LineNumber, e.LineNumber.ToString()), new LogValue(LogValueType.ColumnNumber, e.ColumnNumber.ToString()), new LogValue(LogValueType.ErrorCode, e.Code), e.Message, new FileReference(e.ProjectFile));
 			}
 
 			void eventSource_WarningRaised(object Sender, BuildWarningEventArgs e)
 			{
-				string Message = $"{e.File}({e.LineNumber},{e.ColumnNumber}): warning {e.Code}: {e.Message} ({e.ProjectFile})";
-
-
 				{
 					// workaround for warnings that appear after revert of net6.0 upgrade. Delete this block when the net6.0 upgrade is done.
 					// ...\Engine\Binaries\ThirdParty\DotNet\Windows\sdk\3.1.403\Microsoft.Common.CurrentVersion.targets(3036,5): warning MSB3088: Could not read state file "obj\Development\[projectname].csproj.GenerateResource.cache". The input stream is not a valid binary format.
 					// The starting contents (in bytes) are: 06-01-01-00-00-00-01-19-50-72-6F-70-65-72-74-69-65 ... (...\[projectname].csproj)
 					if (String.Equals(e.Code, "MSB3088", StringComparison.Ordinal))
 					{
-						Log.TraceLog("Suppressed warning: " + Message);
+						Log.Logger.LogDebug("{File}({Line},{Column}): suppressed warning {Code}: {Message} ({ProjectFile})", new FileReference(e.File), new LogValue(LogValueType.LineNumber, e.LineNumber.ToString()), new LogValue(LogValueType.ColumnNumber, e.ColumnNumber.ToString()), new LogValue(LogValueType.ErrorCode, e.Code), e.Message, new FileReference(e.ProjectFile));
 						return;
 					}
 				}
 
 				if (bFirstError)
                 {
-					Trace.WriteLine("");
-					Log.WriteLine(LogEventType.Console, "");
+					Log.Logger.LogInformation("");
 					bFirstError = false;
                 }
 
-				Trace.WriteLine(Message); // double-clickable message in VS output
-				Log.WriteLine(LogEventType.Console, Message);
+				Log.Logger.LogWarning("{File}({Line},{Column}): warning {Code}: {Message} ({ProjectFile})", new FileReference(e.File), new LogValue(LogValueType.LineNumber, e.LineNumber.ToString()), new LogValue(LogValueType.ColumnNumber, e.ColumnNumber.ToString()), new LogValue(LogValueType.ErrorCode, e.Code), e.Message, new FileReference(e.ProjectFile));
 			}
 
 			void eventSource_MessageRaised(object Sender, BuildMessageEventArgs e)
@@ -83,7 +87,7 @@
 					//if (!String.Equals(e.SenderName, "ResolveAssemblyReference"))
 					//if (e.Message.Contains("atic"))
 					{
-						Log.WriteLine(LogEventType.Console, $"{e.SenderName}: {e.Message}");
+						Log.Logger.LogDebug("{SenderName}: {Message}", e.SenderName, e.Message);
 					}
 				}
 			}
@@ -92,7 +96,7 @@
 			{
 				if (bVeryVerboseLog)
 				{
-					Log.WriteLine(LogEventType.Console, $"{e.SenderName}: {e.Message}");
+					Log.Logger.LogDebug("{SenderName}: {Message}", e.SenderName, e.Message);
 				}
 			}
 
@@ -100,7 +104,7 @@
 			{
 				if (bVeryVerboseLog)
 				{
-					Log.WriteLine(LogEventType.Console, $"{e.SenderName}: {e.Message}");
+					Log.Logger.LogDebug("{SenderName}: {Message}", e.SenderName, e.Message);
 				}
 			}
 
@@ -108,11 +112,11 @@
 			{
 				if (bVeryVerboseLog)
 				{
-					Log.WriteLine(LogEventType.Console, $"{e.SenderName}: {e.Message}");
+					Log.Logger.LogDebug("{SenderName}: {Message}", e.SenderName, e.Message);
 				}
 			}
 
-			void ILogger.Shutdown()
+			void IBuildLogger.Shutdown()
 			{
 			}
 		}
@@ -142,18 +146,18 @@
 		}
 
 		public static Dictionary<FileReference, (CsProjBuildRecord, FileReference)> Build(HashSet<FileReference> FoundProjects,
-			bool bForceCompile, out bool bBuildSuccess, CsProjBuildHook Hook, List<DirectoryReference> BaseDirectories, Action<int> OnBuildingProjects)
+			bool bForceCompile, out bool bBuildSuccess, CsProjBuildHook Hook, List<DirectoryReference> BaseDirectories, Action<int> OnBuildingProjects, ILogger Logger)
 		{
 
 			// Register the MS build path prior to invoking the internal routine.  By not having the internal routine
 			// inline, we avoid having the issue of the Microsoft.Build libraries being resolved prior to the build path
 			// being set.
 			RegisterMsBuildPath(Hook);
-			return BuildInternal(FoundProjects, bForceCompile, out bBuildSuccess, Hook, BaseDirectories, OnBuildingProjects);
+			return BuildInternal(FoundProjects, bForceCompile, out bBuildSuccess, Hook, BaseDirectories, OnBuildingProjects, Logger);
 		}
 
 		private static Dictionary<FileReference, (CsProjBuildRecord, FileReference)> BuildInternal(HashSet<FileReference> FoundProjects,
-			bool bForceCompile, out bool bBuildSuccess, CsProjBuildHook Hook, List<DirectoryReference> BaseDirectories, Action<int> OnBuildingProjects)
+			bool bForceCompile, out bool bBuildSuccess, CsProjBuildHook Hook, List<DirectoryReference> BaseDirectories, Action<int> OnBuildingProjects, ILogger Logger)
 		{
 			Dictionary<string, string> GlobalProperties = new Dictionary<string, string>
 			{
@@ -167,9 +171,9 @@
 
 			Dictionary<FileReference, (CsProjBuildRecord BuildRecord, FileReference BuildRecordPath)> BuildRecords = new Dictionary<FileReference, (CsProjBuildRecord, FileReference)>();
 
-			using var ProjectCollection = new ProjectCollection(GlobalProperties);
-			var Projects = new Dictionary<string, Project>();
-			var	SkippedProjects = new HashSet<string>();
+			using ProjectCollection ProjectCollection = new ProjectCollection(GlobalProperties);
+			Dictionary<string, Project> Projects = new Dictionary<string, Project>();
+			HashSet<string> SkippedProjects = new HashSet<string>();
 
 			// Microsoft.Build.Evaluation.Project provides access to information stored in the .csproj xml that is 
 			// not available when using Microsoft.Build.Execution.ProjectInstance (used later in this function and
@@ -193,20 +197,20 @@
 						}
 						catch (Microsoft.Build.Exceptions.InvalidProjectFileException IPFEx)
 						{
-							Log.TraceError($"Could not load project file {ProjectPath}");
-							Log.TraceError(IPFEx.BaseMessage);
+							Logger.LogError("Could not load project file {ProjectPath}", ProjectPath);
+							Logger.LogError("{Message}", IPFEx.BaseMessage);
 
 							if (!String.IsNullOrEmpty(ReferencedBy))
 							{
-								Log.TraceError($"Referenced by: {ReferencedBy}");
+								Logger.LogError("Referenced by: {ReferencedBy}", ReferencedBy);
 							}
 							if (Projects.Count > 0)
 							{
-								Log.TraceError("See the log file for the list of previously loaded projects.");
-								Log.TraceLog("Loaded projects (most recently loaded first):");
+								Logger.LogError("See the log file for the list of previously loaded projects.");
+								Logger.LogError("Loaded projects (most recently loaded first):");
 								foreach (string Path in Projects.Keys.Reverse())
 								{
-									Log.TraceLog($"  {Path}");
+									Logger.LogError("  {Path}", Path);
 								}
 							}
 							throw IPFEx;
@@ -219,7 +223,7 @@
 							if (Project.GetProperty("TargetFramework").EvaluatedValue.Contains("windows", StringComparison.Ordinal))
 							{
 								SkippedProjects.Add(ProjectPath);
-								Log.TraceInformation($"Skipping windows-only project {ProjectPath}");
+								Logger.LogInformation("Skipping windows-only project {ProjectPath}", ProjectPath);
 								return;
 							}
 						}
@@ -274,7 +278,7 @@
 				}
 
 				// References: e.g. Ionic.Zip.Reduced.dll, fastJSON.dll
-				foreach (var Item in Project.GetItems("Reference"))
+				foreach (ProjectItem Item in Project.GetItems("Reference"))
 				{
 					BuildRecord.Dependencies.Add(Item.GetMetadataValue("HintPath"));
 				}
@@ -327,7 +331,7 @@
 				// This also returns a lot more information than we care for - MSBuildGlob objects,
 				// which have a range of precomputed values. It may be possible to take source for
 				// GetAllGlobs() and construct a version that does less.
-				var Globs = Project.GetAllGlobs();
+				List<GlobResult> Globs = Project.GetAllGlobs();
 
 				// FileMatcher.IsMatch() requires directory separators in glob strings to match the
 				// local flavor. There's probably a better way.
@@ -336,7 +340,7 @@
 					char Sep = Path.DirectorySeparatorChar;
 					char NotSep = Sep == '/' ? '\\' : '/'; // AltDirectorySeparatorChar isn't always what we need (it's '/' on Mac)
 
-					var Chars = GlobString.ToCharArray();
+					char[] Chars = GlobString.ToCharArray();
 					int P = 0;
 					for (int I = 0; I < GlobString.Length; ++I, ++P)
 					{
@@ -360,7 +364,7 @@
 					return new string(Chars, 0, P);
 				}
 
-				foreach (var Glob in Globs)
+				foreach (GlobResult Glob in Globs)
 				{
 					if (String.Equals("None", Glob.ItemElement.ItemType, StringComparison.Ordinal))
 					{
@@ -397,7 +401,7 @@
 
 			if (bForceCompile)
 			{
-				Log.TraceLog("Script modules will build: '-Compile' on command line");
+				Logger.LogDebug("Script modules will build: '-Compile' on command line");
 				BuildProjectGraph = InputProjectGraph;
 			}
 			else
@@ -422,7 +426,7 @@
 			if (BuildProjectGraph != null)
 			{
 				OnBuildingProjects(BuildProjectGraph.EntryPointNodes.Count);
-				bBuildSuccess = BuildProjects(BuildProjectGraph, GlobalProperties);
+				bBuildSuccess = BuildProjects(BuildProjectGraph, GlobalProperties, Logger);
 			}
 			else
 			{
@@ -443,14 +447,14 @@
 				if (FileReference.WriteAllTextIfDifferent(BuildRecordPath,
 					JsonSerializer.Serialize<CsProjBuildRecord>(BuildRecord, new JsonSerializerOptions { WriteIndented = true })))
 				{
-					Log.TraceLog($"Wrote script module build record to {BuildRecordPath}");
+					Logger.LogDebug("Wrote script module build record to {BuildRecordPath}", BuildRecordPath);
 				}
 			}
 
 			// todo: re-verify build records after a build to verify that everything is actually up to date
 
 			// even if only a subset was built, this function returns the full list of target assembly paths
-			var OutDict = new Dictionary<FileReference, (CsProjBuildRecord BuildRecord, FileReference BuildRecordPath)>();
+			Dictionary<FileReference, (CsProjBuildRecord BuildRecord, FileReference BuildRecordPath)> OutDict = new Dictionary<FileReference, (CsProjBuildRecord BuildRecord, FileReference BuildRecordPath)>();
 			foreach (ProjectGraphNode EntryPointNode in InputProjectGraph.EntryPointNodes)
 			{
 				FileReference ProjectPath = FileReference.FromString(EntryPointNode.ProjectInstance.FullPath);
@@ -459,10 +463,10 @@
 			return OutDict;
 		}
 
-		private static bool BuildProjects(ProjectGraph ProjectGraph, Dictionary<string, string> GlobalProperties)
+		private static bool BuildProjects(ProjectGraph ProjectGraph, Dictionary<string, string> GlobalProperties, ILogger Logger)
 		{
-			var StartTime = DateTime.UtcNow;
-			var Logger = new MLogger();
+			DateTime StartTime = DateTime.UtcNow;
+			MLogger BuildLogger = new MLogger(Logger);
 
 			string[] TargetsToBuild = { "Restore", "Build" };
 
@@ -470,15 +474,15 @@
 
 			foreach (string TargetToBuild in TargetsToBuild)
 			{
-				var GraphRequest = new GraphBuildRequestData(ProjectGraph, new string[] { TargetToBuild });
+				GraphBuildRequestData GraphRequest = new GraphBuildRequestData(ProjectGraph, new string[] { TargetToBuild });
 
-				var BuildMan = BuildManager.DefaultBuildManager;
+				BuildManager BuildMan = BuildManager.DefaultBuildManager;
 
-				var BuildParameters = new BuildParameters();
+				BuildParameters BuildParameters = new BuildParameters();
 				BuildParameters.AllowFailureWithoutError = false;
 				BuildParameters.DetailedSummary = true;
 
-				BuildParameters.Loggers = new List<ILogger> { Logger };
+				BuildParameters.Loggers = new List<IBuildLogger> { BuildLogger };
 				BuildParameters.MaxNodeCount = 1; // msbuild bug - more than 1 here and the build stalls. Likely related to [REDACTED_PROFILE_URL]/msbuild/issues/1941
 
 				BuildParameters.OnlyLogCriticalEvents = false;
@@ -486,24 +490,24 @@
 
 				BuildParameters.GlobalProperties = GlobalProperties;
 
-				Log.TraceInformation($" {TargetToBuild}...");
+				Logger.LogInformation(" {TargetToBuild}...", TargetToBuild);
 
 				GraphBuildResult BuildResult = BuildMan.Build(BuildParameters, GraphRequest);
 
 				if (BuildResult.OverallResult == BuildResultCode.Failure)
 				{
-					Log.TraceInformation("");
-					foreach (var NodeResult in BuildResult.ResultsByNode)
+					Logger.LogInformation("");
+					foreach (KeyValuePair<ProjectGraphNode, BuildResult> NodeResult in BuildResult.ResultsByNode)
 					{
 						if (NodeResult.Value.OverallResult == BuildResultCode.Failure)
 						{
-							Log.TraceError($"  Failed to build: {NodeResult.Key.ProjectInstance.FullPath}");
+							Logger.LogError("  Failed to build: {ProjectPath}", new FileReference(NodeResult.Key.ProjectInstance.FullPath));
 						}
 					}
 					Result = false;
 				}
 			}
-			Log.TraceInformation("Build projects time: {0:0.00} s", (DateTime.UtcNow - StartTime).TotalMilliseconds / 1000);
+			Logger.LogInformation("Build projects time: {TimeSeconds:0.00} s", (DateTime.UtcNow - StartTime).TotalMilliseconds / 1000);
 
 			return Result;
 		}
diff --git a/Engine/Source/Programs/UnrealBuildTool/Configuration/ModuleRules.cs b/Engine/Source/Programs/UnrealBuildTool/Configuration/ModuleRules.cs
--- a/Engine/Source/Programs/UnrealBuildTool/Configuration/ModuleRules.cs
+++ b/Engine/Source/Programs/UnrealBuildTool/Configuration/ModuleRules.cs
@@ -498,7 +498,7 @@
 				}
 				else
 				{
-					Log.TraceWarning("Cannot apply build order override, file doesn't exist: {0}", File.AbsolutePath);
+					throw new BuildException($"Cannot apply build order override, file doesn't exist: {File.AbsolutePath}");
 				}
 			}
 
diff --git a/Engine/Source/Programs/UnrealBuildTool/Configuration/TargetDescriptor.cs b/Engine/Source/Programs/UnrealBuildTool/Configuration/TargetDescriptor.cs
--- a/Engine/Source/Programs/UnrealBuildTool/Configuration/TargetDescriptor.cs
+++ b/Engine/Source/Programs/UnrealBuildTool/Configuration/TargetDescriptor.cs
@@ -8,6 +8,7 @@
 using System.Threading.Tasks;
 using EpicGames.Core;
 using UnrealBuildBase;
+using Microsoft.Extensions.Logging;
 
 namespace UnrealBuildTool
 {
@@ -195,11 +196,12 @@
 		/// <param name="bUsePrecompiled">Whether to use a precompiled engine distribution</param>
 		/// <param name="bSkipRulesCompile">Whether to skip compiling rules assemblies</param>
 		/// <param name="bForceRulesCompile">Whether to always compile all rules assemblies</param>
+		/// <param name="Logger">Logger for output</param>
 		/// <returns>List of target descriptors</returns>
-		public static List<TargetDescriptor> ParseCommandLine(CommandLineArguments Arguments, bool bUsePrecompiled, bool bSkipRulesCompile, bool bForceRulesCompile)
+		public static List<TargetDescriptor> ParseCommandLine(CommandLineArguments Arguments, bool bUsePrecompiled, bool bSkipRulesCompile, bool bForceRulesCompile, ILogger Logger)
 		{
 			List<TargetDescriptor> TargetDescriptors = new List<TargetDescriptor>();
-			ParseCommandLine(Arguments, bUsePrecompiled, bSkipRulesCompile, bForceRulesCompile, TargetDescriptors);
+			ParseCommandLine(Arguments, bUsePrecompiled, bSkipRulesCompile, bForceRulesCompile, TargetDescriptors, Logger);
 			return TargetDescriptors;
 		}
 
@@ -211,7 +213,8 @@
 		/// <param name="bSkipRulesCompile">Whether to skip compiling rules assemblies</param>
 		/// <param name="bForceRulesCompile">Whether to always compile rules assemblies</param>
 		/// <param name="TargetDescriptors">Receives the list of parsed target descriptors</param>
-		public static void ParseCommandLine(CommandLineArguments Arguments, bool bUsePrecompiled, bool bSkipRulesCompile, bool bForceRulesCompile, List<TargetDescriptor> TargetDescriptors)
+		/// <param name="Logger">Logger for output</param>
+		public static void ParseCommandLine(CommandLineArguments Arguments, bool bUsePrecompiled, bool bSkipRulesCompile, bool bForceRulesCompile, List<TargetDescriptor> TargetDescriptors, ILogger Logger)
 		{
 			List<string> TargetLists;
 			Arguments = Arguments.Remove("-TargetList=", out TargetLists);
@@ -231,7 +234,7 @@
 						if(TrimLine.Length > 0 && TrimLine[0] != ';')
 						{
 							CommandLineArguments NewArguments = Arguments.Append(CommandLineArguments.Split(TrimLine));
-							ParseCommandLine(NewArguments, bUsePrecompiled, bSkipRulesCompile, bForceRulesCompile, TargetDescriptors);
+							ParseCommandLine(NewArguments, bUsePrecompiled, bSkipRulesCompile, bForceRulesCompile, TargetDescriptors, Logger);
 						}
 					}
 				}
@@ -239,13 +242,13 @@
 				foreach(string Target in Targets)
 				{
 					CommandLineArguments NewArguments = Arguments.Append(CommandLineArguments.Split(Target));
-					ParseCommandLine(NewArguments, bUsePrecompiled, bSkipRulesCompile, bForceRulesCompile, TargetDescriptors);
+					ParseCommandLine(NewArguments, bUsePrecompiled, bSkipRulesCompile, bForceRulesCompile, TargetDescriptors, Logger);
 				}
 			}
 			else
 			{
 				// Otherwise just process the whole command line together
-				ParseSingleCommandLine(Arguments, bUsePrecompiled, bSkipRulesCompile, bForceRulesCompile, TargetDescriptors);
+				ParseSingleCommandLine(Arguments, bUsePrecompiled, bSkipRulesCompile, bForceRulesCompile, TargetDescriptors, Logger);
 			}
 		}
 
@@ -257,7 +260,8 @@
 		/// <param name="bSkipRulesCompile">Whether to skip compiling rules assemblies</param>
 		/// <param name="bForceRulesCompile">Whether to always compile all rules assemblies</param>
 		/// <param name="TargetDescriptors">List of target descriptors</param>
-		public static void ParseSingleCommandLine(CommandLineArguments Arguments, bool bUsePrecompiled, bool bSkipRulesCompile, bool bForceRulesCompile, List<TargetDescriptor> TargetDescriptors)
+		/// <param name="Logger">Logger for output</param>
+		public static void ParseSingleCommandLine(CommandLineArguments Arguments, bool bUsePrecompiled, bool bSkipRulesCompile, bool bForceRulesCompile, List<TargetDescriptor> TargetDescriptors, ILogger Logger)
 		{
 			List<UnrealTargetPlatform> Platforms = new List<UnrealTargetPlatform>();
 			List<UnrealTargetConfiguration> Configurations = new List<UnrealTargetConfiguration>();
@@ -390,11 +394,11 @@
 
 							if (ProjectFile == null)
 							{
-								TargetNames.Add(RulesCompiler.CreateEngineRulesAssembly(bUsePrecompiled, bSkipRulesCompile, bForceRulesCompile).GetTargetNameByType(TargetType, Platform, Configuration, Architecture, null));
+								TargetNames.Add(RulesCompiler.CreateEngineRulesAssembly(bUsePrecompiled, bSkipRulesCompile, bForceRulesCompile, Logger).GetTargetNameByType(TargetType, Platform, Configuration, Architecture, null, Logger));
 							}
 							else
 							{
-								TargetNames.Add(RulesCompiler.CreateProjectRulesAssembly(ProjectFile, bUsePrecompiled, bSkipRulesCompile, bForceRulesCompile).GetTargetNameByType(TargetType, Platform, Configuration, Architecture, ProjectFile));
+								TargetNames.Add(RulesCompiler.CreateProjectRulesAssembly(ProjectFile, bUsePrecompiled, bSkipRulesCompile, bForceRulesCompile, Logger).GetTargetNameByType(TargetType, Platform, Configuration, Architecture, ProjectFile, Logger));
 							}
 						}
 
@@ -408,9 +412,9 @@
 						foreach(string TargetName in TargetNames)
 						{
 							// If a project file was not specified see if we can find one
-							if (ProjectFile == null && NativeProjects.TryGetProjectForTarget(TargetName, out ProjectFile))
+							if (ProjectFile == null && NativeProjects.TryGetProjectForTarget(TargetName, Logger, out ProjectFile))
 							{
-								Log.TraceVerbose("Found project file for {0} - {1}", TargetName, ProjectFile);
+								Logger.LogDebug("Found project file for {TargetName} - {ProjectFile}", TargetName, ProjectFile);
 							}
 
 							// Create the target descriptor
diff --git a/Engine/Source/Programs/UnrealBuildTool/Configuration/TargetRules.cs b/Engine/Source/Programs/UnrealBuildTool/Configuration/TargetRules.cs
--- a/Engine/Source/Programs/UnrealBuildTool/Configuration/TargetRules.cs
+++ b/Engine/Source/Programs/UnrealBuildTool/Configuration/TargetRules.cs
@@ -7,6 +7,7 @@
 using System.Reflection;
 using EpicGames.Core;
 using UnrealBuildBase;
+using Microsoft.Extensions.Logging;
 
 namespace UnrealBuildTool
 {
@@ -326,6 +327,11 @@
 		internal FileReference? TargetSourceFile { get; set; }
 
 		/// <summary>
+		/// Logger for output relating to this target. Set before the constructor is run from <see cref="RulesCompiler"/>
+		/// </summary>
+		internal ILogger Logger { get; set; }
+
+		/// <summary>
 		/// Platform that this target is being built for.
 		/// </summary>
 		public readonly UnrealTargetPlatform Platform;
@@ -650,7 +656,7 @@
 		public bool bBuildEditor
 		{
 			get { return (Type == TargetType.Editor || bCompileAgainstEditor); }
-			set { Log.TraceWarning("Setting {0}.bBuildEditor is deprecated. Set {0}.Type instead.", GetType().Name); }
+			set { Logger.LogWarning("Setting {Type}.bBuildEditor is deprecated. Set {Type}.Type instead.", GetType().Name); }
 		}
 
 		/// <summary>
@@ -1835,6 +1841,12 @@
 			this.Version = Target.Version;
 			this.WindowsPlatform = new WindowsTargetRules(this);
 
+			// Make sure the logger was initialized by the caller
+			if (Logger == null)
+			{
+				throw new NotSupportedException("Logger property must be initialized by the caller.");
+			}
+
 			// Read settings from config files
 			Dictionary<ConfigDependencyKey, IReadOnlyList<string>?> ConfigValues = new Dictionary<ConfigDependencyKey, IReadOnlyList<string>?>();
 			foreach (object ConfigurableObject in GetConfigurableObjects())
@@ -1878,7 +1890,7 @@
 			}
 
 			// Setup macros for signing and encryption keys
-			EncryptionAndSigning.CryptoSettings CryptoSettings = EncryptionAndSigning.ParseCryptoSettings(CryptoSettingsDir, Platform);
+			EncryptionAndSigning.CryptoSettings CryptoSettings = EncryptionAndSigning.ParseCryptoSettings(CryptoSettingsDir, Platform, Logger);
 			if (CryptoSettings.IsAnyEncryptionEnabled())
 			{
 				ProjectDefinitions.Add(String.Format("IMPLEMENT_ENCRYPTION_KEY_REGISTRATION()=UE_REGISTER_ENCRYPTION_KEY({0})", FormatHexBytes(CryptoSettings.EncryptionKey!.Key!)));
@@ -2139,7 +2151,7 @@
 						Diagnostics.Add(String.Format(FormatString, ModifiedSetting.Item1, ModifiedSetting.Item2));
 					}
 				}
-				Diagnostics.Add(String.Format("[Upgrade] Suppress this message by setting 'DefaultBuildSettings = BuildSettingsVersion.{1};' in {2}, and explicitly overriding settings that differ from the new defaults.", Version, (BuildSettingsVersion)(BuildSettingsVersion.Latest - 1), File!.GetFileName()));
+				Diagnostics.Add(String.Format("[Upgrade] Suppress this message by setting 'DefaultBuildSettings = BuildSettingsVersion.{0};' in {1}, and explicitly overriding settings that differ from the new defaults.", (BuildSettingsVersion)(BuildSettingsVersion.Latest - 1), File!.GetFileName()));
 				Diagnostics.Add("[Upgrade]");
 			}
 
@@ -2153,7 +2165,7 @@
 				Diagnostics.Add("[Upgrade]");
 			}
 
-			Log.TraceLog("Using EngineIncludeOrderVersion.{0} for target {1}", IncludeOrderVersion, File!.GetFileName());
+			Logger.LogDebug("Using EngineIncludeOrderVersion.{Version} for target {Target}", IncludeOrderVersion, File!.GetFileName());
 		}
 	}
 
diff --git a/Engine/Source/Programs/UnrealBuildTool/Configuration/UEBuildBinary.cs b/Engine/Source/Programs/UnrealBuildTool/Configuration/UEBuildBinary.cs
--- a/Engine/Source/Programs/UnrealBuildTool/Configuration/UEBuildBinary.cs
+++ b/Engine/Source/Programs/UnrealBuildTool/Configuration/UEBuildBinary.cs
@@ -10,6 +10,7 @@ using System.Xml;
 using System.Globalization;
 using EpicGames.Core;
 using UnrealBuildBase;
+using Microsoft.Extensions.Logging;
 
 namespace UnrealBuildTool
 {
@@ -155,11 +156,11 @@ namespace UnrealBuildTool
 		/// <summary>
 		/// Creates all the modules referenced by this target.
 		/// </summary>
-		public void CreateAllDependentModules(UEBuildModule.CreateModuleDelegate CreateModule)
+		public void CreateAllDependentModules(UEBuildModule.CreateModuleDelegate CreateModule, ILogger Logger)
 		{
 			foreach (UEBuildModule Module in Modules)
 			{
-				Module.RecursivelyCreateModules(CreateModule, "Target");
+				Module.RecursivelyCreateModules(CreateModule, "Target", Logger);
 			}
 		}
 
@@ -174,8 +175,9 @@ namespace UnrealBuildTool
 		/// <param name="WorkingSet">The working set of source files</param>
 		/// <param name="ExeDir">Directory containing the output executable</param>
 		/// <param name="Graph">The graph being built</param>
+		/// <param name="Logger">Logger for output</param>
 		/// <returns>Set of built products</returns>
-		public List<FileItem> Build(ReadOnlyTargetRules Target, UEToolChain ToolChain, CppCompileEnvironment CompileEnvironment, LinkEnvironment LinkEnvironment, List<FileReference> SpecificFilesToCompile, ISourceFileWorkingSet WorkingSet, DirectoryReference ExeDir, IActionGraphBuilder Graph)
+		public List<FileItem> Build(ReadOnlyTargetRules Target, UEToolChain ToolChain, CppCompileEnvironment CompileEnvironment, LinkEnvironment LinkEnvironment, List<FileReference> SpecificFilesToCompile, ISourceFileWorkingSet WorkingSet, DirectoryReference ExeDir, IActionGraphBuilder Graph, ILogger Logger)
 		{
 			// Return nothing if we're using precompiled binaries. If we're not linking, we might want just one module to be compiled (eg. a foreign plugin), so allow any actions to run.
 			if (bUsePrecompiled && !(Target.LinkType == TargetLinkType.Monolithic && Target.bDisableLinking))
@@ -184,7 +186,7 @@ namespace UnrealBuildTool
 			}
 
 			// Setup linking environment.
-			LinkEnvironment BinaryLinkEnvironment = SetupBinaryLinkEnvironment(Target, ToolChain, LinkEnvironment, CompileEnvironment, SpecificFilesToCompile, WorkingSet, ExeDir, Graph);
+			LinkEnvironment BinaryLinkEnvironment = SetupBinaryLinkEnvironment(Target, ToolChain, LinkEnvironment, CompileEnvironment, SpecificFilesToCompile, WorkingSet, ExeDir, Graph, Logger);
 
 			// If we're generating projects, we only need include paths and definitions, there is no need to run the linking logic.
 			if (ProjectFileGenerator.bGenerateProjectFiles)
@@ -589,7 +591,7 @@ namespace UnrealBuildTool
 		/// Checks whether the binary output paths are appropriate for the distribution
 		/// level of its direct module dependencies
 		/// </summary>
-		public bool CheckRestrictedFolders(List<DirectoryReference> RootDirectories, Dictionary<UEBuildModule, Dictionary<RestrictedFolder, DirectoryReference>> ModuleRestrictedFolderCache)
+		public bool CheckRestrictedFolders(List<DirectoryReference> RootDirectories, Dictionary<UEBuildModule, Dictionary<RestrictedFolder, DirectoryReference>> ModuleRestrictedFolderCache, ILogger Logger)
 		{
 			// Find all the modules we depend on
 			Dictionary<UEBuildModule, UEBuildModule?> ModuleReferencedBy = new Dictionary<UEBuildModule, UEBuildModule?>();
@@ -647,7 +649,7 @@ namespace UnrealBuildTool
 							{
 								ReferenceChain.Insert(0, ReferencedModule.Name);
 							}
-							Log.TraceError("Output binary \"{0}\" is not in a {1} folder, but references \"{2}\" via {3}.", OutputFilePath, Pair.Key.ToString(), Pair.Value, String.Join(" -> ", ReferenceChain));
+							Logger.LogError("Output binary \"{OutputFilePath}\" is not in a {Folder} folder, but references \"{ReferenceFile}\" via {ReferenceChain}.", OutputFilePath, Pair.Key.ToString(), Pair.Value, String.Join(" -> ", ReferenceChain));
 							bResult = false;
 						}
 					}
@@ -693,7 +695,7 @@ namespace UnrealBuildTool
 			return BinaryCompileEnvironment;
 		}
 
-		private LinkEnvironment SetupBinaryLinkEnvironment(ReadOnlyTargetRules Target, UEToolChain ToolChain, LinkEnvironment LinkEnvironment, CppCompileEnvironment CompileEnvironment, List<FileReference> SpecificFilesToCompile, ISourceFileWorkingSet WorkingSet, DirectoryReference ExeDir, IActionGraphBuilder Graph)
+		private LinkEnvironment SetupBinaryLinkEnvironment(ReadOnlyTargetRules Target, UEToolChain ToolChain, LinkEnvironment LinkEnvironment, CppCompileEnvironment CompileEnvironment, List<FileReference> SpecificFilesToCompile, ISourceFileWorkingSet WorkingSet, DirectoryReference ExeDir, IActionGraphBuilder Graph, ILogger Logger)
 		{
 			LinkEnvironment BinaryLinkEnvironment = new LinkEnvironment(LinkEnvironment);
 			HashSet<UEBuildModule> LinkEnvironmentVisitedModules = new HashSet<UEBuildModule>();
@@ -712,8 +714,8 @@ namespace UnrealBuildTool
 				if (Module.Binary == null || Module.Binary == this)
 				{
 					// Compile each module.
-					Log.TraceVerbose("Compile module: " + Module.Name);
-					LinkInputFiles = Module.Compile(Target, ToolChain, BinaryCompileEnvironment, SpecificFilesToCompile, WorkingSet, Graph);
+					Logger.LogDebug("Compile module: {ModuleName}", Module.Name);
+					LinkInputFiles = Module.Compile(Target, ToolChain, BinaryCompileEnvironment, SpecificFilesToCompile, WorkingSet, Graph, Logger);
 
 					// Save the module outputs. In monolithic builds, this is just the object files.
 					if (Target.LinkType == TargetLinkType.Monolithic)
diff --git a/Engine/Source/Programs/UnrealBuildTool/Configuration/UEBuildDeploy.cs b/Engine/Source/Programs/UnrealBuildTool/Configuration/UEBuildDeploy.cs
--- a/Engine/Source/Programs/UnrealBuildTool/Configuration/UEBuildDeploy.cs
+++ b/Engine/Source/Programs/UnrealBuildTool/Configuration/UEBuildDeploy.cs
@@ -4,6 +4,7 @@
 using System.Linq;
 using System.IO;
 using System;
+using Microsoft.Extensions.Logging;
 
 namespace UnrealBuildTool
 {
@@ -12,13 +13,20 @@
 	/// </summary>
 	abstract class UEBuildDeploy
 	{
+		protected readonly ILogger Logger;
+
+		public UEBuildDeploy(ILogger InLogger)
+		{
+			Logger = InLogger;
+		}
+
 		protected bool CopyFilesRecursively(string FromDir, string ToDir)
 		{
 			var FromDirInfo = new DirectoryInfo(FromDir);
 
 			if (!FromDirInfo.Exists)
 			{
-				Log.TraceError($"Cannot copy files, source directory not found: {FromDirInfo.FullName}");
+				Logger.LogError("Cannot copy files, source directory not found: {FromDir}", FromDirInfo.FullName);
 				return false;
 			}
 
@@ -40,7 +48,7 @@
 					FromFile.CopyTo(ToDestFile, true);
 				} catch (Exception ex)
 				{
-					Log.TraceError($"Error while copying {FromFile.Name} from {FromDir} to {ToDir}: {ex}");
+					Logger.LogError(ex, "Error while copying {FromFile} from {FromDir} to {ToDir}: {Ex}", FromFile, FromDir, ToDir, ex);
 					return false;
 				}
 			}
diff --git a/Engine/Source/Programs/UnrealBuildTool/Configuration/UEBuildModule.cs b/Engine/Source/Programs/UnrealBuildTool/Configuration/UEBuildModule.cs
--- a/Engine/Source/Programs/UnrealBuildTool/Configuration/UEBuildModule.cs
+++ b/Engine/Source/Programs/UnrealBuildTool/Configuration/UEBuildModule.cs
@@ -10,6 +10,7 @@
 using System.Xml;
 using EpicGames.Core;
 using UnrealBuildBase;
+using Microsoft.Extensions.Logging;
 
 namespace UnrealBuildTool
 {
@@ -198,7 +199,8 @@
 		/// </summary>
 		/// <param name="Rules">Rules for this module</param>
 		/// <param name="IntermediateDirectory">Intermediate directory for this module</param>
-		public UEBuildModule(ModuleRules Rules, DirectoryReference IntermediateDirectory)
+		/// <param name="Logger">Logger for output</param>
+		public UEBuildModule(ModuleRules Rules, DirectoryReference IntermediateDirectory, ILogger Logger)
 		{
 			this.Rules = Rules;
 			this.IntermediateDirectory = IntermediateDirectory;
@@ -227,13 +229,13 @@
 				}
 				else if (PublicPreBuildLibraries.Contains(LibraryName))
 				{
-					Log.TraceLog("Library '{0}' was not resolvable to a file when used in Module '{1}'.  Be sure to add either a TargetRules.PreBuildSteps entry or a TargetRules.PreBuildTargets entry to assure it is built for your target.", LibraryName, Name);
+					Logger.LogDebug("Library '{LibraryName}' was not resolvable to a file when used in Module '{Name}'.  Be sure to add either a TargetRules.PreBuildSteps entry or a TargetRules.PreBuildTargets entry to assure it is built for your target.", LibraryName, Name);
 					PublicLibraries.Add(Library.Location);
 				}
 				else
 				{
 					// the library path does not seem to be resolvable as is, lets warn about it as dependency checking will not work for it
-					LogWarningOrThrowError(Rules.Target.DefaultWarningLevel, "Library '{0}' was not resolvable to a file when used in Module '{1}', assuming it is a filename and will search library paths for it. This is slow and dependency checking will not work for it. Please update reference to be fully qualified alternatively use PublicSystemLibraryPaths if you do intended to use this slow path to suppress this warning. ", LibraryName, Name);
+					LogWarningOrThrowError(Logger, Rules.Target.DefaultWarningLevel, "Library '{0}' was not resolvable to a file when used in Module '{1}', assuming it is a filename and will search library paths for it. This is slow and dependency checking will not work for it. Please update reference to be fully qualified alternatively use PublicSystemLibraryPaths if you do intended to use this slow path to suppress this warning. ", LibraryName, Name);
 					PublicSystemLibraries.Add(LibraryName);
 				}
 			}
@@ -280,10 +282,11 @@
 		/// <summary>
 		/// Log a warning or throw an error message
 		/// </summary>
+		/// <param name="Logger"></param>
 		/// <param name="Level"></param>
 		/// <param name="Format"></param>
 		/// <param name="Args"></param>
-		static void LogWarningOrThrowError(WarningLevel Level, string Format, params object[] Args)
+		static void LogWarningOrThrowError(ILogger Logger, WarningLevel Level, string Format, params object[] Args)
 		{
 			if (Level == WarningLevel.Error)
 			{
@@ -291,7 +294,7 @@
 			}
 			else if (Level == WarningLevel.Warning)
 			{
-				Log.TraceWarning(Format, Args);
+				Logger.LogWarning("{Message}", String.Format(Format, Args));
 			}
 		}
 
@@ -852,7 +855,7 @@
 		/// <summary>
 		/// Compiles the module, and returns a list of files output by the compiler.
 		/// </summary>
-		public virtual List<FileItem> Compile(ReadOnlyTargetRules Target, UEToolChain ToolChain, CppCompileEnvironment CompileEnvironment, List<FileReference> SpecificFilesToCompile, ISourceFileWorkingSet WorkingSet, IActionGraphBuilder Graph)
+		public virtual List<FileItem> Compile(ReadOnlyTargetRules Target, UEToolChain ToolChain, CppCompileEnvironment CompileEnvironment, List<FileReference> SpecificFilesToCompile, ISourceFileWorkingSet WorkingSet, IActionGraphBuilder Graph, ILogger Logger)
 		{
 			// Generate type libraries for Windows
 			foreach(ModuleRules.TypeLibrary TypeLibrary in Rules.TypeLibraries)
@@ -928,10 +931,11 @@
 		/// </summary>
 		/// <param name="CreateModule"></param>
 		/// <param name="ReferenceChain"></param>
-		public void RecursivelyCreateModules(CreateModuleDelegate CreateModule, string ReferenceChain)
+		/// <param name="Logger"></param>
+		public void RecursivelyCreateModules(CreateModuleDelegate CreateModule, string ReferenceChain, ILogger Logger)
 		{
 			List<UEBuildModule> ReferenceStack = new List<UEBuildModule>();
-			RecursivelyCreateModules(CreateModule, ReferenceChain, ReferenceStack);
+			RecursivelyCreateModules(CreateModule, ReferenceChain, ReferenceStack, Logger);
 		}
 
 		/// <summary>
@@ -940,7 +944,8 @@
 		/// <param name="CreateModule">Delegate to create a module with a given name</param>
 		/// <param name="ReferenceChain">Chain of references before reaching this module</param>
 		/// <param name="ReferenceStack">Stack of module dependencies that led to this module</param>
-		protected void RecursivelyCreateModules(CreateModuleDelegate CreateModule, string ReferenceChain, List<UEBuildModule> ReferenceStack)
+		/// <param name="Logger">Logger for output</param>
+		protected void RecursivelyCreateModules(CreateModuleDelegate CreateModule, string ReferenceChain, List<UEBuildModule> ReferenceStack, ILogger Logger)
 		{
 			// Name of this reference
 			string ThisRefName = (RulesFile == null) ? Name : RulesFile.GetFileName();
@@ -984,16 +989,16 @@
 				if (string.IsNullOrEmpty(GuiltyModule))
 				{
 					string CycleChain = string.Join(" -> ", ReferenceStack);
-					Log.TraceError("Circular dependency on {0} detected.\n" +
-						"\tFull Route: {1}\n" +
-						"\tCycled Route: is {2}.\n" +
+					Logger.LogError("Circular dependency on {Name} detected.\n" +
+						"\tFull Route: {FullRoute}\n" +
+						"\tCycled Route: is {CycleRoot}.\n" +
 						"Break this loop by moving dependencies into a separate module or using Private/PublicIncludePathModuleNames to reference declarations\n", 
 						ThisRefName, NextReferenceChain, CycleChain);
 
 				}
 				else
 				{
-					Log.TraceVerbose("Found circular reference to {0}, but {1} declares a cycle on {2} which breaks the chain", ThisRefName, GuiltyModule, VictimModule);
+					Logger.LogDebug("Found circular reference to {ThisRefName}, but {GuiltyModule} declares a cycle on {VictimModule} which breaks the chain", ThisRefName, GuiltyModule, VictimModule);
 				}
 			}		
 
@@ -1020,15 +1025,15 @@
 				RecursivelyCreateIncludePathModulesByName(Rules.PrivateIncludePathModuleNames, ref PrivateIncludePathModules, ref bDependsOnVerse, CreateModule, NextReferenceChain);
 
 				// Create all the dependency modules - pass through the reference stack so we can check for cycles
-				RecursivelyCreateModulesByName(Rules.PublicDependencyModuleNames, ref PublicDependencyModules, ref bDependsOnVerse, CreateModule, NextReferenceChain, ReferenceStack);
+				RecursivelyCreateModulesByName(Rules.PublicDependencyModuleNames, ref PublicDependencyModules, ref bDependsOnVerse, CreateModule, NextReferenceChain, ReferenceStack, Logger);
 				if (Rules.Target.IsTestTarget)
 				{
 					// Move the test runner dependency to last position to give it the opportunity to build its special dependencies such as CoreUObject, ApplicationCore etc.
 					MoveTestsRunnerDependencyToLastPosition();
 				}
-				RecursivelyCreateModulesByName(Rules.PrivateDependencyModuleNames, ref PrivateDependencyModules, ref bDependsOnVerse, CreateModule, NextReferenceChain, ReferenceStack);
+				RecursivelyCreateModulesByName(Rules.PrivateDependencyModuleNames, ref PrivateDependencyModules, ref bDependsOnVerse, CreateModule, NextReferenceChain, ReferenceStack, Logger);
 				// Dynamic loads aren't considered a reference chain so start with an empty stack
-				RecursivelyCreateModulesByName(Rules.DynamicallyLoadedModuleNames, ref DynamicallyLoadedModules, ref bDependsOnVerse, CreateModule, NextReferenceChain, new List<UEBuildModule>());
+				RecursivelyCreateModulesByName(Rules.DynamicallyLoadedModuleNames, ref DynamicallyLoadedModules, ref bDependsOnVerse, CreateModule, NextReferenceChain, new List<UEBuildModule>(), Logger);
 			}
 
 			// pop us off the current stack
@@ -1053,7 +1058,7 @@
 			}
 		}
 
-		private static void RecursivelyCreateModulesByName(List<string> ModuleNames, ref List<UEBuildModule>? Modules, ref bool bDependsOnVerse, CreateModuleDelegate CreateModule, string ReferenceChain, List<UEBuildModule> ReferenceStack)
+		private static void RecursivelyCreateModulesByName(List<string> ModuleNames, ref List<UEBuildModule>? Modules, ref bool bDependsOnVerse, CreateModuleDelegate CreateModule, string ReferenceChain, List<UEBuildModule> ReferenceStack, ILogger Logger)
 		{
 			// Check whether the module list is already set. We set this immediately (via the ref) to avoid infinite recursion.
 			if (Modules == null)
@@ -1064,7 +1069,7 @@
 					UEBuildModule Module = CreateModule(ModuleName, ReferenceChain);
 					if (!Modules.Contains(Module))
 					{
-						Module.RecursivelyCreateModules(CreateModule, ReferenceChain, ReferenceStack);
+						Module.RecursivelyCreateModules(CreateModule, ReferenceChain, ReferenceStack, Logger);
 						Modules.Add(Module);
 						bDependsOnVerse |= Module.bDependsOnVerse;
 					}
diff --git a/Engine/Source/Programs/UnrealBuildTool/Configuration/UEBuildModuleCPP.cs b/Engine/Source/Programs/UnrealBuildTool/Configuration/UEBuildModuleCPP.cs
--- a/Engine/Source/Programs/UnrealBuildTool/Configuration/UEBuildModuleCPP.cs
+++ b/Engine/Source/Programs/UnrealBuildTool/Configuration/UEBuildModuleCPP.cs
@@ -6,6 +6,7 @@ using System.IO;
 using System.Linq;
 using System.Text;
 using EpicGames.Core;
+using Microsoft.Extensions.Logging;
 using UnrealBuildBase;
 
 namespace UnrealBuildTool
@@ -203,8 +204,8 @@ namespace UnrealBuildTool
             new KeyValuePair<string, string>("UnrealEd", "MeshPaint"),
         };
 
-		public UEBuildModuleCPP(ModuleRules Rules, DirectoryReference IntermediateDirectory, DirectoryReference? GeneratedCodeDirectory)
-			: base(Rules, IntermediateDirectory)
+		public UEBuildModuleCPP(ModuleRules Rules, DirectoryReference IntermediateDirectory, DirectoryReference? GeneratedCodeDirectory, ILogger Logger)
+			: base(Rules, IntermediateDirectory, Logger)
 		{
 			this.GeneratedCodeDirectory = GeneratedCodeDirectory;
 
@@ -218,12 +219,12 @@ namespace UnrealBuildTool
 
 			foreach (string Def in PublicDefinitions)
 			{
-				Log.TraceVerbose("Compile Env {0}: {1}", Name, Def);
+				Logger.LogDebug("Compile Env {Name}: {Def}", Name, Def);
 			}
 
 			foreach (string Def in Rules.PrivateDefinitions)
 			{
-				Log.TraceVerbose("Compile Env {0}: {1}", Name, Def);
+				Logger.LogDebug("Compile Env {Name}: {Def}", Name, Def);
 			}
 
 			if (Rules.bValidateCircularDependencies || Rules.bTreatAsEngineModule)
@@ -234,8 +235,8 @@ namespace UnrealBuildTool
 					    !CircularDependenciesAllowList.Any(x =>
 						    x.Key == Name && x.Value == CircularlyReferencedModuleName))
 					{
-						Log.TraceWarning(
-							"Found reference between '{0}' and '{1}'. Support for circular references is being phased out; please do not introduce new ones.",
+						Logger.LogWarning(
+							"Found reference between '{Source}' and '{Target}'. Support for circular references is being phased out; please do not introduce new ones.",
 							Name, CircularlyReferencedModuleName);
 					}
 				}
@@ -372,11 +373,11 @@ namespace UnrealBuildTool
 		}
 
 		// UEBuildModule interface.
-		public override List<FileItem> Compile(ReadOnlyTargetRules Target, UEToolChain ToolChain, CppCompileEnvironment BinaryCompileEnvironment, List<FileReference> SpecificFilesToCompile, ISourceFileWorkingSet WorkingSet, IActionGraphBuilder Graph)
+		public override List<FileItem> Compile(ReadOnlyTargetRules Target, UEToolChain ToolChain, CppCompileEnvironment BinaryCompileEnvironment, List<FileReference> SpecificFilesToCompile, ISourceFileWorkingSet WorkingSet, IActionGraphBuilder Graph, ILogger Logger)
 		{
 			//UEBuildPlatform BuildPlatform = UEBuildPlatform.GetBuildPlatform(BinaryCompileEnvironment.Platform);
 
-			List<FileItem> LinkInputFiles = base.Compile(Target, ToolChain, BinaryCompileEnvironment, SpecificFilesToCompile, WorkingSet, Graph);
+			List<FileItem> LinkInputFiles = base.Compile(Target, ToolChain, BinaryCompileEnvironment, SpecificFilesToCompile, WorkingSet, Graph, Logger);
 
 			CppCompileEnvironment ModuleCompileEnvironment = CreateModuleCompileEnvironment(Target, BinaryCompileEnvironment);
 
@@ -475,28 +476,28 @@ namespace UnrealBuildTool
 			{
 				if (Target.bForceUnityBuild)
 				{
-					Log.TraceVerbose("Module '{0}' using unity build mode (bForceUnityBuild enabled for this module)", this.Name);
+					Logger.LogDebug("Module '{ModuleName}' using unity build mode (bForceUnityBuild enabled for this module)", this.Name);
 					bModuleUsesUnityBuild = true;
 				}
 				else if (!Rules.bUseUnity)
 				{
-					Log.TraceVerbose("Module '{0}' not using unity build mode (bUseUnity disabled for this module)", this.Name);
+					Logger.LogDebug("Module '{ModuleName}' not using unity build mode (bUseUnity disabled for this module)", this.Name);
 					bModuleUsesUnityBuild = false;
 				}
 				else if (InputFiles.CPPFiles.Count < MinSourceFilesForUnityBuild)
 				{
-					Log.TraceVerbose("Module '{0}' not using unity build mode (module with fewer than {1} source files)", this.Name, MinSourceFilesForUnityBuild);
+					Logger.LogDebug("Module '{ModuleName}' not using unity build mode (module with fewer than {NumFiles} source files)", this.Name, MinSourceFilesForUnityBuild);
 					bModuleUsesUnityBuild = false;
 				}
 				else
 				{
-					Log.TraceVerbose("Module '{0}' using unity build mode", this.Name);
+					Logger.LogDebug("Module '{ModuleName}' using unity build mode", this.Name);
 					bModuleUsesUnityBuild = true;
 				}
 			}
 			else
 			{
-				Log.TraceVerbose("Module '{0}' not using unity build mode", this.Name);
+				Logger.LogDebug("Module '{ModuleName}' not using unity build mode", this.Name);
 			}
 
 			// Set up the environment with which to compile the CPP files
@@ -528,19 +529,19 @@ namespace UnrealBuildTool
 			// Compile CPP files
 			if (bModuleUsesUnityBuild)
 			{
-				Unity.GenerateUnityCPPs(Target, InputFiles.CPPFiles, InputFiles.HeaderFiles, CompileEnvironment, WorkingSet, Rules.ShortName ?? Name, IntermediateDirectory, Graph, SourceFileToUnityFile,
+				Unity.GenerateUnityCPPs(Target, InputFiles.CPPFiles, InputFiles.HeaderFiles, CompileEnvironment, WorkingSet, Rules.ShortName ?? Name, IntermediateDirectory, Graph, SourceFileToUnityFile, 
 					out List<FileItem> NormalFiles, out List<FileItem> AdaptiveFiles);
-				LinkInputFiles.AddRange(CompileFilesWithToolChain(Target, ToolChain, CompileEnvironment, ModuleCompileEnvironment, NormalFiles, AdaptiveFiles, Graph).ObjectFiles);
+				LinkInputFiles.AddRange(CompileFilesWithToolChain(Target, ToolChain, CompileEnvironment, ModuleCompileEnvironment, NormalFiles, AdaptiveFiles, Graph, Logger).ObjectFiles);
 			}
 			else if (SpecificFilesToCompile.Count == 0)
 			{
-				Unity.GetAdaptiveFiles(Target, InputFiles.CPPFiles, InputFiles.HeaderFiles, CompileEnvironment, WorkingSet, Rules.ShortName ?? Name, IntermediateDirectory, Graph,
+				Unity.GetAdaptiveFiles(Target, InputFiles.CPPFiles, InputFiles.HeaderFiles, CompileEnvironment, WorkingSet, Rules.ShortName ?? Name, IntermediateDirectory, Graph, 
 					out List<FileItem> NormalFiles, out List<FileItem> AdaptiveFiles);
-				LinkInputFiles.AddRange(CompileFilesWithToolChain(Target, ToolChain, CompileEnvironment, ModuleCompileEnvironment, NormalFiles, AdaptiveFiles, Graph).ObjectFiles);
+				LinkInputFiles.AddRange(CompileFilesWithToolChain(Target, ToolChain, CompileEnvironment, ModuleCompileEnvironment, NormalFiles, AdaptiveFiles, Graph, Logger).ObjectFiles);
 			}
 			else
 			{
-				LinkInputFiles.AddRange(CompileFilesWithToolChain(Target, ToolChain, CompileEnvironment, ModuleCompileEnvironment, InputFiles.CPPFiles, new List<FileItem>(), Graph).ObjectFiles);
+				LinkInputFiles.AddRange(CompileFilesWithToolChain(Target, ToolChain, CompileEnvironment, ModuleCompileEnvironment, InputFiles.CPPFiles, new List<FileItem>(), Graph, Logger).ObjectFiles);
 			}
 
 			// Compile all the generated CPP files
@@ -591,9 +592,9 @@ namespace UnrealBuildTool
 
 					if (bModuleUsesUnityBuild)
 					{
-						Unity.GenerateUnityCPPs(Target, GeneratedFileItems, new List<FileItem>(), CompileEnvironment, WorkingSet, (Rules.ShortName ?? Name) + ".gen", IntermediateDirectory, Graph, SourceFileToUnityFile,
+						Unity.GenerateUnityCPPs(Target, GeneratedFileItems, new List<FileItem>(), CompileEnvironment, WorkingSet, (Rules.ShortName ?? Name) + ".gen", IntermediateDirectory, Graph, SourceFileToUnityFile, 
 							out List<FileItem> NormalGeneratedFiles, out List<FileItem> AdaptiveGeneratedFiles);
-						LinkInputFiles.AddRange(CompileFilesWithToolChain(Target, ToolChain, GeneratedCPPCompileEnvironment, ModuleCompileEnvironment, NormalGeneratedFiles, AdaptiveGeneratedFiles, Graph).ObjectFiles);
+						LinkInputFiles.AddRange(CompileFilesWithToolChain(Target, ToolChain, GeneratedCPPCompileEnvironment, ModuleCompileEnvironment, NormalGeneratedFiles, AdaptiveGeneratedFiles, Graph, Logger).ObjectFiles);
 					}
 					else
 					{
@@ -961,7 +962,8 @@ namespace UnrealBuildTool
 			CppCompileEnvironment ModuleCompileEnvironment,
 			List<FileItem> NormalFiles,
 			List<FileItem> AdaptiveFiles,
-			IActionGraphBuilder Graph)
+			IActionGraphBuilder Graph,
+			ILogger Logger)
 		{
 			bool bAdaptiveUnityDisablesPCH = false;
 			if(Rules.PCHUsage == ModuleRules.PCHUsageMode.UseExplicitOrSharedPCHs)
@@ -1410,11 +1412,11 @@ namespace UnrealBuildTool
 			}
 		}
 
-		public CppCompileEnvironment CreateCompileEnvironmentForIntellisense(ReadOnlyTargetRules Target, CppCompileEnvironment BaseCompileEnvironment)
+		public CppCompileEnvironment CreateCompileEnvironmentForIntellisense(ReadOnlyTargetRules Target, CppCompileEnvironment BaseCompileEnvironment, ILogger Logger)
 		{
 			CppCompileEnvironment CompileEnvironment = CreateModuleCompileEnvironment(Target, BaseCompileEnvironment);
-			CompileEnvironment = SetupPrecompiledHeaders(Target, null, CompileEnvironment, new List<FileItem>(), new NullActionGraphBuilder());
-			CreateHeaderForDefinitions(CompileEnvironment, IntermediateDirectory, null, new NullActionGraphBuilder());
+			CompileEnvironment = SetupPrecompiledHeaders(Target, null, CompileEnvironment, new List<FileItem>(), new NullActionGraphBuilder(Logger));
+			CreateHeaderForDefinitions(CompileEnvironment, IntermediateDirectory, null, new NullActionGraphBuilder(Logger));
 			return CompileEnvironment;
 		}
 
diff --git a/Engine/Source/Programs/UnrealBuildTool/Configuration/UEBuildModuleExternal.cs b/Engine/Source/Programs/UnrealBuildTool/Configuration/UEBuildModuleExternal.cs
--- a/Engine/Source/Programs/UnrealBuildTool/Configuration/UEBuildModuleExternal.cs
+++ b/Engine/Source/Programs/UnrealBuildTool/Configuration/UEBuildModuleExternal.cs
@@ -1,22 +1,23 @@
-// Copyright Epic Games, Inc. All Rights Reserved.
-
-using System;
-using System.Collections.Generic;
-using System.Linq;
-using System.Text;
-using System.Threading.Tasks;
-using EpicGames.Core;
-
-namespace UnrealBuildTool
-{
-	/// <summary>
-	/// A module that is never compiled by us, and is only used to group include paths and libraries into a dependency unit.
-	/// </summary>
-	class UEBuildModuleExternal : UEBuildModule
-	{
-		public UEBuildModuleExternal(ModuleRules Rules, DirectoryReference IntermediateDirectory)
-			: base(Rules, IntermediateDirectory)
-		{
-		}
-	}
-}
+// Copyright Epic Games, Inc. All Rights Reserved.
+
+using System;
+using System.Collections.Generic;
+using System.Linq;
+using System.Text;
+using System.Threading.Tasks;
+using EpicGames.Core;
+using Microsoft.Extensions.Logging;
+
+namespace UnrealBuildTool
+{
+	/// <summary>
+	/// A module that is never compiled by us, and is only used to group include paths and libraries into a dependency unit.
+	/// </summary>
+	class UEBuildModuleExternal : UEBuildModule
+	{
+		public UEBuildModuleExternal(ModuleRules Rules, DirectoryReference IntermediateDirectory, ILogger Logger)
+			: base(Rules, IntermediateDirectory, Logger)
+		{
+		}
+	}
+}
diff --git a/Engine/Source/Programs/UnrealBuildTool/Configuration/UEBuildPlatform.cs b/Engine/Source/Programs/UnrealBuildTool/Configuration/UEBuildPlatform.cs
--- a/Engine/Source/Programs/UnrealBuildTool/Configuration/UEBuildPlatform.cs
+++ b/Engine/Source/Programs/UnrealBuildTool/Configuration/UEBuildPlatform.cs
@@ -10,6 +10,7 @@
 using System.Reflection;
 using System.Diagnostics.CodeAnalysis;
 using OpenTracing.Util;
+using Microsoft.Extensions.Logging;
 
 namespace UnrealBuildTool
 {
@@ -26,6 +27,11 @@
 		public readonly UnrealTargetPlatform Platform;
 
 		/// <summary>
+		/// Logger for this platform
+		/// </summary>
+		protected readonly ILogger Logger;
+
+		/// <summary>
 		/// All the platform folder names
 		/// </summary>
 		private static string[]? CachedPlatformFolderNames;
@@ -45,9 +51,11 @@
 		/// </summary>
 		/// <param name="InPlatform">The enum value for this platform</param>
 		/// <param name="SDK">The SDK management object for this platform</param>
-		public UEBuildPlatform(UnrealTargetPlatform InPlatform, UEBuildPlatformSDK SDK)
+		/// <param name="InLogger">Logger for output</param>
+		public UEBuildPlatform(UnrealTargetPlatform InPlatform, UEBuildPlatformSDK SDK, ILogger InLogger)
 		{
 			Platform = InPlatform;
+			Logger = InLogger;
 
 			// check DDPI to see if the platform is enabled on this host platform
 			string IniPlatformName = ConfigHierarchy.GetIniPlatformName(Platform);
@@ -66,7 +74,8 @@
 		/// </summary>
 		/// <param name="bIncludeNonInstalledPlatforms">Whether to register platforms that are not installed</param>
 		/// <param name="bHostPlatformOnly">Only register the host platform</param>
-		public static void RegisterPlatforms(bool bIncludeNonInstalledPlatforms, bool bHostPlatformOnly)
+		/// <param name="Logger">Logger for output</param>
+		public static void RegisterPlatforms(bool bIncludeNonInstalledPlatforms, bool bHostPlatformOnly, ILogger Logger)
 		{
 			// Initialize the installed platform info
 			using (GlobalTracer.Instance.BuildSpan("Initializing InstalledPlatformInfo").StartActive())
@@ -88,7 +97,7 @@
 				{
 					if (CheckType.IsSubclassOf(typeof(UEBuildPlatformFactory)))
 					{
-						Log.TraceVerbose("    Registering build platform: {0}", CheckType.ToString());
+						Logger.LogDebug("    Registering build platform: {Platform}", CheckType.ToString());
 						using (GlobalTracer.Instance.BuildSpan(CheckType.Name).StartActive())
 						{
 							UEBuildPlatformFactory TempInst = (UEBuildPlatformFactory)Activator.CreateInstance(CheckType)!;
@@ -102,7 +111,7 @@
 							// We need all platforms to be registered when we run -validateplatform command to check SDK status of each
 							if (bIncludeNonInstalledPlatforms || InstalledPlatformInfo.IsValidPlatform(TempInst.TargetPlatform))
 							{
-								TempInst.RegisterBuildPlatforms();
+								TempInst.RegisterBuildPlatforms(Logger);
 							}
 						}
 					}
@@ -396,13 +405,14 @@
 		/// Register the given platforms UEBuildPlatform instance
 		/// </summary>
 		/// <param name="InBuildPlatform"> The UEBuildPlatform instance to use for the InPlatform</param>
-		public static void RegisterBuildPlatform(UEBuildPlatform InBuildPlatform)
+		/// <param name="Logger">Logger for output</param>
+		public static void RegisterBuildPlatform(UEBuildPlatform InBuildPlatform, ILogger Logger)
 		{
-			Log.TraceVerbose("        Registering build platform: {0} - buildable: {1}", InBuildPlatform.Platform, InBuildPlatform.HasRequiredSDKsInstalled() == SDKStatus.Valid);
+			Logger.LogDebug("        Registering build platform: {Platform} - buildable: {Buildable}", InBuildPlatform.Platform, InBuildPlatform.HasRequiredSDKsInstalled() == SDKStatus.Valid);
 
 			if (BuildPlatformDictionary.ContainsKey(InBuildPlatform.Platform) == true)
 			{
-				Log.TraceWarning("RegisterBuildPlatform Warning: Registering build platform {0} for {1} when it is already set to {2}",
+				Logger.LogWarning("RegisterBuildPlatform Warning: Registering build platform {Platform} for {ForPlatform} when it is already set to {CurPlatform}",
 					InBuildPlatform.ToString(), InBuildPlatform.Platform.ToString(), BuildPlatformDictionary[InBuildPlatform.Platform].ToString());
 				BuildPlatformDictionary[InBuildPlatform.Platform] = InBuildPlatform;
 			}
@@ -508,8 +518,7 @@
 				return ";";
 			}
 
-			Log.TraceWarning("PATH variable delimiter unknown for platform " + BuildHostPlatform.Current.Platform.ToString() + " using ';'");
-			return ";";
+			throw new InvalidOperationException($"PATH variable delimiter unknown for platform {BuildHostPlatform.Current.Platform}");
 		}
 
 		/// <summary>
@@ -663,7 +672,7 @@
 			return Configurations;
 		}
 
-		protected static bool DoProjectSettingsMatchDefault(UnrealTargetPlatform Platform, DirectoryReference ProjectDirectoryName, string Section, string[]? BoolKeys, string[]? IntKeys, string[]? StringKeys)
+		protected static bool DoProjectSettingsMatchDefault(UnrealTargetPlatform Platform, DirectoryReference ProjectDirectoryName, string Section, string[]? BoolKeys, string[]? IntKeys, string[]? StringKeys, ILogger Logger)
 		{
 			ConfigHierarchy ProjIni = ConfigCache.ReadHierarchy(ConfigHierarchyType.Engine, ProjectDirectoryName, Platform);
 			ConfigHierarchy DefaultIni = ConfigCache.ReadHierarchy(ConfigHierarchyType.Engine, (DirectoryReference?)null, Platform);
@@ -725,7 +734,7 @@
 			};
 
 			return DoProjectSettingsMatchDefault(Platform, ProjectDirectoryName, "/Script/BuildSettings.BuildSettings",
-				BoolKeys, null, null);
+				BoolKeys, null, null, Logger);
 		}
 
 		/// <summary>
diff --git a/Engine/Source/Programs/UnrealBuildTool/Configuration/UEBuildPlatformFactory.cs b/Engine/Source/Programs/UnrealBuildTool/Configuration/UEBuildPlatformFactory.cs
--- a/Engine/Source/Programs/UnrealBuildTool/Configuration/UEBuildPlatformFactory.cs
+++ b/Engine/Source/Programs/UnrealBuildTool/Configuration/UEBuildPlatformFactory.cs
@@ -1,29 +1,30 @@
-﻿// Copyright Epic Games, Inc. All Rights Reserved.
-
-using System;
-using System.Collections.Generic;
-using System.Linq;
-using System.Text;
-using System.Threading.Tasks;
-
-namespace UnrealBuildTool
-{
-	/// <summary>
-	/// Factory class for registering platforms at startup
-	/// </summary>
-	abstract class UEBuildPlatformFactory
-	{
-		/// <summary>
-		/// Gets the target platform for an individual factory
-		/// </summary>
-		public abstract UnrealTargetPlatform TargetPlatform
-		{
-			get;
-		}
-
-		/// <summary>
-		/// Register the platform with the UEBuildPlatform class
-		/// </summary>
-		public abstract void RegisterBuildPlatforms();
-	}
-}
+// Copyright Epic Games, Inc. All Rights Reserved.
+
+using System;
+using System.Collections.Generic;
+using System.Linq;
+using System.Text;
+using System.Threading.Tasks;
+using Microsoft.Extensions.Logging;
+
+namespace UnrealBuildTool
+{
+	/// <summary>
+	/// Factory class for registering platforms at startup
+	/// </summary>
+	abstract class UEBuildPlatformFactory
+	{
+		/// <summary>
+		/// Gets the target platform for an individual factory
+		/// </summary>
+		public abstract UnrealTargetPlatform TargetPlatform
+		{
+			get;
+		}
+
+		/// <summary>
+		/// Register the platform with the UEBuildPlatform class
+		/// </summary>
+		public abstract void RegisterBuildPlatforms(ILogger Logger);
+	}
+}
