UGS: Allow workspaces without an owner field.

diff --git a/Engine/Source/Programs/Shared/EpicGames.Perforce/Records/ClientRecord.cs b/Engine/Source/Programs/Shared/EpicGames.Perforce/Records/ClientRecord.cs
--- a/Engine/Source/Programs/Shared/EpicGames.Perforce/Records/ClientRecord.cs
+++ b/Engine/Source/Programs/Shared/EpicGames.Perforce/Records/ClientRecord.cs
@@ -19,8 +19,8 @@ namespace EpicGames.Perforce
 		/// <summary>
 		/// The Perforce user name of the user who owns the workspace.
 		/// </summary>
-		[PerforceTag("Owner")]
-		public string Owner { get; set; }
+		[PerforceTag("Owner", Optional = true)]
+		public string? Owner { get; set; }
 
 		/// <summary>
 		/// The date the workspace specification was last modified.
