Added TIsUEnumClass trait with UHT integration

diff --git a/Engine/Source/Programs/Shared/EpicGames.UHT/Exporters/CodeGen/UhtHeaderCodeGeneratorHFile.cs b/Engine/Source/Programs/Shared/EpicGames.UHT/Exporters/CodeGen/UhtHeaderCodeGeneratorHFile.cs
--- a/Engine/Source/Programs/Shared/EpicGames.UHT/Exporters/CodeGen/UhtHeaderCodeGeneratorHFile.cs
+++ b/Engine/Source/Programs/Shared/EpicGames.UHT/Exporters/CodeGen/UhtHeaderCodeGeneratorHFile.cs
@@ -271,7 +271,6 @@ namespace EpicGames.UHT.Exporters.CodeGen
 			}
 			builder.Append("\r\n");
 
-			// Forward declare the StaticEnum<> specialization for enum classes
 			if (enumObj.CppForm == UhtEnumCppForm.EnumClass)
 			{
 				builder.Append("\r\n");
@@ -281,6 +280,11 @@ namespace EpicGames.UHT.Exporters.CodeGen
 					builder.Append(" : ").Append(enumObj.UnderlyingType.ToString());
 				}
 				builder.Append(";\r\n");
+				
+				// Add TIsUEnumClass typetraits
+				builder.Append("template<> struct TIsUEnumClass<").Append(enumObj.CppType).Append("> { enum { Value = true }; };\r\n");
+
+				// Forward declare the StaticEnum<> specialization for enum classes
 				builder.Append("template<> ").Append(this.PackageApi).Append("UEnum* StaticEnum<").Append(enumObj.CppType).Append(">();\r\n");
 				builder.Append("\r\n");
 			}
diff --git a/Engine/Source/Programs/UnrealHeaderTool/Private/CodeGenerator.cpp b/Engine/Source/Programs/UnrealHeaderTool/Private/CodeGenerator.cpp
--- a/Engine/Source/Programs/UnrealHeaderTool/Private/CodeGenerator.cpp
+++ b/Engine/Source/Programs/UnrealHeaderTool/Private/CodeGenerator.cpp
@@ -3335,7 +3335,6 @@ void FNativeClassHeaderGenerator::ExportEnum(FOutputDevice& Out, FUnrealEnumDefi
 	}
 	Out.Logf( TEXT("\r\n") );
 
-	// Forward declare the StaticEnum<> specialisation for enum classes
 	if (EnumDef.GetCppForm() == UEnum::ECppForm::EnumClass)
 	{
 		FString UnderlyingTypeString;
@@ -3361,6 +3360,11 @@ void FNativeClassHeaderGenerator::ExportEnum(FOutputDevice& Out, FUnrealEnumDefi
 
 		Out.Logf( TEXT("\r\n") );
 		Out.Logf( TEXT("enum class %s%s;\r\n"), *EnumDef.GetCppType(), *UnderlyingTypeString );
+		
+		// Add TIsUEnumClass typetraits
+		Out.Logf( TEXT("template<> struct TIsUEnumClass<%s> { enum { Value = true }; };\r\n"), *EnumDef.GetCppType());
+
+		// Forward declare the StaticEnum<> specialisation for enum classes
 		Out.Logf( TEXT("template<> %sUEnum* StaticEnum<%s>();\r\n"), *GetAPIString(), *EnumDef.GetCppType());
 		Out.Logf( TEXT("\r\n") );
 	}
diff --git a/Engine/Source/Runtime/CoreUObject/Public/Templates/IsUEnumClass.h b/Engine/Source/Runtime/CoreUObject/Public/Templates/IsUEnumClass.h
new file mode 100644
--- /dev/null
+++ b/Engine/Source/Runtime/CoreUObject/Public/Templates/IsUEnumClass.h
@@ -0,0 +1,15 @@
+// Copyright Epic Games, Inc. All Rights Reserved.
+
+#pragma once
+
+#include "Misc/StaticAssertCompleteType.h"
+
+/**
+ * Traits class which tests if a type is a UEnum class.
+ */
+template <typename T>
+struct TIsUEnumClass
+{
+	UE_STATIC_ASSERT_COMPLETE_TYPE(T, "TIsUEnumClass instantiated with an incomplete type");
+	enum { Value = false };
+};
diff --git a/Engine/Source/Runtime/CoreUObject/Public/UObject/UnrealType.h b/Engine/Source/Runtime/CoreUObject/Public/UObject/UnrealType.h
--- a/Engine/Source/Runtime/CoreUObject/Public/UObject/UnrealType.h
+++ b/Engine/Source/Runtime/CoreUObject/Public/UObject/UnrealType.h
@@ -15,6 +15,7 @@
 #include "Serialization/MemoryImage.h"
 #include "Templates/Casts.h"
 #include "Templates/Greater.h"
+#include "Templates/IsUEnumClass.h"
 #include "Templates/IsFloatingPoint.h"
 #include "Templates/IsIntegral.h"
 #include "Templates/IsSigned.h"
