Ported optional editor only property checks to C# UHT.

#rnx

diff --git a/Engine/Source/Programs/Shared/EpicGames.UHT/Types/Properties/UhtContainerProperty.cs b/Engine/Source/Programs/Shared/EpicGames.UHT/Types/Properties/UhtContainerProperty.cs
--- a/Engine/Source/Programs/Shared/EpicGames.UHT/Types/Properties/UhtContainerProperty.cs
+++ b/Engine/Source/Programs/Shared/EpicGames.UHT/Types/Properties/UhtContainerProperty.cs
@@ -73,5 +73,11 @@ namespace EpicGames.UHT.Types
 		{
 			return this.ValueProperty.ScanForInstancedReferenced(deepScan);
 		}
+
+		///<inheritdoc/>
+		public override bool ContainsEditorOnlyProperties()
+		{
+			return this.ValueProperty.ContainsEditorOnlyProperties();
+		}
 	}
 }
diff --git a/Engine/Source/Programs/Shared/EpicGames.UHT/Types/Properties/UhtStructProperty.cs b/Engine/Source/Programs/Shared/EpicGames.UHT/Types/Properties/UhtStructProperty.cs
--- a/Engine/Source/Programs/Shared/EpicGames.UHT/Types/Properties/UhtStructProperty.cs
+++ b/Engine/Source/Programs/Shared/EpicGames.UHT/Types/Properties/UhtStructProperty.cs
@@ -206,6 +206,22 @@ namespace EpicGames.UHT.Types
 			return referencingProperty.Session.ValidateScriptStructOkForNet(referencingProperty, this.ScriptStruct);
 		}
 
+		///<inheritdoc/>
+		public override bool ContainsEditorOnlyProperties()
+		{
+			foreach (UhtType Child in this.ScriptStruct.Children)
+			{
+				if (Child is UhtProperty Property)
+				{
+					if (Property.PropertyFlags.HasAnyFlags(EPropertyFlags.EditorOnly) || Property.ContainsEditorOnlyProperties())
+					{
+						return true;
+					}
+				}
+			}
+			return false;
+		}
+
 		#region Structure default value sanitizers
 		[UhtStructDefaultValue(Name = "FVector")]
 		[SuppressMessage("CodeQuality", "IDE0051:Remove unused private members", Justification = "Attribute accessed method")]
diff --git a/Engine/Source/Programs/Shared/EpicGames.UHT/Types/UhtClass.cs b/Engine/Source/Programs/Shared/EpicGames.UHT/Types/UhtClass.cs
--- a/Engine/Source/Programs/Shared/EpicGames.UHT/Types/UhtClass.cs
+++ b/Engine/Source/Programs/Shared/EpicGames.UHT/Types/UhtClass.cs
@@ -1791,9 +1791,17 @@ namespace EpicGames.UHT.Types
 					}
 
 					// Validate if we are using editor only data in a class or struct definition
-					if (property.PropertyFlags.HasAnyFlags(EPropertyFlags.EditorOnly) && this.ClassFlags.HasAnyFlags(EClassFlags.Optional))
+					if (this.ClassFlags.HasAnyFlags(EClassFlags.Optional))
 					{
-						property.LogError("Cannot specify an editor only property inside an optional class.");
+						if (property.PropertyFlags.HasAnyFlags(EPropertyFlags.EditorOnly))
+						{
+							property.LogError("Cannot specify an editor only property inside an optional class.");
+						}
+						else if (property.ContainsEditorOnlyProperties())
+						{
+							// TODO: this should technically be an error, but some code already relies on this at this time and should hence 
+							property.LogInfo("Do not specify struct property containing editor only properties inside an optional class.");
+						}
 					}
 
 					// Check for getter/setter
diff --git a/Engine/Source/Programs/Shared/EpicGames.UHT/Types/UhtProperty.cs b/Engine/Source/Programs/Shared/EpicGames.UHT/Types/UhtProperty.cs
--- a/Engine/Source/Programs/Shared/EpicGames.UHT/Types/UhtProperty.cs
+++ b/Engine/Source/Programs/Shared/EpicGames.UHT/Types/UhtProperty.cs
@@ -1860,6 +1860,16 @@ namespace EpicGames.UHT.Types
 			errorType = null;
 			return false;
 		}
+
+		/// <summary>
+		/// Test to see if the property contains any editor only properties.  This does not 
+		/// test to see if the property itself is editor only.  This is implemented for struct properties.
+		/// </summary>
+		/// <returns></returns>
+		public virtual bool ContainsEditorOnlyProperties()
+		{
+			return false;
+		}
 		#endregion
 
 		#region Reference support
