P4 MergeTool stomped the static variable when merging the file from Release-21.10

diff --git a/Engine/Source/Programs/Shared/EpicGames.Build/UEBuildPlatformSDK.cs b/Engine/Source/Programs/Shared/EpicGames.Build/UEBuildPlatformSDK.cs
--- a/Engine/Source/Programs/Shared/EpicGames.Build/UEBuildPlatformSDK.cs
+++ b/Engine/Source/Programs/Shared/EpicGames.Build/UEBuildPlatformSDK.cs
@@ -778,6 +778,7 @@ namespace EpicGames.Core
 		#region Print SDK Info
 		private static bool bHasShownTurnkey = false;
 		private SDKStatus? SDKInfoValidity = null;
+		public static bool bSuppressSDKWarnings = false;
 
 		public virtual SDKStatus PrintSDKInfoAndReturnValidity(LogEventType Verbosity = LogEventType.Console, LogFormatOptions Options = LogFormatOptions.None,
 			LogEventType ErrorVerbosity = LogEventType.Error, LogFormatOptions ErrorOptions = LogFormatOptions.None)
