Horde: Use long-lived regex instances for matching patterns.

diff --git a/Engine/Source/Programs/AutomationTool/AutomationUtils/Matchers/ContentEventMatcher.cs b/Engine/Source/Programs/AutomationTool/AutomationUtils/Matchers/ContentEventMatcher.cs
--- a/Engine/Source/Programs/AutomationTool/AutomationUtils/Matchers/ContentEventMatcher.cs
+++ b/Engine/Source/Programs/AutomationTool/AutomationUtils/Matchers/ContentEventMatcher.cs
@@ -14,6 +14,8 @@ namespace AutomationUtils.Matchers
 	/// </summary>
 	class ContentEventMatcher : ILogEventMatcher
 	{
+		static readonly Regex s_errorWarningPattern = new Regex("Error:|Warning:");
+
 		static readonly Regex s_pattern = new Regex(
 			@"^\s*" +
 			@"(?:\[[\d\.\-: ]+\])*" +
@@ -28,7 +30,7 @@ namespace AutomationUtils.Matchers
 		{
 			// Do the match in two phases so we can early out if the strings "error" or "warning" are not present. The patterns before these strings can
 			// produce many false positives, making them very slow to execute.
-			if (input.IsMatch("Error:|Warning:"))
+			if (input.IsMatch(s_errorWarningPattern))
 			{
 				Match? match = s_pattern.Match(input.CurrentLine!);
 				if (match.Success)
diff --git a/Engine/Source/Programs/AutomationTool/AutomationUtils/Matchers/CrashEventMatcher.cs b/Engine/Source/Programs/AutomationTool/AutomationUtils/Matchers/CrashEventMatcher.cs
--- a/Engine/Source/Programs/AutomationTool/AutomationUtils/Matchers/CrashEventMatcher.cs
+++ b/Engine/Source/Programs/AutomationTool/AutomationUtils/Matchers/CrashEventMatcher.cs
@@ -13,23 +13,25 @@ namespace AutomationUtils.Matchers
 	/// </summary>
 	class CrashEventMatcher : ILogEventMatcher
 	{
+		static readonly Regex s_exitCodePattern = new Regex(@"ExitCode=(3|139|255)(?!\d)");
+
 		public LogEventMatch? Match(ILogCursor cursor)
 		{
-			if (cursor.IsMatch("begin: stack for UAT"))
+			if (cursor.Contains("begin: stack for UAT"))
 			{
 				for (int maxOffset = 1; maxOffset < 100; maxOffset++)
 				{
-					if (cursor.IsMatch(maxOffset, "end: stack for UAT"))
+					if (cursor.Contains(maxOffset, "end: stack for UAT"))
 					{
 						LogEventBuilder builder = new LogEventBuilder(cursor, lineCount: maxOffset + 1);
 						return builder.ToMatch(LogEventPriority.BelowNormal, GetLogLevel(cursor), KnownLogEvents.Engine_Crash);
 					}
 				}
 			}
-			if (cursor.IsMatch("AutomationTool: Stack:"))
+			if (cursor.Contains("AutomationTool: Stack:"))
 			{
 				LogEventBuilder builder = new LogEventBuilder(cursor);
-				while (builder.Current.IsMatch(1, "AutomationTool: Stack:"))
+				while (builder.Current.Contains(1, "AutomationTool: Stack:"))
 				{
 					builder.MoveNext();
 				}
@@ -37,7 +39,7 @@ namespace AutomationUtils.Matchers
 			}
 
 			Match? match;
-			if (cursor.TryMatch(@"ExitCode=(3|139|255)(?!\d)", out match))
+			if (cursor.TryMatch(s_exitCodePattern, out match))
 			{
 				LogEventBuilder builder = new LogEventBuilder(cursor);
 				builder.Annotate("exitCode", match.Groups[1]);
@@ -46,13 +48,16 @@ namespace AutomationUtils.Matchers
 			return null;
 		}
 
+		static readonly Regex s_errorPattern = new Regex("[Ee]rror:");
+		static readonly Regex s_warningPattern = new Regex("[Ww]arning:");
+
 		static LogLevel GetLogLevel(ILogCursor cursor)
 		{
-			if(cursor.IsMatch(0, "[Ee]rror:"))
+			if(cursor.IsMatch(0, s_errorPattern))
 			{
 				return LogLevel.Error;
 			}
-			else if(cursor.IsMatch(0, "[Ww]arning:"))
+			else if(cursor.IsMatch(0, s_warningPattern))
 			{
 				return LogLevel.Warning;
 			}
diff --git a/Engine/Source/Programs/AutomationTool/AutomationUtils/Matchers/ExceptionEventMatcher.cs b/Engine/Source/Programs/AutomationTool/AutomationUtils/Matchers/ExceptionEventMatcher.cs
--- a/Engine/Source/Programs/AutomationTool/AutomationUtils/Matchers/ExceptionEventMatcher.cs
+++ b/Engine/Source/Programs/AutomationTool/AutomationUtils/Matchers/ExceptionEventMatcher.cs
@@ -2,6 +2,7 @@
 
 using EpicGames.Core;
 using Microsoft.Extensions.Logging;
+using System.Text.RegularExpressions;
 
 #nullable enable
 
@@ -12,13 +13,16 @@ namespace AutomationUtils.Matchers
 	/// </summary>
 	class ExceptionEventMatcher : ILogEventMatcher
 	{
+		static readonly Regex s_pattern = new Regex(@"^\s*Unhandled Exception: ");
+		static readonly Regex s_atPattern = new Regex(@"^\s*at ");
+
 		/// <inheritdoc/>
 		public LogEventMatch? Match(ILogCursor cursor)
 		{
-			if (cursor.IsMatch(@"^\s*Unhandled Exception: "))
+			if (cursor.IsMatch(s_pattern))
 			{
 				LogEventBuilder builder = new LogEventBuilder(cursor);
-				while(builder.Current.IsMatch(1, @"^\s*at "))
+				while(builder.Current.IsMatch(1, s_atPattern))
 				{
 					builder.MoveNext();
 				}
diff --git a/Engine/Source/Programs/AutomationTool/AutomationUtils/Matchers/ExitCodeEventMatcher.cs b/Engine/Source/Programs/AutomationTool/AutomationUtils/Matchers/ExitCodeEventMatcher.cs
--- a/Engine/Source/Programs/AutomationTool/AutomationUtils/Matchers/ExitCodeEventMatcher.cs
+++ b/Engine/Source/Programs/AutomationTool/AutomationUtils/Matchers/ExitCodeEventMatcher.cs
@@ -2,6 +2,7 @@
 
 using EpicGames.Core;
 using Microsoft.Extensions.Logging;
+using System.Text.RegularExpressions;
 
 #nullable enable
 
@@ -12,24 +13,15 @@ namespace AutomationUtils.Matchers
 	/// </summary>
 	class ExitCodeEventMatcher : ILogEventMatcher
 	{
+		static readonly Regex s_pattern = new Regex(
+			@"Editor terminated with exit code [1-9]|AutomationTool exiting with ExitCode=[1-9]|BUILD FAILED|(Error executing.+)(tool returned code)(.+)");
+
 		public LogEventMatch? Match(ILogCursor cursor)
 		{
 			int numLines = 0;
 			for (; ; )
 			{
-				if (cursor.IsMatch(numLines, "Editor terminated with exit code [1-9]"))
-				{
-					numLines++;
-				}
-				else if (cursor.IsMatch(numLines, "AutomationTool exiting with ExitCode=[1-9]"))
-				{
-					numLines++;
-				}
-				else if (cursor.IsMatch(numLines, "BUILD FAILED"))
-				{
-					numLines++;
-				}
-				else if (cursor.IsMatch(numLines, "(Error executing.+)(tool returned code)(.+)"))
+				if (cursor.IsMatch(numLines, s_pattern))
 				{
 					numLines++;
 				}
diff --git a/Engine/Source/Programs/AutomationTool/AutomationUtils/Matchers/GauntletEventMatcher.cs b/Engine/Source/Programs/AutomationTool/AutomationUtils/Matchers/GauntletEventMatcher.cs
--- a/Engine/Source/Programs/AutomationTool/AutomationUtils/Matchers/GauntletEventMatcher.cs
+++ b/Engine/Source/Programs/AutomationTool/AutomationUtils/Matchers/GauntletEventMatcher.cs
@@ -13,14 +13,24 @@ namespace AutomationUtils.Matchers
 	/// </summary>
 	class GauntletEventMatcher : ILogEventMatcher
 	{
+		static readonly Regex s_engineTestPattern = new Regex(@"^(?<indent>\s*)Error: EngineTest.RunTests Group:(?<group>[^\s]+) \(");
+
+		static readonly Regex s_mdHeadingPattern = new Regex(@"^\s*#{1,3} ");
+		static readonly Regex s_mdTestNamePattern = new Regex(@"^\s*#####\s+(?<friendly_name>.*):\s*(?<name>\S+)\s*");
+		static readonly Regex s_mdTestListPattern = new Regex(@"^\s*### The following tests failed:");
+
+		static readonly Regex s_screenshotPattern = new Regex(@"Error: Screenshot '(?<screenshot>[^']+)' test failed");
+
+		static readonly Regex s_genericPattern = new Regex("\\[ERROR\\] (.*)$");
+
 		public LogEventMatch? Match(ILogCursor cursor)
 		{
 			Match? match;
-			if(cursor.TryMatch(@"^(?<indent>\s*)Error: EngineTest.RunTests Group:(?<group>[^\s]+) \(", out match))
+			if(cursor.TryMatch(s_engineTestPattern, out match))
 			{
 				string indent = match.Groups["indent"].Value + " ";
 
-				LogEventBuilder builder = new LogEventBuilder(cursor);
+				LogEventBuilder builder = new LogEventBuilder(cursor.Hanging());
 				builder.Annotate(match.Groups["group"]);
 
 				string group = match.Groups["group"].Value;
@@ -30,17 +40,17 @@ namespace AutomationUtils.Matchers
 
 				//				int LineCount = 1;
 				EventId eventId = KnownLogEvents.Gauntlet;
-				while (builder.Next.IsMatch($"^(?:{indent}.*|\\s*)$"))
+				while (builder.Next.CurrentLine != null)
 				{
 					builder.MoveNext();
 					if (inErrorList)
 					{
 						Match? testNameMatch;
-						if (builder.Current.IsMatch(@"^\s*#{1,3} "))
+						if (builder.Current.IsMatch(s_mdHeadingPattern))
 						{
 							inErrorList = false;
 						}
-						else if (builder.Current.TryMatch(@"^\s*#####\s+(?<friendly_name>.*):\s*(?<name>\S+)\s*", out testNameMatch))
+						else if (builder.Current.TryMatch(s_mdTestNamePattern, out testNameMatch))
 						{
 							builder.AddProperty("group", group);//.Annotate().AddProperty("group", Group);
 
@@ -52,7 +62,7 @@ namespace AutomationUtils.Matchers
 					}
 					else
 					{
-						if (builder.Current.IsMatch(@"^\s*### The following tests failed:"))
+						if (builder.Current.IsMatch(s_mdTestListPattern))
 						{
 							inErrorList = true;
 						}
@@ -62,14 +72,14 @@ namespace AutomationUtils.Matchers
 				return builder.ToMatch(LogEventPriority.High, LogLevel.Error, eventId);
 			}
 
-			if (cursor.TryMatch(@"Error: Screenshot '(?<screenshot>[^']+)' test failed", out match))
+			if (cursor.TryMatch(s_screenshotPattern, out match))
 			{
 				LogEventBuilder builder = new LogEventBuilder(cursor);
 				builder.Annotate(match.Groups["screenshot"], LogEventMarkup.ScreenshotTest);//.MarkAsScreenshotTest();
 				return builder.ToMatch(LogEventPriority.High, LogLevel.Error, KnownLogEvents.Gauntlet_ScreenshotTest);
 			}
 
-			if (cursor.TryMatch("\\[ERROR\\] (.*)$", out _))
+			if (cursor.TryMatch(s_genericPattern, out _))
 			{
 				LogEventBuilder builder = new LogEventBuilder(cursor);
 				return builder.ToMatch(LogEventPriority.High, LogLevel.Error, KnownLogEvents.Generic);
diff --git a/Engine/Source/Programs/AutomationTool/AutomationUtils/Matchers/GenericEventMatcher.cs b/Engine/Source/Programs/AutomationTool/AutomationUtils/Matchers/GenericEventMatcher.cs
--- a/Engine/Source/Programs/AutomationTool/AutomationUtils/Matchers/GenericEventMatcher.cs
+++ b/Engine/Source/Programs/AutomationTool/AutomationUtils/Matchers/GenericEventMatcher.cs
@@ -14,15 +14,24 @@ namespace AutomationUtils.Matchers
 	/// </summary>
 	class GenericEventMatcher : ILogEventMatcher
 	{
+		static readonly Regex s_fatalPattern = new Regex(
+			@"^\s*(FATAL|fatal error):");
+
+		static readonly Regex s_warningErrorPattern = new Regex(
+			@"(?<!\w)(?i)(WARNING|ERROR) ?(\([^)]+\)|\[[^\]]+\])?: ");
+
+		static readonly Regex s_errorPattern = new Regex(
+			@"[Ee]rror [A-Z]\d+\s:");
+
 		/// <inheritdoc/>
 		public LogEventMatch? Match(ILogCursor cursor)
 		{
 			Match? match;
-			if (cursor.TryMatch(@"^\s*(FATAL|fatal error):", out _))
+			if (cursor.TryMatch(s_fatalPattern, out _))
 			{
 				return new LogEventBuilder(cursor).ToMatch(LogEventPriority.Low, LogLevel.Error, KnownLogEvents.Generic);
 			}
-			if (cursor.TryMatch(@"(?<!\w)(?i)(WARNING|ERROR) ?(\([^)]+\)|\[[^\]]+\])?: ", out match))
+			if (cursor.TryMatch(s_warningErrorPattern, out match))
 			{
 				// Careful to match the first WARNING or ERROR in the line here.
 				LogLevel level = LogLevel.Error;
@@ -31,28 +40,18 @@ namespace AutomationUtils.Matchers
 					level = LogLevel.Warning;
 				}
 
-				LogEventBuilder builder = new LogEventBuilder(cursor);
-				while (builder.Current.IsMatch(1, String.Format(@"^({0} | *$)", ExtractIndent(cursor[0]!))))
+				LogEventBuilder builder = new LogEventBuilder(cursor.Hanging());
+				while (builder.Next.CurrentLine != null)
 				{
 					builder.MoveNext();
 				}
 				return builder.ToMatch(LogEventPriority.Lowest, level, KnownLogEvents.Generic);
 			}
-			if (cursor.IsMatch(@"[Ee]rror [A-Z]\d+\s:"))
+			if (cursor.IsMatch(s_errorPattern))
 			{
 				return new LogEventBuilder(cursor).ToMatch(LogEventPriority.Lowest, LogLevel.Error, KnownLogEvents.Generic);
 			}
 			return null;
 		}
-
-		static string ExtractIndent(string line)
-		{
-			int length = 0;
-			while (length < line.Length && line[length] == ' ')
-			{
-				length++;
-			}
-			return new string(' ', length);
-		}
 	}
 }
diff --git a/Engine/Source/Programs/AutomationTool/AutomationUtils/Matchers/GradleEventMatcher.cs b/Engine/Source/Programs/AutomationTool/AutomationUtils/Matchers/GradleEventMatcher.cs
--- a/Engine/Source/Programs/AutomationTool/AutomationUtils/Matchers/GradleEventMatcher.cs
+++ b/Engine/Source/Programs/AutomationTool/AutomationUtils/Matchers/GradleEventMatcher.cs
@@ -13,27 +13,31 @@ namespace AutomationUtils.Matchers
 	/// </summary>
 	class GradleEventMatcher : ILogEventMatcher
 	{
+		static readonly Regex s_failurePattern = new Regex(@"^\s*FAILURE:");
+		static readonly Regex s_blankPattern = new Regex(@"^\s*$");
+		static readonly Regex s_indentPattern = new Regex(@"^(\s*)\*");
+
 		public LogEventMatch? Match(ILogCursor cursor)
 		{
-			if (cursor.IsMatch(@"^\s*FAILURE:"))
+			if (cursor.IsMatch(s_failurePattern))
 			{
 				int maxOffset = 1;
 				for (; ; )
 				{
 					int newMaxOffset = maxOffset;
-					while (cursor.IsMatch(newMaxOffset, @"^\s*$"))
+					while (cursor.IsMatch(newMaxOffset, s_blankPattern))
 					{
 						newMaxOffset++;
 					}
 
 					Match? match;
-					if (!cursor.TryMatch(newMaxOffset, @"^(\s*)\*", out match))
+					if (!cursor.TryMatch(newMaxOffset, s_indentPattern, out match))
 					{
 						break;
 					}
 					maxOffset = newMaxOffset + 1;
 
-					while (cursor.IsMatch(maxOffset, $"^{match.Groups[1].Value}"))
+					while (cursor.IsMatch(maxOffset, new Regex($"^{match.Groups[1].Value}")))
 					{
 						maxOffset++;
 					}
diff --git a/Engine/Source/Programs/AutomationTool/AutomationUtils/Matchers/LocalizationEventMatcher.cs b/Engine/Source/Programs/AutomationTool/AutomationUtils/Matchers/LocalizationEventMatcher.cs
--- a/Engine/Source/Programs/AutomationTool/AutomationUtils/Matchers/LocalizationEventMatcher.cs
+++ b/Engine/Source/Programs/AutomationTool/AutomationUtils/Matchers/LocalizationEventMatcher.cs
@@ -22,24 +22,23 @@ namespace AutomationUtils.Matchers
 			@"(?:\((?<line>\d+)\))?" +
 			@":");
 
+		static readonly Regex s_indentPattern = new Regex(@"^\s+");
+
 		/// <inheritdoc/>
 		public LogEventMatch? Match(ILogCursor input)
 		{
 			Match? match = s_pattern.Match(input.CurrentLine!);
 			if (match.Success)
 			{
-				LogEventBuilder builder = new LogEventBuilder(input);
+				LogEventBuilder builder = new LogEventBuilder(input.Hanging());
 				builder.Annotate(match.Groups["channel"], LogEventMarkup.Channel);
 				builder.Annotate(match.Groups["severity"], LogEventMarkup.Severity);
 				builder.AnnotateSourceFile(match.Groups["file"], "Engine");
 				builder.Annotate(match.Groups["line"], LogEventMarkup.LineNumber);
 
-				if (builder.Next.TryMatch(@"^\s+", out Match? indent))
+				while (builder.Next.CurrentLine != null)
 				{
-					while (builder.Next.IsMatch($"^{indent.Value}"))
-					{
-						builder.MoveNext();
-					}
+					builder.MoveNext();
 				}
 
 				LogLevel level = match.Groups["severity"].Value switch
diff --git a/Engine/Source/Programs/AutomationTool/AutomationUtils/Matchers/LogChannelEventMatcher.cs b/Engine/Source/Programs/AutomationTool/AutomationUtils/Matchers/LogChannelEventMatcher.cs
--- a/Engine/Source/Programs/AutomationTool/AutomationUtils/Matchers/LogChannelEventMatcher.cs
+++ b/Engine/Source/Programs/AutomationTool/AutomationUtils/Matchers/LogChannelEventMatcher.cs
@@ -20,22 +20,21 @@ namespace AutomationUtils.Matchers
 			@"(?<severity>Error|Warning|Display): "
 		);
 
+		readonly static Regex s_indentPattern = new Regex(@"^\s+");
+
 		/// <inheritdoc/>
 		public LogEventMatch? Match(ILogCursor input)
 		{
 			Match? match = s_pattern.Match(input.CurrentLine!);
 			if (match.Success)
 			{
-				LogEventBuilder builder = new LogEventBuilder(input);
+				LogEventBuilder builder = new LogEventBuilder(input.Hanging());
 				builder.Annotate(match.Groups["channel"], LogEventMarkup.Channel);
 				builder.Annotate(match.Groups["severity"], LogEventMarkup.Severity);
 				
-				if (builder.Next.TryMatch(@"^\s+", out Match? indent))
-				{						
-					while (builder.Next.IsMatch($"^{indent.Value}"))
-					{
-						builder.MoveNext();
-					}					
+				while(builder.Next.CurrentLine != null)
+				{
+					builder.MoveNext();
 				}
 
 				LogLevel level = match.Groups["severity"].Value switch
diff --git a/Engine/Source/Programs/AutomationTool/AutomationUtils/Matchers/SourceFileLineEventMatcher.cs b/Engine/Source/Programs/AutomationTool/AutomationUtils/Matchers/SourceFileLineEventMatcher.cs
--- a/Engine/Source/Programs/AutomationTool/AutomationUtils/Matchers/SourceFileLineEventMatcher.cs
+++ b/Engine/Source/Programs/AutomationTool/AutomationUtils/Matchers/SourceFileLineEventMatcher.cs
@@ -30,11 +30,14 @@ namespace AutomationUtils.Matchers
 		const string LinePattern =
 			@"(?<line>\d+)";
 
+		static readonly Regex s_pattern = new Regex($"^\\s*{SeverityPattern}: {FilePattern}(?:\\({LinePattern}\\))?: ");
+		static readonly Regex s_copyright = new Regex("copyright");
+
 		/// <inheritdoc/>
 		public LogEventMatch? Match(ILogCursor input)
 		{
 			Match? match;
-			if (input.TryMatch($"^\\s*{SeverityPattern}: {FilePattern}(?:\\({LinePattern}\\))?: ", out match))
+			if (input.TryMatch(s_pattern, out match))
 			{
 				LogLevel level = GetLogLevelFromSeverity(match);
 
@@ -45,7 +48,7 @@ namespace AutomationUtils.Matchers
 				builder.TryAnnotate(match.Groups["line"], LogEventMarkup.LineNumber);
 
 				EventId eventId;
-				if (input.IsMatch("copyright"))
+				if (input.IsMatch(s_copyright))
 				{
 					eventId = KnownLogEvents.AutomationTool_MissingCopyright;
 				}
diff --git a/Engine/Source/Programs/AutomationTool/AutomationUtils/Matchers/SystemicEventMatcher.cs b/Engine/Source/Programs/AutomationTool/AutomationUtils/Matchers/SystemicEventMatcher.cs
--- a/Engine/Source/Programs/AutomationTool/AutomationUtils/Matchers/SystemicEventMatcher.cs
+++ b/Engine/Source/Programs/AutomationTool/AutomationUtils/Matchers/SystemicEventMatcher.cs
@@ -2,6 +2,7 @@
 
 using EpicGames.Core;
 using Microsoft.Extensions.Logging;
+using System.Text.RegularExpressions;
 
 #nullable enable
 
@@ -9,22 +10,29 @@ namespace AutomationUtils.Matchers
 {
 	class SystemicEventMatcher : ILogEventMatcher
 	{
+		static readonly Regex s_ddc = new Regex(@"LogDerivedDataCache: .*queries/writes will be limited");
+
+		static readonly Regex s_pdbUtil = new Regex(@"^\s*ERROR: Error: EC_OK");
+		static readonly Regex s_pdbUtilSuffix = new Regex(@"^\s*ERROR:\s*$");
+
+		static readonly Regex s_roboMerge = new Regex(@"RoboMerge\/gates.*already locked on Commit Server by buildmachine");
+
 		public LogEventMatch? Match(ILogCursor cursor)
 		{
-			if (cursor.IsMatch(@"LogDerivedDataCache: .*queries/writes will be limited"))
+			if (cursor.IsMatch(s_ddc))
 			{
 				return new LogEventBuilder(cursor).ToMatch(LogEventPriority.High, LogLevel.Information, KnownLogEvents.Systemic_SlowDDC);
 			}
-			if (cursor.IsMatch(@"^\s*ERROR: Error: EC_OK"))
+			if (cursor.IsMatch(s_pdbUtil))
 			{
 				LogEventBuilder builder = new LogEventBuilder(cursor);
-				while (builder.Next.IsMatch(@"^\s*ERROR:\s*$"))
+				while (builder.Next.IsMatch(s_pdbUtilSuffix))
 				{
 					builder.MoveNext();
 				}
 				return new LogEventBuilder(cursor).ToMatch(LogEventPriority.High, LogLevel.Information, KnownLogEvents.Systemic_PdbUtil);
 			}
-			if (cursor.IsMatch(@"RoboMerge\/gates.*already locked on Commit Server by buildmachine"))
+			if (cursor.IsMatch(s_roboMerge))
 			{
 				return new LogEventBuilder(cursor).ToMatch(LogEventPriority.Low, LogLevel.Information, KnownLogEvents.Systemic_RoboMergeGateLocked);
 			}
diff --git a/Engine/Source/Programs/Horde/Horde.Agent.Tests/LogParserTests.cs b/Engine/Source/Programs/Horde/Horde.Agent.Tests/LogParserTests.cs
--- a/Engine/Source/Programs/Horde/Horde.Agent.Tests/LogParserTests.cs
+++ b/Engine/Source/Programs/Horde/Horde.Agent.Tests/LogParserTests.cs
@@ -778,6 +778,36 @@ namespace Horde.Agent.Tests
 			}
 		}
 
+		[TestMethod]
+		public void HangingIndentMatcher()
+		{
+			string[] lines =
+			{
+				@"first line",
+				@"first line in multi-line message",
+				@"  this is a hanging indent",
+				@"   this is also hanging",
+				@"this is a separate item",
+			};
+
+			LogBuffer buffer = new LogBuffer(10);
+			buffer.AddLines(lines);
+
+			Assert.AreEqual(buffer.Hanging()[0], "first line");
+			Assert.AreEqual(buffer.Hanging()[1], null);
+			buffer.MoveNext();
+
+			Assert.AreEqual(buffer.Hanging()[0], "first line in multi-line message");
+			Assert.AreEqual(buffer.Hanging()[1], "  this is a hanging indent");
+			Assert.AreEqual(buffer.Hanging()[2], "   this is also hanging");
+			Assert.AreEqual(buffer.Hanging()[3], null);
+			buffer.Advance(3);
+
+			Assert.AreEqual(buffer.Hanging()[0], "this is a separate item");
+			Assert.AreEqual(buffer.Hanging()[1], null);
+			buffer.MoveNext();
+		}
+
 		static List<LogEvent> Parse(IEnumerable<string> lines)
 		{
 			return Parse(String.Join("\n", lines));
diff --git a/Engine/Source/Programs/Shared/EpicGames.Core/ILogCursor.cs b/Engine/Source/Programs/Shared/EpicGames.Core/ILogCursor.cs
--- a/Engine/Source/Programs/Shared/EpicGames.Core/ILogCursor.cs
+++ b/Engine/Source/Programs/Shared/EpicGames.Core/ILogCursor.cs
@@ -1,5 +1,6 @@
 // Copyright Epic Games, Inc. All Rights Reserved.
 
+using System;
 using System.Diagnostics.CodeAnalysis;
 using System.Text.RegularExpressions;
 
@@ -61,6 +62,74 @@ namespace EpicGames.Core
 			public int CurrentLineNumber => _baseLineNumber;
 		}
 
+		class HangingIndentLogCursor : ILogCursor
+		{
+			readonly ILogCursor _inner;
+			readonly int _indent;
+
+			public HangingIndentLogCursor(ILogCursor inner)
+			{
+				_inner = inner;
+				_indent = 0;
+
+				string? line = inner.CurrentLine;
+				if (line != null)
+				{
+					while (_indent < line.Length && Char.IsWhiteSpace(line[_indent]))
+					{
+						_indent++;
+					}
+				}
+			}
+
+			public string? this[int index]
+			{
+				get
+				{
+					if (index <= 0)
+					{
+						return _inner[index];
+					}
+
+					string? line = _inner[index];
+					if (line != null && HasPrefix(line, _inner.CurrentLine!, _indent))
+					{
+						return line;
+					}
+
+					return null;
+				}
+			}
+
+			static bool HasPrefix(string line, string prefix, int prefixLength)
+			{
+				if (line.Length > prefixLength && String.CompareOrdinal(line, 0, prefix, 0, prefixLength) == 0 && Char.IsWhiteSpace(line[prefixLength]))
+				{
+					return true;
+				}
+				if(IsWhitespace(line))
+				{
+					return true;
+				}
+				return false;
+			}
+
+			static bool IsWhitespace(string line)
+			{
+				for (int idx = 0; idx < line.Length; idx++)
+				{
+					if (!Char.IsWhiteSpace(line[idx]))
+					{
+						return false;
+					}
+				}
+				return true;
+			}
+
+			public string? CurrentLine => _inner.CurrentLine;
+			public int CurrentLineNumber => _inner.CurrentLineNumber;
+		}
+
 		/// <summary>
 		/// Creates a new log cursor based at an offset from the current line
 		/// </summary>
@@ -72,6 +141,16 @@ namespace EpicGames.Core
 			return new RebasedLogCursor(cursor, cursor.CurrentLineNumber + offset);
 		}
 
+		/// <summary>
+		/// Creates a new log cursor which enumerates lines with a hanging indent
+		/// </summary>
+		/// <param name="cursor">The current log cursor instance</param>
+		/// <returns>New log cursor instance</returns>
+		public static ILogCursor Hanging(this ILogCursor cursor)
+		{
+			return new HangingIndentLogCursor(cursor);
+		}
+
 		/// <summary>
 		/// Attempts to get a line at the given offset
 		/// </summary>
@@ -85,13 +164,55 @@ namespace EpicGames.Core
 			return nextLine != null;
 		}
 
+		/// <summary>
+		/// Tests whether the current line contains the given text
+		/// </summary>
+		/// <param name="cursor">The log cursor instance</param>
+		/// <param name="text">The text to look for</param>
+		/// <returns>True if the line contains the given string</returns>
+		public static bool Contains(this ILogCursor cursor, string text) => Contains(cursor, 0, text);
+
+		/// <summary>
+		/// Tests whether the current line contains the given text
+		/// </summary>
+		/// <param name="cursor">The log cursor instance</param>
+		/// <param name="offset">Offset of the line to match</param>
+		/// <param name="text">The text to look for</param>
+		/// <returns>True if the line contains the given string</returns>
+		public static bool Contains(this ILogCursor cursor, int offset, string text)
+		{
+			string? line;
+			return cursor.TryGetLine(offset, out line) && line.Contains(text, StringComparison.Ordinal);
+		}
+
+		/// <summary>
+		/// Tests whether the current line starts with the given text
+		/// </summary>
+		/// <param name="cursor">The log cursor instance</param>
+		/// <param name="text">The text to look for</param>
+		/// <returns>True if the line starts with the given string</returns>
+		public static bool StartsWith(this ILogCursor cursor, string text) => StartsWith(cursor, 0, text);
+
+		/// <summary>
+		/// Tests whether the current line starts with the given text
+		/// </summary>
+		/// <param name="cursor">The log cursor instance</param>
+		/// <param name="offset">Offset of the line to match</param>
+		/// <param name="text">The text to look for</param>
+		/// <returns>True if the line starts with the given string</returns>
+		public static bool StartsWith(this ILogCursor cursor, int offset, string text)
+		{
+			string? line;
+			return cursor.TryGetLine(offset, out line) && line.StartsWith(text, StringComparison.Ordinal);
+		}
+
 		/// <summary>
 		/// Determines if the current line matches the given regex
 		/// </summary>
 		/// <param name="cursor">The log cursor instance</param>
 		/// <param name="pattern">The regex pattern to match</param>
 		/// <returns>True if the current line matches the given patter</returns>
-		public static bool IsMatch(this ILogCursor cursor, string pattern)
+		public static bool IsMatch(this ILogCursor cursor, Regex pattern)
 		{
 			return IsMatch(cursor, 0, pattern);
 		}
@@ -103,10 +224,10 @@ namespace EpicGames.Core
 		/// <param name="offset">Offset of the line to match</param>
 		/// <param name="pattern">The regex pattern to match</param>
 		/// <returns>True if the requested line matches the given patter</returns>
-		public static bool IsMatch(this ILogCursor cursor, int offset, string pattern)
+		public static bool IsMatch(this ILogCursor cursor, int offset, Regex pattern)
 		{
 			string? line;
-			return cursor.TryGetLine(offset, out line) && Regex.IsMatch(line!, pattern);
+			return cursor.TryGetLine(offset, out line) && pattern.IsMatch(line!);
 		}
 
 		/// <summary>
@@ -116,7 +237,7 @@ namespace EpicGames.Core
 		/// <param name="pattern">The regex pattern to match</param>
 		/// <param name="outMatch">On success, receives the match result</param>
 		/// <returns>True if the current line matches the given patter</returns>
-		public static bool TryMatch(this ILogCursor cursor, string pattern, [NotNullWhen(true)] out Match? outMatch)
+		public static bool TryMatch(this ILogCursor cursor, Regex pattern, [NotNullWhen(true)] out Match? outMatch)
 		{
 			return TryMatch(cursor, 0, pattern, out outMatch);
 		}
@@ -129,7 +250,7 @@ namespace EpicGames.Core
 		/// <param name="pattern">The regex pattern to match</param>
 		/// <param name="outMatch">On success, receives the match result</param>
 		/// <returns>True if the current line matches the given patter</returns>
-		public static bool TryMatch(this ILogCursor cursor, int offset, string pattern, [NotNullWhen(true)] out Match? outMatch)
+		public static bool TryMatch(this ILogCursor cursor, int offset, Regex pattern, [NotNullWhen(true)] out Match? outMatch)
 		{
 			string? line;
 			if (!cursor.TryGetLine(offset, out line))
@@ -138,7 +259,7 @@ namespace EpicGames.Core
 				return false;
 			}
 
-			Match match = Regex.Match(line, pattern);
+			Match match = pattern.Match(line);
 			if (!match.Success)
 			{
 				outMatch = null;
@@ -156,7 +277,7 @@ namespace EpicGames.Core
 		/// <param name="offset">Initial offset</param>
 		/// <param name="pattern">Pattern to match</param>
 		/// <returns>Offset of the last line that still matches the pattern (inclusive)</returns>
-		public static int MatchForwards(this ILogCursor cursor, int offset, string pattern)
+		public static int MatchForwards(this ILogCursor cursor, int offset, Regex pattern)
 		{
 			while (IsMatch(cursor, offset + 1, pattern))
 			{
@@ -172,12 +293,12 @@ namespace EpicGames.Core
 		/// <param name="offset">Initial offset</param>
 		/// <param name="pattern">Pattern to match</param>
 		/// <returns>Offset of the line that matches the pattern (inclusive), or EOF is encountered</returns>
-		public static int MatchForwardsUntil(this ILogCursor cursor, int offset, string pattern)
+		public static int MatchForwardsUntil(this ILogCursor cursor, int offset, Regex pattern)
 		{
 			string? nextLine;
 			for (int nextOffset = offset + 1; cursor.TryGetLine(nextOffset, out nextLine); nextOffset++)
 			{
-				if (Regex.IsMatch(nextLine, pattern))
+				if (pattern.IsMatch(nextLine))
 				{
 					return nextOffset;
 				}
diff --git a/Engine/Source/Programs/Shared/EpicGames.Core/LogBuffer.cs b/Engine/Source/Programs/Shared/EpicGames.Core/LogBuffer.cs
--- a/Engine/Source/Programs/Shared/EpicGames.Core/LogBuffer.cs
+++ b/Engine/Source/Programs/Shared/EpicGames.Core/LogBuffer.cs
@@ -6,9 +6,9 @@ using System.Collections.Generic;
 namespace EpicGames.Core
 {
 	/// <summary>
-	/// Internal implementation of <see cref="ILogCursor"/> used to parse events
+	/// Implementation of <see cref="ILogCursor"/> used to parse events from a stream of lines
 	/// </summary>
-	class LogBuffer : ILogCursor
+	public class LogBuffer : ILogCursor
 	{
 		int _lineNumber;
 		readonly string?[] _history;
@@ -16,6 +16,10 @@ namespace EpicGames.Core
 		int _historyCount;
 		readonly List<string?> _nextLines;
 
+		/// <summary>
+		/// Constructor
+		/// </summary>
+		/// <param name="historySize">Number of previous lines to keep in the buffer</param>
 		public LogBuffer(int historySize)
 		{
 			_lineNumber = 1;
@@ -23,24 +27,50 @@ namespace EpicGames.Core
 			_nextLines = new List<string?>();
 		}
 
+		/// <summary>
+		/// Whether the buffer needs more data
+		/// </summary>
 		public bool NeedMoreData
 		{
 			get;
 			private set;
 		}
 
+		/// <inheritdoc/>
 		public string? CurrentLine => this[0];
 
+		/// <inheritdoc/>
 		public int CurrentLineNumber => _lineNumber;
 
+		/// <summary>
+		/// Number of forward lines in the buffer
+		/// </summary>
 		public int Length => _nextLines.Count;
 
+		/// <summary>
+		/// Add a new line to the buffer
+		/// </summary>
+		/// <param name="line">The new line</param>
 		public void AddLine(string? line)
 		{
 			_nextLines.Add(line);
 			NeedMoreData = false;
 		}
 
+		/// <summary>
+		/// Add a collection of lines to the buffer
+		/// </summary>
+		/// <param name="lines">The new lines</param>
+		public void AddLines(IEnumerable<string?> lines)
+		{
+			_nextLines.AddRange(lines);
+			NeedMoreData = false;
+		}
+
+		/// <summary>
+		/// Move forward a given number of lines
+		/// </summary>
+		/// <param name="count">Number of lines to advance by</param>
 		public void Advance(int count)
 		{
 			for (int idx = 0; idx < count; idx++)
@@ -49,6 +79,9 @@ namespace EpicGames.Core
 			}
 		}
 
+		/// <summary>
+		/// Move to the next line
+		/// </summary>
 		public void MoveNext()
 		{
 			if (_nextLines.Count == 0)
@@ -72,6 +105,7 @@ namespace EpicGames.Core
 			_lineNumber++;
 		}
 
+		/// <inheritdoc/>
 		public string? this[int offset]
 		{
 			get
@@ -104,6 +138,7 @@ namespace EpicGames.Core
 			}
 		}
 
+		/// <inheritdoc/>
 		public override string? ToString()
 		{
 			return CurrentLine;
diff --git a/Engine/Source/Programs/UnrealBuildTool/Matchers/CompileEventMatcher.cs b/Engine/Source/Programs/UnrealBuildTool/Matchers/CompileEventMatcher.cs
--- a/Engine/Source/Programs/UnrealBuildTool/Matchers/CompileEventMatcher.cs
+++ b/Engine/Source/Programs/UnrealBuildTool/Matchers/CompileEventMatcher.cs
@@ -44,6 +44,10 @@ namespace UnrealBuildTool.Matchers
 		const string ClangSeverity =
 			@"(?<severity>error|warning)";
 
+		static readonly Regex s_preludePattern = new Regex("^\\s*(?:In (member )?function|In file included from)");
+		static readonly Regex s_errorWarningPattern = new Regex("error|warning");
+		static readonly Regex s_clangDiagnosticPattern = new Regex($"^\\s*{FilePattern}\\s*{ClangLocationPattern}:\\s*{ClangSeverity}\\s*:");
+
 		static readonly string[] s_invalidExtensions =
 		{
 			".obj",
@@ -58,14 +62,14 @@ namespace UnrealBuildTool.Matchers
 		{
 			// Match the prelude to any error
 			int maxOffset = 0;
-			while (input.IsMatch(maxOffset, "^\\s*(?:In (member )?function|In file included from)"))
+			while (input.IsMatch(maxOffset, s_preludePattern))
 			{
 				maxOffset++;
 			}
 
 			// Do the match in two phases so we can early out if the strings "error" or "warning" are not present. The patterns before these strings can
 			// produce many false positives, making them very slow to execute.
-			if (input.IsMatch(maxOffset, "error|warning"))
+			if (input.IsMatch(maxOffset, s_errorWarningPattern))
 			{
 				LogEventBuilder builder = new LogEventBuilder(input, maxOffset + 1);
 
@@ -103,7 +107,7 @@ namespace UnrealBuildTool.Matchers
 
 				// Try to match a Clang diagnostic
 				Match? match;
-				if (builder.Current.TryMatch($"^\\s*{FilePattern}\\s*{ClangLocationPattern}:\\s*{ClangSeverity}\\s*:", out match) && IsSourceFile(match))
+				if (builder.Current.TryMatch(s_clangDiagnosticPattern, out match) && IsSourceFile(match))
 				{
 					LogLevel level = GetLogLevelFromSeverity(match);
 
@@ -114,7 +118,8 @@ namespace UnrealBuildTool.Matchers
 
 					string indent = ExtractIndent(input[0]!);
 
-					while (builder.Current.TryMatch(1, $"^(?:{indent} |{indent}\\s*{FilePattern}\\s*{ClangLocationPattern}\\s*note:| *$)", out match))
+					Regex notePattern = new Regex($"^(?:{indent} |{indent}\\s*{FilePattern}\\s*{ClangLocationPattern}\\s*note:| *$)");
+					while (builder.Current.TryMatch(1, notePattern, out match))
 					{
 						builder.MoveNext();
 
@@ -132,10 +137,13 @@ namespace UnrealBuildTool.Matchers
 			return null;
 		}
 
+		static readonly Regex s_msvcPattern = new Regex($"^\\s*(?:ERROR: |WARNING: )?{FilePattern}(?:{VisualCppLocationPattern})? ?:\\s+{VisualCppSeverity}:");
+		static readonly Regex s_projectPattern = new Regex(@"\[(?<project>[^[\]]+)]\s*$");
+
 		bool TryMatchVisualCppEvent(LogEventBuilder builder, [NotNullWhen(true)] out LogEventMatch? outEvent)
 		{
 			Match? match;
-			if(!builder.Current.TryMatch($"^\\s*(?:ERROR: |WARNING: )?{FilePattern}(?:{VisualCppLocationPattern})? ?:\\s+{VisualCppSeverity}:", out match) || !IsSourceFile(match))
+			if(!builder.Current.TryMatch(s_msvcPattern, out match) || !IsSourceFile(match))
 			{
 				outEvent = null;
 				return false;
@@ -154,7 +162,7 @@ namespace UnrealBuildTool.Matchers
 				if (codeGroup.Value.StartsWith("CS", StringComparison.Ordinal))
 				{
 					Match? projectMatch;
-					if (builder.Current.TryMatch(@"\[(?<project>[^[\]]+)]\s*$", out projectMatch))
+					if (builder.Current.TryMatch(s_projectPattern, out projectMatch))
 					{
 						builder.AnnotateSourceFile(projectMatch.Groups[1], "");
 						sourceFileBaseDir = GetPlatformAgnosticDirectoryName(projectMatch.Groups[1].Value) ?? sourceFileBaseDir;
@@ -169,7 +177,7 @@ namespace UnrealBuildTool.Matchers
 					}
 
 					Match? projectMatch;
-					if (builder.Current.TryMatch(@"\[(?<file>[^[\]]+)]\s*$", out projectMatch))
+					if (builder.Current.TryMatch(s_projectPattern, out projectMatch))
 					{
 						builder.AnnotateSourceFile(projectMatch.Groups[1], "");
 						outEvent = builder.ToMatch(LogEventPriority.High, level, KnownLogEvents.MSBuild);
@@ -184,7 +192,7 @@ namespace UnrealBuildTool.Matchers
 
 			string indent = ExtractIndent(builder.Current.CurrentLine ?? String.Empty);
 
-			while (builder.Current.TryMatch(1, $"^(?:{indent} |{indent}\\s*{FilePattern}(?:{VisualCppLocationPattern})?\\s*: note:| *$)", out match))
+			while (builder.Current.TryMatch(1, new Regex($"^(?:{indent} |{indent}\\s*{FilePattern}(?:{VisualCppLocationPattern})?\\s*: note:| *$)"), out match))
 			{
 				builder.MoveNext();
 
@@ -197,7 +205,7 @@ namespace UnrealBuildTool.Matchers
 				}
 			}
 
-			string pattern = $"^{indent} |: note:";
+			Regex pattern = new Regex($"^{indent} |: note:");
 			while (builder.Current.IsMatch(1, pattern))
 			{
 				builder.MoveNext();
diff --git a/Engine/Source/Programs/UnrealBuildTool/Matchers/LinkEventMatcher.cs b/Engine/Source/Programs/UnrealBuildTool/Matchers/LinkEventMatcher.cs
--- a/Engine/Source/Programs/UnrealBuildTool/Matchers/LinkEventMatcher.cs
+++ b/Engine/Source/Programs/UnrealBuildTool/Matchers/LinkEventMatcher.cs
@@ -12,16 +12,28 @@ namespace UnrealBuildTool.Matchers
 	/// </summary>
 	class LinkEventMatcher : ILogEventMatcher
 	{
+		static readonly Regex s_undefinedReferencePattern = new Regex(
+			@": undefined reference to |undefined symbol|^\s*(ld|ld.lld|(.*clang)):|^[^:]*[^a-zA-Z][a-z]?ld: |^\s*>>>");
+
+		static readonly Regex s_undefinedSymbolsForArchPattern = new Regex(
+			@"^(\s*)Undefined symbols for architecture");
+
+		static readonly Regex s_ldClangPattern = new Regex(
+			@"^\s*(ld|clang):");
+
+		static readonly Regex s_microsoftErrorPattern = new Regex(
+			@"error (?<code>LNK\d+):");
+
 		/// <inheritdoc/>
 		public LogEventMatch? Match(ILogCursor cursor)
 		{
 			int lineCount = 0;
 			bool isError = false;
 			bool isWarning = false;
-			while (cursor.IsMatch(lineCount, @": undefined reference to |undefined symbol|^\s*(ld|ld.lld|(.*clang)):|^[^:]*[^a-zA-Z][a-z]?ld: |^\s*>>>"))
+			while (cursor.IsMatch(lineCount, s_undefinedReferencePattern))
 			{
-				isError |= cursor.IsMatch("error:");
-				isWarning |= cursor.IsMatch("warning:");
+				isError |= cursor.Contains("error:");
+				isWarning |= cursor.Contains("warning:");
 				lineCount++;
 			}
 			if (lineCount > 0)
@@ -36,15 +48,15 @@ namespace UnrealBuildTool.Matchers
 				}
 				for (; ; )
 				{
-					if (builder.Next.IsMatch("ld:"))
+					if (builder.Next.Contains("ld:"))
 					{
 						break;
 					}
-					else if (builder.Next.IsMatch("error:"))
+					else if (builder.Next.Contains("error:"))
 					{
 						isError = true;
 					}
-					else if (builder.Next.IsMatch("warning:"))
+					else if (builder.Next.Contains("warning:"))
 					{
 						isWarning = true;
 					}
@@ -63,34 +75,34 @@ namespace UnrealBuildTool.Matchers
 			}
 
 			Match? match;
-			if (cursor.TryMatch(@"^(\s*)Undefined symbols for architecture", out match))
+			if (cursor.TryMatch(s_undefinedSymbolsForArchPattern, out match))
 			{
 				LogEventBuilder builder = new LogEventBuilder(cursor);
 				AddSymbolMarkupForLine(builder);
 
 				string prefix = $"^(?<prefix>{match.Groups[1].Value}\\s+)";
-				while (builder.Next.TryMatch(prefix + @"""(?<symbol>[^""]+)""", out match))
+				while (builder.Next.TryMatch(new Regex(prefix + @"""(?<symbol>[^""]+)"""), out match))
 				{
 					string nextPrefix = $"^{match.Groups["prefix"].Value}\\s+";
 
 					builder.MoveNext();
 					builder.AnnotateSymbol(match.Groups["symbol"]);
 
-					while(builder.Next.TryMatch(nextPrefix + "(?<symbol>[^ ].*) in ", out match))
+					while(builder.Next.TryMatch(new Regex(nextPrefix + "(?<symbol>[^ ].*) in "), out match))
 					{
 						builder.MoveNext();
 						builder.AnnotateSymbol(match.Groups["symbol"]);
 					}
 				}
 
-				while (builder.Next.IsMatch(@"^\s*(ld|clang):"))
+				while (builder.Next.IsMatch(s_ldClangPattern))
 				{
 					builder.MoveNext();
 				}
 
 				return builder.ToMatch(LogEventPriority.Normal, LogLevel.Error, KnownLogEvents.Linker_UndefinedSymbol);
 			}
-			if (cursor.TryMatch(@"error (?<code>LNK\d+):", out match))
+			if (cursor.TryMatch(s_microsoftErrorPattern, out match))
 			{
 				LogEventBuilder builder = new LogEventBuilder(cursor);
 				AddSymbolMarkupForLine(builder);
diff --git a/Engine/Source/Programs/UnrealBuildTool/Matchers/MicrosoftEventMatcher.cs b/Engine/Source/Programs/UnrealBuildTool/Matchers/MicrosoftEventMatcher.cs
--- a/Engine/Source/Programs/UnrealBuildTool/Matchers/MicrosoftEventMatcher.cs
+++ b/Engine/Source/Programs/UnrealBuildTool/Matchers/MicrosoftEventMatcher.cs
@@ -11,21 +11,29 @@ namespace UnrealBuildTool.Matchers
 	/// </summary>
 	class MicrosoftEventMatcher : ILogEventMatcher
 	{
+		static readonly Regex s_warningCodePattern = new Regex(
+			@"(?<severity>(?:error|warning)) " +
+			@"(?<code>[a-zA-Z]+[0-9]+)\s*:");
+
+		static readonly Regex s_fileOrToolPattern = new Regex(@"^\s*(.*[^\s])\s*:");
+
+		static readonly Regex s_fileLinePattern = new Regex(@"^\s*(?<file>.*)\((?<line>\d+)(?:, (?<column>\d+))?\)\s*:$");
+
 		public LogEventMatch? Match(ILogCursor cursor)
 		{
 			// filename(line# [, column#]) | toolname} : [ any text ] {error | warning} code+number:localizable string [ any text ]
 
 			Match? match;
-			if (cursor.TryMatch(@"(?<severity>(?:error|warning)) (?<code>[a-zA-Z]+[0-9]+)\s*:", out match))
+			if (cursor.TryMatch(s_warningCodePattern, out match))
 			{
 				string prefix = cursor.CurrentLine!.Substring(0, match.Index);
 
-				Match fileOrToolMatch = Regex.Match(prefix, @"^\s*(.*[^\s])\s*:");
+				Match fileOrToolMatch = s_fileOrToolPattern.Match(prefix);
 				if (fileOrToolMatch.Success)
 				{
 					LogEventBuilder builder = new LogEventBuilder(cursor);
 
-					Match fileMatch = Regex.Match(fileOrToolMatch.Value, @"^\s*(?<file>.*)\((?<line>\d+)(?:, (?<column>\d+))?\)\s*:$");
+					Match fileMatch = s_fileLinePattern.Match(fileOrToolMatch.Value);
 					if (fileMatch.Success)
 					{
 						builder.AnnotateSourceFile(fileMatch.Groups["file"], null);
diff --git a/Engine/Source/Programs/UnrealBuildTool/Matchers/XoreaxEventMatcher.cs b/Engine/Source/Programs/UnrealBuildTool/Matchers/XoreaxEventMatcher.cs
--- a/Engine/Source/Programs/UnrealBuildTool/Matchers/XoreaxEventMatcher.cs
+++ b/Engine/Source/Programs/UnrealBuildTool/Matchers/XoreaxEventMatcher.cs
@@ -8,24 +8,31 @@ namespace UnrealBuildTool.Matchers
 {
 	class XoreaxEventMatcher : ILogEventMatcher
 	{
+		static readonly Regex s_buildService = new Regex(@"\(BuildService.exe\) is not running");
+
+		static readonly Regex s_xgConsole = new Regex(@"BUILD FAILED: (.*)xgConsole\.exe(.*)");
+
+		static readonly Regex s_buildSystemWarning = new Regex(@"^\s*--------------------Build System Warning[- ]");
+		static readonly Regex s_buildSystemWarningNext = new Regex(@"^(\s*)([^ ].*):");
+
 		public LogEventMatch? Match(ILogCursor cursor)
 		{
-			if (cursor.IsMatch(@"\(BuildService.exe\) is not running"))
+			if (cursor.IsMatch(s_buildService))
 			{
 				LogEventBuilder builder = new LogEventBuilder(cursor);
 				return builder.ToMatch(LogEventPriority.High, LogLevel.Information, KnownLogEvents.Systemic_Xge_ServiceNotRunning);
 			}
 
-			if (cursor.IsMatch(@"BUILD FAILED: (.*)xgConsole\.exe(.*)"))
+			if (cursor.IsMatch(s_xgConsole))
 			{
 				LogEventBuilder builder = new LogEventBuilder(cursor);
 				return builder.ToMatch(LogEventPriority.High, LogLevel.Error, KnownLogEvents.Systemic_Xge_BuildFailed);
 			}
 
-			if (cursor.IsMatch(@"^\s*--------------------Build System Warning[- ]"))
+			if (cursor.IsMatch(s_buildSystemWarning))
 			{
 				LogEventBuilder builder = new LogEventBuilder(cursor);
-				if (builder.Next.TryMatch(@"^(\s*)([^ ].*):", out Match? prefix))
+				if (builder.Next.TryMatch(s_buildSystemWarningNext, out Match? prefix))
 				{
 					builder.MoveNext();
 
@@ -37,7 +44,7 @@ namespace UnrealBuildTool.Matchers
 						eventId = KnownLogEvents.Systemic_Xge_Standalone;
 					}
 
-					string prefixPattern = $"^{prefix.Groups[1].Value}\\s";
+					Regex prefixPattern = new Regex($"^{Regex.Escape(prefix.Groups[1].Value)}\\s");
 					while (builder.Next.IsMatch(prefixPattern))
 					{
 						builder.MoveNext();
