Output a JSON log file automatically as part of the core logger.

diff --git a/Engine/Source/Programs/Shared/EpicGames.Core/Log.cs b/Engine/Source/Programs/Shared/EpicGames.Core/Log.cs
--- a/Engine/Source/Programs/Shared/EpicGames.Core/Log.cs
+++ b/Engine/Source/Programs/Shared/EpicGames.Core/Log.cs
@@ -10,6 +10,8 @@ using System.Reflection;
 using System.Text;
 using System.Text.Json;
 using System.Text.RegularExpressions;
+using System.Threading.Channels;
+using System.Threading.Tasks;
 using JetBrains.Annotations;
 using Microsoft.Extensions.DependencyInjection;
 using Microsoft.Extensions.Logging;
@@ -191,6 +193,15 @@ namespace EpicGames.Core
 			LegacyLogger.SetInnerLogger(logger);
 		}
 
+		/// <summary>
+		/// Flush the current log output
+		/// </summary>
+		/// <returns></returns>
+		public static async Task FlushAsync()
+		{
+			await DefaultLogger.FlushAsync();
+		}
+
 		/// <summary>
 		/// Adds a trace listener that writes to a log file.
 		/// If Log.DuplicateLogFiles is true, two files will be created - one with the requested name,
@@ -831,7 +842,7 @@ namespace EpicGames.Core
 	/// <summary>
 	/// Default log output device
 	/// </summary>
-	class DefaultLogger : ILogger
+	class DefaultLogger : ILogger, IDisposable
 	{
 		/// <summary>
 		/// Temporary status message displayed on the console.
@@ -948,6 +959,21 @@ namespace EpicGames.Core
 		/// </summary>
 		private readonly Stopwatch _statusTimer = new Stopwatch();
 
+		/// <summary>
+		/// Background task for writing to files
+		/// </summary>
+		private Task _writeTask;
+
+		/// <summary>
+		/// Channel for new log events
+		/// </summary>
+		private Channel<JsonLogEvent> _eventChannel = Channel.CreateUnbounded<JsonLogEvent>();
+
+		/// <summary>
+		/// Output streams for structured log data
+		/// </summary>
+		private IReadOnlyList<FileStream> _jsonStreams = Array.Empty<FileStream>();
+
 		/// <summary>
 		/// Constructor
 		/// </summary>
@@ -963,6 +989,59 @@ namespace EpicGames.Core
 			{
 				WriteJsonToStdOut = true;
 			}
+
+			_writeTask = Task.Run(() => WriteFilesAsync());
+		}
+
+		/// <inheritdoc/>
+		public void Dispose()
+		{
+			_eventChannel.Writer.TryComplete();
+			_writeTask.Wait();
+		}
+
+		/// <summary>
+		/// Flush the stream
+		/// </summary>
+		/// <returns></returns>
+		public async Task FlushAsync()
+		{
+			lock (_syncObject)
+			{
+				Channel<JsonLogEvent> prevEventChannel = _eventChannel;
+				_eventChannel = Channel.CreateUnbounded<JsonLogEvent>();
+				prevEventChannel.Writer.TryComplete();
+			}
+
+			await _writeTask;
+			_writeTask = Task.Run(() => WriteFilesAsync());
+		}
+
+		/// <summary>
+		/// Background task to write events to sinks
+		/// </summary>
+		async Task WriteFilesAsync()
+		{
+			byte[] newline = new byte[] { (byte)'\n' };
+			while (await _eventChannel.Reader.WaitToReadAsync())
+			{
+				IReadOnlyList<FileStream> streams = _jsonStreams;
+
+				JsonLogEvent logEvent;
+				while (_eventChannel.Reader.TryRead(out logEvent))
+				{
+					foreach (FileStream stream in streams)
+					{
+						await stream.WriteAsync(logEvent.Data);
+						await stream.WriteAsync(newline);
+					}
+				}
+
+				foreach (FileStream stream in streams)
+				{
+					await stream.FlushAsync();
+				}
+			}
 		}
 
 		/// <summary>
@@ -982,6 +1061,10 @@ namespace EpicGames.Core
 				{
 					Trace.Listeners.Add(logTraceListener);
 					WriteInitialTimestamp();
+
+					List<FileStream> newJsonStreams = new List<FileStream>(_jsonStreams);
+					newJsonStreams.Add(FileReference.Open(outputFile.ChangeExtension(".json"), FileMode.Create, FileAccess.Write, FileShare.Read | FileShare.Delete));
+					_jsonStreams = newJsonStreams;
 				}
 				return logTraceListener;
 			}
@@ -1108,9 +1191,11 @@ namespace EpicGames.Core
 					}
 					try
 					{
+						JsonLogEvent jsonLogEvent = JsonLogEvent.FromLoggerState(logLevel, eventId, state, exception, formatter);
+						_eventChannel.Writer.TryWrite(jsonLogEvent);
+
 						if (WriteJsonToStdOut)
 						{
-							JsonLogEvent jsonLogEvent = JsonLogEvent.FromLoggerState(logLevel, eventId, state, exception, formatter);
 							Console.WriteLine(Encoding.UTF8.GetString(jsonLogEvent.Data.Span));
 						}
 						else
