UBT: Additional logging for the event matchers added in UBT.

diff --git a/Engine/Source/Programs/Shared/EpicGames.Core/LogEventParser.cs b/Engine/Source/Programs/Shared/EpicGames.Core/LogEventParser.cs
--- a/Engine/Source/Programs/Shared/EpicGames.Core/LogEventParser.cs
+++ b/Engine/Source/Programs/Shared/EpicGames.Core/LogEventParser.cs
@@ -124,6 +124,7 @@ namespace EpicGames.Core
 			{
 				if (type.IsClass && typeof(ILogEventMatcher).IsAssignableFrom(type))
 				{
+					_logger.LogDebug("Adding event matcher: {Type}", type.Name);
 					ILogEventMatcher matcher = (ILogEventMatcher)Activator.CreateInstance(type)!;
 					Matchers.Add(matcher);
 				}
