Fix native P4 library not being able to write to files with non-ANSI paths.

diff --git a/Engine/Source/Programs/Shared/EpicGames.Perforce.Native/Client.cpp b/Engine/Source/Programs/Shared/EpicGames.Perforce.Native/Client.cpp
--- a/Engine/Source/Programs/Shared/EpicGames.Perforce.Native/Client.cpp
+++ b/Engine/Source/Programs/Shared/EpicGames.Perforce.Native/Client.cpp
@@ -9,6 +9,8 @@ THIRD_PARTY_INCLUDES_START
 #include "p4/clientapi.h"
 #include "p4/strtable.h"
 #include "p4/datetime.h"
+#include "p4/i18napi.h"
+#include "p4/charset.h"
 #include "p4/md5.h"
 #include "ThirdParty/gzip.h"
 #pragma warning(pop)
@@ -702,12 +704,15 @@ extern "C" NATIVE_API FClient* Client_Create(const FSettings* Settings, FWriteBu
 		}
 	}
 	Client->ClientApi.SetProtocol("tag", "");
+	Client->ClientApi.SetProtocol("unicode", "");
 
 	Error Err;
 	Client->ClientApi.Init(&Err);
 	Client->User.HandleError(&Err);
 	Client->User.Flush();
 
+	GlobalCharSet::Set(CharSetApi::UTF_8);
+
 	return Client;
 }
 
