Fix mangled log output in UGS due to data being read into log buffer at incorrect offset.

diff --git a/Engine/Source/Programs/Shared/EpicGames.Core/ManagedProcess.cs b/Engine/Source/Programs/Shared/EpicGames.Core/ManagedProcess.cs
--- a/Engine/Source/Programs/Shared/EpicGames.Core/ManagedProcess.cs
+++ b/Engine/Source/Programs/Shared/EpicGames.Core/ManagedProcess.cs
@@ -1108,7 +1108,7 @@ namespace EpicGames.Core
 				}
 
 				// Read more data into the buffer
-				int numBytesRead = await ReadAsync(_buffer, _bufferPos, _buffer.Length - _bufferPos, cancellationToken);
+				int numBytesRead = await ReadAsync(_buffer, _bufferLen, _buffer.Length - _bufferLen, cancellationToken);
 				if (numBytesRead == 0)
 				{
 					string? line = null;
