Fix blank lines being output when redirecting stdout from UBT/UAT and generating project files.

diff --git a/Engine/Source/Programs/Shared/EpicGames.Core/Log.cs b/Engine/Source/Programs/Shared/EpicGames.Core/Log.cs
--- a/Engine/Source/Programs/Shared/EpicGames.Core/Log.cs
+++ b/Engine/Source/Programs/Shared/EpicGames.Core/Log.cs
@@ -1337,7 +1337,7 @@ namespace EpicGames.Core
 					numCommonChars++;
 				}
 
-				if (!AllowStatusUpdates && numCommonChars < _statusText.Length)
+				if (!AllowStatusUpdates && numCommonChars < _statusText.Length && _statusText.Length > 0)
 				{
 					// Prevent writing backspace characters if the console doesn't support it
 					Console.WriteLine();
