Add a string utility String.StripQuoteArgument to provide an easy way to reverse String.QuoteArgument

#rnx

diff --git a/Engine/Source/Programs/Shared/EpicGames.Core/StringUtils.cs b/Engine/Source/Programs/Shared/EpicGames.Core/StringUtils.cs
--- a/Engine/Source/Programs/Shared/EpicGames.Core/StringUtils.cs
+++ b/Engine/Source/Programs/Shared/EpicGames.Core/StringUtils.cs
@@ -464,6 +464,23 @@ namespace EpicGames.Core
 			}
 		}
 
+		/// <summary>
+		/// Removes the quotes from the beginning and end of a string (if any), can be used to reverse String.QuoteArgument
+		/// </summary>
+		/// <param name="str">The string to remove the quotes from</param>
+		/// <returns>A string without surrounding quotes</returns>
+		public static string StripQuoteArgument(this string str)
+		{
+			if (str.StartsWith('\"') && str.EndsWith('\"'))
+			{
+				return str.Substring(1, str.Length - 2);
+			}
+			else
+			{
+				return str;
+			}
+		}
+
 		/// <summary>
 		/// Formats bytes into a human readable string
 		/// </summary>
