Add stats/checks to log event parser to make sure it's not a bottleneck when parsing output from child processes.

diff --git a/Engine/Source/Programs/Shared/EpicGames.Core/KnownLogEvents.cs b/Engine/Source/Programs/Shared/EpicGames.Core/KnownLogEvents.cs
--- a/Engine/Source/Programs/Shared/EpicGames.Core/KnownLogEvents.cs
+++ b/Engine/Source/Programs/Shared/EpicGames.Core/KnownLogEvents.cs
@@ -195,6 +195,11 @@ namespace EpicGames.Core
 		/// </summary>
 		public static EventId Systemic_MissingXgeControlWorker { get; } = new EventId(761);
 
+		/// <summary>
+		/// Log parser is taking a significant amount of time; may be bottleneck.
+		/// </summary>
+		public static EventId Systemic_LogParserBottleneck { get; } = new EventId(762);
+
 		/// <summary>
 		/// Maximum code for systemic events. Add new events in the 700-799 range.
 		/// </summary>
diff --git a/Engine/Source/Programs/Shared/EpicGames.Core/LogEventParser.cs b/Engine/Source/Programs/Shared/EpicGames.Core/LogEventParser.cs
--- a/Engine/Source/Programs/Shared/EpicGames.Core/LogEventParser.cs
+++ b/Engine/Source/Programs/Shared/EpicGames.Core/LogEventParser.cs
@@ -2,6 +2,7 @@
 
 using System;
 using System.Collections.Generic;
+using System.Diagnostics;
 using System.IO;
 using System.Linq;
 using System.Reflection;
@@ -92,6 +93,21 @@ namespace EpicGames.Core
 		/// </summary>
 		ILogger _logger;
 
+		/// <summary>
+		/// Timer for the parser being active
+		/// </summary>
+		readonly Stopwatch _timer = Stopwatch.StartNew();
+
+		/// <summary>
+		/// Amount of time that the log parser has been processing events
+		/// </summary>
+		readonly Stopwatch _activeTimer = new Stopwatch();
+
+		/// <summary>
+		/// Number of lines parsed in the last interval
+		/// </summary>
+		int _linesParsed = 0;
+
 		/// <summary>
 		/// Public accessor for the logger
 		/// </summary>
@@ -256,6 +272,8 @@ namespace EpicGames.Core
 		/// <param name="bFlush">Whether we've reached the end of the stream</param>
 		void ProcessData(bool bFlush)
 		{
+			_activeTimer.Start();
+			int startLineCount = _buffer.Length;
 			while (_buffer.Length > 0)
 			{
 				// Try to match an event
@@ -305,6 +323,26 @@ namespace EpicGames.Core
 					_buffer.MoveNext();
 				}
 			}
+			_linesParsed += startLineCount - _buffer.Length;
+			_activeTimer.Stop();
+
+			const double UpdateIntervalSeconds = 30.0;
+			double elapsedSeconds = _timer.Elapsed.TotalSeconds;
+			if (elapsedSeconds > UpdateIntervalSeconds)
+			{
+				const double WarnPct = 0.5;
+				double activeSeconds = _activeTimer.Elapsed.TotalSeconds;
+				double activePct = activeSeconds / elapsedSeconds;
+
+				if (activePct > WarnPct)
+				{
+					_logger.LogInformation(KnownLogEvents.Systemic_LogParserBottleneck, "EpicGames.Core.LogEventParser is taking a significant amount of CPU time: {Active:n1}s/{Total:n1}s ({Pct:n1}%). Processed {NumLines} lines in last {Interval} seconds ({NumLinesInBuffer} in buffer).", activeSeconds, elapsedSeconds, activePct * 100.0, _linesParsed, UpdateIntervalSeconds, _buffer.Length);
+				}
+
+				_activeTimer.Reset();
+				_timer.Restart();
+				_linesParsed = 0;
+			}
 		}
 
 		/// <summary>
