Catch and warn about exceptions inside log matchers.

diff --git a/Engine/Source/Programs/Shared/EpicGames.Core/KnownLogEvents.cs b/Engine/Source/Programs/Shared/EpicGames.Core/KnownLogEvents.cs
--- a/Engine/Source/Programs/Shared/EpicGames.Core/KnownLogEvents.cs
+++ b/Engine/Source/Programs/Shared/EpicGames.Core/KnownLogEvents.cs
@@ -200,6 +200,11 @@ namespace EpicGames.Core
 		/// </summary>
 		public static EventId Systemic_LogParserBottleneck { get; } = new EventId(762);
 
+		/// <summary>
+		/// Exception parsing a log event.
+		/// </summary>
+		public static EventId Systemic_LogEventMatcher { get; } = new EventId(763);
+
 		/// <summary>
 		/// Maximum code for systemic events. Add new events in the 700-799 range.
 		/// </summary>
diff --git a/Engine/Source/Programs/Shared/EpicGames.Core/LogEventParser.cs b/Engine/Source/Programs/Shared/EpicGames.Core/LogEventParser.cs
--- a/Engine/Source/Programs/Shared/EpicGames.Core/LogEventParser.cs
+++ b/Engine/Source/Programs/Shared/EpicGames.Core/LogEventParser.cs
@@ -354,7 +354,15 @@ namespace EpicGames.Core
 			LogEventMatch? currentMatch = null;
 			foreach (ILogEventMatcher matcher in Matchers)
 			{
-				LogEventMatch? match = matcher.Match(_buffer);
+				LogEventMatch? match = null;
+				try
+				{
+					match = matcher.Match(_buffer);
+				}
+				catch (Exception ex)
+				{
+					_logger.LogWarning(KnownLogEvents.Systemic_LogEventMatcher, ex, "Exception while parsing log events with {Type}. Buffer size {Length}.", matcher.GetType().Name, _buffer.Length);
+				}
 				if(match != null)
 				{
 					if (currentMatch == null || match.Priority > currentMatch.Priority)
