UBT: Read ignore pattern files when launching UBT/UAT.

diff --git a/Engine/Source/Programs/AutomationTool/AutomationUtils/Automation.cs b/Engine/Source/Programs/AutomationTool/AutomationUtils/Automation.cs
--- a/Engine/Source/Programs/AutomationTool/AutomationUtils/Automation.cs
+++ b/Engine/Source/Programs/AutomationTool/AutomationUtils/Automation.cs
@@ -79,6 +79,9 @@ namespace AutomationTool
 				Assembly UnrealBuildToolAssembly = typeof(UnrealBuildTool.BuildVersion).Assembly;
 				Log.EventParser.AddMatchersFromAssembly(UnrealBuildToolAssembly);
 
+				// Read all the ignore patterns
+				await Log.EventParser.ReadIgnorePatternsAsync(Unreal.EngineDirectory);
+
 				// Get the path to the telemetry file, if present
 				string TelemetryFile = GlobalCommandLine.TelemetryPath;
 				JsonTracer Tracer = JsonTracer.TryRegisterAsGlobalTracer();
diff --git a/Engine/Source/Programs/Shared/EpicGames.Core/LogEventParser.cs b/Engine/Source/Programs/Shared/EpicGames.Core/LogEventParser.cs
--- a/Engine/Source/Programs/Shared/EpicGames.Core/LogEventParser.cs
+++ b/Engine/Source/Programs/Shared/EpicGames.Core/LogEventParser.cs
@@ -8,6 +8,7 @@ using System.Linq;
 using System.Reflection;
 using System.Text;
 using System.Text.RegularExpressions;
+using System.Threading.Tasks;
 using Microsoft.Extensions.Logging;
 
 namespace EpicGames.Core
@@ -147,6 +148,84 @@ namespace EpicGames.Core
 			}
 		}
 
+		/// <summary>
+		/// Read ignore patterns from the given root directory
+		/// </summary>
+		/// <param name="rootDir"></param>
+		/// <returns></returns>
+		public async Task ReadIgnorePatternsAsync(DirectoryReference rootDir)
+		{
+			Stopwatch timer = Stopwatch.StartNew();
+
+			List<DirectoryReference> baseDirs = new List<DirectoryReference>();
+			baseDirs.Add(rootDir);
+			AddRestrictedDirs(baseDirs, "Restricted");
+			AddRestrictedDirs(baseDirs, "Platforms");
+
+			List<(FileReference, Task)> tasks = new List<(FileReference, Task)>();
+			foreach (DirectoryReference baseDir in baseDirs)
+			{
+				FileReference ignorePatternFile = FileReference.Combine(baseDir, "Build", "Horde", "IgnorePatterns.txt");
+				if (FileReference.Exists(ignorePatternFile))
+				{
+					_logger.LogDebug("Reading ignore patterns from {File}...", ignorePatternFile);
+					tasks.Add((ignorePatternFile, ReadIgnorePatternsAsync(ignorePatternFile)));
+				}
+			}
+
+			foreach ((FileReference file, Task task) in tasks)
+			{
+				try
+				{
+					await task;
+				}
+				catch (Exception ex)
+				{
+					_logger.LogError(ex, "Exception reading patterns from {File}: {Message}", file, ex.Message);
+				}
+			}
+
+			_logger.LogDebug("Took {TimeMs}ms to read ignore patterns", timer.ElapsedMilliseconds);
+		}
+
+		/// <summary>
+		/// Read ignore patterns from a single file
+		/// </summary>
+		/// <param name="ignorePatternFile"></param>
+		/// <returns></returns>
+		public async Task ReadIgnorePatternsAsync(FileReference ignorePatternFile)
+		{
+			string[] lines = await FileReference.ReadAllLinesAsync(ignorePatternFile);
+
+			List<Regex> patterns = new List<Regex>();
+			foreach (string line in lines)
+			{
+				string trimLine = line.Trim();
+				if (trimLine.Length > 0 && trimLine[0] != '#')
+				{
+					patterns.Add(new Regex(trimLine));
+				}
+			}
+
+			lock (IgnorePatterns)
+			{
+				IgnorePatterns.AddRange(patterns);
+			}
+		}
+
+		static void AddRestrictedDirs(List<DirectoryReference> directories, string subFolder)
+		{
+			int numDirs = directories.Count;
+			for (int idx = 0; idx < numDirs; idx++)
+			{
+				DirectoryReference subDir = DirectoryReference.Combine(directories[idx], subFolder);
+				if (DirectoryReference.Exists(subDir))
+				{
+					directories.AddRange(DirectoryReference.EnumerateDirectories(subDir));
+				}
+			}
+		}
+
 		/// <summary>
 		/// Writes a line to the event filter
 		/// </summary>
diff --git a/Engine/Source/Programs/UnrealBuildTool/UnrealBuildTool.cs b/Engine/Source/Programs/UnrealBuildTool/UnrealBuildTool.cs
--- a/Engine/Source/Programs/UnrealBuildTool/UnrealBuildTool.cs
+++ b/Engine/Source/Programs/UnrealBuildTool/UnrealBuildTool.cs
@@ -9,6 +9,7 @@ using System.Reflection;
 using System.Text;
 using System.Text.Json;
 using System.Threading;
+using System.Threading.Tasks;
 using EpicGames.Core;
 using Microsoft.Extensions.Logging;
 using OpenTracing.Util;
@@ -450,7 +451,7 @@ namespace UnrealBuildTool
 		/// </summary>
 		/// <param name="ArgumentsArray">Command line arguments</param>
 		/// <returns>Zero on success, non-zero on error</returns>
-		private static int Main(string[] ArgumentsArray)
+		private static async Task<int> Main(string[] ArgumentsArray)
 		{
 			FileReference? RunFile = null;
 			SingleInstanceMutex? Mutex = null;
@@ -505,6 +506,9 @@ namespace UnrealBuildTool
 				// Add all the default event matchers from the UBT assembly
 				Log.EventParser.AddMatchersFromAssembly(Assembly.GetExecutingAssembly());
 
+				// Read all the ignore patterns
+				await Log.EventParser.ReadIgnorePatternsAsync(Unreal.EngineDirectory);
+
 				// Configure the progress writer
 				ProgressWriter.bWriteMarkup = Options.bWriteProgressMarkup;
 
