UGS work-around for hang on shutdown: allow the runtime to terminate background threads.

diff --git a/Engine/Source/Programs/Shared/EpicGames.Perforce/NativePerforceConnection.cs b/Engine/Source/Programs/Shared/EpicGames.Perforce/NativePerforceConnection.cs
--- a/Engine/Source/Programs/Shared/EpicGames.Perforce/NativePerforceConnection.cs
+++ b/Engine/Source/Programs/Shared/EpicGames.Perforce/NativePerforceConnection.cs
@@ -1,4 +1,4 @@
-﻿// Copyright Epic Games, Inc. All Rights Reserved.
+// Copyright Epic Games, Inc. All Rights Reserved.
 
 using System;
 using System.Collections.Concurrent;
@@ -310,6 +310,7 @@ namespace EpicGames.Perforce
 			_responseCompleteEvent = new ManualResetEvent(false);
 
 			_backgroundThread = new Thread(BackgroundThreadProc);
+			_backgroundThread.IsBackground = true;
 			_backgroundThread.Start();
 		}
 
