EpicGames.AspNet: upgrade to net6.0

Used by:
 Horde.Build
 Horde.Storage
 Jupiter.Common

All are already net6.0

#rnx

diff --git a/Engine/Source/Programs/Shared/EpicGames.AspNet/CbConvertersAspnet.cs b/Engine/Source/Programs/Shared/EpicGames.AspNet/CbConvertersAspnet.cs
--- a/Engine/Source/Programs/Shared/EpicGames.AspNet/CbConvertersAspnet.cs
+++ b/Engine/Source/Programs/Shared/EpicGames.AspNet/CbConvertersAspnet.cs
@@ -52,7 +52,7 @@ namespace EpicGames.AspNet
             writer.WriteField(name, ToCbObject(problemDetails).AsField());
         }
 
-        private CbObject ToCbObject(ProblemDetails problemDetails)
+        private static CbObject ToCbObject(ProblemDetails problemDetails)
         {
             CbWriter objectWriter = new CbWriter();
             objectWriter.BeginObject();
diff --git a/Engine/Source/Programs/Shared/EpicGames.AspNet/EpicGames.AspNet.csproj b/Engine/Source/Programs/Shared/EpicGames.AspNet/EpicGames.AspNet.csproj
--- a/Engine/Source/Programs/Shared/EpicGames.AspNet/EpicGames.AspNet.csproj
+++ b/Engine/Source/Programs/Shared/EpicGames.AspNet/EpicGames.AspNet.csproj
@@ -1,6 +1,6 @@
 <Project Sdk="Microsoft.NET.Sdk.Web">
   <PropertyGroup>
-    <TargetFramework>netcoreapp3.1</TargetFramework>
+    <TargetFramework>net6.0</TargetFramework>
     <Configuration Condition=" '$(Configuration)' == '' ">Development</Configuration>
     <OutputType>Library</OutputType>
     <RootNamespace>EpicGames.AspNet</RootNamespace>
diff --git a/Engine/Source/Programs/Shared/EpicGames.AspNet/ServerTimings.cs b/Engine/Source/Programs/Shared/EpicGames.AspNet/ServerTimings.cs
--- a/Engine/Source/Programs/Shared/EpicGames.AspNet/ServerTimings.cs
+++ b/Engine/Source/Programs/Shared/EpicGames.AspNet/ServerTimings.cs
@@ -45,7 +45,7 @@ namespace EpicGames.AspNet
             {
                 sb.Append(";desc=\"");
                 sb.Append(_description);
-                sb.Append("\"");
+                sb.Append('"');
             }
 
             _serverTimingMetric = sb.ToString();
@@ -62,7 +62,7 @@ namespace EpicGames.AspNet
         public IReadOnlyCollection<ServerTimingMetric> Metrics { get; }
     }
 
-    public class ServerTimingMetricScoped : IDisposable
+    public sealed class ServerTimingMetricScoped : IDisposable
     {
         private readonly IServerTiming _timingManager;
         private readonly string _metricName;
