EpicGames.Tracing: net6.0 upgrade

#rnx

diff --git a/Engine/Source/Programs/Shared/EpicGames.Tracing.Tests/EpicGames.Tracing.Tests.csproj b/Engine/Source/Programs/Shared/EpicGames.Tracing.Tests/EpicGames.Tracing.Tests.csproj
--- a/Engine/Source/Programs/Shared/EpicGames.Tracing.Tests/EpicGames.Tracing.Tests.csproj
+++ b/Engine/Source/Programs/Shared/EpicGames.Tracing.Tests/EpicGames.Tracing.Tests.csproj
@@ -1,7 +1,7 @@
 <Project Sdk="Microsoft.NET.Sdk">
 
   <PropertyGroup>
-    <TargetFramework>netcoreapp3.1</TargetFramework>
+    <TargetFramework>net6.0</TargetFramework>
     <Configuration Condition=" '$(Configuration)' == '' ">Development</Configuration>
     <OutputType>Library</OutputType>
     <Deterministic>true</Deterministic>
diff --git a/Engine/Source/Programs/Shared/EpicGames.Tracing.Tests/UnrealInsights/EventTypeTest.cs b/Engine/Source/Programs/Shared/EpicGames.Tracing.Tests/UnrealInsights/EventTypeTest.cs
--- a/Engine/Source/Programs/Shared/EpicGames.Tracing.Tests/UnrealInsights/EventTypeTest.cs
+++ b/Engine/Source/Programs/Shared/EpicGames.Tracing.Tests/UnrealInsights/EventTypeTest.cs
@@ -8,12 +8,13 @@ using Microsoft.VisualStudio.TestTools.UnitTesting;
 namespace EpicGames.Tracing.Tests.UnrealInsights
 {
 	[TestClass]
+	[System.Diagnostics.CodeAnalysis.SuppressMessage("Naming", "CA1707:Identifiers should not contain underscores")]
 	public class EventTypeTest
 	{
-		private string TraceThreadInfoHex = "00 00 4A 00 24 00 04 07 06 0A 00 00 04 00 02 08 04 00 04 00 02 08 08 00 04 00 02 08 0C 00 00 00 88 04 24 54 72 61 63 65 54 68 72 65 61 64 49 6E 66 6F 54 68 72 65 61 64 49 64 53 79 73 74 65 6D 49 64 53 6F 72 74 48 69 6E 74 4E 61 6D 65";
-		private string TraceNewTraceHex = "00 00 55 00 10 00 04 05 06 08 00 00 08 00 03 0A 08 00 08 00 03 0E 10 00 02 00 01 06 12 00 01 00 00 0B 24 54 72 61 63 65 4E 65 77 54 72 61 63 65 53 74 61 72 74 43 79 63 6C 65 43 79 63 6C 65 46 72 65 71 75 65 6E 63 79 45 6E 64 69 61 6E 50 6F 69 6E 74 65 72 53 69 7A 65";
-		private string DiagnosticsSession2Hex = "00 00 9A 00 22 00 08 07 0B 08 00 00 00 00 88 08 00 00 00 00 88 07 00 00 00 00 89 0B 00 00 00 00 89 06 00 00 00 00 89 0C 00 00 04 00 02 0A 04 00 01 00 00 11 05 00 01 00 00 0A 44 69 61 67 6E 6F 73 74 69 63 73 53 65 73 73 69 6F 6E 32 50 6C 61 74 66 6F 72 6D 41 70 70 4E 61 6D 65 43 6F 6D 6D 61 6E 64 4C 69 6E 65 42 72 61 6E 63 68 42 75 69 6C 64 56 65 72 73 69 6F 6E 43 68 61 6E 67 65 6C 69 73 74 43 6F 6E 66 69 67 75 72 61 74 69 6F 6E 54 79 70 65 54 61 72 67 65 74 54 79 70 65";
-		private string LoggingLogMessageHex = "00 00 40 00 1A 00 03 06 07 0A 00 00 08 00 03 08 08 00 08 00 03 05 10 00 00 00 80 0A 4C 6F 67 67 69 6E 67 4C 6F 67 4D 65 73 73 61 67 65 4C 6F 67 50 6F 69 6E 74 43 79 63 6C 65 46 6F 72 6D 61 74 41 72 67 73";
+		private readonly string TraceThreadInfoHex = "00 00 4A 00 24 00 04 07 06 0A 00 00 04 00 02 08 04 00 04 00 02 08 08 00 04 00 02 08 0C 00 00 00 88 04 24 54 72 61 63 65 54 68 72 65 61 64 49 6E 66 6F 54 68 72 65 61 64 49 64 53 79 73 74 65 6D 49 64 53 6F 72 74 48 69 6E 74 4E 61 6D 65";
+		private readonly string TraceNewTraceHex = "00 00 55 00 10 00 04 05 06 08 00 00 08 00 03 0A 08 00 08 00 03 0E 10 00 02 00 01 06 12 00 01 00 00 0B 24 54 72 61 63 65 4E 65 77 54 72 61 63 65 53 74 61 72 74 43 79 63 6C 65 43 79 63 6C 65 46 72 65 71 75 65 6E 63 79 45 6E 64 69 61 6E 50 6F 69 6E 74 65 72 53 69 7A 65";
+		private readonly string DiagnosticsSession2Hex = "00 00 9A 00 22 00 08 07 0B 08 00 00 00 00 88 08 00 00 00 00 88 07 00 00 00 00 89 0B 00 00 00 00 89 06 00 00 00 00 89 0C 00 00 04 00 02 0A 04 00 01 00 00 11 05 00 01 00 00 0A 44 69 61 67 6E 6F 73 74 69 63 73 53 65 73 73 69 6F 6E 32 50 6C 61 74 66 6F 72 6D 41 70 70 4E 61 6D 65 43 6F 6D 6D 61 6E 64 4C 69 6E 65 42 72 61 6E 63 68 42 75 69 6C 64 56 65 72 73 69 6F 6E 43 68 61 6E 67 65 6C 69 73 74 43 6F 6E 66 69 67 75 72 61 74 69 6F 6E 54 79 70 65 54 61 72 67 65 74 54 79 70 65";
+		private readonly string LoggingLogMessageHex = "00 00 40 00 1A 00 03 06 07 0A 00 00 08 00 03 08 08 00 08 00 03 05 10 00 00 00 80 0A 4C 6F 67 67 69 6E 67 4C 6F 67 4D 65 73 73 61 67 65 4C 6F 67 50 6F 69 6E 74 43 79 63 6C 65 46 6F 72 6D 61 74 41 72 67 73";
 		
 		[TestMethod]
 		public void Deserialize_Trace_NewTrace()
diff --git a/Engine/Source/Programs/Shared/EpicGames.Tracing.Tests/UnrealInsights/GenericEventTest.cs b/Engine/Source/Programs/Shared/EpicGames.Tracing.Tests/UnrealInsights/GenericEventTest.cs
--- a/Engine/Source/Programs/Shared/EpicGames.Tracing.Tests/UnrealInsights/GenericEventTest.cs
+++ b/Engine/Source/Programs/Shared/EpicGames.Tracing.Tests/UnrealInsights/GenericEventTest.cs
@@ -9,32 +9,33 @@ using Microsoft.VisualStudio.TestTools.UnitTesting;
 namespace EpicGames.Tracing.Tests.UnrealInsights
 {
 	[TestClass]
+	[System.Diagnostics.CodeAnalysis.SuppressMessage("Naming", "CA1707:Identifiers should not contain underscores")]
 	public class GenericEventTest
 	{
-		private string CpuProfilerEventSpecHex = "21 00 77 00 34 25 00 00 01 C1 0D 00 4B 32 4E 6F 64 65 5F 56 61 72 69 61 62 6C 65 47 65 74 20 2F 45 6E 67 69 6E 65 2F 45 6E 67 69 6E 65 53 6B 79 2F 42 50 5F 53 6B 79 5F 53 70 68 65 72 65 2E 42 50 5F 53 6B 79 5F 53 70 68 65 72 65 3A 55 73 65 72 43 6F 6E 73 74 72 75 63 74 69 6F 6E 53 63 72 69 70 74 2E 4B 32 4E 6F 64 65 5F 56 61 72 69 61 62 6C 65 47 65 74 5F 39 38 33 03";
-		private string CpuProfilerEventBatchHex1 = "4F 00 02 20 1E 00 87 A7 EB AE 8A D1 0B 01 A7 02 02 04 FB 23 03 A4 0A 61 04 0B 05 F2 80 42 AA 1C EB 5C 06 C0 BC 03 17 07 02 47 08 C2 02 AB 15 09 AA 03 07 0A AE 01 05 0B D8 0D 05 0C B7 01 0D D8 03 00 03 0E A4 03 6F 0F 89 FF 03 10 F3 C4 04 11 60 06 A7 02 10 AB 19 12 B3 04 10 9F 11 13 0A 04 02 04 ED 03 10 91 11 14 80 01 02 97 02 10 A9 D6 0B 15 10 06 DD 03 10 CB 85 04 16 12 06 D3 01 10 89 16 17 0C 02 D6 02 07 18 DA 84 28 99 01 19 18 05 1A 0A A5 10 1B CD 01 1C 06 88 6A C5 49 1D B8 EB 24 A9 17 20 AF 01 10 CD F9 01 21 71 22 F2 AE 12 85 04 23 D1 0D 24 85 74 25 EA 58 91 0D 25 80 E6 02 A9 1C 25 88 CD 06 A7 27 25 D0 DA 09 8B 20 25 DE AD 0B AF 33 25 C0 A9 09 ED 23 25 C4 81 0A F5 2D 25 9A EC 0A F3 35 25 D4 CF 03 FF 19 25 BE 92 01 B1 0D 25 98 22 06";
-		// Unused at the moment -- private string CpuProfilerEventBatchHex2 = "4F 00 02 40 1E 00 E4 AD A2 B2 8A D1 0B 17 41 12 0B 41 12 0B 41 04 11 41 2C 0F 41 0E 17 41 14 13 41 10 11 41 26 0F 41 12 B1 01 41 28 0D 41 04 0F 41 18 09 41 04 0F 41 16 0B 41 16 19 41 24 0D 41 14 17 41 0E 0B 41 2A 0B 41 48 31 41 06 15 41 06 5B 41 14 0B 41 0C 25 41 14 25 41 14 25 41 10 15 41 12 2D 41 0E 15 41 04 25 41 10 19 41 0C 17 41 16 17 41 1E 17 41 14 0B 41 0E 2D 41 0E 45 41 04 15 41 14 17 41 14 13 41 18 19 41 06 0D 41 10 11 41 10 0B 41 0E 0B 41 0E 0D 41 10 0B 41 10 0F 41 18 0F 41 12 0B 41 18 19 41 0E 1F 41 16 0D 41 14 1F 41 16 33 41 10 0D 41 24 25 41 18 15 41 04 23 41 06 25 41 04 0B 41 14 0F 41 0E 15 41 12 77 41 12 0D 41 14 23 41 18 0F 41 16 3B 41 18 19 41 20 81 01 41 18 4D 41 16 2B 41 16 1F 41 1C 11 41 10 11 41 10 0B 41 0E 15 41 06";
-		private string LoggingLogMessageHex = "35 00 B4 FD CA E1 FB 7F 00 00 A4 AB DB 52 44 17 00 00 02 A2 03 00 02 C2 44 61 00 71 00 50 00 72 00 6F 00 66 00 2E 00 64 00 6C 00 6C 00 00 00 7E 00 00 00 06";
-		private string MemoryMemoryScopeHex1 = "25 00 2D 00 00 27 00 00 00";
-		private string MemoryMemoryScopeHex2 = "25 00 2E 00 00 27 00 00 00";
-		private string MemoryMemoryScopeHex3 = "25 00 3D 59 00 33 00 00 00";
-		private string TraceNewTraceHex = "10 00 13 00 D9 6E DA 52 44 17 00 00 80 96 98 00 00 00 00 00 4D 52 08";
-		private string DiagnosticsSession2Hex = "22 00 EA 00 DC 26 0A 01 03 04 01 A0 00 00 57 69 6E 36 34 01 81 01 00 55 6E 72 65 61 6C 45 64 69 74 6F 72 01 C2 0F 00 20 00 44 00 3A 00 5C 00 64 00 65 00 70 00 6F 00 74 00 5C 00 73 00 74 00 61 00 72 00 73 00 68 00 69 00 70 00 2D 00 6D 00 61 00 69 00 6E 00 5C 00 51 00 41 00 47 00 61 00 6D 00 65 00 5C 00 51 00 41 00 47 00 61 00 6D 00 65 00 2E 00 75 00 70 00 72 00 6F 00 6A 00 65 00 63 00 74 00 20 00 2D 00 74 00 72 00 61 00 63 00 65 00 3D 00 63 00 70 00 75 00 20 00 2D 00 67 00 61 00 6D 00 65 00 01 83 02 00 2B 00 2B 00 55 00 45 00 35 00 2B 00 4D 00 61 00 69 00 6E 00 01 84 05 00 2B 00 2B 00 55 00 45 00 35 00 2B 00 4D 00 61 00 69 00 6E 00 2D 00 43 00 4C 00 2D 00 31 00 37 00 34 00 34 00 32 00 35 00 32 00 34 00 03";
-		private string TraceThreadInfoHex1 = "24 00 1B 00 02 00 00 00 74 9D 02 00 FF FF FF FF 01 43 01 00 47 61 6D 65 54 68 72 65 61 64 03";
-		private string TraceThreadInfoHex2 = "24 00 16 00 03 00 00 00 00 00 00 00 FF FF FF 7F 01 A3 00 00 54 72 61 63 65 03";
-		private string TraceThreadInfoHex3 = "24 00 25 00 04 00 00 00 D8 BB 06 00 50 00 00 00 01 83 02 00 46 6F 72 65 67 72 6F 75 6E 64 20 57 6F 72 6B 65 72 20 23 30 03";
+		private readonly string CpuProfilerEventSpecHex = "21 00 77 00 34 25 00 00 01 C1 0D 00 4B 32 4E 6F 64 65 5F 56 61 72 69 61 62 6C 65 47 65 74 20 2F 45 6E 67 69 6E 65 2F 45 6E 67 69 6E 65 53 6B 79 2F 42 50 5F 53 6B 79 5F 53 70 68 65 72 65 2E 42 50 5F 53 6B 79 5F 53 70 68 65 72 65 3A 55 73 65 72 43 6F 6E 73 74 72 75 63 74 69 6F 6E 53 63 72 69 70 74 2E 4B 32 4E 6F 64 65 5F 56 61 72 69 61 62 6C 65 47 65 74 5F 39 38 33 03";
+		private readonly string CpuProfilerEventBatchHex1 = "4F 00 02 20 1E 00 87 A7 EB AE 8A D1 0B 01 A7 02 02 04 FB 23 03 A4 0A 61 04 0B 05 F2 80 42 AA 1C EB 5C 06 C0 BC 03 17 07 02 47 08 C2 02 AB 15 09 AA 03 07 0A AE 01 05 0B D8 0D 05 0C B7 01 0D D8 03 00 03 0E A4 03 6F 0F 89 FF 03 10 F3 C4 04 11 60 06 A7 02 10 AB 19 12 B3 04 10 9F 11 13 0A 04 02 04 ED 03 10 91 11 14 80 01 02 97 02 10 A9 D6 0B 15 10 06 DD 03 10 CB 85 04 16 12 06 D3 01 10 89 16 17 0C 02 D6 02 07 18 DA 84 28 99 01 19 18 05 1A 0A A5 10 1B CD 01 1C 06 88 6A C5 49 1D B8 EB 24 A9 17 20 AF 01 10 CD F9 01 21 71 22 F2 AE 12 85 04 23 D1 0D 24 85 74 25 EA 58 91 0D 25 80 E6 02 A9 1C 25 88 CD 06 A7 27 25 D0 DA 09 8B 20 25 DE AD 0B AF 33 25 C0 A9 09 ED 23 25 C4 81 0A F5 2D 25 9A EC 0A F3 35 25 D4 CF 03 FF 19 25 BE 92 01 B1 0D 25 98 22 06";
+		// Unused at the moment -- private readonly string CpuProfilerEventBatchHex2 = "4F 00 02 40 1E 00 E4 AD A2 B2 8A D1 0B 17 41 12 0B 41 12 0B 41 04 11 41 2C 0F 41 0E 17 41 14 13 41 10 11 41 26 0F 41 12 B1 01 41 28 0D 41 04 0F 41 18 09 41 04 0F 41 16 0B 41 16 19 41 24 0D 41 14 17 41 0E 0B 41 2A 0B 41 48 31 41 06 15 41 06 5B 41 14 0B 41 0C 25 41 14 25 41 14 25 41 10 15 41 12 2D 41 0E 15 41 04 25 41 10 19 41 0C 17 41 16 17 41 1E 17 41 14 0B 41 0E 2D 41 0E 45 41 04 15 41 14 17 41 14 13 41 18 19 41 06 0D 41 10 11 41 10 0B 41 0E 0B 41 0E 0D 41 10 0B 41 10 0F 41 18 0F 41 12 0B 41 18 19 41 0E 1F 41 16 0D 41 14 1F 41 16 33 41 10 0D 41 24 25 41 18 15 41 04 23 41 06 25 41 04 0B 41 14 0F 41 0E 15 41 12 77 41 12 0D 41 14 23 41 18 0F 41 16 3B 41 18 19 41 20 81 01 41 18 4D 41 16 2B 41 16 1F 41 1C 11 41 10 11 41 10 0B 41 0E 15 41 06";
+		private readonly string LoggingLogMessageHex = "35 00 B4 FD CA E1 FB 7F 00 00 A4 AB DB 52 44 17 00 00 02 A2 03 00 02 C2 44 61 00 71 00 50 00 72 00 6F 00 66 00 2E 00 64 00 6C 00 6C 00 00 00 7E 00 00 00 06";
+		private readonly string MemoryMemoryScopeHex1 = "25 00 2D 00 00 27 00 00 00";
+		private readonly string MemoryMemoryScopeHex2 = "25 00 2E 00 00 27 00 00 00";
+		private readonly string MemoryMemoryScopeHex3 = "25 00 3D 59 00 33 00 00 00";
+		private readonly string TraceNewTraceHex = "10 00 13 00 D9 6E DA 52 44 17 00 00 80 96 98 00 00 00 00 00 4D 52 08";
+		private readonly string DiagnosticsSession2Hex = "22 00 EA 00 DC 26 0A 01 03 04 01 A0 00 00 57 69 6E 36 34 01 81 01 00 55 6E 72 65 61 6C 45 64 69 74 6F 72 01 C2 0F 00 20 00 44 00 3A 00 5C 00 64 00 65 00 70 00 6F 00 74 00 5C 00 73 00 74 00 61 00 72 00 73 00 68 00 69 00 70 00 2D 00 6D 00 61 00 69 00 6E 00 5C 00 51 00 41 00 47 00 61 00 6D 00 65 00 5C 00 51 00 41 00 47 00 61 00 6D 00 65 00 2E 00 75 00 70 00 72 00 6F 00 6A 00 65 00 63 00 74 00 20 00 2D 00 74 00 72 00 61 00 63 00 65 00 3D 00 63 00 70 00 75 00 20 00 2D 00 67 00 61 00 6D 00 65 00 01 83 02 00 2B 00 2B 00 55 00 45 00 35 00 2B 00 4D 00 61 00 69 00 6E 00 01 84 05 00 2B 00 2B 00 55 00 45 00 35 00 2B 00 4D 00 61 00 69 00 6E 00 2D 00 43 00 4C 00 2D 00 31 00 37 00 34 00 34 00 32 00 35 00 32 00 34 00 03";
+		private readonly string TraceThreadInfoHex1 = "24 00 1B 00 02 00 00 00 74 9D 02 00 FF FF FF FF 01 43 01 00 47 61 6D 65 54 68 72 65 61 64 03";
+		private readonly string TraceThreadInfoHex2 = "24 00 16 00 03 00 00 00 00 00 00 00 FF FF FF 7F 01 A3 00 00 54 72 61 63 65 03";
+		private readonly string TraceThreadInfoHex3 = "24 00 25 00 04 00 00 00 D8 BB 06 00 50 00 00 00 01 83 02 00 46 6F 72 65 67 72 6F 75 6E 64 20 57 6F 72 6B 65 72 20 23 30 03";
 		
 		[TestMethod]
 		public void Deserialize_Memory_MemoryScope()
 		{
 			GenericEvent Event1 = DeserializeAndAssertGenericEvent(MemoryMemoryScopeHex1, MemoryMemoryScopeEvent.EventType, false, 18, 45);
-			Assert.AreEqual(0x27, Event1.Fields[0].Int!.Value);
+			Assert.AreEqual(0x27, Event1.GetFields()[0].Int!.Value);
 			
 			GenericEvent Event2 = DeserializeAndAssertGenericEvent(MemoryMemoryScopeHex2, MemoryMemoryScopeEvent.EventType, false, 18, 46);
-			Assert.AreEqual(0x27, Event2.Fields[0].Int!.Value);
+			Assert.AreEqual(0x27, Event2.GetFields()[0].Int!.Value);
 			
 			GenericEvent Event3 = DeserializeAndAssertGenericEvent(MemoryMemoryScopeHex3, MemoryMemoryScopeEvent.EventType, false, 18, 5832765);
-			Assert.AreEqual(0x33, Event3.Fields[0].Int!.Value);
+			Assert.AreEqual(0x33, Event3.GetFields()[0].Int!.Value);
 		}
 		
 		[TestMethod]
@@ -42,17 +43,17 @@ namespace EpicGames.Tracing.Tests.UnrealInsights
 		{
 			SerializeAndAssertGenericEvent(MemoryMemoryScopeHex1, MemoryMemoryScopeEvent.EventType, false, 18, new []
 			{
-				GenericEvent.Field.FromInt(0x27),
+				Field.FromInt(0x27),
 			}, 45);
 			
 			SerializeAndAssertGenericEvent(MemoryMemoryScopeHex2, MemoryMemoryScopeEvent.EventType, false, 18, new []
 			{
-				GenericEvent.Field.FromInt(0x27),
+				Field.FromInt(0x27),
 			}, 46);
 			
 			SerializeAndAssertGenericEvent(MemoryMemoryScopeHex3, MemoryMemoryScopeEvent.EventType, false, 18, new []
 			{
-				GenericEvent.Field.FromInt(0x33),
+				Field.FromInt(0x33),
 			}, 5832765);
 		}
 		
@@ -60,10 +61,11 @@ namespace EpicGames.Tracing.Tests.UnrealInsights
 		public void Deserialize_Trace_NewTrace()
 		{
 			GenericEvent Event = DeserializeAndAssertGenericEvent(TraceNewTraceHex, TraceNewTraceEvent.EventType, true, 16);
-			Assert.AreEqual(25582215261913, Event.Fields[0].Long!.Value);
-			Assert.AreEqual(10000000, Event.Fields[1].Long!.Value);
-			Assert.AreEqual(21069, Event.Fields[2].Int!.Value);
-			Assert.AreEqual(8, Event.Fields[3].Int!.Value);
+			Field[] Fields = Event.GetFields();
+			Assert.AreEqual(25582215261913, Fields[0].Long!.Value);
+			Assert.AreEqual(10000000, Fields[1].Long!.Value);
+			Assert.AreEqual(21069, Fields[2].Int!.Value);
+			Assert.AreEqual(8, Fields[3].Int!.Value);
 		}
 		
 		[TestMethod]
@@ -71,10 +73,10 @@ namespace EpicGames.Tracing.Tests.UnrealInsights
 		{
 			SerializeAndAssertGenericEvent(TraceNewTraceHex, TraceNewTraceEvent.EventType, true, 16, new []
 			{
-				GenericEvent.Field.FromLong(25582215261913),
-				GenericEvent.Field.FromLong(10000000),
-				GenericEvent.Field.FromInt(21069),
-				GenericEvent.Field.FromInt(8),
+				Field.FromLong(25582215261913),
+				Field.FromLong(10000000),
+				Field.FromInt(21069),
+				Field.FromInt(8),
 			});
 		}
 		
@@ -90,26 +92,29 @@ namespace EpicGames.Tracing.Tests.UnrealInsights
 		{
 			{
 				GenericEvent Event = DeserializeAndAssertGenericEvent(TraceThreadInfoHex1, TraceThreadInfoEvent.EventType, true, 36);
-				Assert.AreEqual(2, Event.Fields[0].Int!.Value);
-				Assert.AreEqual(171380, Event.Fields[1].Int!.Value);
-				Assert.AreEqual(-1, Event.Fields[2].Int!.Value);
-				Assert.AreEqual("GameThread", Event.Fields[3].String!);
+				Field[] Fields = Event.GetFields();
+				Assert.AreEqual(2, Fields[0].Int!.Value);
+				Assert.AreEqual(171380, Fields[1].Int!.Value);
+				Assert.AreEqual(-1, Fields[2].Int!.Value);
+				Assert.AreEqual("GameThread", Fields[3].String!);
 			}
 			
 			{
 				GenericEvent Event = DeserializeAndAssertGenericEvent(TraceThreadInfoHex2, TraceThreadInfoEvent.EventType, true, 36);
-				Assert.AreEqual(3, Event.Fields[0].Int!.Value);
-				Assert.AreEqual(0, Event.Fields[1].Int!.Value);
-				Assert.AreEqual(2147483647, Event.Fields[2].Int!.Value);
-				Assert.AreEqual("Trace", Event.Fields[3].String!);
+				Field[] Fields = Event.GetFields();
+				Assert.AreEqual(3, Fields[0].Int!.Value);
+				Assert.AreEqual(0, Fields[1].Int!.Value);
+				Assert.AreEqual(2147483647, Fields[2].Int!.Value);
+				Assert.AreEqual("Trace", Fields[3].String!);
 			}
 			
 			{
 				GenericEvent Event = DeserializeAndAssertGenericEvent(TraceThreadInfoHex3, TraceThreadInfoEvent.EventType, true, 36);
-				Assert.AreEqual(4, Event.Fields[0].Int!.Value);
-				Assert.AreEqual(441304, Event.Fields[1].Int!.Value);
-				Assert.AreEqual(80, Event.Fields[2].Int!.Value);
-				Assert.AreEqual("Foreground Worker #0", Event.Fields[3].String!);
+				Field[] Fields = Event.GetFields();
+				Assert.AreEqual(4, Fields[0].Int!.Value);
+				Assert.AreEqual(441304, Fields[1].Int!.Value);
+				Assert.AreEqual(80, Fields[2].Int!.Value);
+				Assert.AreEqual("Foreground Worker #0", Fields[3].String!);
 			}
 		}
 		
@@ -136,14 +141,15 @@ namespace EpicGames.Tracing.Tests.UnrealInsights
 		public void Deserialize_Diagnostics_Session2()
 		{
 			GenericEvent Event = DeserializeAndAssertGenericEvent(DiagnosticsSession2Hex, DiagnosticsSession2Event.EventType, true, 34);
-			Assert.AreEqual("Win64", Event.Fields[0].String!);
-			Assert.AreEqual("UnrealEditor", Event.Fields[1].String!);
-			Assert.AreEqual(@" D:\depot\starship-main\QAGame\QAGame.uproject -trace=cpu -game", Event.Fields[2].String!);
-			Assert.AreEqual("++UE5+Main", Event.Fields[3].String!);
-			Assert.AreEqual("++UE5+Main-CL-17442524", Event.Fields[4].String!);
-			Assert.AreEqual(17442524, Event.Fields[5].Int!.Value);
-			Assert.AreEqual(3, Event.Fields[6].Int!.Value);
-			Assert.AreEqual(4, Event.Fields[7].Int!.Value);
+			Field[] Fields = Event.GetFields();
+			Assert.AreEqual("Win64", Fields[0].String!);
+			Assert.AreEqual("UnrealEditor", Fields[1].String!);
+			Assert.AreEqual(@" D:\depot\starship-main\QAGame\QAGame.uproject -trace=cpu -game", Fields[2].String!);
+			Assert.AreEqual("++UE5+Main", Fields[3].String!);
+			Assert.AreEqual("++UE5+Main-CL-17442524", Fields[4].String!);
+			Assert.AreEqual(17442524, Fields[5].Int!.Value);
+			Assert.AreEqual(3, Fields[6].Int!.Value);
+			Assert.AreEqual(4, Fields[7].Int!.Value);
 		}
 		
 		[TestMethod]
@@ -151,14 +157,14 @@ namespace EpicGames.Tracing.Tests.UnrealInsights
 		{
 			SerializeAndAssertGenericEvent(DiagnosticsSession2Hex, DiagnosticsSession2Event.EventType, true, 34, new []
 			{
-				GenericEvent.Field.FromString("Win64"),
-				GenericEvent.Field.FromString("UnrealEditor"),
-				GenericEvent.Field.FromString(@" D:\depot\starship-main\QAGame\QAGame.uproject -trace=cpu -game"),
-				GenericEvent.Field.FromString("++UE5+Main"),
-				GenericEvent.Field.FromString("++UE5+Main-CL-17442524"),
-				GenericEvent.Field.FromInt(17442524),
-				GenericEvent.Field.FromInt(3),
-				GenericEvent.Field.FromInt(4)
+				Field.FromString("Win64"),
+				Field.FromString("UnrealEditor"),
+				Field.FromString(@" D:\depot\starship-main\QAGame\QAGame.uproject -trace=cpu -game"),
+				Field.FromString("++UE5+Main"),
+				Field.FromString("++UE5+Main-CL-17442524"),
+				Field.FromInt(17442524),
+				Field.FromInt(3),
+				Field.FromInt(4)
 			});
 		}
 		
@@ -173,17 +179,17 @@ namespace EpicGames.Tracing.Tests.UnrealInsights
 		public void CpuProfiler_EventSpec_Deserialize()
 		{
 			GenericEvent Event = DeserializeAndAssertGenericEvent(CpuProfilerEventSpecHex, CpuProfilerEventSpecEvent.EventType, true, 33);
-			Assert.AreEqual(9524, Event.Fields[0].Int!.Value);
-			Assert.AreEqual("K2Node_VariableGet /Engine/EngineSky/BP_Sky_Sphere.BP_Sky_Sphere:UserConstructionScript.K2Node_VariableGet_983", Event.Fields[1].String!);
+			Assert.AreEqual(9524, Event.GetFields()[0].Int!.Value);
+			Assert.AreEqual("K2Node_VariableGet /Engine/EngineSky/BP_Sky_Sphere.BP_Sky_Sphere:UserConstructionScript.K2Node_VariableGet_983", Event.GetFields()[1].String!);
 		}
 		
 		[TestMethod]
 		public void CpuProfiler_EventSpec_Serialize_Generic()
 		{
-			SerializeAndAssertGenericEvent(CpuProfilerEventSpecHex, CpuProfilerEventSpecEvent.EventType, true, 33, new GenericEvent.Field[]
+			SerializeAndAssertGenericEvent(CpuProfilerEventSpecHex, CpuProfilerEventSpecEvent.EventType, true, 33, new Field[]
 			{
-				GenericEvent.Field.FromInt(9524),
-				GenericEvent.Field.FromString("K2Node_VariableGet /Engine/EngineSky/BP_Sky_Sphere.BP_Sky_Sphere:UserConstructionScript.K2Node_VariableGet_983")
+				Field.FromInt(9524),
+				Field.FromString("K2Node_VariableGet /Engine/EngineSky/BP_Sky_Sphere.BP_Sky_Sphere:UserConstructionScript.K2Node_VariableGet_983")
 			});
 		}
 		
@@ -198,7 +204,7 @@ namespace EpicGames.Tracing.Tests.UnrealInsights
 		public void CpuProfiler_EventBatch_Deserialize()
 		{
 			GenericEvent Event = DeserializeAndAssertGenericEvent(CpuProfilerEventBatchHex1, CpuProfilerEventBatchEvent.EventType, false, 39);
-			byte[] EventData = Event.Fields[0].Array!;
+			byte[] EventData = Event.GetFields()[0].GetArray()!;
 			Assert.AreEqual(241, EventData.Length);
 			
 			CpuProfilerSerializer Serializer = new CpuProfilerSerializer(10000000);
@@ -246,13 +252,15 @@ namespace EpicGames.Tracing.Tests.UnrealInsights
 		public void Deserialize_Logging_LogMessage()
 		{
 			GenericEvent Event = DeserializeAndAssertGenericEvent(LoggingLogMessageHex, LoggingLogMessage.EventType, false, 26);
-			Assert.AreEqual(140719801695668, Event.Fields[0].Long!.Value);
-			Assert.AreEqual(25582215343012, Event.Fields[1].Long!.Value);
-			Assert.AreEqual(29, Event.Fields[2].Array!.Length);
-			Assert.AreEqual(0x02, Event.Fields[2].Array![0]);
-			Assert.AreEqual(0xC2, Event.Fields[2].Array![1]);
-			Assert.AreEqual(0x44, Event.Fields[2].Array![2]);
-			Assert.AreEqual(0x7E, Event.Fields[2].Array![25]);
+			Field[] Fields = Event.GetFields();
+			Assert.AreEqual(140719801695668, Fields[0].Long!.Value);
+			Assert.AreEqual(25582215343012, Fields[1].Long!.Value);
+			byte[] Bytes = Fields[2].GetArray()!;
+			Assert.AreEqual(29, Bytes.Length);
+			Assert.AreEqual(0x02, Bytes[0]);
+			Assert.AreEqual(0xC2, Bytes[1]);
+			Assert.AreEqual(0x44, Bytes[2]);
+			Assert.AreEqual(0x7E, Bytes[25]);
 		}
 		
 		[TestMethod]
@@ -262,10 +270,10 @@ namespace EpicGames.Tracing.Tests.UnrealInsights
 			using MemoryStream Ms = new MemoryStream();
 			using BinaryWriter Writer = new BinaryWriter(Ms);
 
-			GenericEvent.Field[] Fields =
+			Field[] Fields =
 			{
-				GenericEvent.Field.FromInt(9524),
-				GenericEvent.Field.FromString("K2Node_VariableGet /Engine/EngineSky/BP_Sky_Sphere.BP_Sky_Sphere:UserConstructionScript.K2Node_VariableGet_983"),
+				Field.FromInt(9524),
+				Field.FromString("K2Node_VariableGet /Engine/EngineSky/BP_Sky_Sphere.BP_Sky_Sphere:UserConstructionScript.K2Node_VariableGet_983"),
 			};
 			GenericEvent Event = new GenericEvent(98989, Fields, CpuProfilerEventSpecEvent.EventType);
 			Assert.AreEqual(0x77, Event.Size);
@@ -320,12 +328,12 @@ namespace EpicGames.Tracing.Tests.UnrealInsights
 		
 		internal static byte[] HexStringToBytes(string HexData)
 		{
-			return StringToByteArray(HexData.Replace(" ", ""));
+			return StringToByteArray(HexData.Replace(" ", "", StringComparison.Ordinal));
 		}
 		
 		public static string ToHexString(byte[] ba)
 		{
-			return BitConverter.ToString(ba).Replace("-"," ");
+			return BitConverter.ToString(ba).Replace("-", " ", StringComparison.Ordinal);
 		}
 
 		internal static void AssertHexString(string Expected, string Actual)
@@ -373,11 +381,11 @@ namespace EpicGames.Tracing.Tests.UnrealInsights
 			Reader.EnsureEntireStreamIsConsumed();
 			
 			Assert.AreEqual(ExpectedSerial, Event.Serial);
-			Assert.AreEqual(EventType.Fields.Count, Event.Fields.Length);
+			Assert.AreEqual(EventType.Fields.Count, Event.GetFields().Length);
 			return Event;
 		}
 
-		private static void SerializeAndAssertGenericEvent(string ExpectedHexData, EventType EventType, bool IsImportant, ushort Uid, GenericEvent.Field[] Fields, uint Serial = 0)
+		private static void SerializeAndAssertGenericEvent(string ExpectedHexData, EventType EventType, bool IsImportant, ushort Uid, Field[] Fields, uint Serial = 0)
 		{
 			using MemoryStream Ms = new MemoryStream();
 			using BinaryWriter Writer = new BinaryWriter(Ms);
diff --git a/Engine/Source/Programs/Shared/EpicGames.Tracing.Tests/UnrealInsights/TransportPacketTest.cs b/Engine/Source/Programs/Shared/EpicGames.Tracing.Tests/UnrealInsights/TransportPacketTest.cs
--- a/Engine/Source/Programs/Shared/EpicGames.Tracing.Tests/UnrealInsights/TransportPacketTest.cs
+++ b/Engine/Source/Programs/Shared/EpicGames.Tracing.Tests/UnrealInsights/TransportPacketTest.cs
@@ -88,7 +88,7 @@ namespace EpicGames.Tracing.Tests.UnrealInsights
 		{
 			void Assert7Bit(ulong ExpectedValue, string ActualBytesHex)
 			{
-				byte[] ActualBytes = GenericEventTest.StringToByteArray(ActualBytesHex.Replace(" ", ""));
+				byte[] ActualBytes = GenericEventTest.StringToByteArray(ActualBytesHex.Replace(" ", "", StringComparison.Ordinal));
 				using MemoryStream Ms = new MemoryStream(ActualBytes);
 				using BinaryReader Reader = new BinaryReader(Ms);
 				Assert.AreEqual(ExpectedValue, TraceUtils.Read7BitUint(Reader));
@@ -126,7 +126,7 @@ namespace EpicGames.Tracing.Tests.UnrealInsights
 		{
 			void Assert7Bit(ulong Value, string ExpectedBytesHex)
 			{
-				byte[] ExpectedBytes = GenericEventTest.StringToByteArray(ExpectedBytesHex.Replace(" ", ""));
+				byte[] ExpectedBytes = GenericEventTest.StringToByteArray(ExpectedBytesHex.Replace(" ", "", StringComparison.Ordinal));
 				using MemoryStream Ms = new MemoryStream();
 				using BinaryWriter Writer = new BinaryWriter(Ms);
 
@@ -183,8 +183,8 @@ namespace EpicGames.Tracing.Tests.UnrealInsights
 				using BinaryReader Reader = new BinaryReader(Ms);
 
 				TransportPacket DeserializedPacket = TransportPacket.Deserialize(Reader);
-				Assert.AreEqual(0, DeserializedPacket.ThreadId);
-				Assert.AreEqual(Event1.Size + Event2.Size, DeserializedPacket.Data.Length);
+				Assert.AreEqual(0, DeserializedPacket.ThreadIdAndMarkers);
+				Assert.AreEqual(Event1.Size + Event2.Size, DeserializedPacket.GetData().Length);
 			}
 			
 			{
@@ -197,8 +197,8 @@ namespace EpicGames.Tracing.Tests.UnrealInsights
 				using BinaryReader Reader = new BinaryReader(Ms);
 
 				TransportPacket DeserializedPacket = TransportPacket.Deserialize(Reader);
-				Assert.AreEqual(1, DeserializedPacket.ThreadId);
-				Assert.AreEqual(0, DeserializedPacket.Data.Length);
+				Assert.AreEqual(1, DeserializedPacket.ThreadIdAndMarkers);
+				Assert.AreEqual(0, DeserializedPacket.GetData().Length);
 			}
 			
 			{
@@ -211,8 +211,8 @@ namespace EpicGames.Tracing.Tests.UnrealInsights
 				using BinaryReader Reader = new BinaryReader(Ms);
 
 				TransportPacket DeserializedPacket = TransportPacket.Deserialize(Reader);
-				Assert.AreEqual(12345, DeserializedPacket.ThreadId);
-				Assert.AreEqual(0, DeserializedPacket.Data.Length);
+				Assert.AreEqual(12345, DeserializedPacket.ThreadIdAndMarkers);
+				Assert.AreEqual(0, DeserializedPacket.GetData().Length);
 			}
 		}
 		
@@ -237,8 +237,8 @@ namespace EpicGames.Tracing.Tests.UnrealInsights
 				using BinaryReader Reader = new BinaryReader(Ms);
 
 				TransportPacket DeserializedPacket = TransportPacket.Deserialize(Reader);
-				Assert.AreEqual(0, DeserializedPacket.ThreadId);
-				Assert.AreEqual(Event1.Size + Event2.Size, DeserializedPacket.Data.Length);
+				Assert.AreEqual(0, DeserializedPacket.ThreadIdAndMarkers);
+				Assert.AreEqual(Event1.Size + Event2.Size, DeserializedPacket.GetData().Length);
 				Reader.EnsureEntireStreamIsConsumed();
 			}
 			
@@ -257,8 +257,8 @@ namespace EpicGames.Tracing.Tests.UnrealInsights
 				using BinaryReader Reader = new BinaryReader(Ms);
 
 				TransportPacket DeserializedPacket = TransportPacket.Deserialize(Reader);
-				Assert.AreEqual(0, DeserializedPacket.ThreadId);
-				Assert.AreEqual(Event1.Type.Size + Event2.Type.Size + Event3.Type.Size, DeserializedPacket.Data.Length);
+				Assert.AreEqual(0, DeserializedPacket.ThreadIdAndMarkers);
+				Assert.AreEqual(Event1.Type.Size + Event2.Type.Size + Event3.Type.Size, DeserializedPacket.GetData().Length);
 				Reader.EnsureEntireStreamIsConsumed();
 			}
 		}
diff --git a/Engine/Source/Programs/Shared/EpicGames.Tracing.Tests/UnrealInsights/UnrealInsightsTest.cs b/Engine/Source/Programs/Shared/EpicGames.Tracing.Tests/UnrealInsights/UnrealInsightsTest.cs
--- a/Engine/Source/Programs/Shared/EpicGames.Tracing.Tests/UnrealInsights/UnrealInsightsTest.cs
+++ b/Engine/Source/Programs/Shared/EpicGames.Tracing.Tests/UnrealInsights/UnrealInsightsTest.cs
@@ -13,8 +13,8 @@ namespace EpicGames.Tracing.Tests.UnrealInsights
 {
 	public class StubTraceEvent : ITraceEvent
 	{
-		public readonly uint Val1;
-		public readonly byte Val2;
+		readonly uint Val1;
+		readonly byte Val2;
 
 		public StubTraceEvent(uint Val1, byte Val2)
 		{
@@ -56,7 +56,7 @@ namespace EpicGames.Tracing.Tests.UnrealInsights
 			Writer.Write(Test3);
 			
 			Ms.Position = 0;
-			BinaryReader Reader = new BinaryReader(Ms);
+			using BinaryReader Reader = new BinaryReader(Ms);
 			TraceImportantEventHeader EventHeader = TraceImportantEventHeader.Deserialize(Reader);
 			Assert.AreEqual(1000, EventHeader.Uid);
 			Assert.AreEqual(EventSize, EventHeader.EventSize);
@@ -100,29 +100,33 @@ namespace EpicGames.Tracing.Tests.UnrealInsights
 			
 			{
 				GenericEvent Event = (GenericEvent) Reader.EventsPerThread[1][0];
-				Assert.AreEqual(1001, Event.Fields[0].Long!.Value);
-				Assert.AreEqual(2001, Event.Fields[1].Long!.Value);
-				Assert.AreEqual(31, Event.Fields[2].Int!.Value);
-				Assert.AreEqual(8, Event.Fields[3].Int!.Value);
+				Field[] Fields = Event.GetFields();
+				Assert.AreEqual(1001, Fields[0].Long!.Value);
+				Assert.AreEqual(2001, Fields[1].Long!.Value);
+				Assert.AreEqual(31, Fields[2].Int!.Value);
+				Assert.AreEqual(8, Fields[3].Int!.Value);
 			}
 			{
 				GenericEvent Event = (GenericEvent) Reader.EventsPerThread[1][1];
-				Assert.AreEqual(1002, Event.Fields[0].Long!.Value);
-				Assert.AreEqual(2002, Event.Fields[1].Long!.Value);
-				Assert.AreEqual(32, Event.Fields[2].Int!.Value);
-				Assert.AreEqual(8, Event.Fields[3].Int!.Value);
+				Field[] Fields = Event.GetFields();
+				Assert.AreEqual(1002, Fields[0].Long!.Value);
+				Assert.AreEqual(2002, Fields[1].Long!.Value);
+				Assert.AreEqual(32, Fields[2].Int!.Value);
+				Assert.AreEqual(8, Fields[3].Int!.Value);
 			}
 			{
 				GenericEvent Event = (GenericEvent) Reader.EventsPerThread[1][2];
-				Assert.AreEqual(1003, Event.Fields[0].Long!.Value);
-				Assert.AreEqual(2003, Event.Fields[1].Long!.Value);
-				Assert.AreEqual(33, Event.Fields[2].Int!.Value);
-				Assert.AreEqual(8, Event.Fields[3].Int!.Value);
+				Field[] Fields = Event.GetFields();
+				Assert.AreEqual(1003, Fields[0].Long!.Value);
+				Assert.AreEqual(2003, Fields[1].Long!.Value);
+				Assert.AreEqual(33, Fields[2].Int!.Value);
+				Assert.AreEqual(8, Fields[3].Int!.Value);
 			}
 			{
 				GenericEvent Event = (GenericEvent) Reader.EventsPerThread[1][3];
-				Assert.AreEqual(400, Event.Fields[0].Int!.Value);
-				Assert.AreEqual("MyCpuEventSpecEvent1", Event.Fields[1].String!);
+				Field[] Fields = Event.GetFields();
+				Assert.AreEqual(400, Fields[0].Int!.Value);
+				Assert.AreEqual("MyCpuEventSpecEvent1", Fields[1].String!);
 			}
 		}
 		
@@ -185,7 +189,8 @@ namespace EpicGames.Tracing.Tests.UnrealInsights
 			Fs.Close();
 
 			UnrealInsightsReader Reader = new UnrealInsightsReader();
-			Reader.Read(File.Open(FileName, FileMode.Open));
+			using FileStream Stream = File.Open(FileName, FileMode.Open);
+			Reader.Read(Stream);
 			Reader.PrintEventSummary();
 
 			Assert.AreEqual(0, Reader.EventsPerThread[TransportPacket.ThreadIdEvents].Count);
@@ -197,7 +202,8 @@ namespace EpicGames.Tracing.Tests.UnrealInsights
 			string ExampleDecompTrace = "UnrealInsights/example_trace.decomp.utrace";
 
 			UnrealInsightsReader Reader = new UnrealInsightsReader();
-			Reader.Read(File.Open(ExampleDecompTrace, FileMode.Open));
+			using FileStream Stream = File.Open(ExampleDecompTrace, FileMode.Open);
+			Reader.Read(Stream);
 			Reader.PrintEventSummary();
 			
 			Assert.AreEqual(6281, Reader.NumTransportPacketsRead);
diff --git a/Engine/Source/Programs/Shared/EpicGames.Tracing/EpicGames.Tracing.csproj b/Engine/Source/Programs/Shared/EpicGames.Tracing/EpicGames.Tracing.csproj
--- a/Engine/Source/Programs/Shared/EpicGames.Tracing/EpicGames.Tracing.csproj
+++ b/Engine/Source/Programs/Shared/EpicGames.Tracing/EpicGames.Tracing.csproj
@@ -1,6 +1,6 @@
 <Project Sdk="Microsoft.NET.Sdk">
   <PropertyGroup>
-    <TargetFramework>netcoreapp3.1</TargetFramework>
+    <TargetFramework>net6.0</TargetFramework>
     <Configuration Condition=" '$(Configuration)' == '' ">Development</Configuration>
     <OutputType>Library</OutputType>
     <Deterministic>true</Deterministic>
diff --git a/Engine/Source/Programs/Shared/EpicGames.Tracing/UnrealInsights/Events.cs b/Engine/Source/Programs/Shared/EpicGames.Tracing/UnrealInsights/Events.cs
--- a/Engine/Source/Programs/Shared/EpicGames.Tracing/UnrealInsights/Events.cs
+++ b/Engine/Source/Programs/Shared/EpicGames.Tracing/UnrealInsights/Events.cs
@@ -159,7 +159,7 @@ namespace EpicGames.Tracing.UnrealInsights
 	{
 		public ushort Size => 7;
 		public EventType Type => EventType.WellKnown(PredefinedEventUid.EnterScope_T, "EnterScopeTimestamp");
-		public readonly ulong Timestamp;
+		public ulong Timestamp { get; }
 
 		public EnterScopeEventTimestamp(ulong Timestamp)
 		{
@@ -187,7 +187,7 @@ namespace EpicGames.Tracing.UnrealInsights
 	{
 		public ushort Size => 7;
 		public EventType Type => EventType.WellKnown(PredefinedEventUid.EnterScope_T, "LeaveScopeTimestamp");
-		public readonly ulong Timestamp;
+		readonly ulong Timestamp;
 
 		public LeaveScopeEventTimestamp(ulong Timestamp)
 		{
@@ -213,14 +213,16 @@ namespace EpicGames.Tracing.UnrealInsights
 
 	public class TraceImportantEventHeader
 	{
-		public readonly ushort Uid;
-		public readonly ushort EventSize;
+		public ushort Uid { get; }
+		public ushort EventSize { get; }
 
 		public const ushort HeaderSize = sizeof(ushort) + sizeof(ushort); 
 
 		public TraceImportantEventHeader(ushort Uid, ushort EventSize)
 		{
+#pragma warning disable CA1508 // Avoid dead conditional code
 			Debug.Assert(HeaderSize == 4);
+#pragma warning restore CA1508 // Avoid dead conditional code
 			this.Uid = Uid;
 			this.EventSize = EventSize;
 		}
diff --git a/Engine/Source/Programs/Shared/EpicGames.Tracing/UnrealInsights/Events/CpuProfiler.cs b/Engine/Source/Programs/Shared/EpicGames.Tracing/UnrealInsights/Events/CpuProfiler.cs
--- a/Engine/Source/Programs/Shared/EpicGames.Tracing/UnrealInsights/Events/CpuProfiler.cs
+++ b/Engine/Source/Programs/Shared/EpicGames.Tracing/UnrealInsights/Events/CpuProfiler.cs
@@ -8,13 +8,17 @@ namespace EpicGames.Tracing.UnrealInsights.Events
 {
 	public class CpuProfilerEventSpecEvent : ITraceEvent
 	{
-		public static readonly EventType EventType;
+		public static readonly EventType EventType = new EventType(0, "CpuProfiler", "EventSpec", EventType.FlagImportant | EventType.FlagMaybeHasAux | EventType.FlagNoSync,
+			new List<EventTypeField>() {
+				new EventTypeField(0, 4, EventTypeField.TypeInt32, "Id"),
+				new EventTypeField(4, 0, EventTypeField.TypeAnsiString, "Name")
+			});
 		
 		public ushort Size => (ushort) (GenericEvent.Size + TraceImportantEventHeader.HeaderSize);
 		public EventType Type => EventType;
 
-		public readonly uint Id;
-		public readonly string Name;
+		public uint Id { get; }
+		readonly string Name;
 
 		private readonly GenericEvent GenericEvent;
 
@@ -23,22 +27,15 @@ namespace EpicGames.Tracing.UnrealInsights.Events
 			this.Id = Id;
 			this.Name = Name;
 			
-			GenericEvent.Field[] Fields =
+			Field[] Fields =
 			{
-				GenericEvent.Field.FromInt((int) Id),
-				GenericEvent.Field.FromString(Name),
+				Field.FromInt((int) Id),
+				Field.FromString(Name),
 			};
 			
 			GenericEvent = new GenericEvent(0, Fields, EventType);
 		}
 
-		static CpuProfilerEventSpecEvent()
-		{
-			EventType = new EventType("CpuProfiler", "EventSpec", EventType.FlagImportant | EventType.FlagMaybeHasAux | EventType.FlagNoSync);
-			EventType.AddEventType(0, 4, EventTypeField.TypeInt32, "Id");
-			EventType.AddEventType(4, 0, EventTypeField.TypeAnsiString, "Name");
-		}
-
 		public void Serialize(ushort Uid, BinaryWriter Writer)
 		{
 			new TraceImportantEventHeader(Uid, GenericEvent.Size).Serialize(Writer);
@@ -48,25 +45,20 @@ namespace EpicGames.Tracing.UnrealInsights.Events
 	
 	public class CpuProfilerEventBatchEvent : ITraceEvent
 	{
-		public static readonly EventType EventType;
+		public static readonly EventType EventType = new EventType(0, "CpuProfiler", "EventBatch", EventType.FlagMaybeHasAux | EventType.FlagNoSync,
+			new List<EventTypeField>() { new EventTypeField(0, 0, EventTypeField.TypeArray, "Data") });
 		
 		public ushort Size => (ushort) (GenericEvent.Size + sizeof(ushort));
 		public EventType Type => EventType;
 
-		public readonly GenericEvent GenericEvent;
+		readonly GenericEvent GenericEvent;
 
 		public CpuProfilerEventBatchEvent(byte[] Data)
 		{
-			GenericEvent.Field[] Fields = {GenericEvent.Field.FromArray(Data)};
+			Field[] Fields = {Field.FromArray(Data)};
 			GenericEvent = new GenericEvent(0, Fields, EventType);
 		}
 		
-		static CpuProfilerEventBatchEvent()
-		{
-			EventType = new EventType("CpuProfiler", "EventBatch", EventType.FlagMaybeHasAux | EventType.FlagNoSync);
-			EventType.AddEventType(0, 0, EventTypeField.TypeArray, "Data");
-		}
-
 		public void Serialize(ushort Uid, BinaryWriter Writer)
 		{
 			Writer.WritePackedUid(Uid);
@@ -76,9 +68,9 @@ namespace EpicGames.Tracing.UnrealInsights.Events
 	
 	public class CpuProfilerScopeEvent
 	{
-		public readonly ulong Timestamp;
-		public readonly bool IsEnterScope;
-		public readonly uint? SpecId;
+		public ulong Timestamp { get; }
+		public bool IsEnterScope { get; }
+		public uint? SpecId { get; }
 
 		public CpuProfilerScopeEvent(ulong Timestamp, bool IsEnterScope, uint? SpecId)
 		{
@@ -100,7 +92,7 @@ namespace EpicGames.Tracing.UnrealInsights.Events
 	public class CpuProfilerSerializer
 	{
 		private readonly ulong CycleFrequency;
-		public readonly List<CpuProfilerScopeEvent> ScopeEvents = new List<CpuProfilerScopeEvent>();
+		public List<CpuProfilerScopeEvent> ScopeEvents { get; } = new List<CpuProfilerScopeEvent>();
 
 		public CpuProfilerSerializer(ulong CycleFrequency)
 		{
diff --git a/Engine/Source/Programs/Shared/EpicGames.Tracing/UnrealInsights/Events/Diagnostics.cs b/Engine/Source/Programs/Shared/EpicGames.Tracing/UnrealInsights/Events/Diagnostics.cs
--- a/Engine/Source/Programs/Shared/EpicGames.Tracing/UnrealInsights/Events/Diagnostics.cs
+++ b/Engine/Source/Programs/Shared/EpicGames.Tracing/UnrealInsights/Events/Diagnostics.cs
@@ -1,12 +1,23 @@
 // Copyright Epic Games, Inc. All Rights Reserved.
 
+using System.Collections.Generic;
 using System.IO;
 
 namespace EpicGames.Tracing.UnrealInsights.Events
 {
 	public class DiagnosticsSession2Event : ITraceEvent
 	{
-		public static readonly EventType EventType;
+		public static readonly EventType EventType = new EventType(0, "Diagnostics", "Session2", EventType.FlagImportant | EventType.FlagMaybeHasAux | EventType.FlagNoSync,
+			new List<EventTypeField>() {
+				new EventTypeField(0, 0, EventTypeField.TypeAnsiString, "Platform"),
+				new EventTypeField(0, 0, EventTypeField.TypeAnsiString, "AppName"),
+				new EventTypeField(0, 0, EventTypeField.TypeWideString, "CommandLine"),
+				new EventTypeField(0, 0, EventTypeField.TypeWideString, "Branch"),
+				new EventTypeField(0, 0, EventTypeField.TypeWideString, "BuildVersion"),
+				new EventTypeField(0, 4, EventTypeField.TypeInt32, "Changelist"),
+				new EventTypeField(4, 1, EventTypeField.TypeInt8, "ConfigurationType"),
+				new EventTypeField(5, 1, EventTypeField.TypeInt8, "TargetType"),
+			});
 		
 		public ushort Size => (ushort) (GenericEvent.Size + TraceImportantEventHeader.HeaderSize);
 		public EventType Type => EventType;
@@ -14,34 +25,21 @@ namespace EpicGames.Tracing.UnrealInsights.Events
 
 		public DiagnosticsSession2Event(string Platform, string AppName, string CommandLine, string Branch, string BuildVersion, int ChangeList, int ConfigurationType, int TargetType)
 		{
-			GenericEvent.Field[] Fields =
+			Field[] Fields =
 			{
-				GenericEvent.Field.FromString(Platform),
-				GenericEvent.Field.FromString(AppName),
-				GenericEvent.Field.FromString(CommandLine),
-				GenericEvent.Field.FromString(Branch),
-				GenericEvent.Field.FromString(BuildVersion),
-				GenericEvent.Field.FromInt(ChangeList),
-				GenericEvent.Field.FromInt(ConfigurationType),
-				GenericEvent.Field.FromInt(TargetType),
+				Field.FromString(Platform),
+				Field.FromString(AppName),
+				Field.FromString(CommandLine),
+				Field.FromString(Branch),
+				Field.FromString(BuildVersion),
+				Field.FromInt(ChangeList),
+				Field.FromInt(ConfigurationType),
+				Field.FromInt(TargetType),
 			};
 
 			GenericEvent = new GenericEvent(0, Fields, EventType);
 		}
 		
-		static DiagnosticsSession2Event()
-		{
-			EventType = new EventType("Diagnostics", "Session2", EventType.FlagImportant | EventType.FlagMaybeHasAux | EventType.FlagNoSync);
-			EventType.AddEventType(0, 0, EventTypeField.TypeAnsiString, "Platform");
-			EventType.AddEventType(0, 0, EventTypeField.TypeAnsiString, "AppName");
-			EventType.AddEventType(0, 0, EventTypeField.TypeWideString, "CommandLine");
-			EventType.AddEventType(0, 0, EventTypeField.TypeWideString, "Branch");
-			EventType.AddEventType(0, 0, EventTypeField.TypeWideString, "BuildVersion");
-			EventType.AddEventType(0, 4, EventTypeField.TypeInt32, "Changelist");
-			EventType.AddEventType(4, 1, EventTypeField.TypeInt8, "ConfigurationType");
-			EventType.AddEventType(5, 1, EventTypeField.TypeInt8, "TargetType");
-		}
-
 		public void Serialize(ushort Uid, BinaryWriter Writer)
 		{
 			new TraceImportantEventHeader(Uid, GenericEvent.Size).Serialize(Writer);
diff --git a/Engine/Source/Programs/Shared/EpicGames.Tracing/UnrealInsights/Events/GenericEvent.cs b/Engine/Source/Programs/Shared/EpicGames.Tracing/UnrealInsights/Events/GenericEvent.cs
--- a/Engine/Source/Programs/Shared/EpicGames.Tracing/UnrealInsights/Events/GenericEvent.cs
+++ b/Engine/Source/Programs/Shared/EpicGames.Tracing/UnrealInsights/Events/GenericEvent.cs
@@ -6,38 +6,43 @@ using System.Text;
 
 namespace EpicGames.Tracing.UnrealInsights.Events
 {
+	public class Field
+	{
+		public bool? Bool { get; }
+		public int? Int { get; }
+		public long? Long { get; }
+		public double? Float { get; }
+		public string? String { get; }
+		readonly byte[]? Array;
+		public byte[]? GetArray()
+		{
+			return Array;
+		}
+
+		private Field(bool Value) { Bool = Value; }
+		private Field(int Value) { Int = Value; }
+		private Field(long Value) { Long = Value; }
+		private Field(float Value) { Float = Value; }
+		private Field(double Value) { Float = Value; }
+		private Field(string Value) { String = Value; }
+		private Field(byte[] Value) { Array = Value; }
+		public static Field FromBool(bool Value) => new Field(Value);
+		public static Field FromInt(int Value) => new Field(Value);
+		public static Field FromLong(long Value) => new Field(Value);
+		public static Field FromFloat(float Value) => new Field(Value);
+		public static Field FromDouble(double Value) => new Field(Value);
+		public static Field FromString(string Value) => new Field(Value);
+		public static Field FromArray(byte[] Value) => new Field(Value);
+	}
+
 	/// <summary>
 	/// Represents a generic event in the stream with fields stored dynamically
 	/// </summary>
 	public class GenericEvent : ITraceEvent
 	{
-		public class Field
-		{
-			public readonly bool? Bool;
-			public readonly int? Int;
-			public readonly long? Long;
-			public readonly double? Float;
-			public readonly string? String;
-			public readonly byte[]? Array;
-
-			private Field(bool Value) { Bool = Value; }
-			private Field(int Value) { Int = Value; }
-			private Field(long Value) { Long = Value; }
-			private Field(float Value) { Float = Value; }
-			private Field(double Value) { Float = Value; }
-			private Field(string Value) { String = Value; }
-			private Field(byte[] Value) { Array = Value; }
-			public static Field FromBool(bool Value) => new Field(Value);
-			public static Field FromInt(int Value) => new Field(Value);
-			public static Field FromLong(long Value) => new Field(Value);
-			public static Field FromFloat(float Value) => new Field(Value);
-			public static Field FromDouble(double Value) => new Field(Value);
-			public static Field FromString(string Value) => new Field(Value);
-			public static Field FromArray(byte[] Value) => new Field(Value);
-		}
-
-		public readonly uint Serial;
-		public readonly Field[] Fields;
+		public uint Serial { get; }
+		readonly Field[] Fields;
+		public Field[] GetFields() => Fields;
 		public EventType Type => EventType;
 		private EventType EventType;
 
@@ -74,7 +79,7 @@ namespace EpicGames.Tracing.UnrealInsights.Events
 						}
 						else
 						{
-							AuxDataSize = (ushort) Fields[i].Array!.Length;
+							AuxDataSize = (ushort) Fields[i].GetArray()!.Length;
 						}
 
 						TotalSize += sizeof(uint); // AuxHeader
@@ -141,7 +146,7 @@ namespace EpicGames.Tracing.UnrealInsights.Events
 					}
 					else
 					{
-						Data = Fields[i].Array!;
+						Data = Fields[i].GetArray()!;
 					}
 
 					uint AuxHeader = 0;
diff --git a/Engine/Source/Programs/Shared/EpicGames.Tracing/UnrealInsights/Events/Logging.cs b/Engine/Source/Programs/Shared/EpicGames.Tracing/UnrealInsights/Events/Logging.cs
--- a/Engine/Source/Programs/Shared/EpicGames.Tracing/UnrealInsights/Events/Logging.cs
+++ b/Engine/Source/Programs/Shared/EpicGames.Tracing/UnrealInsights/Events/Logging.cs
@@ -1,25 +1,23 @@
 // Copyright Epic Games, Inc. All Rights Reserved.
 
 using System;
+using System.Collections.Generic;
 using System.IO;
 
 namespace EpicGames.Tracing.UnrealInsights.Events
 {
 	public class LoggingLogMessage : ITraceEvent
 	{
-		public static readonly EventType EventType;
+		public static readonly EventType EventType = new EventType(0, "Logging", "LogMessage", EventType.FlagMaybeHasAux | EventType.FlagNoSync,
+			new List<EventTypeField>() {
+				new EventTypeField(0, 8, EventTypeField.TypeInt64, "LogPoint"),
+				new EventTypeField(8, 8, EventTypeField.TypeInt64, "Cycle"),
+				new EventTypeField(16, 0, EventTypeField.TypeArray, "FormatArgs")
+			});
 
 		public ushort Size => throw new NotImplementedException();
 		public EventType Type => EventType;
 
-		static LoggingLogMessage()
-		{
-			EventType = new EventType("Logging", "LogMessage", EventType.FlagMaybeHasAux | EventType.FlagNoSync);
-			EventType.AddEventType(0, 8, EventTypeField.TypeInt64, "LogPoint");
-			EventType.AddEventType(8, 8, EventTypeField.TypeInt64, "Cycle");
-			EventType.AddEventType(16, 0, EventTypeField.TypeArray, "FormatArgs");
-		}
-		
 		public void Serialize(ushort Uid, BinaryWriter Writer)
 		{
 			throw new NotImplementedException();
diff --git a/Engine/Source/Programs/Shared/EpicGames.Tracing/UnrealInsights/Events/Memory.cs b/Engine/Source/Programs/Shared/EpicGames.Tracing/UnrealInsights/Events/Memory.cs
--- a/Engine/Source/Programs/Shared/EpicGames.Tracing/UnrealInsights/Events/Memory.cs
+++ b/Engine/Source/Programs/Shared/EpicGames.Tracing/UnrealInsights/Events/Memory.cs
@@ -1,23 +1,19 @@
 // Copyright Epic Games, Inc. All Rights Reserved.
 
 using System;
+using System.Collections.Generic;
 using System.IO;
 
 namespace EpicGames.Tracing.UnrealInsights.Events
 {
 	public class MemoryMemoryScopeEvent : ITraceEvent
 	{
-		public static readonly EventType EventType;
-		
+		public static readonly EventType EventType = new EventType(0, "Memory", "MemoryScope", EventType.FlagNone,
+			new List<EventTypeField>() { new EventTypeField(0, 4, EventTypeField.TypeInt32, "Id") });
+
 		public ushort Size => throw new NotImplementedException();
 		public EventType Type => EventType;
 		
-		static MemoryMemoryScopeEvent()
-		{
-			EventType = new EventType("Memory", "MemoryScope", EventType.FlagNone);
-			EventType.AddEventType(0, 4, EventTypeField.TypeInt32, "Id");
-		}
-
 		public void Serialize(ushort Uid, BinaryWriter Writer)
 		{
 			throw new NotImplementedException();
diff --git a/Engine/Source/Programs/Shared/EpicGames.Tracing/UnrealInsights/Events/NewEvent.cs b/Engine/Source/Programs/Shared/EpicGames.Tracing/UnrealInsights/Events/NewEvent.cs
--- a/Engine/Source/Programs/Shared/EpicGames.Tracing/UnrealInsights/Events/NewEvent.cs
+++ b/Engine/Source/Programs/Shared/EpicGames.Tracing/UnrealInsights/Events/NewEvent.cs
@@ -8,6 +8,7 @@ using System.Text;
 
 namespace EpicGames.Tracing.UnrealInsights.Events
 {
+	[System.Diagnostics.CodeAnalysis.SuppressMessage("Performance", "CA1823:Avoid unused private fields")]
 	public class EventTypeField
 	{
 		public ushort Offset { get; private set; }
@@ -16,7 +17,6 @@ namespace EpicGames.Tracing.UnrealInsights.Events
 		public byte NameSize { get; private set; }
 		public string Name { get; private set; } = "<unknown>";
 		public const ushort StructSize = 2 + 2 + 1 + 1;
-
 		private static readonly byte Field_CategoryMask = ToOctalByte("300");
 		private static readonly byte Field_Integer = ToOctalByte("000");
 		private static readonly byte Field_Float = ToOctalByte("100");
diff --git a/Engine/Source/Programs/Shared/EpicGames.Tracing/UnrealInsights/Events/Trace.cs b/Engine/Source/Programs/Shared/EpicGames.Tracing/UnrealInsights/Events/Trace.cs
--- a/Engine/Source/Programs/Shared/EpicGames.Tracing/UnrealInsights/Events/Trace.cs
+++ b/Engine/Source/Programs/Shared/EpicGames.Tracing/UnrealInsights/Events/Trace.cs
@@ -1,12 +1,19 @@
 // Copyright Epic Games, Inc. All Rights Reserved.
 
+using System.Collections.Generic;
 using System.IO;
 
 namespace EpicGames.Tracing.UnrealInsights.Events
 {
 	public class TraceNewTraceEvent : ITraceEvent
 	{
-		public static readonly EventType EventType;
+		public static readonly EventType EventType = new EventType(0, "$Trace", "NewTrace", EventType.FlagImportant | EventType.FlagNoSync,
+			new List<EventTypeField>() {
+				new EventTypeField(0, 8, EventTypeField.TypeInt64, "StartCycle"),
+				new EventTypeField(8, 8, EventTypeField.TypeInt64, "CycleFrequency"),
+				new EventTypeField(16, 2, EventTypeField.TypeInt16, "Endian"),
+				new EventTypeField(18, 1, EventTypeField.TypeInt8, "PointerSize")
+			});
 		
 		public ulong StartCycle { get; private set; }
 		public ulong CycleFrequency { get; private set; }
@@ -24,15 +31,6 @@ namespace EpicGames.Tracing.UnrealInsights.Events
 			this.PointerSize = PointerSize;
 		}
 		
-		static TraceNewTraceEvent()
-		{
-			EventType = new EventType("$Trace", "NewTrace", EventType.FlagImportant | EventType.FlagNoSync);
-			EventType.AddEventType(0, 8, EventTypeField.TypeInt64, "StartCycle");
-			EventType.AddEventType(8, 8, EventTypeField.TypeInt64, "CycleFrequency");
-			EventType.AddEventType(16, 2, EventTypeField.TypeInt16, "Endian");
-			EventType.AddEventType(18, 1, EventTypeField.TypeInt8, "PointerSize");
-		}
-
 		public void Serialize(ushort Uid, BinaryWriter Writer)
 		{
 			new TraceImportantEventHeader(Uid, EventType.GetEventSize()).Serialize(Writer);
@@ -51,10 +49,17 @@ namespace EpicGames.Tracing.UnrealInsights.Events
 			return new TraceNewTraceEvent(StartCycle, CycleFrequency, Endian, PointerSize);
 		}
 	}
-	
+
 	public class TraceThreadInfoEvent : ITraceEvent
 	{
-		public static readonly EventType EventType;
+		public static readonly EventType EventType = new EventType(0, "$Trace", "ThreadInfo", EventType.FlagImportant | EventType.FlagMaybeHasAux | EventType.FlagNoSync,
+			new List<EventTypeField>() {
+				new EventTypeField(0, 4, EventTypeField.TypeInt32, "ThreadId"),
+				new EventTypeField(4, 4, EventTypeField.TypeInt32, "SystemId"),
+				new EventTypeField(8, 4, EventTypeField.TypeInt32, "SortHint"),
+				new EventTypeField(12, 0, EventTypeField.TypeAnsiString, "Name")
+			});
+			
 
 		public ushort Size => (ushort) (GenericEvent.Size + TraceImportantEventHeader.HeaderSize);
 		public EventType Type => EventType;
@@ -62,26 +67,17 @@ namespace EpicGames.Tracing.UnrealInsights.Events
 		
 		public TraceThreadInfoEvent(int ThreadId, int SystemId, int SortHint, string Name)
 		{
-			GenericEvent.Field[] Fields =
+			Field[] Fields =
 			{
-				GenericEvent.Field.FromInt((int) ThreadId),
-				GenericEvent.Field.FromInt((int) SystemId),
-				GenericEvent.Field.FromInt((int) SortHint),
-				GenericEvent.Field.FromString(Name),
+				Field.FromInt((int) ThreadId),
+				Field.FromInt((int) SystemId),
+				Field.FromInt((int) SortHint),
+				Field.FromString(Name),
 			};
 
 			GenericEvent = new GenericEvent(0, Fields, EventType);
 		}
 
-		static TraceThreadInfoEvent()
-		{
-			EventType = new EventType("$Trace", "ThreadInfo", EventType.FlagImportant | EventType.FlagMaybeHasAux | EventType.FlagNoSync);
-			EventType.AddEventType(0, 4, EventTypeField.TypeInt32, "ThreadId");
-			EventType.AddEventType(4, 4, EventTypeField.TypeInt32, "SystemId");
-			EventType.AddEventType(8, 4, EventTypeField.TypeInt32, "SortHint");
-			EventType.AddEventType(12, 0, EventTypeField.TypeAnsiString, "Name");
-		}
-
 		public void Serialize(ushort Uid, BinaryWriter Writer)
 		{
 			new TraceImportantEventHeader(Uid, GenericEvent.Size).Serialize(Writer);
diff --git a/Engine/Source/Programs/Shared/EpicGames.Tracing/UnrealInsights/Reader.cs b/Engine/Source/Programs/Shared/EpicGames.Tracing/UnrealInsights/Reader.cs
--- a/Engine/Source/Programs/Shared/EpicGames.Tracing/UnrealInsights/Reader.cs
+++ b/Engine/Source/Programs/Shared/EpicGames.Tracing/UnrealInsights/Reader.cs
@@ -8,7 +8,8 @@ using EpicGames.Tracing.UnrealInsights.Events;
 
 namespace EpicGames.Tracing.UnrealInsights
 {
-	public class PredefinedEventUid
+	[System.Diagnostics.CodeAnalysis.SuppressMessage("Naming", "CA1707:Identifiers should not contain underscores")]
+	public static class PredefinedEventUid
 	{
 		public const ushort NewEvent = 0;
 		public const ushort AuxData = 1;
@@ -31,8 +32,8 @@ namespace EpicGames.Tracing.UnrealInsights
 	
 	public class UnrealInsightsReader
 	{
-		public readonly Dictionary<ushort, EventType> EventTypes = new Dictionary<ushort, EventType>();
-		public readonly Dictionary<ushort, List<ITraceEvent>> EventsPerThread = new Dictionary<ushort, List<ITraceEvent>>();
+		public Dictionary<ushort, EventType> EventTypes { get; } = new Dictionary<ushort, EventType>();
+		public Dictionary<ushort, List<ITraceEvent>> EventsPerThread { get; } = new Dictionary<ushort, List<ITraceEvent>>();
 
 		internal int NumTransportPacketsRead { get; private set; } = 0;
 
@@ -93,7 +94,7 @@ namespace EpicGames.Tracing.UnrealInsights
 					ThreadStreams[Packet.GetThreadId()] = ThreadStream;
 				}
 
-				ThreadStream.Write(Packet.Data);
+				ThreadStream.Write(Packet.GetData());
 				NumTransportPacketsRead++;
 			}
 			Reader.EnsureEntireStreamIsConsumed();
diff --git a/Engine/Source/Programs/Shared/EpicGames.Tracing/UnrealInsights/StreamHeader.cs b/Engine/Source/Programs/Shared/EpicGames.Tracing/UnrealInsights/StreamHeader.cs
--- a/Engine/Source/Programs/Shared/EpicGames.Tracing/UnrealInsights/StreamHeader.cs
+++ b/Engine/Source/Programs/Shared/EpicGames.Tracing/UnrealInsights/StreamHeader.cs
@@ -10,13 +10,13 @@ namespace EpicGames.Tracing.UnrealInsights
 		public const uint MagicTrc = 1414677317; // TRCE
 		public const uint MagicTrc2 = 1414677298; // TRC2
 
-		public uint MagicFourCc;
-		public ushort MetadataSize;
-		public ushort MetadataField0;
-		public ushort ControlPort;
+		uint MagicFourCc;
+		ushort MetadataSize;
+		ushort MetadataField0;
+		ushort ControlPort;
 
-		public byte TransportVersion;
-		public byte ProtocolVersion;
+		byte TransportVersion;
+		byte ProtocolVersion;
 
 		private StreamHeader()
 		{
diff --git a/Engine/Source/Programs/Shared/EpicGames.Tracing/UnrealInsights/TraceUtils.cs b/Engine/Source/Programs/Shared/EpicGames.Tracing/UnrealInsights/TraceUtils.cs
--- a/Engine/Source/Programs/Shared/EpicGames.Tracing/UnrealInsights/TraceUtils.cs
+++ b/Engine/Source/Programs/Shared/EpicGames.Tracing/UnrealInsights/TraceUtils.cs
@@ -4,7 +4,7 @@ using System.IO;
 
 namespace EpicGames.Tracing.UnrealInsights
 {
-	public class TraceUtils
+	public static class TraceUtils
 	{
 		public static ulong Read7BitUint(BinaryReader Reader)
 		{
diff --git a/Engine/Source/Programs/Shared/EpicGames.Tracing/UnrealInsights/TransportPacket.cs b/Engine/Source/Programs/Shared/EpicGames.Tracing/UnrealInsights/TransportPacket.cs
--- a/Engine/Source/Programs/Shared/EpicGames.Tracing/UnrealInsights/TransportPacket.cs
+++ b/Engine/Source/Programs/Shared/EpicGames.Tracing/UnrealInsights/TransportPacket.cs
@@ -19,9 +19,13 @@ namespace EpicGames.Tracing.UnrealInsights
 		private const ushort ThreadIdMask = PartialMarker - 1;
 		
 		public ushort PacketSize { get; private set; }
-		public ushort ThreadId { get; private set; }
+		public ushort ThreadIdAndMarkers { get; private set; }
 		public ushort DecodedSize { get; private set; }
-		public byte[] Data { get; private set; } = Array.Empty<byte>();
+		byte[] Data = Array.Empty<byte>();
+		public byte[] GetData()
+		{
+			return Data;
+		}
 	
 		private TransportPacket()
 		{
@@ -29,17 +33,17 @@ namespace EpicGames.Tracing.UnrealInsights
 		
 		public bool IsEncoded()
 		{
-			return (ThreadId & EncodedMarker) != 0;
+			return (ThreadIdAndMarkers & EncodedMarker) != 0;
 		}
 		
 		public bool IsPartial()
 		{
-			return (ThreadId & PartialMarker) != 0;
+			return (ThreadIdAndMarkers & PartialMarker) != 0;
 		}
 		
 		public ushort GetThreadId()
 		{
-			return (ushort) (ThreadId & ThreadIdMask);
+			return (ushort) (ThreadIdAndMarkers & ThreadIdMask);
 		}
 		
 		public static bool IsNormalThread(ushort ThreadId)
@@ -60,7 +64,7 @@ namespace EpicGames.Tracing.UnrealInsights
 			TotalSize += 4; // The two uint16 writes below are included
 			
 			Writer.Write(TotalSize);
-			Writer.Write(ThreadId);
+			Writer.Write(ThreadIdAndMarkers);
 
 			foreach ((ushort Uid, ITraceEvent Event) in Events)
 			{
@@ -77,7 +81,7 @@ namespace EpicGames.Tracing.UnrealInsights
 		{
 			TransportPacket Packet = new TransportPacket();
 			Packet.PacketSize = PacketSize;
-			Packet.ThreadId = ThreadId;
+			Packet.ThreadIdAndMarkers = ThreadId;
 			return Packet;
 		}
 
@@ -85,7 +89,7 @@ namespace EpicGames.Tracing.UnrealInsights
 		{
 			TransportPacket Packet = new TransportPacket();
 			Packet.PacketSize = Reader.ReadUInt16();
-			Packet.ThreadId = Reader.ReadUInt16();
+			Packet.ThreadIdAndMarkers = Reader.ReadUInt16();
 
 			int HeaderSize = sizeof(ushort) + sizeof(ushort); // PacketSize + ThreadId
 			if (Packet.IsEncoded())
diff --git a/Engine/Source/Programs/Shared/EpicGames.Tracing/UnrealInsights/Writer.cs b/Engine/Source/Programs/Shared/EpicGames.Tracing/UnrealInsights/Writer.cs
--- a/Engine/Source/Programs/Shared/EpicGames.Tracing/UnrealInsights/Writer.cs
+++ b/Engine/Source/Programs/Shared/EpicGames.Tracing/UnrealInsights/Writer.cs
@@ -10,9 +10,9 @@ namespace EpicGames.Tracing.UnrealInsights
 {
 	public class UnrealInsightsWriter
 	{
-		public readonly Dictionary<EventType, ushort> EventTypeToUids = new Dictionary<EventType, ushort>();
-		public readonly Dictionary<ushort, EventType> UidToEventTypes = new Dictionary<ushort, EventType>();
-		public readonly Dictionary<ushort, List<ITraceEvent>> ThreadToEvents = new Dictionary<ushort, List<ITraceEvent>>();
+		readonly Dictionary<EventType, ushort> EventTypeToUids = new Dictionary<EventType, ushort>();
+		readonly Dictionary<ushort, EventType> UidToEventTypes = new Dictionary<ushort, EventType>();
+		readonly Dictionary<ushort, List<ITraceEvent>> ThreadToEvents = new Dictionary<ushort, List<ITraceEvent>>();
 		
 		private ushort UidCounter = PredefinedEventUid._WellKnownNum; // IDs below 16 are reserved for well-known events
 		private object Lock = new object();
