EpicGames.Redis: net6.0 upgrade (with warnings)

Remaining warnings:
RedisChannel.cs(15,25,15,37): warning CA1066: Type EpicGames.Redis.RedisChannel<T> should implement IEquatable<T> because it overrides Equals
Utility\RedisLock.cs(14,15,14,24): warning CA1063: Provide an overridable implementation of Dispose(bool) on 'RedisLock' or mark the type as sealed. A call to Dispose(false) should only clean up native resources. A call to Dispose(true) should clean up both managed and native resources.
Utility\RedisLock.cs(39,26,39,38): warning CA1816: Change RedisLock.DisposeAsync() to call GC.SuppressFinalize(object). This will prevent derived types that introduce a finalizer from needing to re-implement 'IDisposable' to call it.
Utility\RedisLock.cs(79,11,79,96): warning CA2008: Do not create tasks without passing a TaskScheduler
Utility\RedisLock.cs(33,15,33,22): warning CA1063: Modify 'RedisLock.Dispose' so that it calls Dispose(true), then calls GC.SuppressFinalize on the current object instance ('this' or 'Me' in Visual Basic), and then returns
Utility\RedisLock.cs(33,15,33,22): warning CA1816: Change RedisLock.Dispose() to call GC.SuppressFinalize(object). This will prevent derived types that introduce a finalizer from needing to re-implement 'IDisposable' to call it.

#rnx

diff --git a/Engine/Source/Programs/Shared/EpicGames.Redis/EpicGames.Redis.csproj b/Engine/Source/Programs/Shared/EpicGames.Redis/EpicGames.Redis.csproj
--- a/Engine/Source/Programs/Shared/EpicGames.Redis/EpicGames.Redis.csproj
+++ b/Engine/Source/Programs/Shared/EpicGames.Redis/EpicGames.Redis.csproj
@@ -1,6 +1,6 @@
 <Project Sdk="Microsoft.NET.Sdk">
   <PropertyGroup>
-    <TargetFramework>netcoreapp3.1</TargetFramework>
+    <TargetFramework>net6.0</TargetFramework>
     <Configuration Condition=" '$(Configuration)' == '' ">Development</Configuration>
     <OutputType>Library</OutputType>
     <Deterministic>true</Deterministic>
diff --git a/Engine/Source/Programs/Shared/EpicGames.Redis/RedisChannel.cs b/Engine/Source/Programs/Shared/EpicGames.Redis/RedisChannel.cs
--- a/Engine/Source/Programs/Shared/EpicGames.Redis/RedisChannel.cs
+++ b/Engine/Source/Programs/Shared/EpicGames.Redis/RedisChannel.cs
@@ -45,7 +45,7 @@ namespace EpicGames.Redis
 	/// Subscription to a <see cref="RedisChannel{Task}"/>
 	/// </summary>
 	/// <typeparam name="T"></typeparam>
-	public class RedisChannelSubscription<T> : IDisposable, IAsyncDisposable
+	public sealed class RedisChannelSubscription<T> : IDisposable, IAsyncDisposable
 	{
 		/// <summary>
 		/// The subscriber to register with
@@ -105,7 +105,7 @@ namespace EpicGames.Redis
 		/// <param name="message"></param>
 		void UntypedHandler(RedisChannel _, RedisValue message)
 		{
-			_handler(Channel, RedisSerializer.Deserialize<T>(message));
+			_handler(Channel, RedisSerializer.Deserialize<T>(message)!);
 		}
 
 		/// <inheritdoc/>
diff --git a/Engine/Source/Programs/Shared/EpicGames.Redis/RedisHash.cs b/Engine/Source/Programs/Shared/EpicGames.Redis/RedisHash.cs
--- a/Engine/Source/Programs/Shared/EpicGames.Redis/RedisHash.cs
+++ b/Engine/Source/Programs/Shared/EpicGames.Redis/RedisHash.cs
@@ -54,7 +54,7 @@ namespace EpicGames.Redis
 		public async Task<TValue> GetAsync(TName name, CommandFlags flags = CommandFlags.None)
 		{
 			RedisValue value = await Database.HashGetAsync(Key, RedisSerializer.Serialize(name), flags);
-			return value.IsNull ? default! : RedisSerializer.Deserialize<TValue>(value);
+			return value.IsNull ? default! : RedisSerializer.Deserialize<TValue>(value)!;
 		}
 
 		/// <inheritdoc cref="IDatabaseAsync.HashGetAsync(RedisKey, RedisValue[], CommandFlags)"/>
@@ -62,21 +62,21 @@ namespace EpicGames.Redis
 		{
 			RedisValue[] nameArray = names.Select(x => RedisSerializer.Serialize(x)).ToArray();
 			RedisValue[] valueArray = await Database.HashGetAsync(Key, nameArray, flags);
-			return Array.ConvertAll(valueArray, x => RedisSerializer.Deserialize<TValue>(x));
+			return Array.ConvertAll(valueArray, x => RedisSerializer.Deserialize<TValue>(x)!);
 		}
 
 		/// <inheritdoc cref="IDatabaseAsync.HashGetAllAsync(RedisKey, CommandFlags)"/>
 		public async Task<HashEntry<TName, TValue>[]> GetAllAsync(CommandFlags flags = CommandFlags.None)
 		{
 			HashEntry[] entries = await Database.HashGetAllAsync(Key, flags);
-			return Array.ConvertAll(entries, x => new HashEntry<TName, TValue>(RedisSerializer.Deserialize<TName>(x.Name), RedisSerializer.Deserialize<TValue>(x.Value)));
+			return Array.ConvertAll(entries, x => new HashEntry<TName, TValue>(RedisSerializer.Deserialize<TName>(x.Name)!, RedisSerializer.Deserialize<TValue>(x.Value)!));
 		}
 
 		/// <inheritdoc cref="IDatabaseAsync.HashKeysAsync(RedisKey, CommandFlags)"/>
 		public async Task<TName[]> KeysAsync(CommandFlags flags = CommandFlags.None)
 		{
 			RedisValue[] nameArray = await Database.HashKeysAsync(Key, flags);
-			return Array.ConvertAll(nameArray, x => RedisSerializer.Deserialize<TName>(x));
+			return Array.ConvertAll(nameArray, x => RedisSerializer.Deserialize<TName>(x)!);
 		}
 
 		/// <inheritdoc cref="IDatabaseAsync.HashLengthAsync(RedisKey, CommandFlags)"/>
@@ -101,7 +101,7 @@ namespace EpicGames.Redis
 		public async Task<TValue[]> ValuesAsync(CommandFlags flags = CommandFlags.None)
 		{
 			RedisValue[] values = await Database.HashValuesAsync(Key, flags);
-			return Array.ConvertAll(values, x => RedisSerializer.Deserialize<TValue>(x));
+			return Array.ConvertAll(values, x => RedisSerializer.Deserialize<TValue>(x)!);
 		}
 	}
 
diff --git a/Engine/Source/Programs/Shared/EpicGames.Redis/RedisList.cs b/Engine/Source/Programs/Shared/EpicGames.Redis/RedisList.cs
--- a/Engine/Source/Programs/Shared/EpicGames.Redis/RedisList.cs
+++ b/Engine/Source/Programs/Shared/EpicGames.Redis/RedisList.cs
@@ -37,7 +37,7 @@ namespace EpicGames.Redis
 		public async Task<TElement> GetByIndexAsync(long index, CommandFlags flags = CommandFlags.None)
 		{
 			RedisValue value = await Database.ListGetByIndexAsync(Key, index, flags);
-			return value.IsNull? default! : RedisSerializer.Deserialize<TElement>(value);
+			return value.IsNull? default! : RedisSerializer.Deserialize<TElement>(value)!;
 		}
 
 		/// <inheritdoc cref="IDatabaseAsync.ListInsertAfterAsync(RedisKey, RedisValue, RedisValue, CommandFlags)"/>
@@ -60,7 +60,7 @@ namespace EpicGames.Redis
 		public async Task<TElement> LeftPopAsync(CommandFlags flags = CommandFlags.None)
 		{
 			RedisValue value = await Database.ListLeftPopAsync(Key, flags);
-			return value.IsNull ? default! : RedisSerializer.Deserialize<TElement>(value);
+			return value.IsNull ? default! : RedisSerializer.Deserialize<TElement>(value)!;
 		}
 
 		/// <inheritdoc cref="IDatabaseAsync.ListLeftPushAsync(RedisKey, RedisValue, When, CommandFlags)"/>
@@ -86,7 +86,7 @@ namespace EpicGames.Redis
 		public async Task<TElement[]> RangeAsync(long start = 0, long stop = -1, CommandFlags flags = CommandFlags.None)
 		{
 			RedisValue[] values = await Database.ListRangeAsync(Key, start, stop, flags);
-			return Array.ConvertAll(values, (Converter<RedisValue, TElement>)(x => (TElement)RedisSerializer.Deserialize<TElement>(x)));
+			return Array.ConvertAll(values, (Converter<RedisValue, TElement>)(x => (TElement)RedisSerializer.Deserialize<TElement>(x)!));
 		}
 
 		/// <inheritdoc cref="IDatabaseAsync.ListRemoveAsync(RedisKey, RedisValue, Int64, CommandFlags)"/>
@@ -100,7 +100,7 @@ namespace EpicGames.Redis
 		public async Task<TElement> RightPopAsync(CommandFlags flags = CommandFlags.None)
 		{
 			RedisValue value = await Database.ListRightPopAsync(Key, flags);
-			return value.IsNull? default! : RedisSerializer.Deserialize<TElement>(value);
+			return value.IsNull? default! : RedisSerializer.Deserialize<TElement>(value)!;
 		}
 
 		/// <inheritdoc cref="IDatabaseAsync.ListRightPushAsync(RedisKey, RedisValue, When, CommandFlags)"/>
diff --git a/Engine/Source/Programs/Shared/EpicGames.Redis/RedisSerializer.cs b/Engine/Source/Programs/Shared/EpicGames.Redis/RedisSerializer.cs
--- a/Engine/Source/Programs/Shared/EpicGames.Redis/RedisSerializer.cs
+++ b/Engine/Source/Programs/Shared/EpicGames.Redis/RedisSerializer.cs
@@ -13,7 +13,7 @@ namespace EpicGames.Redis
 	/// Attribute specifying the converter type to use for a class
 	/// </summary>
 	[AttributeUsage(AttributeTargets.Class | AttributeTargets.Struct)]
-	public class RedisConverterAttribute : Attribute
+	public sealed class RedisConverterAttribute : Attribute
 	{
 		/// <summary>
 		/// Type of the converter to use
@@ -47,7 +47,7 @@ namespace EpicGames.Redis
 		/// </summary>
 		/// <param name="value"></param>
 		/// <returns></returns>
-		T FromRedisValue(RedisValue value);
+		T? FromRedisValue(RedisValue value);
 	}
 
 	/// <summary>
@@ -63,7 +63,7 @@ namespace EpicGames.Redis
 		}
 
 		/// <inheritdoc/>
-		public T FromRedisValue(RedisValue value)
+		public T? FromRedisValue(RedisValue value)
 		{
 			return CbSerializer.Deserialize<T>(new CbField((byte[])value));
 		}
@@ -83,8 +83,8 @@ namespace EpicGames.Redis
 				_typeConverter = typeConverter;
 			}
 
-			public RedisValue ToRedisValue(T value) => (string)_typeConverter.ConvertTo(value, typeof(string));
-			public T FromRedisValue(RedisValue value) => (T)_typeConverter.ConvertFrom((string)value);
+			public RedisValue ToRedisValue(T value) => (string?)_typeConverter.ConvertTo(value, typeof(string));
+			public T? FromRedisValue(RedisValue value) => (T?)_typeConverter.ConvertFrom((string)value);
 		}
 
 		class RedisNativeConverter<T> : IRedisConverter<T>
@@ -228,7 +228,7 @@ namespace EpicGames.Redis
 		/// <typeparam name="T"></typeparam>
 		/// <param name="value"></param>
 		/// <returns></returns>
-		public static T Deserialize<T>(RedisValue value)
+		public static T? Deserialize<T>(RedisValue value)
 		{
 			return CachedConverter<T>.Converter.FromRedisValue(value);
 		}
diff --git a/Engine/Source/Programs/Shared/EpicGames.Redis/RedisSet.cs b/Engine/Source/Programs/Shared/EpicGames.Redis/RedisSet.cs
--- a/Engine/Source/Programs/Shared/EpicGames.Redis/RedisSet.cs
+++ b/Engine/Source/Programs/Shared/EpicGames.Redis/RedisSet.cs
@@ -60,21 +60,21 @@ namespace EpicGames.Redis
 		public async Task<TElement[]> MembersAsync(CommandFlags flags = CommandFlags.None)
 		{
 			RedisValue[] values = await _database.SetMembersAsync(Key, flags);
-			return Array.ConvertAll(values, (Converter<RedisValue, TElement>)(x => RedisSerializer.Deserialize<TElement>(x)));
+			return Array.ConvertAll(values, (Converter<RedisValue, TElement>)(x => RedisSerializer.Deserialize<TElement>(x)!));
 		}
 
 		/// <inheritdoc cref="IDatabaseAsync.SetPopAsync(RedisKey, CommandFlags)"/>
 		public async Task<TElement> PopAsync(CommandFlags flags = CommandFlags.None)
 		{
 			RedisValue value = await _database.SetPopAsync(Key, flags);
-			return value.IsNull ? default! : RedisSerializer.Deserialize<TElement>(value);
+			return value.IsNull ? default! : RedisSerializer.Deserialize<TElement>(value)!;
 		}
 
 		/// <inheritdoc cref="IDatabaseAsync.SetPopAsync(RedisKey, Int64, CommandFlags)"/>
 		public async Task<TElement[]> PopAsync(long count, CommandFlags flags = CommandFlags.None)
 		{
 			RedisValue[] values = await _database.SetPopAsync(Key, count, flags);
-			return Array.ConvertAll<RedisValue, TElement>(values, x => RedisSerializer.Deserialize<TElement>(x));
+			return Array.ConvertAll<RedisValue, TElement>(values, x => RedisSerializer.Deserialize<TElement>(x)!);
 		}
 
 		/// <inheritdoc cref="IDatabaseAsync.SetRemoveAsync(RedisKey, RedisValue, CommandFlags)"/>
diff --git a/Engine/Source/Programs/Shared/EpicGames.Redis/RedisSortedSet.cs b/Engine/Source/Programs/Shared/EpicGames.Redis/RedisSortedSet.cs
--- a/Engine/Source/Programs/Shared/EpicGames.Redis/RedisSortedSet.cs
+++ b/Engine/Source/Programs/Shared/EpicGames.Redis/RedisSortedSet.cs
@@ -35,7 +35,7 @@ namespace EpicGames.Redis
 		/// <param name="entry"></param>
 		public SortedSetEntry(SortedSetEntry entry)
 		{
-			Element = RedisSerializer.Deserialize<T>(entry.Element);
+			Element = RedisSerializer.Deserialize<T>(entry.Element)!;
 			ElementValue = entry.Element;
 			Score = entry.Score;
 		}
@@ -118,7 +118,7 @@ namespace EpicGames.Redis
 		public async Task<TElement[]> RangeByRankAsync(long start, long stop = -1, Order order = Order.Ascending, CommandFlags flags = CommandFlags.None)
 		{
 			RedisValue[] values = await _database.SortedSetRangeByRankAsync(Key, start, stop, order, flags);
-			return Array.ConvertAll(values, x => RedisSerializer.Deserialize<TElement>(x));
+			return Array.ConvertAll(values, x => RedisSerializer.Deserialize<TElement>(x)!);
 		}
 
 		/// <inheritdoc cref="IDatabaseAsync.SortedSetRangeByRankWithScoresAsync(RedisKey, Int64, Int64, Order, CommandFlags)"/>
@@ -132,7 +132,7 @@ namespace EpicGames.Redis
 		public async Task<TElement[]> RangeByScoreAsync(double start = Double.NegativeInfinity, double stop = Double.PositiveInfinity, Exclude exclude = Exclude.None, Order order = Order.Ascending, long skip = 0L, long take = -1L, CommandFlags flags = CommandFlags.None)
 		{
 			RedisValue[] values = await _database.SortedSetRangeByScoreAsync(Key, start, stop, exclude, order, skip, take, flags);
-			return Array.ConvertAll(values, x => RedisSerializer.Deserialize<TElement>(x));
+			return Array.ConvertAll(values, x => RedisSerializer.Deserialize<TElement>(x)!);
 		}
 
 		/// <inheritdoc cref="IDatabaseAsync.SortedSetRangeByScoreWithScoresAsync(RedisKey, Double, Double, Exclude, Order, Int64, Int64, CommandFlags)"/>
@@ -148,7 +148,7 @@ namespace EpicGames.Redis
 			RedisValue minValue = RedisSerializer.Serialize(min);
 			RedisValue maxValue = RedisSerializer.Serialize(max);
 			RedisValue[] values = await _database.SortedSetRangeByValueAsync(Key, minValue, maxValue, exclude, order, skip, take, flags);
-			return Array.ConvertAll(values, x => RedisSerializer.Deserialize<TElement>(x));
+			return Array.ConvertAll(values, x => RedisSerializer.Deserialize<TElement>(x)!);
 		}
 
 		/// <inheritdoc cref="IDatabaseAsync.SortedSetRankAsync(RedisKey, RedisValue, Order, CommandFlags)"/>
diff --git a/Engine/Source/Programs/Shared/EpicGames.Redis/Utility/RedisLock.cs b/Engine/Source/Programs/Shared/EpicGames.Redis/Utility/RedisLock.cs
--- a/Engine/Source/Programs/Shared/EpicGames.Redis/Utility/RedisLock.cs
+++ b/Engine/Source/Programs/Shared/EpicGames.Redis/Utility/RedisLock.cs
@@ -76,7 +76,7 @@ namespace EpicGames.Redis.Utility
 			Stopwatch timer = Stopwatch.StartNew();
 			for (; ; )
 			{
-				await Task.Delay(duration / 2, cancellationToken).ContinueWith(x => { }); // Do not throw
+				await Task.Delay(duration / 2, cancellationToken).ContinueWith(x => { }, cancellationToken); // Do not throw
 				if (cancellationToken.IsCancellationRequested)
 				{
 					timer.Stop();
