EpicGames.Jupiter: net6.0 upgrade

#rnx

diff --git a/Engine/Source/Programs/AutomationTool/BuildGraph/Tasks/UploadTask.cs b/Engine/Source/Programs/AutomationTool/BuildGraph/Tasks/UploadTask.cs
--- a/Engine/Source/Programs/AutomationTool/BuildGraph/Tasks/UploadTask.cs
+++ b/Engine/Source/Programs/AutomationTool/BuildGraph/Tasks/UploadTask.cs
@@ -160,7 +160,7 @@ namespace AutomationTool.Tasks
 				}
 			}
 			// Upload the tree to Jupiter
-			Dictionary<FileReference, List<string>> Mapping = await FileTree.UploadToJupiter(Parameters.JupiterUrl, Parameters.JupiterNamespace, Parameters.JupiterKey, Metadata);
+			Dictionary<FileReference, List<string>> Mapping = await FileTree.UploadToJupiter(new Uri(Parameters.JupiterUrl), Parameters.JupiterNamespace, Parameters.JupiterKey, Metadata);
 
 			// Debug output of which files mapped to which blobs, can be useful to determine which files are constantly being uploaded
 			// Json.Save(FileReference.Combine(AutomationTool.Unreal.RootDirectory, "JupiterUpload.json"), Mapping);
diff --git a/Engine/Source/Programs/AutomationTool/Scripts/DownloadJupiterBuild.cs b/Engine/Source/Programs/AutomationTool/Scripts/DownloadJupiterBuild.cs
--- a/Engine/Source/Programs/AutomationTool/Scripts/DownloadJupiterBuild.cs
+++ b/Engine/Source/Programs/AutomationTool/Scripts/DownloadJupiterBuild.cs
@@ -4,6 +4,7 @@ using AutomationTool;
 using System.Threading.Tasks;
 using EpicGames.Core;
 using EpicGames.Jupiter;
+using System;
 
 namespace BuildScripts.Automation
 {
@@ -21,7 +22,7 @@ namespace BuildScripts.Automation
 			JupiterFileTree Tree = new JupiterFileTree(TargetDirectory, true);
 			
 			FileReference LocalManifest = FileReference.Combine(TargetDirectory, "Jupiter-Manifest.json");
-			Task DownloadTask = Tree.DownloadFromJupiter(LocalManifest, JupiterUrl, Namespace, Key);
+			Task DownloadTask = Tree.DownloadFromJupiter(LocalManifest, new Uri(JupiterUrl), Namespace, Key);
 			DownloadTask.Wait();
 		}
 	}
diff --git a/Engine/Source/Programs/Shared/EpicGames.Jupiter/EpicGames.Jupiter.csproj b/Engine/Source/Programs/Shared/EpicGames.Jupiter/EpicGames.Jupiter.csproj
--- a/Engine/Source/Programs/Shared/EpicGames.Jupiter/EpicGames.Jupiter.csproj
+++ b/Engine/Source/Programs/Shared/EpicGames.Jupiter/EpicGames.Jupiter.csproj
@@ -1,6 +1,6 @@
 <Project Sdk="Microsoft.NET.Sdk">
   <PropertyGroup>
-    <TargetFramework>netcoreapp3.1</TargetFramework>
+    <TargetFramework>net6.0</TargetFramework>
     <Configuration Condition=" '$(Configuration)' == '' ">Development</Configuration>
     <OutputType>Library</OutputType>
     <Deterministic>true</Deterministic>
diff --git a/Engine/Source/Programs/Shared/EpicGames.Jupiter/JupierUtils.cs b/Engine/Source/Programs/Shared/EpicGames.Jupiter/JupierUtils.cs
--- a/Engine/Source/Programs/Shared/EpicGames.Jupiter/JupierUtils.cs
+++ b/Engine/Source/Programs/Shared/EpicGames.Jupiter/JupierUtils.cs
@@ -63,10 +63,12 @@ namespace EpicGames.Jupiter
 	{
 		public static byte[] GetSHA1(byte[] payload)
 		{
+#pragma warning disable CA5350 // Do Not Use Weak Cryptographic Algorithms
 			using (SHA1 sha1 = SHA1.Create())
 			{
 				return sha1.ComputeHash(payload);
 			}
+#pragma warning restore CA5350 // Do Not Use Weak Cryptographic Algorithms
 		}
 
 		public static string GetSHA1AsString(byte[] payload)
diff --git a/Engine/Source/Programs/Shared/EpicGames.Jupiter/JupiterTree.cs b/Engine/Source/Programs/Shared/EpicGames.Jupiter/JupiterTree.cs
--- a/Engine/Source/Programs/Shared/EpicGames.Jupiter/JupiterTree.cs
+++ b/Engine/Source/Programs/Shared/EpicGames.Jupiter/JupiterTree.cs
@@ -90,11 +90,11 @@ namespace EpicGames.Jupiter
 			}
 		}
 
-		public async Task UploadToJupiter(string JupiterUrl, string JupiterNamespace, string JupiterTreeKey, Dictionary<string, object> Metadata = null)
+		public async Task UploadToJupiter(Uri JupiterUrl, string JupiterNamespace, string JupiterTreeKey, Dictionary<string, object> Metadata = null)
 		{
 			string TreeHashString = SHA1Utils.FormatAsHexString(CalculateTreeHash());
 			List<JupiterTree> AllTrees = GetAllTrees();
-			using (HttpClient JupiterClient = new HttpClient {BaseAddress = new Uri(JupiterUrl)})
+			using (HttpClient JupiterClient = new HttpClient {BaseAddress = JupiterUrl})
 			{
 				Log.TraceInformation("Creating a new tree root in Jupiter with id: {0}", JupiterTreeKey);
 
@@ -117,9 +117,9 @@ namespace EpicGames.Jupiter
 						Metadata = Metadata,
 					};
 					string PutTreeRootString = JsonSerializer.Serialize(TreeRoot);
-					StringContent PutTreeRootContent = new StringContent(PutTreeRootString, Encoding.UTF8, "application/json");
+					using StringContent PutTreeRootContent = new StringContent(PutTreeRootString, Encoding.UTF8, "application/json");
 					// upload tree information
-					HttpResponseMessage PutTreeResult = await JupiterClient.PutAsync(string.Format("api/v1/c/tree-root/{0}", JupiterNamespace), PutTreeRootContent);
+					HttpResponseMessage PutTreeResult = await JupiterClient.PutAsync(new Uri(string.Format("api/v1/c/tree-root/{0}", JupiterNamespace)), PutTreeRootContent);
 					if (!PutTreeResult.IsSuccessStatusCode)
 					{
 						string ErrorMsg = await PutTreeResult.Content.ReadAsStringAsync();
@@ -140,8 +140,8 @@ namespace EpicGames.Jupiter
 							Blobs = Tree.ContentHashes,
 						};
 					string PutTreeString = JsonSerializer.Serialize(PutTreesRequest);
-					StringContent PutTreeContent = new StringContent(PutTreeString, Encoding.UTF8, "application/json");
-					HttpResponseMessage PutTreeContentResult = await JupiterClient.PutAsync(string.Format("api/v1/c/tree/{0}", JupiterNamespace), PutTreeContent);
+					using StringContent PutTreeContent = new StringContent(PutTreeString, Encoding.UTF8, "application/json");
+					HttpResponseMessage PutTreeContentResult = await JupiterClient.PutAsync(new Uri(string.Format("api/v1/c/tree/{0}", JupiterNamespace)), PutTreeContent);
 					if (!PutTreeContentResult.IsSuccessStatusCode)
 					{
 						string ErrorMsg = await PutTreeContentResult.Content.ReadAsStringAsync();
@@ -166,8 +166,8 @@ namespace EpicGames.Jupiter
 					};
 
 					string FilterBlobString = JsonSerializer.Serialize(FilterBlobRequest);
-					StringContent FilterBlobContent = new StringContent(FilterBlobString, Encoding.UTF8, "application/json");
-					HttpResponseMessage FilterBlobResponse = await JupiterClient.PostAsync("api/v1/s", FilterBlobContent);
+					using StringContent FilterBlobContent = new StringContent(FilterBlobString, Encoding.UTF8, "application/json");
+					HttpResponseMessage FilterBlobResponse = await JupiterClient.PostAsync(new Uri("api/v1/s"), FilterBlobContent);
 					if (!FilterBlobResponse.IsSuccessStatusCode)
 					{
 						string ErrorMsg = await FilterBlobResponse.Content.ReadAsStringAsync();
@@ -197,7 +197,7 @@ namespace EpicGames.Jupiter
 									ByteArrayContent Content = new ByteArrayContent(await JupiterTree.GetContent(BlobHash));
 									Content.Headers.Remove("Content-Type");
 									Content.Headers.Add("Content-Type", "application/octet-stream");
-									HttpResponseMessage PutBlobResultResult = await JupiterClient.PutAsync(string.Format("api/v1/s/{0}/{1}", JupiterNamespace, BlobHash), Content);
+									HttpResponseMessage PutBlobResultResult = await JupiterClient.PutAsync(new Uri(string.Format("api/v1/s/{0}/{1}", JupiterNamespace, BlobHash)), Content);
 
 									if (!PutBlobResultResult.IsSuccessStatusCode)
 									{
@@ -210,12 +210,12 @@ namespace EpicGames.Jupiter
 						}
 					}
 
-					Task.WaitAll(SubmitTasks.ToArray());
+					await Task.WhenAll(SubmitTasks);
 				}
 
 				// verify the build upload
 				{
-					HttpResponseMessage FinalizeTreeResult = await JupiterClient.PutAsync(string.Format("api/v1/c/tree-root/{0}/{1}/finalize", JupiterNamespace, JupiterTreeKey), null);
+					HttpResponseMessage FinalizeTreeResult = await JupiterClient.PutAsync(new Uri(string.Format("api/v1/c/tree-root/{0}/{1}/finalize", JupiterNamespace, JupiterTreeKey)), null);
 					if (!FinalizeTreeResult.IsSuccessStatusCode)
 					{
 						string ErrorMsg = await FinalizeTreeResult.Content.ReadAsStringAsync();
@@ -228,9 +228,9 @@ namespace EpicGames.Jupiter
 			}
 		}
 
-		public static async Task<JupiterTree> FromJupiterKey(string JupiterUrl, string JupiterNamespace, string JupiterTreeKey)
+		public static async Task<JupiterTree> FromJupiterKey(Uri JupiterUrl, string JupiterNamespace, string JupiterTreeKey)
 		{
-			using (HttpClient JupiterClient = new HttpClient {BaseAddress = new Uri(JupiterUrl)})
+			using (HttpClient JupiterClient = new HttpClient {BaseAddress = JupiterUrl})
 			{
 				Log.TraceInformation("Downloading tree with key \"{0}\" from Jupiter at {1} in namespace {2}", JupiterTreeKey, JupiterUrl, JupiterNamespace);
 
@@ -238,7 +238,7 @@ namespace EpicGames.Jupiter
 
 				{
 					// download the tree root
-					HttpResponseMessage GetTreeRootResult = await JupiterClient.GetAsync(string.Format("api/v1/c/tree-root/{0}/{1}", JupiterNamespace, JupiterTreeKey));
+					HttpResponseMessage GetTreeRootResult = await JupiterClient.GetAsync(new Uri(string.Format("api/v1/c/tree-root/{0}/{1}", JupiterNamespace, JupiterTreeKey)));
 
 					string GetTreeRootResultString = await GetTreeRootResult.Content.ReadAsStringAsync();
 					if (!GetTreeRootResult.IsSuccessStatusCode)
@@ -254,7 +254,7 @@ namespace EpicGames.Jupiter
 				Dictionary<string, TreeContents> TreeMapping;
 				{
 					// Flatten this tree, and download the tree descriptions of all those trees
-					HttpResponseMessage GetFlattenTreeResult = await JupiterClient.GetAsync(string.Format("api/v1/c/tree/{0}/{1}/flattend", JupiterNamespace, TopTreeHash));
+					HttpResponseMessage GetFlattenTreeResult = await JupiterClient.GetAsync(new Uri(string.Format("api/v1/c/tree/{0}/{1}/flattend", JupiterNamespace, TopTreeHash)));
 
 					string GetFlattenTreeResultString = await GetFlattenTreeResult.Content.ReadAsStringAsync();
 					if (!GetFlattenTreeResult.IsSuccessStatusCode)
@@ -270,7 +270,7 @@ namespace EpicGames.Jupiter
 					foreach (string Tree in TreesToDownload)
 					{
 						// Flatten this tree, and download the tree descriptions of all those trees
-						HttpResponseMessage GetTreeResult = await JupiterClient.GetAsync(string.Format("api/v1/c/tree/{0}/{1}", JupiterNamespace, Tree));
+						HttpResponseMessage GetTreeResult = await JupiterClient.GetAsync(new Uri(string.Format("api/v1/c/tree/{0}/{1}", JupiterNamespace, Tree)));
 
 						string GetTreeResultString = await GetTreeResult.Content.ReadAsStringAsync();
 						if (!GetTreeResult.IsSuccessStatusCode)
@@ -290,7 +290,7 @@ namespace EpicGames.Jupiter
 			}
 		}
 
-		private static JupiterTree CreateTreeFromTreeDescriptions(string JupiterUrl, string JupiterNamespace, string RootTreeHash, Dictionary<string, TreeContents> TreeDescriptions)
+		private static JupiterTree CreateTreeFromTreeDescriptions(Uri JupiterUrl, string JupiterNamespace, string RootTreeHash, Dictionary<string, TreeContents> TreeDescriptions)
 		{
 			TreeContents RootTreeContents = TreeDescriptions[RootTreeHash];
 			JupiterIoContentProvider ContentProvider = new JupiterIoContentProvider();
@@ -305,20 +305,20 @@ namespace EpicGames.Jupiter
 			return NewTree;
 		}
 
-		public class TreeRootContents
+		class TreeRootContents
 		{
 			public string treeHash { get; set; }
 			public Dictionary<string, object> metadata { get; set; }
 		}
 
-		public class FlattendTreeContents
+		class FlattendTreeContents
 		{
 			public string[] allTrees { get; set; }
 
 			public string[] allBlobs { get; set; }
 		}
 
-		public class TreeContents
+		class TreeContents
 		{
 			public string treeHash { get; set; }
 			public List<string> trees { get; set; }
@@ -328,15 +328,15 @@ namespace EpicGames.Jupiter
 
 	public class JupiterIoContentProvider : JupiterTreeContentProvider
 	{
-		public string BaseUrl { get; private set; }
+		public Uri BaseUrl { get; private set; }
 		public List<string> Blobs { get; private set; }
 		public string JupiterNamespace { get; private set; }
 
 		public override async Task<byte[]> GetContent(string Sha1)
 		{
-			using (HttpClient JupiterClient = new HttpClient {BaseAddress = new Uri(BaseUrl)})
+			using (HttpClient JupiterClient = new HttpClient {BaseAddress = BaseUrl})
 			{
-				HttpResponseMessage GetContentsResult = await JupiterClient.GetAsync(string.Format("api/v1/s/{0}/{1}", JupiterNamespace, Sha1));
+				HttpResponseMessage GetContentsResult = await JupiterClient.GetAsync(new Uri(string.Format("api/v1/s/{0}/{1}", JupiterNamespace, Sha1)));
 
 				if (!GetContentsResult.IsSuccessStatusCode)
 				{
@@ -357,7 +357,7 @@ namespace EpicGames.Jupiter
 
 		public override event EventHandler OnContentChanged;
 
-		public void AddContentReference(string JupiterUrl, string InJupiterNamespace, List<string> InBlobs)
+		public void AddContentReference(Uri JupiterUrl, string InJupiterNamespace, List<string> InBlobs)
 		{
 			BaseUrl = JupiterUrl;
 			JupiterNamespace = InJupiterNamespace;
@@ -485,7 +485,7 @@ namespace EpicGames.Jupiter
 			Files.Add(File);
 		}
 
-		public async Task<Dictionary<FileReference, List<string>>> UploadToJupiter(string JupiterUrl, string JupiterNamespace, string JupiterTreeKey, Dictionary<string, object> Metadata = null)
+		public async Task<Dictionary<FileReference, List<string>>> UploadToJupiter(Uri JupiterUrl, string JupiterNamespace, string JupiterTreeKey, Dictionary<string, object> Metadata = null)
 		{
 			Manifest Manifest = new Manifest(BaseDir);
 
@@ -537,7 +537,7 @@ namespace EpicGames.Jupiter
 			return FileToChunkMapping;
 		}
 
-		public async Task<List<FileReference>> DownloadFromJupiter(FileReference LocalManifestPath, string JupiterUrl, string JupiterNamespace, string JupiterTreeKey, IProgress<Tuple<float, FileReference>> Progress = null)
+		public async Task<List<FileReference>> DownloadFromJupiter(FileReference LocalManifestPath, Uri JupiterUrl, string JupiterNamespace, string JupiterTreeKey, IProgress<Tuple<float, FileReference>> Progress = null)
 		{
 			Manifest LocalManifest = null;
 			if (LocalManifestPath.ToFileInfo().Exists)
@@ -596,14 +596,14 @@ namespace EpicGames.Jupiter
 			return FilesWritten;
 		}
 
-		public class Manifest
+		class Manifest
 		{
 			private readonly DirectoryReference BaseDir;
 
 			// ReSharper disable once MemberCanBePrivate.Global , fastjson requires this to be public so it can be serialized
-			public class ManifestData
+			class ManifestData
 			{
-				public string SchemaVersion = "v1"; // added in case we need to update this schema in the future
+				public string SchemaVersion { get; set; } = "v1"; // added in case we need to update this schema in the future
 				public Dictionary<string, string> Files = new Dictionary<string, string>();
 			}
 			private readonly ManifestData Data;
@@ -628,7 +628,7 @@ namespace EpicGames.Jupiter
 			public void AddFile(FileReference File, byte[] TreeHash)
 			{
 				// we escape the forward slashes to make it valid json keys
-				Files.Add(File.MakeRelativeTo(BaseDir).Replace(@"\", @"\\"), SHA1Utils.FormatAsHexString(TreeHash));
+				Files.Add(File.MakeRelativeTo(BaseDir).Replace(@"\", @"\\", StringComparison.Ordinal), SHA1Utils.FormatAsHexString(TreeHash));
 			}
 
 			public JupiterTree AsTree()
diff --git a/Engine/Source/Programs/UnrealGameSync/UnrealGameSync/Jupiter/JupiterMonitor.cs b/Engine/Source/Programs/UnrealGameSync/UnrealGameSync/Jupiter/JupiterMonitor.cs
--- a/Engine/Source/Programs/UnrealGameSync/UnrealGameSync/Jupiter/JupiterMonitor.cs
+++ b/Engine/Source/Programs/UnrealGameSync/UnrealGameSync/Jupiter/JupiterMonitor.cs
@@ -244,7 +244,7 @@ namespace UnrealGameSync
 
 					DirectoryReference RootDirectory = LocalRootPath;
 					JupiterFileTree FileTree = new JupiterFileTree(RootDirectory, InDeferReadingFiles: true);
-					Task<List<FileReference>> DownloadTask = FileTree.DownloadFromJupiter(JupiterManifestFileReference, JupiterUrl, JupiterNamespace, ArchiveKey, ProgressCallback);
+					Task<List<FileReference>> DownloadTask = FileTree.DownloadFromJupiter(JupiterManifestFileReference, new Uri(JupiterUrl), JupiterNamespace, ArchiveKey, ProgressCallback);
 					DownloadTask.Wait();
 
 					List<FileReference> WrittenFiles = DownloadTask.Result;
