Horde: Remove reference to deleted interface

#skipci

diff --git a/Engine/Source/Programs/Shared/EpicGames.Horde.Tests/StorageCompressorTest.cs b/Engine/Source/Programs/Shared/EpicGames.Horde.Tests/StorageCompressorTest.cs
--- a/Engine/Source/Programs/Shared/EpicGames.Horde.Tests/StorageCompressorTest.cs
+++ b/Engine/Source/Programs/Shared/EpicGames.Horde.Tests/StorageCompressorTest.cs
@@ -15,17 +15,17 @@
 	/// Fake compressor used for testing. Simply prepends the data with const prefix.
 	/// Allows for easy debugging and is enough to cause a changed hash.
 	/// </summary>
-	public class FakeStorageCompressor : IStorageCompressor
+	public class FakeStorageCompressor
 	{
 		private readonly byte[] _prefix = Encoding.UTF8.GetBytes("FAKECOMP-");
 		
-		public async Task CompressAsync(Stream input, Stream output, CancellationToken cancellationToken)
+		public async Task CompressAsync(Stream input, Stream output, CancellationToken cancellationToken = default)
 		{
 			await output.WriteAsync(_prefix, cancellationToken);
 			await input.CopyToAsync(output, cancellationToken);
 		}
 
-		public async Task DecompressAsync(Stream input, Stream output, CancellationToken cancellationToken)
+		public async Task DecompressAsync(Stream input, Stream output, CancellationToken cancellationToken = default)
 		{
 			using MemoryStream ms = new();
 			using BinaryReader reader = new(input);
@@ -44,7 +44,7 @@
 		[TestMethod]
 		public async Task CompressAndDecompress()
 		{
-			IStorageCompressor compressor = new FakeStorageCompressor();
+			FakeStorageCompressor compressor = new ();
 			byte[] uncompressedData = { 0x41, 0x42, 0x43 };
 			using MemoryStream input = new(uncompressedData);
 			using MemoryStream output = new();
