Prevent Redis exceptions breaking in the debugger. Restore original code, but pass CancellationToken.None to make intent clear.

diff --git a/Engine/Source/Programs/Shared/EpicGames.Redis/Utility/RedisLock.cs b/Engine/Source/Programs/Shared/EpicGames.Redis/Utility/RedisLock.cs
--- a/Engine/Source/Programs/Shared/EpicGames.Redis/Utility/RedisLock.cs
+++ b/Engine/Source/Programs/Shared/EpicGames.Redis/Utility/RedisLock.cs
@@ -1,5 +1,6 @@
 // Copyright Epic Games, Inc. All Rights Reserved.
 
+using EpicGames.Core;
 using StackExchange.Redis;
 using System;
 using System.Diagnostics;
@@ -92,13 +93,7 @@ namespace EpicGames.Redis.Utility
 			Stopwatch timer = Stopwatch.StartNew();
 			for (; ; )
 			{
-				try
-				{
-					await Task.Delay(duration / 2, cancellationToken);
-				}
-				catch (TaskCanceledException)
-				{
-				}
+				await Task.Delay(duration / 2, cancellationToken).ContinueWith(x => { }, CancellationToken.None, TaskContinuationOptions.None, TaskScheduler.Default); // Do not throw
 
 				if (cancellationToken.IsCancellationRequested)
 				{
