Match multi-line localization errors as a single log event.

diff --git a/Engine/Source/Programs/AutomationTool/AutomationUtils/Matchers/ContentEventMatcher.cs b/Engine/Source/Programs/AutomationTool/AutomationUtils/Matchers/ContentEventMatcher.cs
--- a/Engine/Source/Programs/AutomationTool/AutomationUtils/Matchers/ContentEventMatcher.cs
+++ b/Engine/Source/Programs/AutomationTool/AutomationUtils/Matchers/ContentEventMatcher.cs
@@ -35,14 +35,14 @@ namespace AutomationUtils.Matchers
 				Match? match = s_pattern.Match(input.CurrentLine!);
 				if (match.Success)
 				{
-					LogEventBuilder builder = new LogEventBuilder(input.Hanging());
+					LogEventBuilder builder = new LogEventBuilder(input);
 
 					builder.Annotate(match.Groups["channel"], LogEventMarkup.Channel);
 					builder.Annotate(match.Groups["severity"], LogEventMarkup.Severity);
 					builder.AnnotateAsset(match.Groups["asset"]);
 					builder.Annotate(match.Groups["message"], LogEventMarkup.Message);
 
-					while(builder.Current.TryGetLine(1, out _))
+					while(builder.IsNextLineHanging())
 					{
 						builder.MoveNext();
 					}
diff --git a/Engine/Source/Programs/AutomationTool/AutomationUtils/Matchers/GauntletEventMatcher.cs b/Engine/Source/Programs/AutomationTool/AutomationUtils/Matchers/GauntletEventMatcher.cs
--- a/Engine/Source/Programs/AutomationTool/AutomationUtils/Matchers/GauntletEventMatcher.cs
+++ b/Engine/Source/Programs/AutomationTool/AutomationUtils/Matchers/GauntletEventMatcher.cs
@@ -30,7 +30,7 @@ namespace AutomationUtils.Matchers
 			{
 				string indent = match.Groups["indent"].Value + " ";
 
-				LogEventBuilder builder = new LogEventBuilder(cursor.Hanging());
+				LogEventBuilder builder = new LogEventBuilder(cursor);
 				builder.Annotate(match.Groups["group"]);
 
 				string group = match.Groups["group"].Value;
@@ -40,7 +40,7 @@ namespace AutomationUtils.Matchers
 
 				//				int LineCount = 1;
 				EventId eventId = KnownLogEvents.Gauntlet;
-				while (builder.Next.CurrentLine != null)
+				while (builder.Next.CurrentLine != null && builder.IsNextLineAligned())
 				{
 					builder.MoveNext();
 					if (inErrorList)
diff --git a/Engine/Source/Programs/AutomationTool/AutomationUtils/Matchers/GenericEventMatcher.cs b/Engine/Source/Programs/AutomationTool/AutomationUtils/Matchers/GenericEventMatcher.cs
--- a/Engine/Source/Programs/AutomationTool/AutomationUtils/Matchers/GenericEventMatcher.cs
+++ b/Engine/Source/Programs/AutomationTool/AutomationUtils/Matchers/GenericEventMatcher.cs
@@ -40,8 +40,8 @@ namespace AutomationUtils.Matchers
 					level = LogLevel.Warning;
 				}
 
-				LogEventBuilder builder = new LogEventBuilder(cursor.Hanging());
-				while (builder.Next.CurrentLine != null)
+				LogEventBuilder builder = new LogEventBuilder(cursor);
+				while (builder.IsNextLineHanging())
 				{
 					builder.MoveNext();
 				}
diff --git a/Engine/Source/Programs/AutomationTool/AutomationUtils/Matchers/LocalizationEventMatcher.cs b/Engine/Source/Programs/AutomationTool/AutomationUtils/Matchers/LocalizationEventMatcher.cs
--- a/Engine/Source/Programs/AutomationTool/AutomationUtils/Matchers/LocalizationEventMatcher.cs
+++ b/Engine/Source/Programs/AutomationTool/AutomationUtils/Matchers/LocalizationEventMatcher.cs
@@ -22,7 +22,7 @@ namespace AutomationUtils.Matchers
 			@"(?:\((?<line>\d+)\))?" +
 			@":");
 
-		static readonly Regex s_indentPattern = new Regex(@"^\s+");
+		static readonly Regex s_trailingPattern = new Regex(@": See conflicting location\.\s*$");
 
 		/// <inheritdoc/>
 		public LogEventMatch? Match(ILogCursor input)
@@ -30,17 +30,6 @@ namespace AutomationUtils.Matchers
 			Match? match;
 			if (input.TryMatch(s_pattern, out match))
 			{
-				LogEventBuilder builder = new LogEventBuilder(input.Hanging());
-				builder.Annotate(match.Groups["channel"], LogEventMarkup.Channel);
-				builder.Annotate(match.Groups["severity"], LogEventMarkup.Severity);
-				builder.AnnotateSourceFile(match.Groups["file"], "Engine");
-				builder.Annotate(match.Groups["line"], LogEventMarkup.LineNumber);
-
-				while (builder.Next.CurrentLine != null)
-				{
-					builder.MoveNext();
-				}
-
 				LogLevel level = match.Groups["severity"].Value switch
 				{
 					"Error" => LogLevel.Error,
@@ -48,10 +37,26 @@ namespace AutomationUtils.Matchers
 					_ => LogLevel.Information,
 				};
 
-				return builder.ToMatch(LogEventPriority.High, level, KnownLogEvents.Engine_Localization);
+				LogEventBuilder builder = new LogEventBuilder(input);
+				Annotate(builder, match);
+
+				if (builder.IsNextLineAligned() && builder.Next.TryMatch(s_pattern, out match) && builder.Next.IsMatch(s_trailingPattern))
+				{
+					builder.MoveNext();
+					Annotate(builder, match);
+				}
 
+				return builder.ToMatch(LogEventPriority.High, level, KnownLogEvents.Engine_Localization);
 			}
 			return null;
 		}
+
+		static void Annotate(LogEventBuilder builder, Match match)
+		{
+			builder.Annotate(match.Groups["channel"], LogEventMarkup.Channel);
+			builder.Annotate(match.Groups["severity"], LogEventMarkup.Severity);
+			builder.AnnotateSourceFile(match.Groups["file"], null);
+			builder.Annotate(match.Groups["line"], LogEventMarkup.LineNumber);
+		}
 	}
 }
diff --git a/Engine/Source/Programs/AutomationTool/AutomationUtils/Matchers/LogChannelEventMatcher.cs b/Engine/Source/Programs/AutomationTool/AutomationUtils/Matchers/LogChannelEventMatcher.cs
--- a/Engine/Source/Programs/AutomationTool/AutomationUtils/Matchers/LogChannelEventMatcher.cs
+++ b/Engine/Source/Programs/AutomationTool/AutomationUtils/Matchers/LogChannelEventMatcher.cs
@@ -28,11 +28,11 @@ namespace AutomationUtils.Matchers
 			Match? match;
 			if (input.TryMatch(s_pattern, out match))
 			{
-				LogEventBuilder builder = new LogEventBuilder(input.Hanging());
+				LogEventBuilder builder = new LogEventBuilder(input);
 				builder.Annotate(match.Groups["channel"], LogEventMarkup.Channel);
 				builder.Annotate(match.Groups["severity"], LogEventMarkup.Severity);
 				
-				while(builder.Next.CurrentLine != null)
+				while(builder.IsNextLineHanging())
 				{
 					builder.MoveNext();
 				}
diff --git a/Engine/Source/Programs/Shared/EpicGames.Core/ILogCursor.cs b/Engine/Source/Programs/Shared/EpicGames.Core/ILogCursor.cs
--- a/Engine/Source/Programs/Shared/EpicGames.Core/ILogCursor.cs
+++ b/Engine/Source/Programs/Shared/EpicGames.Core/ILogCursor.cs
@@ -62,74 +62,6 @@ namespace EpicGames.Core
 			public int CurrentLineNumber => _baseLineNumber;
 		}
 
-		class HangingIndentLogCursor : ILogCursor
-		{
-			readonly ILogCursor _inner;
-			readonly int _indent;
-
-			public HangingIndentLogCursor(ILogCursor inner)
-			{
-				_inner = inner;
-				_indent = 0;
-
-				string? line = inner.CurrentLine;
-				if (line != null)
-				{
-					while (_indent < line.Length && Char.IsWhiteSpace(line[_indent]))
-					{
-						_indent++;
-					}
-				}
-			}
-
-			public string? this[int index]
-			{
-				get
-				{
-					if (index <= 0)
-					{
-						return _inner[index];
-					}
-
-					string? line = _inner[index];
-					if (line != null && HasPrefix(line, _inner.CurrentLine!, _indent))
-					{
-						return line;
-					}
-
-					return null;
-				}
-			}
-
-			static bool HasPrefix(string line, string prefix, int prefixLength)
-			{
-				if (line.Length > prefixLength && String.CompareOrdinal(line, 0, prefix, 0, prefixLength) == 0 && Char.IsWhiteSpace(line[prefixLength]))
-				{
-					return true;
-				}
-				if(IsWhitespace(line))
-				{
-					return true;
-				}
-				return false;
-			}
-
-			static bool IsWhitespace(string line)
-			{
-				for (int idx = 0; idx < line.Length; idx++)
-				{
-					if (!Char.IsWhiteSpace(line[idx]))
-					{
-						return false;
-					}
-				}
-				return true;
-			}
-
-			public string? CurrentLine => _inner.CurrentLine;
-			public int CurrentLineNumber => _inner.CurrentLineNumber;
-		}
-
 		/// <summary>
 		/// Creates a new log cursor based at an offset from the current line
 		/// </summary>
@@ -141,16 +73,6 @@ namespace EpicGames.Core
 			return new RebasedLogCursor(cursor, cursor.CurrentLineNumber + offset);
 		}
 
-		/// <summary>
-		/// Creates a new log cursor which enumerates lines with a hanging indent
-		/// </summary>
-		/// <param name="cursor">The current log cursor instance</param>
-		/// <returns>New log cursor instance</returns>
-		public static ILogCursor Hanging(this ILogCursor cursor)
-		{
-			return new HangingIndentLogCursor(cursor);
-		}
-
 		/// <summary>
 		/// Attempts to get a line at the given offset
 		/// </summary>
@@ -305,5 +227,63 @@ namespace EpicGames.Core
 			}
 			return offset;
 		}
+
+		/// <summary>
+		/// Check if a line at a given offset is left-aligned with at least this indent of another line
+		/// </summary>
+		/// <param name="cursor">The log cursor</param>
+		/// <param name="offset">Offset of the line to check</param>
+		/// <param name="firstLine">The first line to compare with</param>
+		/// <returns>True if the line at the given offset is aligned at least as much as the current lien</returns>
+		public static bool IsAligned(this ILogCursor cursor, int offset, string? firstLine)
+		{
+			if (firstLine == null || !cursor.TryGetLine(offset, out string? line))
+			{
+				return false;
+			}
+
+			for (int idx = 0; idx < line.Length; idx++)
+			{
+				if (idx == firstLine.Length || !Char.IsWhiteSpace(firstLine[idx]))
+				{
+					return true;
+				}
+				if (line[idx] != firstLine[idx])
+				{
+					return false;
+				}
+			}
+
+			return true;
+		}
+
+		/// <summary>
+		/// Check if a line at a given offset is left-aligned with at least this indent of another line
+		/// </summary>
+		/// <param name="cursor">The log cursor</param>
+		/// <param name="offset">Offset of the line to check</param>
+		/// <param name="firstLine">The first line to compare with</param>
+		/// <returns>True if the line at the given offset is aligned at least as much as the current lien</returns>
+		public static bool IsHanging(this ILogCursor cursor, int offset, string? firstLine)
+		{
+			if (firstLine == null || !cursor.TryGetLine(offset, out string? line))
+			{
+				return false;
+			}
+
+			for (int idx = 0; idx < line.Length; idx++)
+			{
+				if (idx == firstLine.Length || !Char.IsWhiteSpace(firstLine[idx]))
+				{
+					return Char.IsWhiteSpace(line[idx]);
+				}
+				if (line[idx] != firstLine[idx])
+				{
+					return false;
+				}
+			}
+
+			return true;
+		}
 	}
 }
diff --git a/Engine/Source/Programs/Shared/EpicGames.Core/LogEventBuilder.cs b/Engine/Source/Programs/Shared/EpicGames.Core/LogEventBuilder.cs
--- a/Engine/Source/Programs/Shared/EpicGames.Core/LogEventBuilder.cs
+++ b/Engine/Source/Programs/Shared/EpicGames.Core/LogEventBuilder.cs
@@ -352,6 +352,30 @@ namespace EpicGames.Core
 			_properties.Add(name, value);
 		}
 
+		string? GetFirstLine()
+		{
+			if (_lines == null || _lines.Count == 0)
+			{
+				return Current.CurrentLine;
+			}
+			else
+			{
+				return _lines[^1]._message;
+			}
+		}
+
+		/// <summary>
+		/// Check if the next line is aligned or indented from the first line
+		/// </summary>
+		/// <returns>True if the next line is aligned with this</returns>
+		public bool IsNextLineAligned() => Next.IsAligned(0, GetFirstLine());
+
+		/// <summary>
+		/// Check if the next line is indented from the first line
+		/// </summary>
+		/// <returns>True if the next line is aligned with this</returns>
+		public bool IsNextLineHanging() => Next.IsHanging(0, GetFirstLine());
+
 		/// <summary>
 		/// Complete the current line and move to the next
 		/// </summary>
