EpicGames.Perforce*: net6.0 upgrade

#rnx

diff --git a/Engine/Source/Programs/Shared/EpicGames.Perforce.Managed/CachedFileInfo.cs b/Engine/Source/Programs/Shared/EpicGames.Perforce.Managed/CachedFileInfo.cs
--- a/Engine/Source/Programs/Shared/EpicGames.Perforce.Managed/CachedFileInfo.cs
+++ b/Engine/Source/Programs/Shared/EpicGames.Perforce.Managed/CachedFileInfo.cs
@@ -41,22 +41,22 @@ namespace EpicGames.Perforce.Managed
 			FileInfo info = new FileInfo(location.FullName);
 			if (!info.Exists)
 			{
-				logger.LogWarning("warning: {0} was missing from cache.", location);
+				logger.LogWarning("warning: {Location} was missing from cache.", location);
 				return false;
 			}
 			if (info.Length != Length)
 			{
-				logger.LogWarning("warning: {0} was {1:n} bytes; expected {2:n} bytes", location, info.Length, Length);
+				logger.LogWarning("warning: {Location} was {Bytes:n} bytes; expected {ExpectedBytes:n} bytes", location, info.Length, Length);
 				return false;
 			}
 			if (info.LastWriteTimeUtc.Ticks != LastModifiedTicks)
 			{
-				logger.LogWarning("warning: {0} was last modified at {1}; expected {2}", location, info.LastWriteTimeUtc, new DateTime(LastModifiedTicks, DateTimeKind.Utc));
+				logger.LogWarning("warning: {Location} was last modified at {LastWriteTime}; expected {ExpectedLastWriteTime}", location, info.LastWriteTimeUtc, new DateTime(LastModifiedTicks, DateTimeKind.Utc));
 				return false;
 			}
 			if (info.Attributes.HasFlag(FileAttributes.ReadOnly) != BReadOnly)
 			{
-				logger.LogWarning("warning: {0} readonly flag is {1}; expected {2}", info.Attributes.HasFlag(FileAttributes.ReadOnly), BReadOnly);
+				logger.LogWarning("warning: {Location} readonly flag is {ReadOnly}; expected {ExpectedReadOnly}", location, info.Attributes.HasFlag(FileAttributes.ReadOnly), BReadOnly);
 				return false;
 			}
 
diff --git a/Engine/Source/Programs/Shared/EpicGames.Perforce.Managed/EpicGames.Perforce.Managed.csproj b/Engine/Source/Programs/Shared/EpicGames.Perforce.Managed/EpicGames.Perforce.Managed.csproj
--- a/Engine/Source/Programs/Shared/EpicGames.Perforce.Managed/EpicGames.Perforce.Managed.csproj
+++ b/Engine/Source/Programs/Shared/EpicGames.Perforce.Managed/EpicGames.Perforce.Managed.csproj
@@ -1,7 +1,7 @@
 <Project Sdk="Microsoft.NET.Sdk">
 
   <PropertyGroup>
-    <TargetFramework>netcoreapp3.1</TargetFramework>
+    <TargetFramework>net6.0</TargetFramework>
     <Nullable>enable</Nullable>
     <GenerateAssemblyInfo>false</GenerateAssemblyInfo>
     <GenerateTargetFrameworkAttribute>false</GenerateTargetFrameworkAttribute>
diff --git a/Engine/Source/Programs/Shared/EpicGames.Perforce.Managed/ManagedWorkspace.cs b/Engine/Source/Programs/Shared/EpicGames.Perforce.Managed/ManagedWorkspace.cs
--- a/Engine/Source/Programs/Shared/EpicGames.Perforce.Managed/ManagedWorkspace.cs
+++ b/Engine/Source/Programs/Shared/EpicGames.Perforce.Managed/ManagedWorkspace.cs
@@ -39,17 +39,17 @@ namespace EpicGames.Perforce.Managed
 		/// <summary>
 		/// The Perforce connection
 		/// </summary>
-		public IPerforceConnection _perforceClient;
+		public IPerforceConnection PerforceClient { get; }
 
 		/// <summary>
 		/// Stream to sync it to
 		/// </summary>
-		public string _streamName;
+		public string StreamName { get; }
 
 		/// <summary>
 		/// View for this client
 		/// </summary>
-		public IReadOnlyList<string> _view;
+		public IReadOnlyList<string> View { get; }
 
 		/// <summary>
 		/// Constructor
@@ -59,9 +59,9 @@ namespace EpicGames.Perforce.Managed
 		/// <param name="view">List of filters for the stream</param>
 		public PopulateRequest(IPerforceConnection perforceClient, string streamName, IReadOnlyList<string> view)
 		{
-			_perforceClient = perforceClient;
-			_streamName = streamName;
-			_view = view;
+			PerforceClient = perforceClient;
+			StreamName = streamName;
+			View = view;
 		}
 	}
 
@@ -450,7 +450,7 @@ namespace EpicGames.Perforce.Managed
 				const int MaxDisplay = 1000;
 				foreach (string path in paths.OrderBy(x => x).Take(MaxDisplay))
 				{
-					_logger.LogInformation($"  {path}");
+					_logger.LogInformation("  {Path}", path);
 				}
 				if (paths.Count > MaxDisplay)
 				{
@@ -675,13 +675,14 @@ namespace EpicGames.Perforce.Managed
 					}
 
 					// Build a table showing amount of unique content in each stream
-					string[,] cells = new string[streamNames.Count + 1, streamNames.Count + 1];
-					cells[0, 0] = "";
+					string[][] cells = new string[streamNames.Count + 1][];//, streamNames.Count + 1];
 					for (int idx = 0; idx < streamNames.Count; idx++)
 					{
-						cells[idx + 1, 0] = streamNames[idx];
-						cells[0, idx + 1] = streamNames[idx];
+						cells[idx] = new string[streamNames.Count + 1];
+						cells[idx + 1][0] = streamNames[idx];
+						cells[0][idx + 1] = streamNames[idx];
 					}
+					cells[0][0] = "";
 
 					// Populate the table
 					for (int rowIdx = 0; rowIdx < streamNames.Count; rowIdx++)
@@ -690,7 +691,7 @@ namespace EpicGames.Perforce.Managed
 						for (int colIdx = 0; colIdx < streamNames.Count; colIdx++)
 						{
 							long diffSize = files.Where(x => !filesInStream[colIdx].Contains(x.ContentId)).Sum(x => x.Length);
-							cells[rowIdx + 1, colIdx + 1] = String.Format("{0:0.0}mb", diffSize / (1024.0 * 1024.0));
+							cells[rowIdx + 1][colIdx + 1] = String.Format("{0:0.0}mb", diffSize / (1024.0 * 1024.0));
 						}
 					}
 
@@ -700,7 +701,7 @@ namespace EpicGames.Perforce.Managed
 					{
 						for (int rowIdx = 0; rowIdx < streamNames.Count + 1; rowIdx++)
 						{
-							colWidths[colIdx] = Math.Max(colWidths[colIdx], cells[rowIdx, colIdx].Length);
+							colWidths[colIdx] = Math.Max(colWidths[colIdx], cells[rowIdx][colIdx].Length);
 						}
 					}
 
@@ -713,12 +714,12 @@ namespace EpicGames.Perforce.Managed
 						StringBuilder row = new StringBuilder();
 						for (int colIdx = 0; colIdx < streamNames.Count + 1; colIdx++)
 						{
-							string cell = cells[rowIdx, colIdx];
+							string cell = cells[rowIdx][colIdx];
 							row.Append(' ', colWidths[colIdx] - cell.Length);
 							row.Append(cell);
 							row.Append(" | ");
 						}
-						_logger.LogInformation(row.ToString());
+						_logger.LogInformation("{Row}", row.ToString());
 					}
 					_logger.LogInformation("");
 
@@ -743,15 +744,15 @@ namespace EpicGames.Perforce.Managed
 				_logger.LogInformation("Local changes:");
 				foreach (string difference in differences)
 				{
-					if (difference.StartsWith("+"))
+					if (difference.StartsWith("+", StringComparison.Ordinal))
 					{
 						Console.ForegroundColor = ConsoleColor.Green;
 					}
-					else if (difference.StartsWith("-"))
+					else if (difference.StartsWith("-", StringComparison.Ordinal))
 					{
 						Console.ForegroundColor = ConsoleColor.Red;
 					}
-					else if (difference.StartsWith("!"))
+					else if (difference.StartsWith("!", StringComparison.Ordinal))
 					{
 						Console.ForegroundColor = ConsoleColor.Yellow;
 					}
@@ -759,7 +760,7 @@ namespace EpicGames.Perforce.Managed
 					{
 						Console.ResetColor();
 					}
-					_logger.LogInformation("  {0}", difference);
+					_logger.LogInformation("  {Difference}", difference);
 				}
 				Console.ResetColor();
 			}
@@ -921,7 +922,7 @@ namespace EpicGames.Perforce.Managed
 				string localPath = writeFile.Path;
 				_logger.LogInformation("  Writing {LocalPath}", localPath);
 
-				Directory.CreateDirectory(Path.GetDirectoryName(localPath));
+				Directory.CreateDirectory(Path.GetDirectoryName(localPath)!);
 
 				PerforceResponse printResponse = await perforce.TryPrintAsync(localPath, $"{writeFile}@={unshelveChangelist}", cancellationToken);
 				if (!printResponse.Succeeded)
@@ -954,12 +955,12 @@ namespace EpicGames.Perforce.Managed
 			// Revert all changes in each of the unique clients
 			foreach (PopulateRequest request in requests)
 			{
-				using IPerforceConnection perforce = await request._perforceClient.WithoutClientAsync();
+				using IPerforceConnection perforce = await request.PerforceClient.WithoutClientAsync();
 
-				PerforceResponse<ClientRecord> response = await perforce.TryGetClientAsync(request._perforceClient.Settings.ClientName!, cancellationToken);
+				PerforceResponse<ClientRecord> response = await perforce.TryGetClientAsync(request.PerforceClient.Settings.ClientName!, cancellationToken);
 				if (response.Succeeded)
 				{
-					await RevertInternalAsync(request._perforceClient, cancellationToken);
+					await RevertInternalAsync(request.PerforceClient, cancellationToken);
 				}
 			}
 
@@ -971,20 +972,20 @@ namespace EpicGames.Perforce.Managed
 			for (int idx = 0; idx < requests.Count; idx++)
 			{
 				PopulateRequest request = requests[idx];
-				string streamName = request._streamName;
+				string streamName = request.StreamName;
 				_logger.LogInformation("Finding contents of {StreamName}:", streamName);
 
 				using (_logger.BeginIndentScope("  "))
 				{
-					await DeleteClientAsync(request._perforceClient, cancellationToken);
-					await UpdateClientAsync(request._perforceClient, streamName, cancellationToken);
+					await DeleteClientAsync(request.PerforceClient, cancellationToken);
+					await UpdateClientAsync(request.PerforceClient, streamName, cancellationToken);
 
-					int changeNumber = await GetLatestClientChangeAsync(request._perforceClient, cancellationToken);
+					int changeNumber = await GetLatestClientChangeAsync(request.PerforceClient, cancellationToken);
 					_logger.LogInformation("Latest change is CL {CL}", changeNumber);
 
-					await UpdateClientHaveTableAsync(request._perforceClient, changeNumber, request._view, cancellationToken);
+					await UpdateClientHaveTableAsync(request.PerforceClient, changeNumber, request.View, cancellationToken);
 
-					StreamSnapshot contents = await FindClientContentsAsync(request._perforceClient, changeNumber, cancellationToken);
+					StreamSnapshot contents = await FindClientContentsAsync(request.PerforceClient, changeNumber, cancellationToken);
 					streamState[idx] = Tuple.Create(changeNumber, contents);
 
 					GC.Collect();
@@ -1050,20 +1051,20 @@ namespace EpicGames.Perforce.Managed
 			for (int idx = 0; idx < requests.Count; idx++)
 			{
 				PopulateRequest request = requests[idx];
-				string streamName = request._streamName;
+				string streamName = request.StreamName;
 				_logger.LogInformation("Syncing files for {StreamName}:", streamName);
 
 				using (_logger.BeginIndentScope("  "))
 				{
-					await DeleteClientAsync(request._perforceClient, cancellationToken);
-					await UpdateClientAsync(request._perforceClient, streamName, cancellationToken);
+					await DeleteClientAsync(request.PerforceClient, cancellationToken);
+					await UpdateClientAsync(request.PerforceClient, streamName, cancellationToken);
 
 					int changeNumber = streamState[idx].Item1;
-					await UpdateClientHaveTableAsync(request._perforceClient, changeNumber, requests[idx]._view, cancellationToken);
+					await UpdateClientHaveTableAsync(request.PerforceClient, changeNumber, requests[idx].View, cancellationToken);
 
 					StreamSnapshot contents = streamState[idx].Item2;
 					await RemoveFilesFromWorkspaceAsync(contents, cancellationToken);
-					await AddFilesToWorkspaceAsync(request._perforceClient, contents, bFakeSync, cancellationToken);
+					await AddFilesToWorkspaceAsync(request.PerforceClient, contents, bFakeSync, cancellationToken);
 				}
 			}
 
@@ -1136,8 +1137,8 @@ namespace EpicGames.Perforce.Managed
 			FileReference configFile = FileReference.Combine(_baseDir, "p4.ini");
 			using (StreamWriter writer = new StreamWriter(configFile.FullName))
 			{
-				writer.WriteLine("P4PORT={0}", perforceClient.Settings.ServerAndPort);
-				writer.WriteLine("P4CLIENT={0}", perforceClient.Settings.ClientName);
+				await writer.WriteLineAsync($"P4PORT={perforceClient.Settings.ServerAndPort}");
+				await writer.WriteLineAsync($"P4CLIENT={perforceClient.Settings.ClientName}");
 			}
 		}
 
@@ -1201,7 +1202,7 @@ namespace EpicGames.Perforce.Managed
 				}
 
 				// Find all the open changes
-				List<ChangesRecord> changes = await perforceClient.GetChangesAsync(ChangesOptions.None, perforceClient.Settings.ClientName!, -1, ChangeStatus.Pending, null, new string[0], cancellationToken);
+				List<ChangesRecord> changes = await perforceClient.GetChangesAsync(ChangesOptions.None, perforceClient.Settings.ClientName!, -1, ChangeStatus.Pending, null, Array.Empty<string>(), cancellationToken);
 
 				// Delete the changelist
 				foreach (ChangesRecord change in changes)
@@ -1212,7 +1213,7 @@ namespace EpicGames.Perforce.Managed
 					{
 						if (record.Files.Count > 0)
 						{
-							await perforceClient.DeleteShelvedFilesAsync(record.Number, new string[0], cancellationToken);
+							await perforceClient.DeleteShelvedFilesAsync(record.Number, Array.Empty<string>(), cancellationToken);
 						}
 					}
 
@@ -1255,7 +1256,7 @@ namespace EpicGames.Perforce.Managed
 				Stopwatch timer = Stopwatch.StartNew();
 
 				// Sync an initial set of files. Either start with a full workspace and remove files, or start with nothing and add files.
-				if (view.Count == 0 || view[0].StartsWith("-"))
+				if (view.Count == 0 || view[0].StartsWith("-", StringComparison.Ordinal))
 				{
 					await UpdateHaveTablePathAsync(perforceClient, $"//{perforceClient.Settings.ClientName}/...@{changeNumber}", cancellationToken);
 				}
@@ -1268,7 +1269,7 @@ namespace EpicGames.Perforce.Managed
 				foreach (string filter in view)
 				{
 					string syncPath;
-					if (filter.StartsWith("-"))
+					if (filter.StartsWith("-", StringComparison.Ordinal))
 					{
 						syncPath = String.Format("//{0}/{1}#0", perforceClient.Settings.ClientName, RemoveLeadingSlash(filter.Substring(1)));
 					}
@@ -1290,7 +1291,7 @@ namespace EpicGames.Perforce.Managed
 		/// <param name="syncPath">Path to sync</param>
 		/// <param name="cancellationToken">Cancellation token</param>
 		/// <returns>Async task</returns>
-		private async Task UpdateHaveTablePathAsync(IPerforceConnection perforceClient, string syncPath, CancellationToken cancellationToken)
+		private static async Task UpdateHaveTablePathAsync(IPerforceConnection perforceClient, string syncPath, CancellationToken cancellationToken)
 		{
 			PerforceResponseList<SyncSummaryRecord> responseList = await perforceClient.TrySyncQuietAsync(SyncOptions.KeepWorkspaceFiles, -1, new[] { syncPath }, cancellationToken);
 			foreach (PerforceResponse<SyncSummaryRecord> response in responseList)
@@ -1654,7 +1655,7 @@ namespace EpicGames.Perforce.Managed
 			}
 			catch (Exception ex)
 			{
-				_logger.LogWarning(ex, "warning: Unable to delete directory {0}", directoryToDelete);
+				_logger.LogWarning(ex, "warning: Unable to delete directory {Directory}", directoryToDelete);
 				_bRequiresRepair = true;
 			}
 		}
@@ -1721,7 +1722,7 @@ namespace EpicGames.Perforce.Managed
 				// Make sure there's enough space on this drive
 				if (RuntimeInformation.IsOSPlatform(OSPlatform.Windows))
 				{
-					long freeSpace = new DriveInfo(Path.GetPathRoot(_baseDir.FullName)).AvailableFreeSpace;
+					long freeSpace = new DriveInfo(Path.GetPathRoot(_baseDir.FullName)!).AvailableFreeSpace;
 					if (freeSpace - syncSize < MinScratchSpace)
 					{
 						throw new InsufficientSpaceException($"Not enough space to sync new files (free space: {freeSpace / (1024.0 * 1024.0):n1}mb, sync size: {syncSize / (1024.0 * 1024.0):n1}mb, min scratch space: {MinScratchSpace / (1024.0 * 1024.0):n1}mb)");
@@ -1774,7 +1775,7 @@ namespace EpicGames.Perforce.Managed
 							{
 								(int batchBeginIdx, int batchEndIdx) = batches[nextBatchIdx];
 
-								Task task = Task.Run(() => SyncBatch(client, filesToSync, batchBeginIdx, batchEndIdx, bFakeSync, cancellationToken));
+								Task task = Task.Run(() => SyncBatch(client, filesToSync, batchBeginIdx, batchEndIdx, bFakeSync, cancellationToken), cancellationToken);
 								tasks[task] = nextBatchIdx++;
 							}
 
@@ -1836,7 +1837,7 @@ namespace EpicGames.Perforce.Managed
 				{
 					FileReference localFile = filesToSync[idx]._workspaceFile.GetLocation();
 					DirectoryReference.CreateDirectory(localFile.Directory);
-					FileReference.WriteAllBytes(localFile, new byte[0]);
+					FileReference.WriteAllBytes(localFile, Array.Empty<byte>());
 				}
 			}
 			else
@@ -1846,14 +1847,14 @@ namespace EpicGames.Perforce.Managed
 				{
 					for (int idx = beginIdx; idx < endIdx; idx++)
 					{
-						writer.WriteLine("{0}#{1}", filesToSync[idx]._streamFile.Path, filesToSync[idx]._streamFile.Revision);
+						await writer.WriteLineAsync($"{filesToSync[idx]._streamFile.Path}#{filesToSync[idx]._streamFile.Revision}");
 					}
 				}
 
 				using PerforceConnection clientWithFileList = new PerforceConnection(client.Settings, client.Logger);
 				clientWithFileList.GlobalOptions.Add($"-x\"{syncFileName}\"");
 
-				await clientWithFileList.SyncAsync(SyncOptions.Force | SyncOptions.FullDepotSyntax, -1, new string[0], cancellationToken).ToListAsync(cancellationToken);
+				await clientWithFileList.SyncAsync(SyncOptions.Force | SyncOptions.FullDepotSyntax, -1, Array.Empty<string>(), cancellationToken).ToListAsync(cancellationToken);
 			}
 		}
 
@@ -2014,7 +2015,7 @@ namespace EpicGames.Perforce.Managed
 		/// </summary>
 		/// <param name="operation">Name of the operation</param>
 		/// <returns>Disposable object for the trace</returns>
-		private IDisposable Trace(string operation)
+		private static IDisposable Trace(string operation)
 		{
 			return TraceSpan.Create(operation, service: "hordeagent_repository");
 		}
diff --git a/Engine/Source/Programs/Shared/EpicGames.Perforce.Managed/StreamSnapshot.cs b/Engine/Source/Programs/Shared/EpicGames.Perforce.Managed/StreamSnapshot.cs
--- a/Engine/Source/Programs/Shared/EpicGames.Perforce.Managed/StreamSnapshot.cs
+++ b/Engine/Source/Programs/Shared/EpicGames.Perforce.Managed/StreamSnapshot.cs
@@ -22,9 +22,9 @@ namespace EpicGames.Perforce.Managed
 		/// <summary>
 		/// Lookup a directory by reference
 		/// </summary>
-		/// <param name="ref">The reference</param>
+		/// <param name="reference">The reference</param>
 		/// <returns></returns>
-		public abstract StreamTree Lookup(StreamTreeRef @ref);
+		public abstract StreamTree Lookup(StreamTreeRef @reference);
 	}
 
 	/// <summary>
diff --git a/Engine/Source/Programs/Shared/EpicGames.Perforce.Managed/StreamSnapshotFromMemory.cs b/Engine/Source/Programs/Shared/EpicGames.Perforce.Managed/StreamSnapshotFromMemory.cs
--- a/Engine/Source/Programs/Shared/EpicGames.Perforce.Managed/StreamSnapshotFromMemory.cs
+++ b/Engine/Source/Programs/Shared/EpicGames.Perforce.Managed/StreamSnapshotFromMemory.cs
@@ -68,9 +68,9 @@ namespace EpicGames.Perforce.Managed
 		}
 
 		/// <inheritdoc/>
-		public override StreamTree Lookup(StreamTreeRef @ref)
+		public override StreamTree Lookup(StreamTreeRef @reference)
 		{
-			return new StreamTree(@ref.Path, HashToTree[@ref.Hash]);
+			return new StreamTree(@reference.Path, HashToTree[@reference.Hash]);
 		}
 
 		/// <summary>
diff --git a/Engine/Source/Programs/Shared/EpicGames.Perforce.Managed/StreamTree.cs b/Engine/Source/Programs/Shared/EpicGames.Perforce.Managed/StreamTree.cs
--- a/Engine/Source/Programs/Shared/EpicGames.Perforce.Managed/StreamTree.cs
+++ b/Engine/Source/Programs/Shared/EpicGames.Perforce.Managed/StreamTree.cs
@@ -94,9 +94,11 @@ namespace EpicGames.Perforce.Managed
 		/// </summary>
 		public Utf8String Path { get; set; }
 
+
 		/// <summary>
 		/// Hash of the tree
 		/// </summary>
+		[System.Diagnostics.CodeAnalysis.SuppressMessage("Naming", "CA1721:Property names should not match get methods", Justification = "Ignored")]
 		public IoHash Hash { get; set; }
 
 		#region Field names
diff --git a/Engine/Source/Programs/Shared/EpicGames.Perforce.Managed/StreamTreeReader.cs b/Engine/Source/Programs/Shared/EpicGames.Perforce.Managed/StreamTreeReader.cs
--- a/Engine/Source/Programs/Shared/EpicGames.Perforce.Managed/StreamTreeReader.cs
+++ b/Engine/Source/Programs/Shared/EpicGames.Perforce.Managed/StreamTreeReader.cs
@@ -15,9 +15,9 @@ namespace EpicGames.Perforce.Managed
 		/// <summary>
 		/// Reads a node of the tree
 		/// </summary>
-		/// <param name="ref"></param>
+		/// <param name="reference"></param>
 		/// <returns></returns>
-		public abstract Task<StreamTree> ReadAsync(StreamTreeRef @ref);
+		public abstract Task<StreamTree> ReadAsync(StreamTreeRef @reference);
 	}
 
 	/// <summary>
@@ -41,9 +41,9 @@ namespace EpicGames.Perforce.Managed
 		}
 
 		/// <inheritdoc/>
-		public override Task<StreamTree> ReadAsync(StreamTreeRef @ref)
+		public override Task<StreamTree> ReadAsync(StreamTreeRef @reference)
 		{
-			return Task.FromResult(new StreamTree(@ref.Path, _hashToTree[@ref.Hash]));
+			return Task.FromResult(new StreamTree(@reference.Path, _hashToTree[@reference.Hash]));
 		}
 	}
 }
diff --git a/Engine/Source/Programs/Shared/EpicGames.Perforce/EpicGames.Perforce.csproj b/Engine/Source/Programs/Shared/EpicGames.Perforce/EpicGames.Perforce.csproj
--- a/Engine/Source/Programs/Shared/EpicGames.Perforce/EpicGames.Perforce.csproj
+++ b/Engine/Source/Programs/Shared/EpicGames.Perforce/EpicGames.Perforce.csproj
@@ -1,7 +1,7 @@
 <Project Sdk="Microsoft.NET.Sdk">
 
   <PropertyGroup>
-    <TargetFramework>netcoreapp3.1</TargetFramework>
+    <TargetFramework>net6.0</TargetFramework>
     <Configuration Condition=" '$(Configuration)' == '' ">Development</Configuration>
     <OutputType>Library</OutputType>
     <RootNamespace>EpicGames.Perforce</RootNamespace>
diff --git a/Engine/Source/Programs/Shared/EpicGames.Perforce/PerforceEnvironment.cs b/Engine/Source/Programs/Shared/EpicGames.Perforce/PerforceEnvironment.cs
--- a/Engine/Source/Programs/Shared/EpicGames.Perforce/PerforceEnvironment.cs
+++ b/Engine/Source/Programs/Shared/EpicGames.Perforce/PerforceEnvironment.cs
@@ -1,4 +1,4 @@
-﻿// Copyright Epic Games, Inc. All Rights Reserved.
+// Copyright Epic Games, Inc. All Rights Reserved.
 
 using System;
 using System.Collections;
@@ -7,6 +7,7 @@ using System.Diagnostics;
 using System.IO;
 using System.Linq;
 using System.Runtime.InteropServices;
+using System.Runtime.Versioning;
 using EpicGames.Core;
 using Microsoft.Win32;
 
@@ -64,6 +65,7 @@ namespace EpicGames.Perforce
 	/// <summary>
 	/// Default global environment used by Linux and MacOS, which reads settings from the registry.
 	/// </summary>
+	[SupportedOSPlatform("windows")]
 	class WindowsGlobalPerforceEnvironment : GlobalPerforceEnvironment
 	{
 		public WindowsGlobalPerforceEnvironment()
diff --git a/Engine/Source/Programs/Shared/EpicGames.Perforce/PerforceResponseList.cs b/Engine/Source/Programs/Shared/EpicGames.Perforce/PerforceResponseList.cs
--- a/Engine/Source/Programs/Shared/EpicGames.Perforce/PerforceResponseList.cs
+++ b/Engine/Source/Programs/Shared/EpicGames.Perforce/PerforceResponseList.cs
@@ -36,7 +36,7 @@ namespace EpicGames.Perforce
 		/// <summary>
 		/// Returns the first error, or null.
 		/// </summary>
-		public PerforceError FirstError => Errors.FirstOrDefault();
+		public PerforceError? FirstError => Errors.FirstOrDefault();
 
 		/// <summary>
 		/// Sequence of all the data objects from the responses.
