EpicGames.MongoDB: net6.0 upgrade

#rnx

diff --git a/Engine/Source/Programs/Shared/EpicGames.MongoDB/EpicGames.MongoDB.csproj b/Engine/Source/Programs/Shared/EpicGames.MongoDB/EpicGames.MongoDB.csproj
--- a/Engine/Source/Programs/Shared/EpicGames.MongoDB/EpicGames.MongoDB.csproj
+++ b/Engine/Source/Programs/Shared/EpicGames.MongoDB/EpicGames.MongoDB.csproj
@@ -1,6 +1,6 @@
 <Project Sdk="Microsoft.NET.Sdk">
   <PropertyGroup>
-    <TargetFramework>netcoreapp3.1</TargetFramework>
+    <TargetFramework>net6.0</TargetFramework>
     <Configuration Condition=" '$(Configuration)' == '' ">Development</Configuration>
     <OutputType>Library</OutputType>
     <RootNamespace>EpicGames.Core</RootNamespace>
diff --git a/Engine/Source/Programs/Shared/EpicGames.MongoDB/ObjectIdHelpers.cs b/Engine/Source/Programs/Shared/EpicGames.MongoDB/ObjectIdHelpers.cs
--- a/Engine/Source/Programs/Shared/EpicGames.MongoDB/ObjectIdHelpers.cs
+++ b/Engine/Source/Programs/Shared/EpicGames.MongoDB/ObjectIdHelpers.cs
@@ -1,10 +1,6 @@
 // Copyright Epic Games, Inc. All Rights Reserved.
 
 using MongoDB.Bson;
-using System;
-using System.Collections.Generic;
-using System.Linq;
-using System.Threading.Tasks;
 
 namespace EpicGames.MongoDB
 {
diff --git a/Engine/Source/Programs/Shared/EpicGames.MongoDB/SingletonDocument.cs b/Engine/Source/Programs/Shared/EpicGames.MongoDB/SingletonDocument.cs
--- a/Engine/Source/Programs/Shared/EpicGames.MongoDB/SingletonDocument.cs
+++ b/Engine/Source/Programs/Shared/EpicGames.MongoDB/SingletonDocument.cs
@@ -16,7 +16,7 @@ namespace EpicGames.MongoDB
 	/// Attribute specifying the unique id for a singleton document
 	/// </summary>
 	[AttributeUsage(AttributeTargets.Class)]
-	public class SingletonDocumentAttribute : Attribute
+	public sealed class SingletonDocumentAttribute : Attribute
 	{
 		/// <summary>
 		/// Unique id for the singleton document
diff --git a/Engine/Source/Programs/Shared/EpicGames.MongoDB/TransactionBuilder.cs b/Engine/Source/Programs/Shared/EpicGames.MongoDB/TransactionBuilder.cs
--- a/Engine/Source/Programs/Shared/EpicGames.MongoDB/TransactionBuilder.cs
+++ b/Engine/Source/Programs/Shared/EpicGames.MongoDB/TransactionBuilder.cs
@@ -122,7 +122,7 @@ namespace EpicGames.MongoDB
 			PropertyInfo IndexerProperty = Reflection<Dictionary<TKey, TValue>>.Indexer;
 			foreach (KeyValuePair<TKey, TValue?> Update in Updates)
 			{
-				MethodCallExpression Index = Expression.Call(Expr.Body, IndexerProperty.GetMethod, new[] { Expression.Constant(Update.Key) });
+				MethodCallExpression Index = Expression.Call(Expr.Body, IndexerProperty.GetMethod!, new[] { Expression.Constant(Update.Key) });
 				if (Update.Value == null)
 				{
 					Expression<Func<TDocument, object>> IndexExpression = Expression.Lambda<Func<TDocument, object>>(Index, Expr.Parameters[0]);
@@ -160,7 +160,7 @@ namespace EpicGames.MongoDB
 			if (Body.NodeType == ExpressionType.MemberAccess)
 			{
 				MemberExpression MemberExpression = (MemberExpression)Body;
-				object? Object = Evaluate(MemberExpression.Expression, Target);
+				object? Object = Evaluate(MemberExpression.Expression!, Target);
 				PropertyInfo PropertyInfo = (PropertyInfo)MemberExpression.Member;
 				PropertyInfo.SetValue(Object, Value);
 			}
@@ -174,14 +174,14 @@ namespace EpicGames.MongoDB
 			else if (Body.NodeType == ExpressionType.Call)
 			{
 				MethodCallExpression CallExpression = (MethodCallExpression)Body;
-				System.Collections.IDictionary Dictionary = (System.Collections.IDictionary)Evaluate(CallExpression.Object, Target)!;
+				System.Collections.IDictionary Dictionary = (System.Collections.IDictionary)Evaluate(CallExpression.Object!, Target)!;
 				object? Key = Evaluate(CallExpression.Arguments[0], null);
 				Dictionary[Key!] = Value;
 			}
 			else if (Body.NodeType == ExpressionType.Index)
 			{
 				IndexExpression IndexExpression = (IndexExpression)Body;
-				object? Object = Evaluate(IndexExpression.Object, Target);
+				object? Object = Evaluate(IndexExpression.Object!, Target);
 
 				object?[] Arguments = new object?[IndexExpression.Arguments.Count];
 				for (int Idx = 0; Idx < IndexExpression.Arguments.Count; Idx++)
@@ -189,7 +189,7 @@ namespace EpicGames.MongoDB
 					Arguments[Idx] = Evaluate(IndexExpression.Arguments[Idx], null);
 				}
 
-				IndexExpression.Indexer.SetValue(Object, Value, Arguments);
+				IndexExpression.Indexer!.SetValue(Object, Value, Arguments);
 			}
 			else
 			{
@@ -222,14 +222,14 @@ namespace EpicGames.MongoDB
 			else if (Body.NodeType == ExpressionType.Call)
 			{
 				MethodCallExpression CallExpression = (MethodCallExpression)Body;
-				System.Collections.IDictionary? Dictionary = (System.Collections.IDictionary?)Evaluate(CallExpression.Object, Target);
+				System.Collections.IDictionary? Dictionary = (System.Collections.IDictionary?)Evaluate(CallExpression.Object!, Target);
 				object? Key = Evaluate(CallExpression.Arguments[0], null);
 				Dictionary!.Remove(Key!);
 			}
 			else if (Body.NodeType == ExpressionType.Index)
 			{
 				IndexExpression IndexExpression = (IndexExpression)Body;
-				System.Collections.IDictionary? Dictionary = (System.Collections.IDictionary?)Evaluate(IndexExpression.Object, Target);
+				System.Collections.IDictionary? Dictionary = (System.Collections.IDictionary?)Evaluate(IndexExpression.Object!, Target);
 				object? Key = Evaluate(IndexExpression.Arguments[0], null);
 				Dictionary!.Remove(Key!);
 			}
@@ -255,7 +255,7 @@ namespace EpicGames.MongoDB
 			if (Expression.NodeType == ExpressionType.Call)
 			{
 				MethodCallExpression CallExpression = (MethodCallExpression)Expression;
-				object? Object = Evaluate(CallExpression.Object, Parameter);
+				object? Object = Evaluate(CallExpression.Object!, Parameter);
 				object?[] Arguments = CallExpression.Arguments.Select(x => Evaluate(x, Parameter)).ToArray();
 				return CallExpression.Method.Invoke(Object, Arguments);
 			}
@@ -267,7 +267,7 @@ namespace EpicGames.MongoDB
 			else if (Expression.NodeType == ExpressionType.MemberAccess)
 			{
 				MemberExpression MemberExpression = (MemberExpression)Expression;
-				object? Target = Evaluate(MemberExpression.Expression, Parameter);
+				object? Target = Evaluate(MemberExpression.Expression!, Parameter);
 
 				MemberInfo Member = MemberExpression.Member;
 				switch (Member.MemberType)
