Fix differences in serialization behaviors between native and legacy Perforce clients.

IPerforceConnection now contains a CreateRecord method to create a PerforceRecord, which adds numbered fields for lists in the legacy client, and unnumbered fields for the native client.

diff --git a/Engine/Source/Programs/Shared/EpicGames.Perforce/IPerforceConnection.cs b/Engine/Source/Programs/Shared/EpicGames.Perforce/IPerforceConnection.cs
--- a/Engine/Source/Programs/Shared/EpicGames.Perforce/IPerforceConnection.cs
+++ b/Engine/Source/Programs/Shared/EpicGames.Perforce/IPerforceConnection.cs
@@ -41,5 +41,12 @@ namespace EpicGames.Perforce
 		/// <param name="cancellationToken">Token used to cancel the operation</param>
 		/// <returns>Response from the server</returns>
 		Task<IPerforceOutput> LoginCommandAsync(string password, CancellationToken cancellationToken = default);
+
+		/// <summary>
+		/// Creates a record from a set of input fields
+		/// </summary>
+		/// <param name="fields">Fields for the record</param>
+		/// <returns>Serialized record data</returns>
+		PerforceRecord CreateRecord(List<KeyValuePair<string, object>> fields);
 	}
 }
diff --git a/Engine/Source/Programs/Shared/EpicGames.Perforce/NativePerforceConnection.cs b/Engine/Source/Programs/Shared/EpicGames.Perforce/NativePerforceConnection.cs
--- a/Engine/Source/Programs/Shared/EpicGames.Perforce/NativePerforceConnection.cs
+++ b/Engine/Source/Programs/Shared/EpicGames.Perforce/NativePerforceConnection.cs
@@ -580,5 +580,8 @@ namespace EpicGames.Perforce
 
 			return result.ToString();
 		}
+
+		/// <inheritdoc/>
+		public PerforceRecord CreateRecord(List<KeyValuePair<string, object>> fields) => PerforceRecord.FromFields(fields, false);
 	}
 }
diff --git a/Engine/Source/Programs/Shared/EpicGames.Perforce/PerforceConnection.cs b/Engine/Source/Programs/Shared/EpicGames.Perforce/PerforceConnection.cs
--- a/Engine/Source/Programs/Shared/EpicGames.Perforce/PerforceConnection.cs
+++ b/Engine/Source/Programs/Shared/EpicGames.Perforce/PerforceConnection.cs
@@ -1,4 +1,4 @@
-﻿// Copyright Epic Games, Inc. All Rights Reserved.
+// Copyright Epic Games, Inc. All Rights Reserved.
 
 using System;
 using System.Collections.Generic;
@@ -9,6 +9,7 @@ using System.Runtime.CompilerServices;
 using System.Text;
 using System.Threading;
 using System.Threading.Tasks;
+using EpicGames.Core;
 using Microsoft.Extensions.Logging;
 
 namespace EpicGames.Perforce
@@ -360,6 +361,9 @@ namespace EpicGames.Perforce
 			}
 		}
 
+		/// <inheritdoc/>
+		public PerforceRecord CreateRecord(List<KeyValuePair<string, object>> fields) => PerforceRecord.FromFields(fields, true);
+
 		#endregion
 	}
 
@@ -489,41 +493,6 @@ namespace EpicGames.Perforce
 			}
 		}
 
-		/// <summary>
-		/// Serializes a list of key/value pairs into binary format.
-		/// </summary>
-		/// <param name="keyValuePairs">List of key value pairs</param>
-		/// <returns>Serialized record data</returns>
-		static byte[] SerializeRecord(List<KeyValuePair<string, object>> keyValuePairs)
-		{
-			MemoryStream stream = new MemoryStream();
-			using (BinaryWriter writer = new BinaryWriter(stream))
-			{
-				writer.Write((byte)'{');
-				foreach (KeyValuePair<string, object> keyValuePair in keyValuePairs)
-				{
-					writer.Write('s');
-					byte[] keyBytes = Encoding.UTF8.GetBytes(keyValuePair.Key);
-					writer.Write((int)keyBytes.Length);
-					writer.Write(keyBytes);
-
-					if (keyValuePair.Value is string str)
-					{
-						writer.Write('s');
-						byte[] valueBytes = Encoding.UTF8.GetBytes(str);
-						writer.Write((int)valueBytes.Length);
-						writer.Write(valueBytes);
-					}
-					else
-					{
-						throw new PerforceException("Unsupported formatting type for {0}", keyValuePair.Key);
-					}
-				}
-				writer.Write((byte)'0');
-			}
-			return stream.ToArray();
-		}
-
 		/// <summary>
 		/// Execute a command and parse the response
 		/// </summary>
@@ -745,7 +714,7 @@ namespace EpicGames.Perforce
 				throw new PerforceException("'Number' field should be set to -1 to create a new changelist.");
 			}
 
-			PerforceResponse response = await SingleResponseCommandAsync(connection, "change", new List<string> { "-i" }, SerializeRecord(record), null, cancellationToken);
+			PerforceResponse response = await SingleResponseCommandAsync(connection, "change", new List<string> { "-i" }, SerializeRecord(connection, record), null, cancellationToken);
 
 			PerforceError? error = response.Error;
 			if (error != null)
@@ -808,7 +777,7 @@ namespace EpicGames.Perforce
 				arguments.Add("-u");
 			}
 
-			return SingleResponseCommandAsync(connection, "change", arguments, SerializeRecord(record), null, cancellationToken);
+			return SingleResponseCommandAsync(connection, "change", arguments, connection.SerializeRecord(record), null, cancellationToken);
 		}
 
 		/// <summary>
@@ -887,9 +856,10 @@ namespace EpicGames.Perforce
 		/// <summary>
 		/// Serializes a change record to a byte array
 		/// </summary>
+		/// <param name="connection">Connection to the Perforce server</param>
 		/// <param name="input">The record to serialize</param>
 		/// <returns>Serialized record</returns>
-		static byte[] SerializeRecord(ChangeRecord input)
+		static byte[] SerializeRecord(this IPerforceConnection connection, ChangeRecord input)
 		{
 			List<KeyValuePair<string, object>> nameToValue = new List<KeyValuePair<string, object>>();
 			if (input.Number == -1)
@@ -916,11 +886,11 @@ namespace EpicGames.Perforce
 			{
 				nameToValue.Add(new KeyValuePair<string, object>("Description", input.Description));
 			}
-			foreach(string file in input.Files)
+			if (input.Files.Count > 0)
 			{
-				nameToValue.Add(new KeyValuePair<string, object>("Files", file));
+				nameToValue.Add(new KeyValuePair<string, object>("Files", input.Files));
 			}
-			return SerializeRecord(nameToValue);
+			return connection.CreateRecord(nameToValue).Serialize();
 		}
 
 		#endregion
@@ -1182,7 +1152,7 @@ namespace EpicGames.Perforce
 		/// <returns>Response from the server</returns>
 		public static Task<PerforceResponse> TryUpdateClientAsync(this IPerforceConnection connection, ClientRecord record, CancellationToken cancellationToken = default)
 		{
-			return SingleResponseCommandAsync(connection, "client", new List<string> { "-i" }, SerializeRecord(record), null, cancellationToken);
+			return SingleResponseCommandAsync(connection, "client", new List<string> { "-i" }, connection.SerializeRecord(record), null, cancellationToken);
 		}
 
 		/// <summary>
@@ -1348,9 +1318,10 @@ namespace EpicGames.Perforce
 		/// <summary>
 		/// Serializes a client record to a byte array
 		/// </summary>
+		/// <param name="connection">Connection to the Perforce server</param>
 		/// <param name="input">The input record</param>
 		/// <returns>Serialized record data</returns>
-		static byte[] SerializeRecord(ClientRecord input)
+		static byte[] SerializeRecord(this IPerforceConnection connection, ClientRecord input)
 		{
 			List<KeyValuePair<string, object>> nameToValue = new List<KeyValuePair<string, object>>();
 			if (input.Name != null)
@@ -1393,11 +1364,11 @@ namespace EpicGames.Perforce
 			{
 				nameToValue.Add(new KeyValuePair<string, object>("Stream", input.Stream));
 			}
-			for (int idx = 0; idx < input.View.Count; idx++)
+			if (input.View.Count > 0)
 			{
-				nameToValue.Add(new KeyValuePair<string, object>(String.Format("View{0}", idx), input.View[idx]));
+				nameToValue.Add(new KeyValuePair<string, object>("View", input.View));
 			}
-			return SerializeRecord(nameToValue);
+			return connection.CreateRecord(nameToValue).Serialize();
 		}
 
 		#endregion
diff --git a/Engine/Source/Programs/Shared/EpicGames.Perforce/PerforceRecord.cs b/Engine/Source/Programs/Shared/EpicGames.Perforce/PerforceRecord.cs
--- a/Engine/Source/Programs/Shared/EpicGames.Perforce/PerforceRecord.cs
+++ b/Engine/Source/Programs/Shared/EpicGames.Perforce/PerforceRecord.cs
@@ -4,6 +4,9 @@ using System;
 using System.Buffers.Binary;
 using System.Buffers.Text;
 using System.Collections.Generic;
+using System.IO;
+using System.Linq;
+using System.Text;
 using System.Text.RegularExpressions;
 using EpicGames.Core;
 
@@ -44,6 +47,22 @@ namespace EpicGames.Perforce
 			Data = data;
 		}
 
+		/// <summary>
+		/// Create a value from a string
+		/// </summary>
+		/// <param name="text">String value to serialize</param>
+		public PerforceValue(string text)
+		{
+			int length = Encoding.UTF8.GetByteCount(text);
+
+			byte[] buffer = new byte[1 + 4 + length];
+			buffer[0] = (byte)'s';
+			BinaryPrimitives.WriteInt32LittleEndian(buffer.AsSpan(1, 4), length);
+			Encoding.UTF8.GetBytes(text, buffer.AsSpan(5));
+
+			Data = buffer;
+		}
+
 		/// <summary>
 		/// Determines if the value is empty
 		/// </summary>
@@ -217,6 +236,40 @@ namespace EpicGames.Perforce
 		/// </summary>
 		public List<KeyValuePair<Utf8String, PerforceValue>> Rows { get; } = new List<KeyValuePair<Utf8String, PerforceValue>>();
 
+		/// <summary>
+		/// Create a record from a set of fields
+		/// </summary>
+		/// <param name="fields">Fields to serialize</param>
+		/// <param name="numberedListElements">Whether to number list elements when flattening the record</param>
+		/// <returns>Record containing the given fields</returns>
+		public static PerforceRecord FromFields(IEnumerable<KeyValuePair<string, object>> fields, bool numberedListElements)
+		{
+			PerforceRecord record = new PerforceRecord();
+			foreach ((string name, object value) in fields)
+			{
+				if (value is string str)
+				{
+					record.Rows.Add(new KeyValuePair<Utf8String, PerforceValue>(name, new PerforceValue(str)));
+				}
+				else if (value is List<string> list)
+				{
+					if (numberedListElements)
+					{
+						record.Rows.AddRange(Enumerable.Range(0, list.Count).Select(x => new KeyValuePair<Utf8String, PerforceValue>($"{name}{x}", new PerforceValue(list[x]))));
+					}
+					else
+					{
+						record.Rows.AddRange(list.Select(x => new KeyValuePair<Utf8String, PerforceValue>(name, new PerforceValue(x))));
+					}
+				}
+				else
+				{
+					throw new PerforceException("Unsupported formatting type for {0}", name);
+				}
+			}
+			return record;
+		}
+
 		/// <summary>
 		/// Copy this record into the given array of values. This method is O(n) if every record key being in the list of keys in the same order, but O(n^2) if not.
 		/// </summary>
@@ -260,5 +313,63 @@ namespace EpicGames.Perforce
 				}
 			}
 		}
+
+		/// <summary>
+		/// Serialize to a single block of data
+		/// </summary>
+		/// <returns>Serialized data</returns>
+		public byte[] Serialize()
+		{
+			MemoryStream stream = new MemoryStream();
+			using (BinaryWriter writer = new BinaryWriter(stream))
+			{
+				writer.Write((byte)'{');
+				foreach ((Utf8String name, PerforceValue value) in Rows)
+				{
+					writer.Write('s');
+					writer.Write((int)name.Length);
+					writer.Write(name.Span);
+					writer.Write(value.Data.Span);
+				}
+				writer.Write((byte)'0');
+			}
+			return stream.ToArray();
+		}
+
+
+		/// <summary>
+		/// Serializes a list of key/value pairs into binary format.
+		/// </summary>
+		/// <param name="keyValuePairs">List of key value pairs</param>
+		/// <returns>Serialized record data</returns>
+		public static byte[] Serialize(List<KeyValuePair<string, object>> keyValuePairs)
+		{
+			MemoryStream stream = new MemoryStream();
+			using (BinaryWriter writer = new BinaryWriter(stream))
+			{
+				writer.Write((byte)'{');
+				foreach (KeyValuePair<string, object> keyValuePair in keyValuePairs)
+				{
+					writer.Write('s');
+					byte[] keyBytes = Encoding.UTF8.GetBytes(keyValuePair.Key);
+					writer.Write((int)keyBytes.Length);
+					writer.Write(keyBytes);
+
+					if (keyValuePair.Value is string str)
+					{
+						writer.Write('s');
+						byte[] valueBytes = Encoding.UTF8.GetBytes(str);
+						writer.Write((int)valueBytes.Length);
+						writer.Write(valueBytes);
+					}
+					else
+					{
+						throw new PerforceException("Unsupported formatting type for {0}", keyValuePair.Key);
+					}
+				}
+				writer.Write((byte)'0');
+			}
+			return stream.ToArray();
+		}
 	}
 }
