UGS: Stop finalizer in NativePerforceConnection waiting for a thread.

diff --git a/Engine/Source/Programs/Shared/EpicGames.Perforce/NativePerforceConnection.cs b/Engine/Source/Programs/Shared/EpicGames.Perforce/NativePerforceConnection.cs
--- a/Engine/Source/Programs/Shared/EpicGames.Perforce/NativePerforceConnection.cs
+++ b/Engine/Source/Programs/Shared/EpicGames.Perforce/NativePerforceConnection.cs
@@ -351,19 +351,29 @@ namespace EpicGames.Perforce
 		/// </summary>
 		~NativePerforceConnection()
 		{
-			Dispose();
+			Dispose(false);
 		}
 
 		/// <inheritdoc/>
 		public void Dispose()
+		{
+			Dispose(true);
+			GC.SuppressFinalize(this);
+		}
+
+		/// <inheritdoc/>
+		void Dispose(bool disposing)
 		{
 			if (_backgroundThread != null)
 			{
 				_requests.Add(null);
 				_requests.CompleteAdding();
 
-				_backgroundThread.Join();
-				_backgroundThread = null!;
+				if (disposing)
+				{
+					_backgroundThread.Join();
+					_backgroundThread = null!;
+				}
 			}
 
 			_requests.Dispose();
@@ -382,8 +392,6 @@ namespace EpicGames.Perforce
 			{
 				pinnedBuffer.Dispose();
 			}
-
-			GC.SuppressFinalize(this);
 		}
 
 		/// <summary>
