PR #9230: Fix AutomationTool ignoring -Timestamps flag (Contributed by [REDACTED_PERSON])

### Notes
This seems like a typo/oversight from [REDACTED_HANDLE]

### Testing
Tested with a local UBT run

#rnx

diff --git a/Engine/Source/Programs/Shared/EpicGames.Core/Log.cs b/Engine/Source/Programs/Shared/EpicGames.Core/Log.cs
--- a/Engine/Source/Programs/Shared/EpicGames.Core/Log.cs
+++ b/Engine/Source/Programs/Shared/EpicGames.Core/Log.cs
@@ -1174,7 +1174,7 @@ namespace EpicGames.Core
 						foreach (string line in lines)
 						{
 							string lineWithTime = timePrefixActual + line;
-							listener.WriteLine(line);
+							listener.WriteLine(lineWithTime);
 							listener.Flush();
 						}
 					}
