Moved code related to OIDC from OidcToken (helper application) into a library EpicGames.OIDC this will be used both by OidcToken and UnrealGameSync for allocating tokens.

diff --git a/Engine/Source/Programs/Shared/EpicGames.OIDC/EpicGames.OIDC.csproj b/Engine/Source/Programs/Shared/EpicGames.OIDC/EpicGames.OIDC.csproj
new file mode 100644
--- /dev/null
+++ b/Engine/Source/Programs/Shared/EpicGames.OIDC/EpicGames.OIDC.csproj
@@ -0,0 +1,22 @@
+<Project Sdk="Microsoft.NET.Sdk">
+
+  <PropertyGroup>
+    <OutputType>Library</OutputType>
+    <RootNamespace>EpicGames.OIDC</RootNamespace>
+    <TargetFramework>netcoreapp3.1</TargetFramework>
+    <Nullable>enable</Nullable>
+    <GenerateAssemblyInfo>false</GenerateAssemblyInfo>
+    <GenerateTargetFrameworkAttribute>false</GenerateTargetFrameworkAttribute>
+  </PropertyGroup>
+
+  <ItemGroup>
+    <Compile Include="..\MetaData.cs">
+      <Link>Properties\MetaData.cs</Link>
+    </Compile>
+  </ItemGroup>
+  
+  <ItemGroup>
+    <PackageReference Include="IdentityModel.OidcClient" Version="4.0.0" />
+  </ItemGroup>
+
+</Project>
diff --git a/Engine/Source/Programs/Shared/EpicGames.OIDC/ITokenStore.cs b/Engine/Source/Programs/Shared/EpicGames.OIDC/ITokenStore.cs
new file mode 100644
--- /dev/null
+++ b/Engine/Source/Programs/Shared/EpicGames.OIDC/ITokenStore.cs
@@ -0,0 +1,56 @@
+// Copyright Epic Games, Inc. All Rights Reserved.
+
+using System;
+using System.Diagnostics.CodeAnalysis;
+using System.Runtime.InteropServices;
+using Microsoft.Extensions.DependencyInjection;
+
+namespace EpicGames.OIDC
+{
+	public interface ITokenStore: IDisposable
+	{
+		public bool TryGetRefreshToken(string oidcProvider, [MaybeNullWhen(false)] out string refreshToken);
+
+		public void AddRefreshToken(string providerIdentifier, string refreshToken);
+		public void Save();
+	}
+
+	public static class TokenStoreFactory
+	{
+		public static ITokenStore CreateTokenStore(IServiceProvider provider)
+		{
+			if (RuntimeInformation.IsOSPlatform(OSPlatform.Windows))
+			{
+				return ActivatorUtilities.CreateInstance<WindowsTokenStore>(provider);
+			}
+			else if (RuntimeInformation.IsOSPlatform(OSPlatform.OSX))
+			{
+				return ActivatorUtilities.CreateInstance<MacOSTokenStore>(provider);
+			}
+			else if (RuntimeInformation.IsOSPlatform(OSPlatform.Linux))
+			{
+				return ActivatorUtilities.CreateInstance<FilesystemTokenStore>(provider);
+			}
+
+			throw new NotSupportedException("Unknown platform when attempting to create ITokenStore");
+		}
+
+		public static ITokenStore CreateTokenStore()
+		{
+			if (RuntimeInformation.IsOSPlatform(OSPlatform.Windows))
+			{
+				return new WindowsTokenStore();
+			}
+			else if (RuntimeInformation.IsOSPlatform(OSPlatform.OSX))
+			{
+				return new MacOSTokenStore();
+			}
+			else if (RuntimeInformation.IsOSPlatform(OSPlatform.Linux))
+			{
+				return new FilesystemTokenStore();
+			}
+
+			throw new NotSupportedException("Unknown platform when attempting to create ITokenStore");
+		}
+	}
+}
\ No newline at end of file
diff --git a/Engine/Source/Programs/Shared/EpicGames.OIDC/MacOSTokenStore.cs b/Engine/Source/Programs/Shared/EpicGames.OIDC/MacOSTokenStore.cs
new file mode 100644
--- /dev/null
+++ b/Engine/Source/Programs/Shared/EpicGames.OIDC/MacOSTokenStore.cs
@@ -0,0 +1,111 @@
+// Copyright Epic Games, Inc. All Rights Reserved.
+
+using System;
+using System.Collections.Generic;
+using System.Diagnostics;
+using System.Text;
+
+namespace EpicGames.OIDC
+{
+	public class MacOSTokenStore : ITokenStore
+	{
+		public MacOSTokenStore()
+		{
+		}
+
+		public bool TryGetRefreshToken(string oidcProvider, out string refreshToken)
+		{
+			string? token = KeychainHelper.GetKeychainItem(oidcProvider);
+
+			if (token == null)
+			{
+				refreshToken = "";
+				return false;
+			}
+
+			refreshToken = token;
+			return true;
+		}
+
+		public void AddRefreshToken(string providerIdentifier, string refreshToken)
+		{
+			KeychainHelper.AddKeychainItem(refreshToken, providerIdentifier);
+		}
+
+		public void Save()
+		{
+			
+		}
+
+		public void Dispose()
+		{
+		}
+	}
+
+	static class KeychainHelper
+	{
+		private const string OidcTokenIdentifier = "EpicGames OidcToken";
+
+		private static string? RunKeychainHandler(string verb, List<string> args)
+		{
+			ProcessStartInfo startInfo = new ProcessStartInfo()
+			{
+				FileName = "security",
+				RedirectStandardOutput = true,
+			};
+			startInfo.ArgumentList.Add(verb);
+			foreach (string s in args)
+			{
+				startInfo.ArgumentList.Add(s);
+			}
+
+			Process? p = Process.Start(startInfo);
+			if (p == null)
+			{
+				throw new Exception("Failed to run \"security\" to interact with keychain");
+			}
+
+			StringBuilder sb = new StringBuilder();
+			p.OutputDataReceived += (o, e) => { sb.Append(e.Data); };
+			p.BeginOutputReadLine();
+			p.WaitForExit();
+
+			if (p.ExitCode == 44)
+			{
+				return null; // not found
+			}
+
+			if (p.ExitCode != 0)
+			{
+				throw new Exception("Unhandled exitcode {p.ExitCode} when interacting with keychain");
+			}
+
+			return sb.ToString();
+		}
+
+		public static void AddKeychainItem(string dataToEncrypt, string itemIdentifier)
+		{
+			RunKeychainHandler("add-generic-password", new List<string> {
+				"-a",
+				itemIdentifier,
+				"-s",
+				OidcTokenIdentifier,
+				"-w",
+				dataToEncrypt
+			});
+		}
+
+		public static string? GetKeychainItem(string itemIdentifier)
+		{
+			string? output = RunKeychainHandler("find-generic-password", new List<string> {
+				"-a",
+				itemIdentifier,
+				"-s",
+				OidcTokenIdentifier,
+				"-w"
+			});
+
+			return output;
+		}
+	}
+}
\ No newline at end of file
diff --git a/Engine/Source/Programs/Shared/EpicGames.OIDC/OidcTokenManager.cs b/Engine/Source/Programs/Shared/EpicGames.OIDC/OidcTokenManager.cs
new file mode 100644
--- /dev/null
+++ b/Engine/Source/Programs/Shared/EpicGames.OIDC/OidcTokenManager.cs
@@ -0,0 +1,464 @@
+// Copyright Epic Games, Inc. All Rights Reserved.
+
+using System;
+using System.Collections.Generic;
+using System.ComponentModel.DataAnnotations;
+using System.Diagnostics;
+using System.IO;
+using System.Linq;
+using System.Net;
+using System.Net.Http;
+using System.Runtime.InteropServices;
+using System.Text;
+using System.Threading.Tasks;
+using IdentityModel.Client;
+using IdentityModel.OidcClient;
+using IdentityModel.OidcClient.Results;
+using Microsoft.Extensions.DependencyInjection;
+using Microsoft.Extensions.Logging;
+using Microsoft.Extensions.Options;
+
+namespace EpicGames.OIDC
+{
+	public class OidcTokenManager
+	{
+		private readonly Dictionary<string, OidcTokenClient> _tokenClients = new Dictionary<string, OidcTokenClient>();
+		private readonly ITokenStore _tokenStore;
+
+		public OidcTokenManager(IServiceProvider provider, IOptionsMonitor<OidcTokenOptions> settings, ITokenStore tokenStore)
+		{
+			_tokenStore = tokenStore;
+
+			Dictionary<string, string> refreshTokens = new Dictionary<string, string>();
+			foreach (string oidcProvider in settings.CurrentValue.Providers.Keys)
+			{
+				if (tokenStore.TryGetRefreshToken(oidcProvider, out string? refreshToken))
+				{
+					refreshTokens.TryAdd(oidcProvider, refreshToken);
+				}
+			}
+
+			Providers = settings.CurrentValue.Providers;
+			foreach ((string key, ProviderInfo providerInfo) in settings.CurrentValue.Providers)
+			{
+				OidcTokenClient tokenClient = ActivatorUtilities.CreateInstance<OidcTokenClient>(provider, providerInfo);
+
+				if (refreshTokens.TryGetValue(key, out string? refreshToken))
+				{
+					tokenClient.SetRefreshToken(refreshToken);
+				}
+
+				_tokenClients.Add(key, tokenClient);
+			}
+		}
+
+		public OidcTokenManager(Dictionary<string, ProviderInfo> providers, ITokenStore tokenStore)
+		{
+			_tokenStore = tokenStore;
+
+			Dictionary<string, string> refreshTokens = new Dictionary<string, string>();
+			foreach (string oidcProvider in providers.Keys)
+			{
+				if (tokenStore.TryGetRefreshToken(oidcProvider, out string? refreshToken))
+				{
+					refreshTokens.TryAdd(oidcProvider, refreshToken);
+				}
+			}
+
+			Providers = providers;
+			foreach ((string key, ProviderInfo providerInfo) in Providers)
+			{
+				OidcTokenClient tokenClient = new OidcTokenClient(providerInfo);
+
+				if (refreshTokens.TryGetValue(key, out string? refreshToken))
+				{
+					tokenClient.SetRefreshToken(refreshToken);
+				}
+
+				_tokenClients.Add(key, tokenClient);
+			}
+		}
+
+		public Dictionary<string, ProviderInfo> Providers { get; }
+
+		public bool HasUnfinishedLogin()
+		{
+			return _tokenClients.Any(pair => pair.Value.GetStatusForProvider() == OidcStatus.NotLoggedIn);
+		}
+
+		public async Task<OidcTokenInfo> Login(string providerIdentifier)
+		{
+			OidcTokenClient tokenClient = _tokenClients[providerIdentifier];
+
+			OidcTokenInfo tokenInfo = await tokenClient.Login();
+			
+			if (!String.IsNullOrEmpty(tokenInfo.RefreshToken))
+			{
+				_tokenStore.AddRefreshToken(providerIdentifier, tokenInfo.RefreshToken);
+			}
+
+			return tokenInfo;
+		}
+
+		public Task<OidcTokenInfo> GetAccessToken(string providerIdentifier)
+		{
+			return _tokenClients[providerIdentifier].GetAccessToken();
+		}
+
+		public OidcStatus GetStatusForProvider(string providerIdentifier)
+		{
+			return _tokenClients[providerIdentifier].GetStatusForProvider();
+		}
+	}
+
+	public enum OidcStatus
+	{
+		Connected,
+		NotLoggedIn,
+		TokenRefreshRequired
+	}
+
+	public class OidcTokenInfo
+	{
+		public string? IdentityToken { get; set; }
+		public string? RefreshToken { get; set; }
+		public string? AccessToken { get; set; }
+		public DateTimeOffset TokenExpiry { get; set; }
+
+		public bool IsValid => IdentityToken != null && RefreshToken != null && AccessToken != null;
+	};
+
+	class OidcTokenClient
+	{
+		private readonly Uri _authorityUri;
+		private readonly string _clientId;
+		private readonly Uri _redirectUri;
+		private readonly string _scopes;
+
+		private OidcClient? _oidcClient;
+
+		private string? _refreshToken;
+		private string? _accessToken;
+		private DateTimeOffset _tokenExpiry;
+
+		private bool _initialized = false;
+
+		private readonly ILogger<OidcTokenClient>? _logger;
+
+		public OidcTokenClient(ILogger<OidcTokenClient> logger, ProviderInfo providerInfo)
+		{
+			_logger = logger;
+
+			_authorityUri = providerInfo.ServerUri;
+			_clientId = providerInfo.ClientId;
+			_redirectUri = providerInfo.RedirectUri;
+			_scopes = providerInfo.Scopes;
+		}
+
+		public OidcTokenClient(ProviderInfo providerInfo)
+		{
+			_logger = null;
+
+			_authorityUri = providerInfo.ServerUri;
+			_clientId = providerInfo.ClientId;
+			_redirectUri = providerInfo.RedirectUri;
+			_scopes = providerInfo.Scopes;
+		}
+
+		private async Task Initialize()
+		{
+			if (_initialized)
+			{
+				return;
+			}
+
+			OidcClientOptions options = new OidcClientOptions
+			{
+				Authority = _authorityUri.ToString(),
+				Policy = new Policy { Discovery = new DiscoveryPolicy { Authority = _authorityUri.ToString() } },
+				ClientId = _clientId,
+				Scope = _scopes,
+				FilterClaims = false,
+				RedirectUri = _redirectUri.ToString(),
+				//Flow = OidcClientOptions.AuthenticationFlow.AuthorizationCode,
+				//ResponseMode = OidcClientOptions.AuthorizeResponseMode.Redirect
+			};
+
+			// we need to fetch the discovery document ourselves to support OIDC Authorities which have a subresource for it
+			// with Okta has for authorization servers for instance.
+			DiscoveryDocumentResponse discoveryDocument = await GetDiscoveryDocument();
+			options.ProviderInformation = new ProviderInformation
+			{
+				IssuerName = discoveryDocument.Issuer,
+				KeySet = discoveryDocument.KeySet,
+
+				AuthorizeEndpoint = discoveryDocument.AuthorizeEndpoint,
+				TokenEndpoint = discoveryDocument.TokenEndpoint,
+				EndSessionEndpoint = discoveryDocument.EndSessionEndpoint,
+				UserInfoEndpoint = discoveryDocument.UserInfoEndpoint,
+				TokenEndPointAuthenticationMethods = discoveryDocument.TokenEndpointAuthenticationMethodsSupported
+			};
+
+			_oidcClient = new OidcClient(options);
+			_initialized = true;
+		}
+
+		public async Task<OidcTokenInfo> Login()
+		{
+			await Initialize();
+
+			// setup a local http server to listen for the result of the login
+			using HttpListener http = new HttpListener();
+			Uri uri = _redirectUri;
+			// build the url the server should be hosted at
+			string prefix = $"{uri.Scheme}{Uri.SchemeDelimiter}{uri.Authority}/";
+			http.Prefixes.Add(prefix);
+			http.Start();
+
+			// generate the appropriate codes we need to login
+			AuthorizeState loginState = await _oidcClient!.PrepareLoginAsync();
+			// start the user browser
+			OpenBrowser(loginState.StartUrl);
+
+			string? responseData = await ProcessHttpRequest(http);
+
+			// parse the returned url for the tokens needed to complete the login
+			LoginResult loginResult = await _oidcClient.ProcessResponseAsync(responseData, loginState);
+
+			http.Stop();
+
+			if (loginResult.IsError)
+			{
+				throw new LoginFailedException("Failed to login due to error: " + loginResult.Error);
+			}
+
+			_refreshToken = loginResult.RefreshToken;
+			_accessToken = loginResult.AccessToken;
+			_tokenExpiry = loginResult.AccessTokenExpiration;
+
+			return new OidcTokenInfo
+			{
+				IdentityToken = loginResult.IdentityToken,
+				RefreshToken = loginResult.RefreshToken,
+				AccessToken = loginResult.AccessToken,
+				TokenExpiry = loginResult.AccessTokenExpiration
+			};
+		}
+
+		private static async Task<string?> ProcessHttpRequest(HttpListener http)
+		{
+			HttpListenerContext context = await http.GetContextAsync();
+			string? responseData;
+			switch (context.Request.HttpMethod)
+			{
+				case "GET":
+					responseData = context.Request.RawUrl;
+					break;
+				case "POST":
+				{
+					HttpListenerRequest request = context.Request;
+					if (request.ContentType != null && !request.ContentType.Equals("application/x-www-form-urlencoded",
+						StringComparison.OrdinalIgnoreCase))
+					{
+						// we do not support url encoded return types
+						context.Response.StatusCode = 415;
+						return null;
+					}
+
+					// attempt to parse the body
+
+					// if there is no body we can not handle the post
+					if (!context.Request.HasEntityBody)
+					{
+						context.Response.StatusCode = 415;
+						return null;
+					}
+
+					await using Stream body = request.InputStream;
+					using StreamReader reader = new StreamReader(body, request.ContentEncoding);
+					responseData = await reader.ReadToEndAsync();
+					break;
+				}
+				default:
+					// if we receive any other http method something is very odd. Tell them to use a different method.
+					context.Response.StatusCode = 415;
+					return null;
+			}
+
+			// generate a simple http page to show the user
+			HttpListenerResponse response = context.Response;
+			const string httpPage = "<html><head></head><body><h1 style=\"text-align: center;\">Please close this browser to finish your login.<h1></body></html>";
+			byte[] buffer = Encoding.UTF8.GetBytes(httpPage);
+			response.ContentLength64 = buffer.Length;
+			Stream responseOutput = response.OutputStream;
+			await responseOutput.WriteAsync(buffer, 0, buffer.Length);
+			responseOutput.Close();
+
+			return responseData;
+		}
+
+		private void OpenBrowser(string url)
+		{
+			_logger?.LogInformation("Opening system browser to {Url}", url);
+
+			try
+			{
+				Process.Start(url);
+			}
+			catch
+			{
+				// hack because of this: [REDACTED_PROFILE_URL]/corefx/issues/10361
+				if (RuntimeInformation.IsOSPlatform(OSPlatform.Windows))
+				{
+					url = url.Replace("&", "^&", StringComparison.InvariantCultureIgnoreCase);
+					Process.Start(new ProcessStartInfo("cmd", $"/c start {url}") { CreateNoWindow = true });
+				}
+				else if (RuntimeInformation.IsOSPlatform(OSPlatform.Linux))
+				{
+					Process.Start("xdg-open", url);
+				}
+				else if (RuntimeInformation.IsOSPlatform(OSPlatform.OSX))
+				{
+					Process.Start("open", url);
+				}
+				else
+				{
+					throw;
+				}
+			}
+		}
+
+		private async Task<OidcTokenInfo> DoRefreshToken(string inRefreshToken)
+		{
+			await Initialize();
+
+			// use the refresh token to acquire a new access token
+			RefreshTokenResult refreshTokenResult = await _oidcClient!.RefreshTokenAsync(inRefreshToken);
+
+			if (refreshTokenResult.IsError)
+			{
+				if (refreshTokenResult.Error == "invalid_grant")
+				{
+					// the refresh token is no logger valid, resetting it and treating us as not logged in
+					_refreshToken = null;
+					throw new NotLoggedInException();
+				}
+				else
+				{
+					throw new Exception($"Error using the refresh token: {refreshTokenResult.Error} , details: {refreshTokenResult.ErrorDescription}");
+				}
+			}
+			
+			_refreshToken = refreshTokenResult.RefreshToken;
+			_accessToken = refreshTokenResult.AccessToken;
+			_tokenExpiry = refreshTokenResult.AccessTokenExpiration;
+
+			return new OidcTokenInfo
+			{
+				IdentityToken = refreshTokenResult.IdentityToken,
+				RefreshToken = refreshTokenResult.RefreshToken,
+				AccessToken = refreshTokenResult.AccessToken,
+				TokenExpiry = refreshTokenResult.AccessTokenExpiration
+			};
+		}
+
+		private async Task<DiscoveryDocumentResponse> GetDiscoveryDocument()
+		{
+			string discoUrl = $"{_authorityUri}/.well-known/openid-configuration";
+
+			using HttpClient client = new HttpClient();
+			using DiscoveryDocumentRequest doc = new DiscoveryDocumentRequest
+			{
+				Address = discoUrl,
+				Policy =
+				{
+					ValidateEndpoints = false
+				}
+			};
+			DiscoveryDocumentResponse disco = await client.GetDiscoveryDocumentAsync(doc);
+
+			if (disco.IsError)
+			{
+				throw new Exception(disco.Error);
+			}
+
+			return disco;
+		}
+		public async Task<OidcTokenInfo> GetAccessToken()
+		{
+			if (String.IsNullOrEmpty(_refreshToken))
+			{
+				throw new NotLoggedInException();
+			}
+
+			// if the token is valid for another few minutes we can use it
+			// we avoid using a token that is about to expire to make sure we can finish the call we expect to do with it before it expires
+			if (!String.IsNullOrEmpty(_accessToken) && _tokenExpiry.AddMinutes(2) > DateTime.Now)
+			{
+				return new OidcTokenInfo
+				{
+					RefreshToken = _refreshToken,
+					AccessToken = _accessToken,
+					TokenExpiry = _tokenExpiry
+				};
+			}
+
+			return await DoRefreshToken(_refreshToken);
+		}
+
+		public OidcStatus GetStatusForProvider()
+		{
+			if (String.IsNullOrEmpty(_refreshToken))
+			{
+				return OidcStatus.NotLoggedIn;
+			}
+
+			if (String.IsNullOrEmpty(_accessToken))
+			{
+				return OidcStatus.TokenRefreshRequired;
+			}
+
+			if (_tokenExpiry < DateTime.Now)
+			{
+				return OidcStatus.TokenRefreshRequired;
+			}
+
+			return OidcStatus.Connected;
+		}
+
+		public void SetRefreshToken(string inRefreshToken)
+		{
+			_refreshToken = inRefreshToken;
+		}
+	}
+
+	public class LoginFailedException : Exception
+	{
+		public LoginFailedException(string message) : base(message)
+		{
+		}
+	}
+
+	public class NotLoggedInException : Exception
+	{
+	}
+
+	public class OidcTokenOptions
+	{
+		[System.Diagnostics.CodeAnalysis.SuppressMessage("Usage", "CA2227:Collection properties should be read only", Justification = "Used by the configuration system")]
+		public Dictionary<string, ProviderInfo> Providers { get; set; } = new Dictionary<string, ProviderInfo>();
+	}
+
+	public class ProviderInfo
+	{
+		[Required] public Uri ServerUri { get; set; } = null!;
+
+		[Required] public string ClientId { get; set; } = null!;
+
+		[Required] public string DisplayName { get; set; } = null!;
+
+		[Required] public Uri RedirectUri { get; set; } = null!;
+
+		public string Scopes { get; set; } = "openid profile offline_access";
+	}
+}
\ No newline at end of file
diff --git a/Engine/Source/Programs/Shared/EpicGames.OIDC/TokenStore.cs b/Engine/Source/Programs/Shared/EpicGames.OIDC/TokenStore.cs
new file mode 100644
--- /dev/null
+++ b/Engine/Source/Programs/Shared/EpicGames.OIDC/TokenStore.cs
@@ -0,0 +1,148 @@
+// Copyright Epic Games, Inc. All Rights Reserved.
+
+using System;
+using System.Collections.Generic;
+using System.IO;
+using System.Security.Cryptography;
+using System.Text;
+using System.Text.Json;
+using System.Text.Json.Serialization;
+
+namespace EpicGames.OIDC
+{
+	/// <summary>
+	/// A generic token store that saves the offline token in a file on disk, this works on any platform
+	/// </summary>
+	public class FilesystemTokenStore : ITokenStore, IDisposable
+	{
+		private class TokenStoreState
+		{
+			public Dictionary<string, string> Providers { get; set; } = new Dictionary<string, string>();
+			public string Key { get; set; }
+
+			[JsonConstructor]
+			public TokenStoreState(string key, Dictionary<string, string> providers)
+			{
+				Key = key;
+				Providers = providers;
+			}
+
+			public TokenStoreState(byte[] key, Dictionary<string, byte[]> providers)
+			{
+				Key = Convert.ToBase64String(key);
+
+				foreach (KeyValuePair<string, byte[]> pair in providers)
+				{
+					Providers[pair.Key] = Convert.ToBase64String(pair.Value);
+				}
+			}
+		}
+
+		private readonly Dictionary<string, byte[]> _providerToRefreshToken = new Dictionary<string, byte[]>();
+		private readonly SymmetricAlgorithm _crypt = Aes.Create();
+
+		public FilesystemTokenStore()
+		{
+			ReadStoreFromDisk();
+		}
+
+		private static FileInfo GetStorePath()
+		{
+			return new FileInfo(Path.Combine(Environment.GetFolderPath(Environment.SpecialFolder.LocalApplicationData),
+				"oidc-token", "tokenStore.dat"));
+		}
+
+		private void ReadStoreFromDisk()
+		{
+			FileInfo fi = GetStorePath();
+			if (!fi.Exists)
+			{
+				// if we have no store on disk then we just initialize it to empty
+				return;
+			}
+
+			using FileStream fs = fi.OpenRead();
+			using TextReader tr = new StreamReader(fs);
+			TokenStoreState? state = JsonSerializer.Deserialize<TokenStoreState>(tr.ReadToEnd());
+
+			if (state == null)
+			{
+				// if we fail to deserialize the state just drop it, will mean users will need to login again
+				return;
+			}
+
+			foreach ((string key, string value) in state.Providers)
+			{
+				_providerToRefreshToken[key] = Convert.FromBase64String(value);
+			}
+
+			_crypt.IV = Convert.FromBase64String(state.Key);
+		}
+
+		private void SaveStoreToDisk()
+		{
+			FileInfo fi = GetStorePath();
+
+			if (!fi.Directory?.Exists ?? false)
+			{
+				Directory.CreateDirectory(fi.Directory!.FullName);
+			}
+
+			using FileStream fs = fi.Open(FileMode.Create, FileAccess.Write);
+			using Utf8JsonWriter writer = new Utf8JsonWriter(fs);
+			JsonSerializer.Serialize<TokenStoreState>(writer, new TokenStoreState(_crypt.IV, _providerToRefreshToken));
+		}
+
+		public bool TryGetRefreshToken(string oidcProvider, out string refreshToken)
+		{
+			if (!_providerToRefreshToken.TryGetValue(oidcProvider, out byte[]? encryptedToken))
+			{
+				refreshToken = "";
+				return false;
+			}
+
+			ICryptoTransform decryptor = _crypt.CreateDecryptor(GetSeed(), _crypt.IV);
+			byte[] bytes = decryptor.TransformFinalBlock(encryptedToken, 0, encryptedToken.Length);
+			refreshToken = Encoding.Unicode.GetString(bytes);
+			return true;
+		}
+
+		public void AddRefreshToken(string providerIdentifier, string refreshToken)
+		{
+			byte[] bytes = Encoding.Unicode.GetBytes(refreshToken);
+			ICryptoTransform encryptor = _crypt.CreateEncryptor(GetSeed(), _crypt.IV);
+
+			byte[] encryptedToken = encryptor.TransformFinalBlock(bytes, 0, bytes.Length);
+			_providerToRefreshToken[providerIdentifier] = encryptedToken;
+		}
+
+		public void Save()
+		{
+			SaveStoreToDisk();
+		}
+
+		static byte[] GetSeed()
+		{
+			return new byte[]
+			{
+				0x1e, 0x72, 0x5e, 0xe7, 0x08, 0x9e, 0x29, 0x5e, 0xcb, 0xbe, 0x1b, 0xdf, 0x0e, 0xf9, 0x4a, 0x30, 0xd1,
+				0xa9, 0x9b, 0xa2, 0xee, 0x58, 0xc4, 0x8e
+			};
+		}
+
+		protected virtual void Dispose(bool disposing)
+		{
+			if (disposing)
+			{
+				SaveStoreToDisk();
+				_crypt.Dispose();
+			}
+		}
+
+		public void Dispose()
+		{
+			Dispose(true);
+			GC.SuppressFinalize(this);
+		}
+	}
+}
\ No newline at end of file
diff --git a/Engine/Source/Programs/Shared/EpicGames.OIDC/WindowsTokenStore.cs b/Engine/Source/Programs/Shared/EpicGames.OIDC/WindowsTokenStore.cs
new file mode 100644
--- /dev/null
+++ b/Engine/Source/Programs/Shared/EpicGames.OIDC/WindowsTokenStore.cs
@@ -0,0 +1,309 @@
+// Copyright Epic Games, Inc. All Rights Reserved.
+
+using System;
+using System.Collections.Generic;
+using System.ComponentModel;
+using System.IO;
+using System.Reflection;
+using System.Runtime.InteropServices;
+using System.Text;
+using System.Text.Json;
+using System.Text.Json.Serialization;
+
+namespace EpicGames.OIDC
+{
+	public class WindowsTokenStore : ITokenStore, IDisposable
+	{
+		private class TokenStoreState
+		{
+			public Dictionary<string, string> Providers { get; set; } = new Dictionary<string, string>();
+
+			[JsonConstructor]
+			public TokenStoreState(Dictionary<string, string> providers)
+			{
+				Providers = providers;
+			}
+
+			public TokenStoreState(Dictionary<string, byte[]> providers)
+			{
+				foreach (KeyValuePair<string, byte[]> pair in providers)
+				{
+					Providers[pair.Key] = Convert.ToBase64String(pair.Value);
+				}
+			}
+		}
+
+		private readonly Dictionary<string, byte[]> _providerToRefreshToken = new Dictionary<string, byte[]>();
+
+		public WindowsTokenStore()
+		{
+			ReadStoreFromDisk();
+		}
+
+		private static FileInfo GetStorePath()
+		{
+			return new FileInfo(Path.Combine(Environment.GetFolderPath(Environment.SpecialFolder.LocalApplicationData), "oidc-token", "oidcTokenStore.dat"));
+		}
+
+		private void ReadStoreFromDisk()
+		{
+			FileInfo fi = GetStorePath();
+			if (!fi.Exists)
+			{
+				// if we have no store on disk then we just initialize it to empty
+				return;
+			}
+
+			using FileStream fs = fi.OpenRead();
+			using TextReader tr = new StreamReader(fs);
+			TokenStoreState? state = JsonSerializer.Deserialize<TokenStoreState>(tr.ReadToEnd());
+
+			if (state == null)
+			{
+				// if we fail to deserialize the state just drop it, will mean users will need to login again
+				return;
+			}
+
+			foreach ((string key, string value) in state.Providers)
+			{
+				_providerToRefreshToken[key] = Convert.FromBase64String(value);
+			}
+		}
+
+		private void SaveStoreToDisk()
+		{
+			FileInfo fi = GetStorePath();
+
+			if (!fi.Directory?.Exists ?? false)
+			{
+				Directory.CreateDirectory(fi.Directory!.FullName);
+			}
+
+			using FileStream fs = fi.Open(FileMode.Create, FileAccess.Write);
+			using Utf8JsonWriter writer = new Utf8JsonWriter(fs);
+			JsonSerializer.Serialize<TokenStoreState>(writer, new TokenStoreState(_providerToRefreshToken));
+		}
+
+		public bool TryGetRefreshToken(string oidcProvider, out string refreshToken)
+		{
+			if (!_providerToRefreshToken.TryGetValue(oidcProvider, out byte[]? encryptedToken))
+			{
+				refreshToken = "";
+				return false;
+			}
+
+			byte[] bytes = CryptProtectDataHelper.DoCryptUnprotectData(encryptedToken, $"OidcToken-{oidcProvider}", GetEntropy(oidcProvider));
+			refreshToken = Encoding.Unicode.GetString(bytes);
+
+			return true;
+		}
+
+		private static byte[] GetEntropy(string oidcProvider)
+		{
+			byte[] providerBytes = Encoding.UTF8.GetBytes(oidcProvider);
+			byte[] assemblyBytes = Encoding.UTF8.GetBytes(Assembly.GetEntryAssembly()!.GetName().FullName);
+
+			byte[] mergedBytes = new byte[providerBytes.Length + assemblyBytes.Length];
+			Array.Copy(providerBytes, mergedBytes, providerBytes.Length);
+			int bytesWritten = providerBytes.Length;
+			Array.Copy(assemblyBytes, 0, mergedBytes, bytesWritten, assemblyBytes.Length);
+
+			return mergedBytes;
+		}
+
+		public void AddRefreshToken(string providerIdentifier, string refreshToken)
+		{
+			byte[] bytes = Encoding.Unicode.GetBytes(refreshToken);
+			byte[] encryptedToken = CryptProtectDataHelper.DoCryptProtectData(bytes, $"OidcToken-{providerIdentifier}", GetEntropy(providerIdentifier));
+
+			_providerToRefreshToken[providerIdentifier] = encryptedToken;
+		}
+
+		public void Save()
+		{
+			SaveStoreToDisk();
+		}
+
+		protected virtual void Dispose(bool disposing)
+		{
+			if (disposing)
+			{
+				SaveStoreToDisk();
+			}
+		}
+
+		public void Dispose()
+		{
+			Dispose(true);
+			GC.SuppressFinalize(this);
+		}
+	}
+
+	static class CryptProtectDataHelper
+	{
+		
+		[StructLayout(LayoutKind.Sequential, CharSet=CharSet.Unicode)]
+		private struct DataBlob {
+			public int cbData;
+			public IntPtr pbData;
+		}
+
+		[Flags]
+		private enum CryptProtectFlags {
+			// for remote-access situations where ui is not an option
+			// if UI was specified on protect or unprotect operation, the call
+			// will fail and GetLastError() will indicate ERROR_PASSWORD_RESTRICTION
+			CryptprotectUiForbidden = 0x1,
+
+			// per machine protected data -- any user on machine where CryptProtectData
+			// took place may CryptUnprotectData
+			CryptprotectLocalMachine = 0x4,
+
+			// force credential synchronize during CryptProtectData()
+			// Synchronize is only operation that occurs during this operation
+			CryptprotectCredSync = 0x8,
+
+			// Generate an Audit on protect and unprotect operations
+			CryptprotectAudit = 0x10,
+
+			// Protect data with a non-recoverable key
+			CryptprotectNoRecovery = 0x20,
+
+			// Verify the protection of a protected blob
+			CryptprotectVerifyProtection = 0x40
+		}
+
+		[Flags]
+		private enum CryptProtectPromptFlags {
+			// prompt on unprotect
+			CryptprotectPromptOnUnprotect = 0x1,
+
+			// prompt on protect
+			CryptprotectPromptOnProtect = 0x2
+		}
+
+		[StructLayout(LayoutKind.Sequential, CharSet=CharSet.Unicode)]
+		private struct CryptprotectPromptstruct {
+			public int cbSize;
+			public CryptProtectPromptFlags dwPromptFlags;
+			public IntPtr hwndApp;
+			public string szPrompt;
+		}
+
+		[
+			DllImport("Crypt32.dll",
+				SetLastError=true,
+				CharSet=System.Runtime.InteropServices.CharSet.Auto)
+		]
+		[return: MarshalAs(UnmanagedType.Bool)]
+		private static extern bool CryptProtectData(
+			ref DataBlob pDataIn,
+			string szDataDescr,
+			ref DataBlob pOptionalEntropy,
+			IntPtr pvReserved,
+			IntPtr pPromptStruct,
+			CryptProtectFlags dwFlags,
+			ref DataBlob pDataOut
+		);
+
+		[
+			DllImport("Crypt32.dll",
+				SetLastError=true,
+				CharSet=System.Runtime.InteropServices.CharSet.Auto)
+		]
+		[return: MarshalAs(UnmanagedType.Bool)]
+		private static extern bool CryptUnprotectData(
+			ref DataBlob pDataIn,
+			string szDataDescr,
+			ref DataBlob pOptionalEntropy,
+			IntPtr pvReserved,
+			IntPtr pPromptStruct,
+			CryptProtectFlags dwFlags,
+			ref DataBlob pDataOut
+		);
+
+		public static byte[] DoCryptProtectData(byte[] dataToProtect, string description, byte[] entropy)
+		{
+			DataBlob dataOut = new DataBlob();
+
+			GCHandle dataHandle = GCHandle.Alloc(dataToProtect, GCHandleType.Pinned);
+			GCHandle entropyHandle = GCHandle.Alloc(entropy, GCHandleType.Pinned);
+			try
+			{
+				Marshal.Copy(dataToProtect, 0, dataHandle.AddrOfPinnedObject(), dataToProtect.Length);
+				Marshal.Copy(entropy, 0, entropyHandle.AddrOfPinnedObject(), entropy.Length);
+
+				DataBlob data = new DataBlob()
+				{
+					cbData = dataToProtect.Length,
+					pbData = dataHandle.AddrOfPinnedObject()
+				};
+
+				DataBlob entropyBlob = new DataBlob()
+				{
+					cbData = entropy.Length,
+					pbData = entropyHandle.AddrOfPinnedObject()
+				};
+
+				CryptProtectFlags flags = 0;
+
+				if (!CryptProtectData(ref data, description, ref entropyBlob, IntPtr.Zero, IntPtr.Zero, flags, ref dataOut))
+				{
+					throw new Win32Exception();
+				}
+			}
+			finally
+			{
+				dataHandle.Free();
+				entropyHandle.Free();
+			}
+		
+		
+			byte[] buf = new byte[dataOut.cbData];
+			Marshal.Copy(dataOut.pbData, buf, 0, dataOut.cbData);
+			return buf;
+		}
+
+		public static byte[] DoCryptUnprotectData(byte[] dataToDecrypt, string description, byte[] entropy)
+		{
+			DataBlob dataOut = new DataBlob();
+
+			GCHandle dataHandle = GCHandle.Alloc(dataToDecrypt, GCHandleType.Pinned);
+			GCHandle entropyHandle = GCHandle.Alloc(entropy, GCHandleType.Pinned);
+			try
+			{
+				Marshal.Copy(dataToDecrypt, 0, dataHandle.AddrOfPinnedObject(), dataToDecrypt.Length);
+				Marshal.Copy(entropy, 0, entropyHandle.AddrOfPinnedObject(), entropy.Length);
+
+				DataBlob data = new DataBlob()
+				{
+					cbData = dataToDecrypt.Length,
+					pbData = dataHandle.AddrOfPinnedObject()
+				};
+
+				DataBlob entropyBlob = new DataBlob()
+				{
+					cbData = entropy.Length,
+					pbData = entropyHandle.AddrOfPinnedObject()
+				};
+
+				CryptProtectFlags flags = 0;
+
+				if (!CryptUnprotectData(ref data, description, ref entropyBlob, IntPtr.Zero, IntPtr.Zero, flags, ref dataOut))
+				{
+					throw new Win32Exception();
+				}
+			}
+			finally
+			{
+				dataHandle.Free();
+				entropyHandle.Free();
+			}
+		
+		
+			byte[] buf = new byte[dataOut.cbData];
+			Marshal.Copy(dataOut.pbData, buf, 0, dataOut.cbData);
+			return buf;
+		}
+	}
+}
\ No newline at end of file
