Added required configurations to EpicGames.OIDC - resolves error generating projects.

diff --git a/Engine/Source/Programs/Shared/EpicGames.OIDC/EpicGames.OIDC.csproj b/Engine/Source/Programs/Shared/EpicGames.OIDC/EpicGames.OIDC.csproj
--- a/Engine/Source/Programs/Shared/EpicGames.OIDC/EpicGames.OIDC.csproj
+++ b/Engine/Source/Programs/Shared/EpicGames.OIDC/EpicGames.OIDC.csproj
@@ -7,6 +7,16 @@
     <Nullable>enable</Nullable>
     <GenerateAssemblyInfo>false</GenerateAssemblyInfo>
     <GenerateTargetFrameworkAttribute>false</GenerateTargetFrameworkAttribute>
+    <Configurations>Debug;Release;Development</Configurations>
+  </PropertyGroup>
+
+  <PropertyGroup Condition=" '$(Configuration)|$(Platform)' == 'Debug|AnyCPU' ">
+  </PropertyGroup>
+  <PropertyGroup Condition="'$(Configuration)|$(Platform)' == 'Development|AnyCPU'">
+    <Optimize>true</Optimize>
+  </PropertyGroup>
+  <PropertyGroup Condition="'$(Configuration)|$(Platform)' == 'Release|AnyCPU'">
+    <Optimize>true</Optimize>
   </PropertyGroup>
 
   <ItemGroup>
