BuildGraph: Decouple diagnostics from the nodes and agents that contain them.

diff --git a/Engine/Source/Programs/AutomationTool/BuildGraph/BuildGraph.cs b/Engine/Source/Programs/AutomationTool/BuildGraph/BuildGraph.cs
--- a/Engine/Source/Programs/AutomationTool/BuildGraph/BuildGraph.cs
+++ b/Engine/Source/Programs/AutomationTool/BuildGraph/BuildGraph.cs
@@ -695,22 +695,23 @@ namespace AutomationTool
 			// Print out all the diagnostic messages which still apply, unless we're running a step as part of a build system or just listing the contents of the file. 
 			if(SingleNode == null && (!bListOnly || bShowDiagnostics))
 			{
-				foreach(BgGraphDiagnostic Diagnostic in Graph.Diagnostics)
+				List<BgDiagnostic> Diagnostics = Graph.GetAllDiagnostics();
+				foreach(BgDiagnostic Diagnostic in Diagnostics)
 				{
-					if(Diagnostic.EventType == LogEventType.Console)
+					if(Diagnostic.Level == LogLevel.Information)
 					{
-						CommandUtils.LogWarning("{0}({1}): {2}", Diagnostic.Location.File, Diagnostic.Location.LineNumber, Diagnostic.Message);
+						CommandUtils.LogInformation("{0}({1}): {2}", Diagnostic.File, Diagnostic.Line, Diagnostic.Message);
 					}
-					else if(Diagnostic.EventType == LogEventType.Warning)
+					else if(Diagnostic.Level == LogLevel.Warning)
 					{
-						CommandUtils.LogWarning("{0}({1}): warning: {2}", Diagnostic.Location.File, Diagnostic.Location.LineNumber, Diagnostic.Message);
+						CommandUtils.LogWarning("{0}({1}): warning: {2}", Diagnostic.File, Diagnostic.Line, Diagnostic.Message);
 					}
 					else
 					{
-						CommandUtils.LogError("{0}({1}): error: {2}", Diagnostic.Location.File, Diagnostic.Location.LineNumber, Diagnostic.Message);
+						CommandUtils.LogError("{0}({1}): error: {2}", Diagnostic.File, Diagnostic.Line, Diagnostic.Message);
 					}
 				}
-				if(Graph.Diagnostics.Any(x => x.EventType == LogEventType.Error))
+				if(Diagnostics.Any(x => x.Level == LogLevel.Error))
 				{
 					return ExitCode.Error_Unknown;
 				}
diff --git a/Engine/Source/Programs/Shared/EpicGames.BuildGraph/BgAgent.cs b/Engine/Source/Programs/Shared/EpicGames.BuildGraph/BgAgent.cs
--- a/Engine/Source/Programs/Shared/EpicGames.BuildGraph/BgAgent.cs
+++ b/Engine/Source/Programs/Shared/EpicGames.BuildGraph/BgAgent.cs
@@ -29,6 +29,11 @@ namespace EpicGames.BuildGraph
 		/// </summary>
 		public List<BgNode> Nodes { get; set; } = new List<BgNode>();
 
+		/// <summary>
+		/// Diagnostics for this agent
+		/// </summary>
+		public List<BgDiagnostic> Diagnostics { get; } = new List<BgDiagnostic>();
+
 		/// <summary>
 		/// Constructor
 		/// </summary>
diff --git a/Engine/Source/Programs/Shared/EpicGames.BuildGraph/BgDiagnostic.cs b/Engine/Source/Programs/Shared/EpicGames.BuildGraph/BgDiagnostic.cs
--- a/Engine/Source/Programs/Shared/EpicGames.BuildGraph/BgDiagnostic.cs
+++ b/Engine/Source/Programs/Shared/EpicGames.BuildGraph/BgDiagnostic.cs
@@ -1,6 +1,7 @@
 // Copyright Epic Games, Inc. All Rights Reserved.
 
 using EpicGames.Core;
+using Microsoft.Extensions.Logging;
 
 namespace EpicGames.BuildGraph
 {
@@ -8,43 +9,37 @@ namespace EpicGames.BuildGraph
 	/// Diagnostic message from the graph script. These messages are parsed at startup, then culled along with the rest of the graph nodes before output. Doing so
 	/// allows errors and warnings which are only output if a node is part of the graph being executed.
 	/// </summary>
-	public class BgGraphDiagnostic
+	public class BgDiagnostic
 	{
 		/// <summary>
-		/// Location of the diagnostic
+		/// File containing the diagnostic
 		/// </summary>
-		public BgScriptLocation Location { get; }
+		public string File { get; }
 
 		/// <summary>
-		/// The diagnostic event type
+		/// Line number containing the diagnostic
 		/// </summary>
-		public LogEventType EventType { get; }
+		public int Line { get; }
 
 		/// <summary>
-		/// The message to display
-		/// </summary>
-		public string Message { get; }
-
-		/// <summary>
-		/// The node which this diagnostic is declared in. If the node is culled from the graph, the message will not be displayed.
+		/// The diagnostic event type
 		/// </summary>
-		public BgNode? EnclosingNode { get; }
+		public LogLevel Level { get; }
 
 		/// <summary>
-		/// The agent that this diagnostic is declared in. If the entire agent is culled from the graph, the message will not be displayed.
+		/// The message to display
 		/// </summary>
-		public BgAgent? EnclosingAgent { get; }
+		public string Message { get; }
 
 		/// <summary>
 		/// Constructor
 		/// </summary>
-		public BgGraphDiagnostic(BgScriptLocation location, LogEventType eventType, string message, BgNode? enclosingNode, BgAgent? enclosingAgent)
+		public BgDiagnostic(string file, int line, LogLevel level, string message)
 		{
-			Location = location;
-			EventType = eventType;
+			File = file;
+			Line = line;
+			Level = level;
 			Message = message;
-			EnclosingNode = enclosingNode;
-			EnclosingAgent = enclosingAgent;
 		}
 	}
 }
diff --git a/Engine/Source/Programs/Shared/EpicGames.BuildGraph/BgDiagnosticSpec.cs b/Engine/Source/Programs/Shared/EpicGames.BuildGraph/BgDiagnosticSpec.cs
--- a/Engine/Source/Programs/Shared/EpicGames.BuildGraph/BgDiagnosticSpec.cs
+++ b/Engine/Source/Programs/Shared/EpicGames.BuildGraph/BgDiagnosticSpec.cs
@@ -40,7 +40,20 @@ namespace EpicGames.BuildGraph
 		{
 			BgScriptLocation location = new BgScriptLocation("(unknown)", "(unknown)", 1);
 			string messageValue = Message.Compute(context);
-			graph.Diagnostics.Add(new BgGraphDiagnostic(location, (LogEventType)Level, messageValue, enclosingNode, enclosingAgent));
+
+			BgDiagnostic diagnostic = new BgDiagnostic(location.File, location.LineNumber, Level, messageValue);
+			if (enclosingNode != null)
+			{
+				enclosingNode.Diagnostics.Add(diagnostic);
+			}
+			else if (enclosingAgent != null)
+			{
+				enclosingAgent.Diagnostics.Add(diagnostic);
+			}
+			else
+			{
+				graph.Diagnostics.Add(diagnostic);
+			}
 		}
 
 		/// <inheritdoc/>
diff --git a/Engine/Source/Programs/Shared/EpicGames.BuildGraph/BgGraph.cs b/Engine/Source/Programs/Shared/EpicGames.BuildGraph/BgGraph.cs
--- a/Engine/Source/Programs/Shared/EpicGames.BuildGraph/BgGraph.cs
+++ b/Engine/Source/Programs/Shared/EpicGames.BuildGraph/BgGraph.cs
@@ -83,9 +83,9 @@ namespace EpicGames.BuildGraph
 		public List<BgLabel> Labels { get; } = new List<BgLabel>();
 
 		/// <summary>
-		/// Diagnostic messages for this graph
+		/// Diagnostics at graph scope
 		/// </summary>
-		public List<BgGraphDiagnostic> Diagnostics { get; } = new List<BgGraphDiagnostic>();
+		public List<BgDiagnostic> Diagnostics { get; } = new List<BgDiagnostic>();
 
 		/// <summary>
 		/// Default constructor
@@ -104,6 +104,24 @@ namespace EpicGames.BuildGraph
 			return NameToNode.ContainsKey(name) || NameToReport.ContainsKey(name) || NameToAggregate.ContainsKey(name);
 		}
 
+		/// <summary>
+		/// Gets diagnostics from all graph structures
+		/// </summary>
+		/// <returns>List of diagnostics</returns>
+		public List<BgDiagnostic> GetAllDiagnostics()
+		{
+			List<BgDiagnostic> diagnostics = new List<BgDiagnostic>(Diagnostics);
+			foreach (BgAgent agent in Agents)
+			{
+				diagnostics.AddRange(agent.Diagnostics);
+				foreach (BgNode node in agent.Nodes)
+				{
+					diagnostics.AddRange(node.Diagnostics);
+				}
+			}
+			return diagnostics;
+		}
+
 		/// <summary>
 		/// Tries to resolve the given name to one or more nodes. Checks for aggregates, and actual nodes.
 		/// </summary>
@@ -267,9 +285,6 @@ namespace EpicGames.BuildGraph
 
 			// Remove any badges which do not have all their dependencies
 			Badges.RemoveAll(x => x.Nodes.Any(y => !retainNodes.Contains(y)));
-
-			// Remove any diagnostics which are no longer part of the graph
-			Diagnostics.RemoveAll(x => (x.EnclosingNode != null && !retainNodes.Contains(x.EnclosingNode)) || (x.EnclosingAgent != null && !Agents.Contains(x.EnclosingAgent)));
 		}
 
 		/// <summary>
diff --git a/Engine/Source/Programs/Shared/EpicGames.BuildGraph/BgNode.cs b/Engine/Source/Programs/Shared/EpicGames.BuildGraph/BgNode.cs
--- a/Engine/Source/Programs/Shared/EpicGames.BuildGraph/BgNode.cs
+++ b/Engine/Source/Programs/Shared/EpicGames.BuildGraph/BgNode.cs
@@ -57,27 +57,27 @@ namespace EpicGames.BuildGraph
 		/// <summary>
 		/// Array of inputs which this node requires to run
 		/// </summary>
-		public BgNodeOutput[] Inputs { get; }
+		public IReadOnlyList<BgNodeOutput> Inputs { get; }
 
 		/// <summary>
 		/// Array of outputs produced by this node
 		/// </summary>
-		public BgNodeOutput[] Outputs { get; }
+		public IReadOnlyList<BgNodeOutput> Outputs { get; }
 
 		/// <summary>
 		/// Nodes which this node has input dependencies on
 		/// </summary>
-		public BgNode[] InputDependencies { get; set; }
+		public IReadOnlyList<BgNode> InputDependencies { get; set; }
 
 		/// <summary>
 		/// Nodes which this node needs to run after
 		/// </summary>
-		public BgNode[] OrderDependencies { get; set; }
+		public IReadOnlyList<BgNode> OrderDependencies { get; set; }
 
 		/// <summary>
 		/// Tokens which must be acquired for this node to run
 		/// </summary>
-		public FileReference[] RequiredTokens { get; }
+		public IReadOnlyList<FileReference> RequiredTokens { get; }
 
 		/// <summary>
 		/// List of tasks to execute
@@ -109,28 +109,33 @@ namespace EpicGames.BuildGraph
 		/// </summary>
 		public Dictionary<string, string> Annotations { get; } = new Dictionary<string, string>(StringComparer.OrdinalIgnoreCase);
 
+		/// <summary>
+		/// Diagnostics for this node
+		/// </summary>
+		public List<BgDiagnostic> Diagnostics { get; } = new List<BgDiagnostic>();
+
 		/// <summary>
 		/// Constructor
 		/// </summary>
-		/// <param name="inName">The name of this node</param>
-		/// <param name="inInputs">Inputs that this node depends on</param>
-		/// <param name="inOutputNames">Names of the outputs that this node produces</param>
-		/// <param name="inInputDependencies">Nodes which this node is dependent on for its inputs</param>
-		/// <param name="inOrderDependencies">Nodes which this node needs to run after. Should include all input dependencies.</param>
-		/// <param name="inRequiredTokens">Optional tokens which must be required for this node to run</param>
-		public BgNode(string inName, BgNodeOutput[] inInputs, string[] inOutputNames, BgNode[] inInputDependencies, BgNode[] inOrderDependencies, FileReference[] inRequiredTokens)
+		/// <param name="name">The name of this node</param>
+		/// <param name="inputs">Inputs that this node depends on</param>
+		/// <param name="outputNames">Names of the outputs that this node produces</param>
+		/// <param name="inputDependencies">Nodes which this node is dependent on for its inputs</param>
+		/// <param name="orderDependencies">Nodes which this node needs to run after. Should include all input dependencies.</param>
+		/// <param name="requiredTokens">Optional tokens which must be required for this node to run</param>
+		public BgNode(string name, IReadOnlyList<BgNodeOutput> inputs, IReadOnlyList<string> outputNames, IReadOnlyList<BgNode> inputDependencies, IReadOnlyList<BgNode> orderDependencies, IReadOnlyList<FileReference> requiredTokens)
 		{
-			Name = inName;
-			Inputs = inInputs;
+			Name = name;
+			Inputs = inputs;
 
 			List<BgNodeOutput> allOutputs = new List<BgNodeOutput>();
 			allOutputs.Add(new BgNodeOutput(this, "#" + Name));
-			allOutputs.AddRange(inOutputNames.Where(x => String.Compare(x, Name, StringComparison.InvariantCultureIgnoreCase) != 0).Select(x => new BgNodeOutput(this, x)));
+			allOutputs.AddRange(outputNames.Where(x => String.Compare(x, Name, StringComparison.InvariantCultureIgnoreCase) != 0).Select(x => new BgNodeOutput(this, x)));
 			Outputs = allOutputs.ToArray();
 
-			InputDependencies = inInputDependencies;
-			OrderDependencies = inOrderDependencies;
-			RequiredTokens = inRequiredTokens;
+			InputDependencies = inputDependencies;
+			OrderDependencies = orderDependencies;
+			RequiredTokens = requiredTokens;
 		}
 
 		/// <summary>
diff --git a/Engine/Source/Programs/Shared/EpicGames.BuildGraph/BgScriptReader.cs b/Engine/Source/Programs/Shared/EpicGames.BuildGraph/BgScriptReader.cs
--- a/Engine/Source/Programs/Shared/EpicGames.BuildGraph/BgScriptReader.cs
+++ b/Engine/Source/Programs/Shared/EpicGames.BuildGraph/BgScriptReader.cs
@@ -461,13 +461,13 @@ namespace EpicGames.BuildGraph
 						await ReadNotifierAsync(childElement);
 						break;
 					case "Trace":
-						await ReadDiagnosticAsync(childElement, LogEventType.Console);
+						await ReadDiagnosticAsync(childElement, LogLevel.Information);
 						break;
 					case "Warning":
-						await ReadDiagnosticAsync(childElement, LogEventType.Warning);
+						await ReadDiagnosticAsync(childElement, LogLevel.Warning);
 						break;
 					case "Error":
-						await ReadDiagnosticAsync(childElement, LogEventType.Error);
+						await ReadDiagnosticAsync(childElement, LogLevel.Error);
 						break;
 					case "Do":
 						await ReadBlockAsync(childElement, ReadGraphBodyAsync);
@@ -745,13 +745,13 @@ namespace EpicGames.BuildGraph
 						await ReadAggregateAsync(childElement);
 						break;
 					case "Trace":
-						await ReadDiagnosticAsync(childElement, LogEventType.Console);
+						await ReadDiagnosticAsync(childElement, LogLevel.Information);
 						break;
 					case "Warning":
-						await ReadDiagnosticAsync(childElement, LogEventType.Warning);
+						await ReadDiagnosticAsync(childElement, LogLevel.Warning);
 						break;
 					case "Error":
-						await ReadDiagnosticAsync(childElement, LogEventType.Error);
+						await ReadDiagnosticAsync(childElement, LogLevel.Error);
 						break;
 					case "Label":
 						await ReadLabelAsync(childElement);
@@ -827,13 +827,13 @@ namespace EpicGames.BuildGraph
 						await ReadRegexAsync(childElement);
 						break;
 					case "Trace":
-						await ReadDiagnosticAsync(childElement, LogEventType.Console);
+						await ReadDiagnosticAsync(childElement, LogLevel.Information);
 						break;
 					case "Warning":
-						await ReadDiagnosticAsync(childElement, LogEventType.Warning);
+						await ReadDiagnosticAsync(childElement, LogLevel.Warning);
 						break;
 					case "Error":
-						await ReadDiagnosticAsync(childElement, LogEventType.Error);
+						await ReadDiagnosticAsync(childElement, LogLevel.Error);
 						break;
 					case "Do":
 						await ReadBlockAsync(childElement, ReadNodeBodyAsync);
@@ -1006,8 +1006,8 @@ namespace EpicGames.BuildGraph
 		/// Reads a warning from the given element, evaluates the condition on it, and writes it to the log if the condition passes.
 		/// </summary>
 		/// <param name="element">Xml element to read the definition from</param>
-		/// <param name="eventType">The diagnostic event type</param>
-		protected abstract Task ReadDiagnosticAsync(BgScriptElement element, LogEventType eventType);
+		/// <param name="level">The diagnostic event type</param>
+		protected abstract Task ReadDiagnosticAsync(BgScriptElement element, LogLevel level);
 
 		/// <summary>
 		/// Reads an object name from its defining element. Outputs an error if the name is missing.
@@ -2093,15 +2093,26 @@ namespace EpicGames.BuildGraph
 		/// Reads a warning from the given element, evaluates the condition on it, and writes it to the log if the condition passes.
 		/// </summary>
 		/// <param name="element">Xml element to read the definition from</param>
-		/// <param name="eventType">The diagnostic event type</param>
-		protected override async Task ReadDiagnosticAsync(BgScriptElement element, LogEventType eventType)
+		/// <param name="level">The diagnostic event type</param>
+		protected override async Task ReadDiagnosticAsync(BgScriptElement element, LogLevel level)
 		{
 			if (await EvaluateConditionAsync(element))
 			{
 				string message = ReadAttribute(element, "Message");
 
-				BgGraphDiagnostic diagnostic = new BgGraphDiagnostic(element.Location, eventType, message, _enclosingNode, _enclosingAgent);
-				_graph.Diagnostics.Add(diagnostic);
+				BgDiagnostic diagnostic = new BgDiagnostic(element.Location.File, element.Location.LineNumber, level, message);
+				if (_enclosingNode != null)
+				{
+					_enclosingNode.Diagnostics.Add(diagnostic);
+				}
+				else if (_enclosingAgent != null)
+				{
+					_enclosingAgent.Diagnostics.Add(diagnostic);
+				}
+				else
+				{
+					_graph.Diagnostics.Add(diagnostic);
+				}
 			}
 		}
 
