Support for serializing signed and unsigned VarInt types from C#.

diff --git a/Engine/Source/Programs/Shared/EpicGames.Horde/Bundles/BundleObject.cs b/Engine/Source/Programs/Shared/EpicGames.Horde/Bundles/BundleObject.cs
--- a/Engine/Source/Programs/Shared/EpicGames.Horde/Bundles/BundleObject.cs
+++ b/Engine/Source/Programs/Shared/EpicGames.Horde/Bundles/BundleObject.cs
@@ -156,10 +156,10 @@ namespace EpicGames.Horde.Bundles
 				IoHash hash = new IoHash(span);
 				span = span.Slice(IoHash.NumBytes);
 
-				int rank = (int)VarInt.Read(span, out int rankBytes);
+				int rank = (int)VarInt.ReadUnsigned(span, out int rankBytes);
 				span = span.Slice(rankBytes);
 
-				int cost = (int)VarInt.Read(span, out int costBytes);
+				int cost = (int)VarInt.ReadUnsigned(span, out int costBytes);
 				span = span.Slice(costBytes);
 
 				imports.Add(new BundleImport(hash, rank, cost));
@@ -169,7 +169,7 @@ namespace EpicGames.Horde.Bundles
 
 		static byte[] Serialize(List<BundleImport> imports)
 		{
-			byte[] data = new byte[imports.Sum(x => IoHash.NumBytes + VarInt.Measure(x.Rank) + VarInt.Measure(x.Length))];
+			byte[] data = new byte[imports.Sum(x => IoHash.NumBytes + VarInt.MeasureUnsigned(x.Rank) + VarInt.MeasureUnsigned(x.Length))];
 
 			Span<byte> span = data;
 			foreach (BundleImport import in imports)
@@ -177,10 +177,10 @@ namespace EpicGames.Horde.Bundles
 				import.Hash.CopyTo(span);
 				span = span.Slice(IoHash.NumBytes);
 
-				int rankBytes = VarInt.Write(span, import.Rank);
+				int rankBytes = VarInt.WriteUnsigned(span, import.Rank);
 				span = span.Slice(rankBytes);
 
-				int costBytes = VarInt.Write(span, import.Length);
+				int costBytes = VarInt.WriteUnsigned(span, import.Length);
 				span = span.Slice(costBytes);
 			}
 
@@ -265,10 +265,10 @@ namespace EpicGames.Horde.Bundles
 			{
 				BundleCompressionPacket packet = new BundleCompressionPacket(packetOffset);
 
-				packet.EncodedLength = (int)VarInt.Read(span, out int encodedLengthBytes);
+				packet.EncodedLength = (int)VarInt.ReadUnsigned(span, out int encodedLengthBytes);
 				span = span.Slice(encodedLengthBytes);
 
-				packet.DecodedLength = (int)VarInt.Read(span, out int decodedLengthBytes);
+				packet.DecodedLength = (int)VarInt.ReadUnsigned(span, out int decodedLengthBytes);
 				span = span.Slice(decodedLengthBytes);
 
 				packetOffset += packet.EncodedLength;
@@ -279,19 +279,19 @@ namespace EpicGames.Horde.Bundles
 					IoHash hash = new IoHash(span);
 					span = span.Slice(IoHash.NumBytes);
 
-					int rank = (int)VarInt.Read(span, out int rankBytes);
+					int rank = (int)VarInt.ReadUnsigned(span, out int rankBytes);
 					span = span.Slice(rankBytes);
 
-					int length = (int)VarInt.Read(span, out int lengthBytes);
+					int length = (int)VarInt.ReadUnsigned(span, out int lengthBytes);
 					span = span.Slice(lengthBytes);
 
-					int numReferences = (int)VarInt.Read(span, out int numReferencesBytes);
+					int numReferences = (int)VarInt.ReadUnsigned(span, out int numReferencesBytes);
 					span = span.Slice(numReferencesBytes);
 
 					int[] references = new int[numReferences];
 					for (int referenceIdx = 0; referenceIdx < numReferences; referenceIdx++)
 					{
-						references[referenceIdx] = (int)VarInt.Read(span, out int referenceBytes);
+						references[referenceIdx] = (int)VarInt.ReadUnsigned(span, out int referenceBytes);
 						span = span.Slice(referenceBytes);
 					}
 
@@ -320,10 +320,10 @@ namespace EpicGames.Horde.Bundles
 				{
 					CheckPacketSequence(prevPacket, packet);
 
-					int encodedLengthBytes = VarInt.Write(span, packet.EncodedLength);
+					int encodedLengthBytes = VarInt.WriteUnsigned(span, packet.EncodedLength);
 					span = span.Slice(encodedLengthBytes);
 
-					int decodedLengthBytes = VarInt.Write(span, packet.DecodedLength);
+					int decodedLengthBytes = VarInt.WriteUnsigned(span, packet.DecodedLength);
 					span = span.Slice(decodedLengthBytes);
 
 					prevPacket = packet;
@@ -332,18 +332,18 @@ namespace EpicGames.Horde.Bundles
 				export.Hash.CopyTo(span);
 				span = span.Slice(IoHash.NumBytes);
 
-				int rankBytes = VarInt.Write(span, export.Rank);
+				int rankBytes = VarInt.WriteUnsigned(span, export.Rank);
 				span = span.Slice(rankBytes);
 
-				int lengthBytes = VarInt.Write(span, export.Length);
+				int lengthBytes = VarInt.WriteUnsigned(span, export.Length);
 				span = span.Slice(lengthBytes);
 
-				int numReferencesBytes = VarInt.Write(span, export.References.Count);
+				int numReferencesBytes = VarInt.WriteUnsigned(span, export.References.Count);
 				span = span.Slice(numReferencesBytes);
 
 				foreach (int reference in export.References)
 				{
-					int referenceBytes = VarInt.Write(span, reference);
+					int referenceBytes = VarInt.WriteUnsigned(span, reference);
 					span = span.Slice(referenceBytes);
 				}
 			}
@@ -363,11 +363,11 @@ namespace EpicGames.Horde.Bundles
 				if (prevPacket == null || packet.Offset != prevPacket.Offset)
 				{
 					CheckPacketSequence(prevPacket, packet);
-					length += VarInt.Measure(packet.EncodedLength) + VarInt.Measure(packet.DecodedLength);
+					length += VarInt.MeasureUnsigned(packet.EncodedLength) + VarInt.MeasureUnsigned(packet.DecodedLength);
 					prevPacket = packet;
 				}
 
-				length += IoHash.NumBytes + VarInt.Measure(export.Rank) + VarInt.Measure(export.Length) + VarInt.Measure(export.References.Count) + export.References.Sum(x => VarInt.Measure(x));
+				length += IoHash.NumBytes + VarInt.MeasureUnsigned(export.Rank) + VarInt.MeasureUnsigned(export.Length) + VarInt.MeasureUnsigned(export.References.Count) + export.References.Sum(x => VarInt.MeasureUnsigned(x));
 			}
 
 			return length;
diff --git a/Engine/Source/Programs/Shared/EpicGames.Horde/Bundles/Nodes/DirectoryNode.cs b/Engine/Source/Programs/Shared/EpicGames.Horde/Bundles/Nodes/DirectoryNode.cs
--- a/Engine/Source/Programs/Shared/EpicGames.Horde/Bundles/Nodes/DirectoryNode.cs
+++ b/Engine/Source/Programs/Shared/EpicGames.Horde/Bundles/Nodes/DirectoryNode.cs
@@ -388,15 +388,15 @@ namespace EpicGames.Horde.Bundles.Nodes
 		{
 			// Measure the required size of the write buffer
 			int size = 1;
-			size += VarInt.Measure(_nameToFileEntry.Count);
+			size += VarInt.MeasureUnsigned(_nameToFileEntry.Count);
 			foreach (FileEntry entry in _nameToFileEntry.Values)
 			{
-				size += (entry.Name.Length + 1) + 1 + VarInt.Measure((ulong)entry.Length) + IoHash.NumBytes;
+				size += (entry.Name.Length + 1) + 1 + VarInt.MeasureUnsigned((ulong)entry.Length) + IoHash.NumBytes;
 			}
-			size += VarInt.Measure(_nameToDirectoryEntry.Count);
+			size += VarInt.MeasureUnsigned(_nameToDirectoryEntry.Count);
 			foreach (DirectoryEntry entry in _nameToDirectoryEntry.Values)
 			{
-				size += (entry.Name.Length + 1) + VarInt.Measure((ulong)entry.Length) + IoHash.NumBytes;
+				size += (entry.Name.Length + 1) + VarInt.MeasureUnsigned((ulong)entry.Length) + IoHash.NumBytes;
 			}
 
 			// Allocate the buffer and copy the node to it
@@ -406,7 +406,7 @@ namespace EpicGames.Horde.Bundles.Nodes
 			span[0] = TypeId;
 			span = span.Slice(1);
 
-			int fileCountBytes = VarInt.Write(span, _nameToFileEntry.Count);
+			int fileCountBytes = VarInt.WriteUnsigned(span, _nameToFileEntry.Count);
 			span = span.Slice(fileCountBytes);
 
 			foreach (FileEntry fileEntry in _nameToFileEntry.Values.OrderBy(x => x.Name))
@@ -417,14 +417,14 @@ namespace EpicGames.Horde.Bundles.Nodes
 				span[0] = (byte)fileEntry.Flags;
 				span = span.Slice(1);
 
-				int lengthBytes = VarInt.Write(span, fileEntry.Length);
+				int lengthBytes = VarInt.WriteUnsigned(span, fileEntry.Length);
 				span = span.Slice(lengthBytes);
 
 				fileEntry.Hash.CopyTo(span);
 				span = span.Slice(IoHash.NumBytes);
 			}
 
-			int directoryCountBytes = VarInt.Write(span, _nameToDirectoryEntry.Count);
+			int directoryCountBytes = VarInt.WriteUnsigned(span, _nameToDirectoryEntry.Count);
 			span = span.Slice(directoryCountBytes);
 
 			foreach (DirectoryEntry directoryEntry in _nameToDirectoryEntry.Values.OrderBy(x => x.Name))
@@ -432,7 +432,7 @@ namespace EpicGames.Horde.Bundles.Nodes
 				directoryEntry.Name.Span.CopyTo(span);
 				span = span.Slice(directoryEntry.Name.Length + 1);
 
-				int lengthBytes = VarInt.Write(span, directoryEntry.Length);
+				int lengthBytes = VarInt.WriteUnsigned(span, directoryEntry.Length);
 				span = span.Slice(lengthBytes);
 
 				directoryEntry.Hash.CopyTo(span);
@@ -460,7 +460,7 @@ namespace EpicGames.Horde.Bundles.Nodes
 
 			DirectoryNode node = new DirectoryNode();
 
-			int fileCount = (int)VarInt.Read(span, out int fileCountBytes);
+			int fileCount = (int)VarInt.ReadUnsigned(span, out int fileCountBytes);
 			span = span.Slice(fileCountBytes);
 
 			node._nameToFileEntry.EnsureCapacity(fileCount);
@@ -473,7 +473,7 @@ namespace EpicGames.Horde.Bundles.Nodes
 				FileEntryFlags flags = (FileEntryFlags)span[0];
 				span = span.Slice(1);
 
-				long length = (long)VarInt.Read(span, out int lengthBytes);
+				long length = (long)VarInt.ReadUnsigned(span, out int lengthBytes);
 				span = span.Slice(lengthBytes);
 
 				IoHash hash = new IoHash(span);
@@ -482,7 +482,7 @@ namespace EpicGames.Horde.Bundles.Nodes
 				node._nameToFileEntry[name] = new FileEntry(node, name, flags, length, hash);
 			}
 
-			int directoryCount = (int)VarInt.Read(span, out int directoryCountBytes);
+			int directoryCount = (int)VarInt.ReadUnsigned(span, out int directoryCountBytes);
 			span = span.Slice(directoryCountBytes);
 
 			node._nameToDirectoryEntry.EnsureCapacity(directoryCount);
@@ -492,7 +492,7 @@ namespace EpicGames.Horde.Bundles.Nodes
 				Utf8String name = new Utf8String(span.Slice(0, nameLen).ToArray());
 				span = span.Slice(nameLen + 1);
 
-				long length = (long)VarInt.Read(span, out int lengthBytes);
+				long length = (long)VarInt.ReadUnsigned(span, out int lengthBytes);
 				span = span.Slice(lengthBytes);
 
 				IoHash hash = new IoHash(span);
diff --git a/Engine/Source/Programs/Shared/EpicGames.Horde/Bundles/Nodes/FileNode.cs b/Engine/Source/Programs/Shared/EpicGames.Horde/Bundles/Nodes/FileNode.cs
--- a/Engine/Source/Programs/Shared/EpicGames.Horde/Bundles/Nodes/FileNode.cs
+++ b/Engine/Source/Programs/Shared/EpicGames.Horde/Bundles/Nodes/FileNode.cs
@@ -144,7 +144,7 @@ namespace EpicGames.Horde.Bundles.Nodes
 			FileNode node = new FileNode();
 			node._data = new ReadOnlySequence<byte>(data);
 			node._hash = IoHash.Compute(data.Span);
-			node.Depth = (int)VarInt.Read(span.Slice(1), out int depthBytes);
+			node.Depth = (int)VarInt.ReadUnsigned(span.Slice(1), out int depthBytes);
 
 			int headerLength = 1 + depthBytes;
 			span = span.Slice(headerLength);
@@ -158,7 +158,7 @@ namespace EpicGames.Horde.Bundles.Nodes
 			}
 			else
 			{
-				node.Length = (long)VarInt.Read(span, out int lengthBytes);
+				node.Length = (long)VarInt.ReadUnsigned(span, out int lengthBytes);
 				span = span.Slice(lengthBytes);
 
 				node._childNodeRefs = new List<BundleNodeRef<FileNode>>(payload.Length / IoHash.NumBytes);
@@ -449,12 +449,12 @@ namespace EpicGames.Horde.Bundles.Nodes
 
 		private static byte[] CreateFirstSegmentData(int depth, ReadOnlySpan<byte> leafData)
 		{
-			int depthBytes = VarInt.Measure(depth);
+			int depthBytes = VarInt.MeasureUnsigned(depth);
 
 			byte[] buffer = new byte[1 + depthBytes + leafData.Length];
 			buffer[0] = TypeId;
 
-			VarInt.Write(buffer.AsSpan(1, depthBytes), depth);
+			VarInt.WriteUnsigned(buffer.AsSpan(1, depthBytes), depth);
 			leafData.CopyTo(buffer.AsSpan(1 + depthBytes));
 
 			return buffer;
@@ -466,8 +466,8 @@ namespace EpicGames.Horde.Bundles.Nodes
 			{
 				if (_childNodeRefs != null)
 				{
-					int depthBytes = VarInt.Measure(Depth);
-					int lengthBytes = VarInt.Measure((ulong)Length);
+					int depthBytes = VarInt.MeasureUnsigned(Depth);
+					int lengthBytes = VarInt.MeasureUnsigned((ulong)Length);
 
 					byte[] newData = new byte[1 + depthBytes + lengthBytes + (_childNodeRefs.Count * IoHash.NumBytes)];
 					_data = new ReadOnlySequence<byte>(newData);
@@ -476,10 +476,10 @@ namespace EpicGames.Horde.Bundles.Nodes
 					newData[0] = TypeId;
 					Span<byte> span = newData.AsSpan(1);
 
-					VarInt.Write(span, Depth);
+					VarInt.WriteUnsigned(span, Depth);
 					span = span.Slice(depthBytes);
 
-					VarInt.Write(span, Length);
+					VarInt.WriteUnsigned(span, Length);
 					span = span.Slice(lengthBytes);
 
 					foreach (BundleNodeRef<FileNode> childNodeRef in _childNodeRefs)
diff --git a/Engine/Source/Programs/Shared/EpicGames.Serialization.Tests/EngineTests.cs b/Engine/Source/Programs/Shared/EpicGames.Serialization.Tests/EngineTests.cs
--- a/Engine/Source/Programs/Shared/EpicGames.Serialization.Tests/EngineTests.cs
+++ b/Engine/Source/Programs/Shared/EpicGames.Serialization.Tests/EngineTests.cs
@@ -585,7 +585,7 @@ namespace EpicGames.Serialization.Tests
 			// Test CbField(String, OutOfRangeSize)
 			{
 				byte[] payload = new byte[9];
-				VarInt.Write(payload, (ulong)(1) << 31);
+				VarInt.WriteUnsigned(payload, (ulong)(1) << 31);
 				TestFieldError(CbFieldType.String, payload, CbFieldError.RangeError, new Utf8String("ABC"));
 			}
 
@@ -628,7 +628,7 @@ namespace EpicGames.Serialization.Tests
 		{
 			byte[] payload = new byte[9];
 			ulong negative = (ulong)((byte)fieldType & 1);
-			VarInt.Write(payload, magnitude - negative);
+			VarInt.WriteUnsigned(payload, magnitude - negative);
 			ulong defaultValue = 8;
 			ulong expectedValue = (negative != 0)? (ulong)(-(long)(magnitude)) : magnitude;
 			CbField field = new CbField(payload, fieldType);
@@ -722,7 +722,7 @@ namespace EpicGames.Serialization.Tests
 			// Test CbField(Integer+, MaxBinary32) as Float
 			{
 				byte[] payload = new byte[9];
-				VarInt.Write(payload, ((ulong)(1) << 24) - 1); // 16,777,215
+				VarInt.WriteUnsigned(payload, ((ulong)(1) << 24) - 1); // 16,777,215
 				CbField field = new CbField(payload, CbFieldType.IntegerPositive);
 				TestField(CbFieldType.Float32, field, 16_777_215.0f);
 				TestField(CbFieldType.Float64, field, 16_777_215.0);
@@ -731,7 +731,7 @@ namespace EpicGames.Serialization.Tests
 			// Test CbField(Integer+, MaxBinary32+1) as Float
 			{
 				byte[] payload = new byte[9];
-				VarInt.Write(payload, (ulong)(1) << 24); // 16,777,216
+				VarInt.WriteUnsigned(payload, (ulong)(1) << 24); // 16,777,216
 				CbField field = new CbField(payload, CbFieldType.IntegerPositive);
 				TestFieldError(CbFieldType.Float32, field, CbFieldError.RangeError, 8.0f);
 				TestField(CbFieldType.Float64, field, 16_777_216.0);
@@ -740,7 +740,7 @@ namespace EpicGames.Serialization.Tests
 			// Test CbField(Integer+, MaxBinary64) as Float
 			{
 				byte[] payload = new byte[9];
-				VarInt.Write(payload, ((ulong)(1) << 53) - 1); // 9,007,199,254,740,991
+				VarInt.WriteUnsigned(payload, ((ulong)(1) << 53) - 1); // 9,007,199,254,740,991
 				CbField field = new CbField(payload, CbFieldType.IntegerPositive);
 				TestFieldError(CbFieldType.Float32, field, CbFieldError.RangeError, 8.0f);
 				TestField(CbFieldType.Float64, field, 9_007_199_254_740_991.0);
@@ -749,7 +749,7 @@ namespace EpicGames.Serialization.Tests
 			// Test CbField(Integer+, MaxBinary64+1) as Float
 			{
 				byte[] payload = new byte[9];
-				VarInt.Write(payload, (ulong)(1) << 53); // 9,007,199,254,740,992
+				VarInt.WriteUnsigned(payload, (ulong)(1) << 53); // 9,007,199,254,740,992
 				CbField field = new CbField(payload, CbFieldType.IntegerPositive);
 				TestFieldError(CbFieldType.Float32, field, CbFieldError.RangeError, 8.0f);
 				TestFieldError(CbFieldType.Float64, field, CbFieldError.RangeError, 8.0);
@@ -758,7 +758,7 @@ namespace EpicGames.Serialization.Tests
 			// Test CbField(Integer+, MaxUInt64) as Float
 			{
 				byte[] payload = new byte[9];
-				VarInt.Write(payload, ~(ulong)0); // Max uint64
+				VarInt.WriteUnsigned(payload, ~(ulong)0); // Max uint64
 				CbField field = new CbField(payload, CbFieldType.IntegerPositive);
 				TestFieldError(CbFieldType.Float32, field, CbFieldError.RangeError, 8.0f);
 				TestFieldError(CbFieldType.Float64, field, CbFieldError.RangeError, 8.0);
@@ -767,7 +767,7 @@ namespace EpicGames.Serialization.Tests
 			// Test CbField(Integer-, MaxBinary32) as Float
 			{
 				byte[] payload = new byte[9];
-				VarInt.Write(payload, ((ulong)(1) << 24) - 2); // -16,777,215
+				VarInt.WriteUnsigned(payload, ((ulong)(1) << 24) - 2); // -16,777,215
 				CbField field = new CbField(payload, CbFieldType.IntegerNegative);
 				TestField(CbFieldType.Float32, field, -16_777_215.0f);
 				TestField(CbFieldType.Float64, field, -16_777_215.0);
@@ -776,7 +776,7 @@ namespace EpicGames.Serialization.Tests
 			// Test CbField(Integer-, MaxBinary32+1) as Float
 			{
 				byte[] payload = new byte[9];
-				VarInt.Write(payload, ((ulong)(1) << 24) - 1); // -16,777,216
+				VarInt.WriteUnsigned(payload, ((ulong)(1) << 24) - 1); // -16,777,216
 				CbField field = new CbField(payload, CbFieldType.IntegerNegative);
 				TestFieldError(CbFieldType.Float32, field, CbFieldError.RangeError, 8.0f);
 				TestField(CbFieldType.Float64, field, -16_777_216.0);
@@ -785,7 +785,7 @@ namespace EpicGames.Serialization.Tests
 			// Test CbField(Integer-, MaxBinary64) as Float
 			{
 				byte[] payload = new byte[9];
-				VarInt.Write(payload, ((ulong)(1) << 53) - 2); // -9,007,199,254,740,991
+				VarInt.WriteUnsigned(payload, ((ulong)(1) << 53) - 2); // -9,007,199,254,740,991
 				CbField field = new CbField(payload, CbFieldType.IntegerNegative);
 				TestFieldError(CbFieldType.Float32, field, CbFieldError.RangeError, 8.0f);
 				TestField(CbFieldType.Float64, field, -9_007_199_254_740_991.0);
@@ -794,7 +794,7 @@ namespace EpicGames.Serialization.Tests
 			// Test CbField(Integer-, MaxBinary64+1) as Float
 			{
 				byte[] payload = new byte[9];
-				VarInt.Write(payload, ((ulong)(1) << 53) - 1); // -9,007,199,254,740,992
+				VarInt.WriteUnsigned(payload, ((ulong)(1) << 53) - 1); // -9,007,199,254,740,992
 				CbField field = new CbField(payload, CbFieldType.IntegerNegative);
 				TestFieldError(CbFieldType.Float32, field, CbFieldError.RangeError, 8.0f);
 				TestFieldError(CbFieldType.Float64, field, CbFieldError.RangeError, 8.0);
diff --git a/Engine/Source/Programs/Shared/EpicGames.Serialization.Tests/VarIntTests.cs b/Engine/Source/Programs/Shared/EpicGames.Serialization.Tests/VarIntTests.cs
--- a/Engine/Source/Programs/Shared/EpicGames.Serialization.Tests/VarIntTests.cs
+++ b/Engine/Source/Programs/Shared/EpicGames.Serialization.Tests/VarIntTests.cs
@@ -12,12 +12,12 @@ namespace EpicGames.Serialization.Tests
 		{
 			byte[] buffer = new byte[20];
 
-			int length = VarInt.Write(buffer, -1);
+			int length = VarInt.WriteUnsigned(buffer, -1);
 			Assert.AreEqual(9, length);
-			Assert.AreEqual(9, VarInt.Measure(-1));
+			Assert.AreEqual(9, VarInt.MeasureUnsigned(-1));
 
 			Assert.AreEqual(9, VarInt.Measure(buffer));
-			int value = (int)(long)VarInt.Read(buffer, out int bytesRead);
+			int value = (int)(long)VarInt.ReadUnsigned(buffer, out int bytesRead);
 			Assert.AreEqual(9, bytesRead);
 
 			Assert.AreEqual(-1, value);
diff --git a/Engine/Source/Programs/Shared/EpicGames.Serialization/CbField.cs b/Engine/Source/Programs/Shared/EpicGames.Serialization/CbField.cs
--- a/Engine/Source/Programs/Shared/EpicGames.Serialization/CbField.cs
+++ b/Engine/Source/Programs/Shared/EpicGames.Serialization/CbField.cs
@@ -765,7 +765,7 @@ namespace EpicGames.Serialization
 
 			if (CbFieldUtils.HasFieldName(type))
 			{
-				_nameLen = (int)VarInt.Read(data.Slice(offset).Span, out int nameLenByteCount);
+				_nameLen = (int)VarInt.ReadUnsigned(data.Slice(offset).Span, out int nameLenByteCount);
 				offset += nameLenByteCount + _nameLen;
 			}
 
@@ -895,7 +895,7 @@ namespace EpicGames.Serialization
 			{
 				Error = CbFieldError.None;
 
-				ulong length = VarInt.Read(Payload.Span, out int bytesRead);
+				ulong length = VarInt.ReadUnsigned(Payload.Span, out int bytesRead);
 				return Payload.Slice(bytesRead, (int)length);
 			}
 			else
@@ -953,7 +953,7 @@ namespace EpicGames.Serialization
 		{
 			if (CbFieldUtils.IsString(TypeWithFlags))
 			{
-				ulong valueSize = VarInt.Read(Payload.Span, out int valueSizeByteCount);
+				ulong valueSize = VarInt.ReadUnsigned(Payload.Span, out int valueSizeByteCount);
 				if (valueSize >= (1UL << 31))
 				{
 					Error = CbFieldError.RangeError;
@@ -1083,7 +1083,7 @@ namespace EpicGames.Serialization
 				ulong outOfRangeMask = ~(ulong)1 << (magnitudeBits - 1);
 				ulong isNegative = (ulong)(byte)(TypeWithFlags) & 1;
 
-				ulong magnitude = VarInt.Read(Payload.Span, out _);
+				ulong magnitude = VarInt.ReadUnsigned(Payload.Span, out _);
 				ulong value = magnitude ^ (ulong)-(long)(isNegative);
 
 				if ((magnitude & outOfRangeMask) == 0 && (isNegative == 0 || isSigned))
@@ -1119,7 +1119,7 @@ namespace EpicGames.Serialization
 						ulong isNegative = (ulong)TypeWithFlags & 1;
 						ulong outOfRangeMask = ~((1UL << /*FLT_MANT_DIG*/ 24) - 1);
 
-						ulong magnitude = VarInt.Read(Payload.Span, out _) + isNegative;
+						ulong magnitude = VarInt.ReadUnsigned(Payload.Span, out _) + isNegative;
 						bool isInRange = (magnitude & outOfRangeMask) == 0;
 						Error = isInRange ? CbFieldError.None : CbFieldError.RangeError;
 						return isInRange ? (float)((isNegative != 0) ? (float)-(long)magnitude : (float)magnitude) : defaultValue;
@@ -1157,7 +1157,7 @@ namespace EpicGames.Serialization
 						ulong isNegative = (ulong)TypeWithFlags & 1;
 						ulong outOfRangeMask = ~((1UL << /*DBL_MANT_DIG*/ 53) - 1);
 
-						ulong magnitude = VarInt.Read(Payload.Span, out _) + isNegative;
+						ulong magnitude = VarInt.ReadUnsigned(Payload.Span, out _) + isNegative;
 						bool isInRange = (magnitude & outOfRangeMask) == 0;
 						Error = isInRange ? CbFieldError.None : CbFieldError.RangeError;
 						return isInRange ? (double)((isNegative != 0) ? (double)-(long)magnitude : (double)magnitude) : defaultValue;
@@ -1615,7 +1615,7 @@ namespace EpicGames.Serialization
 			{
 				ReadOnlyMemory<byte> payloadBytes = Payload;
 				int payloadSizeByteCount;
-				int payloadSize = (int)VarInt.Read(payloadBytes.Span, out payloadSizeByteCount);
+				int payloadSize = (int)VarInt.ReadUnsigned(payloadBytes.Span, out payloadSizeByteCount);
 				payloadBytes = payloadBytes.Slice(payloadSizeByteCount);
 				int numByteCount = CbFieldUtils.IsArray(localTypeWithFlags) ? (int)VarInt.Measure(payloadBytes.Span) : 0;
 				if (payloadSize > numByteCount)
@@ -1653,7 +1653,7 @@ namespace EpicGames.Serialization
 		/// <returns></returns>
 		private ReadOnlyMemory<byte> GetViewNoType()
 		{
-			int nameSize = CbFieldUtils.HasFieldName(TypeWithFlags) ? _nameLen + (int)VarInt.Measure((uint)_nameLen) : 0;
+			int nameSize = CbFieldUtils.HasFieldName(TypeWithFlags) ? _nameLen + (int)VarInt.MeasureUnsigned((uint)_nameLen) : 0;
 			return Memory.Slice(_payloadOffset - nameSize);
 		}
 
@@ -1698,7 +1698,7 @@ namespace EpicGames.Serialization
 				case CbFieldType.Binary:
 				case CbFieldType.String:
 					{
-						ulong payloadSize = VarInt.Read(Payload.Span, out int bytesRead);
+						ulong payloadSize = VarInt.ReadUnsigned(Payload.Span, out int bytesRead);
 						return payloadSize + (ulong)bytesRead;
 					}
 				case CbFieldType.IntegerPositive:
@@ -2273,7 +2273,7 @@ namespace EpicGames.Serialization
 			{
 				ReadOnlyMemory<byte> payloadBytes = _innerField.Payload;
 				payloadBytes = payloadBytes.Slice(VarInt.Measure(payloadBytes.Span));
-				return (int)VarInt.Read(payloadBytes.Span, out int _);
+				return (int)VarInt.ReadUnsigned(payloadBytes.Span, out int _);
 			}
 		}
 
diff --git a/Engine/Source/Programs/Shared/EpicGames.Serialization/CbWriter.cs b/Engine/Source/Programs/Shared/EpicGames.Serialization/CbWriter.cs
--- a/Engine/Source/Programs/Shared/EpicGames.Serialization/CbWriter.cs
+++ b/Engine/Source/Programs/Shared/EpicGames.Serialization/CbWriter.cs
@@ -283,7 +283,7 @@ namespace EpicGames.Serialization
 				throw new CbWriterException($"Named fields are not allowed within fields of type {scope._fieldType}");
 			}
 
-			int nameVarIntLength = VarInt.Measure(name.Length);
+			int nameVarIntLength = VarInt.MeasureUnsigned(name.Length);
 			if (scope._uniformFieldType == CbFieldType.None)
 			{
 				Span<byte> buffer = Allocate(1 + nameVarIntLength + name.Length).Span;
@@ -456,9 +456,9 @@ namespace EpicGames.Serialization
 		/// <param name="value">Value to write</param>
 		void WriteIntegerPayload(ulong value)
 		{
-			int length = VarInt.Measure(value);
+			int length = VarInt.MeasureUnsigned(value);
 			Span<byte> buffer = Allocate(length).Span;
-			VarInt.Write(buffer, value);
+			VarInt.WriteUnsigned(buffer, value);
 		}
 
 		/// <summary>
@@ -702,7 +702,7 @@ namespace EpicGames.Serialization
 		/// <param name="value">Value to be written</param>
 		static void WriteBinaryPayload(Span<byte> output, ReadOnlySpan<byte> value)
 		{
-			int varIntLength = VarInt.Write(output, value.Length);
+			int varIntLength = VarInt.WriteUnsigned(output, value.Length);
 			output = output[varIntLength..];
 
 			value.CopyTo(output);
@@ -715,7 +715,7 @@ namespace EpicGames.Serialization
 		/// <param name="value">Value to be written</param>
 		void WriteBinaryPayload(ReadOnlySpan<byte> value)
 		{
-			int valueVarIntLength = VarInt.Measure(value.Length);
+			int valueVarIntLength = VarInt.MeasureUnsigned(value.Length);
 			Span<byte> buffer = Allocate(valueVarIntLength + value.Length).Span;
 			WriteBinaryPayload(buffer, value);
 		}
@@ -1031,12 +1031,12 @@ namespace EpicGames.Serialization
 				{
 					case CbFieldType.Object:
 					case CbFieldType.UniformObject:
-						sizeOfChildHeaders += childScope._sizeOfChildHeaders + VarInt.Measure(childScope._length + childScope._sizeOfChildHeaders);
+						sizeOfChildHeaders += childScope._sizeOfChildHeaders + VarInt.MeasureUnsigned(childScope._length + childScope._sizeOfChildHeaders);
 						break;
 					case CbFieldType.Array:
 					case CbFieldType.UniformArray:
-						int arrayCountLength = VarInt.Measure(childScope._count);
-						sizeOfChildHeaders += childScope._sizeOfChildHeaders + VarInt.Measure(childScope._length + childScope._sizeOfChildHeaders + arrayCountLength) + arrayCountLength;
+						int arrayCountLength = VarInt.MeasureUnsigned(childScope._count);
+						sizeOfChildHeaders += childScope._sizeOfChildHeaders + VarInt.MeasureUnsigned(childScope._length + childScope._sizeOfChildHeaders + arrayCountLength) + arrayCountLength;
 						break;
 					default:
 						throw new InvalidOperationException();
@@ -1057,12 +1057,12 @@ namespace EpicGames.Serialization
 			{
 				case CbFieldType.Object:
 				case CbFieldType.UniformObject:
-					return VarInt.Write(span, scope._length + scope._sizeOfChildHeaders);
+					return VarInt.WriteUnsigned(span, scope._length + scope._sizeOfChildHeaders);
 				case CbFieldType.Array:
 				case CbFieldType.UniformArray:
-					int numItemsLength = VarInt.Measure(scope._count);
-					int offset = VarInt.Write(span, scope._length + scope._sizeOfChildHeaders + numItemsLength);
-					return offset + VarInt.Write(span.Slice(offset), scope._count);
+					int numItemsLength = VarInt.MeasureUnsigned(scope._count);
+					int offset = VarInt.WriteUnsigned(span, scope._length + scope._sizeOfChildHeaders + numItemsLength);
+					return offset + VarInt.WriteUnsigned(span.Slice(offset), scope._count);
 				default:
 					throw new InvalidOperationException();
 			}
diff --git a/Engine/Source/Programs/Shared/EpicGames.Serialization/VarInt.cs b/Engine/Source/Programs/Shared/EpicGames.Serialization/VarInt.cs
--- a/Engine/Source/Programs/Shared/EpicGames.Serialization/VarInt.cs
+++ b/Engine/Source/Programs/Shared/EpicGames.Serialization/VarInt.cs
@@ -2,6 +2,7 @@
 
 using System;
 using System.Numerics;
+using EpicGames.Core;
 
 namespace EpicGames.Serialization
 {
@@ -15,18 +16,51 @@ namespace EpicGames.Serialization
 		/// </summary>
 		/// <param name="buffer">A variable-length encoding of an unsigned integer</param>
 		/// <returns></returns>
-		public static ulong Read(ReadOnlySpan<byte> buffer)
+		[Obsolete("Call ReadSigned or ReadUnsigned instead")]
+		public static ulong Read(ReadOnlySpan<byte> buffer) => ReadUnsigned(buffer, out _);
+
+		/// <summary>
+		/// Read a variable-length unsigned integer.
+		/// </summary>
+		/// <param name="buffer">A variable-length encoding of an unsigned integer</param>
+		/// <param name="bytesRead">The number of bytes consumed from the input</param>
+		/// <returns></returns>
+		[Obsolete("Call ReadSigned or ReadUnsigned instead")]
+		public static ulong Read(ReadOnlySpan<byte> buffer, out int bytesRead) => ReadUnsigned(buffer, out bytesRead);
+
+		/// <summary>
+		/// Read a variable-length signed integer.
+		/// </summary>
+		/// <param name="buffer">A variable-length encoding of an unsigned integer</param>
+		/// <returns></returns>
+		public static long ReadSigned(ReadOnlySpan<byte> buffer) => ReadSigned(buffer, out _);
+
+		/// <summary>
+		/// Read a variable-length signed integer.
+		/// </summary>
+		/// <param name="buffer">A variable-length encoding of an unsigned integer</param>
+		/// <param name="bytesRead">The number of bytes consumed from the input</param>
+		/// <returns></returns>
+		public static long ReadSigned(ReadOnlySpan<byte> buffer, out int bytesRead)
 		{
-			return Read(buffer, out _);
+			ulong unsignedValue = ReadUnsigned(buffer, out bytesRead);
+			return DecodeSigned(unsignedValue);
 		}
 
+		/// <summary>
+		/// Read a variable-length unsigned integer.
+		/// </summary>
+		/// <param name="buffer">A variable-length encoding of an unsigned integer</param>
+		/// <returns></returns>
+		public static ulong ReadUnsigned(ReadOnlySpan<byte> buffer) => ReadUnsigned(buffer, out _);
+
 		/// <summary>
 		/// Read a variable-length unsigned integer.
 		/// </summary>
 		/// <param name="buffer">A variable-length encoding of an unsigned integer</param>
 		/// <param name="bytesRead">The number of bytes consumed from the input</param>
 		/// <returns></returns>
-		public static ulong Read(ReadOnlySpan<byte> buffer, out int bytesRead)
+		public static ulong ReadUnsigned(ReadOnlySpan<byte> buffer, out int bytesRead)
 		{
 			bytesRead = (int)Measure(buffer);
 
@@ -51,14 +85,40 @@ namespace EpicGames.Serialization
 			return BitOperations.LeadingZeroCount(b) - 23;
 		}
 
+		/// <summary>
+		/// Measure the number of bytes required to encode the input.
+		/// </summary>
+		/// <param name="value"></param>
+		/// <returns></returns>
+		[Obsolete("Use MeasureSigned or MeasureUnsigned instead")]
+		public static int Measure(uint value) => MeasureUnsigned(value);
+
+		/// <summary>
+		/// Measure the number of bytes required to encode the input.
+		/// </summary>
+		/// <param name="value"></param>
+		/// <returns></returns>
+		[Obsolete("Use MeasureSigned or MeasureUnsigned instead")]
+		public static int Measure(ulong value) => MeasureUnsigned(value);
+
+		/// <summary>
+		/// Measure the number of bytes (1-9) required to encode the input.
+		/// </summary>
+		/// <param name="value"></param>
+		/// <returns></returns>
+		public static int MeasureSigned(long value)
+		{
+			return MeasureUnsigned(EncodeSigned(value));
+		}
+
 		/// <summary>
 		/// Measure the number of bytes (1-5) required to encode the 32-bit input.
 		/// </summary>
 		/// <param name="value"></param>
 		/// <returns></returns>
-		public static int Measure(int value)
+		public static int MeasureUnsigned(int value)
 		{
-			return Measure((ulong)(long)value);
+			return MeasureUnsigned((ulong)(long)value);
 		}
 
 		/// <summary>
@@ -66,7 +126,7 @@ namespace EpicGames.Serialization
 		/// </summary>
 		/// <param name="value"></param>
 		/// <returns></returns>
-		public static int Measure(uint value)
+		public static int MeasureUnsigned(uint value)
 		{
 			return BitOperations.Log2(value) / 7 + 1;
 		}
@@ -76,7 +136,7 @@ namespace EpicGames.Serialization
 		/// </summary>
 		/// <param name="value"></param>
 		/// <returns></returns>
-		public static int Measure(ulong value)
+		public static int MeasureUnsigned(ulong value)
 		{
 			return Math.Min(BitOperations.Log2(value) / 7 + 1, 9);
 		}
@@ -87,9 +147,28 @@ namespace EpicGames.Serialization
 		/// <param name="value">An unsigned integer to encode</param>
 		/// <param name="buffer">A buffer of at least 9 bytes to write the output to.</param>
 		/// <returns>The number of bytes used in the output</returns>
-		public static int Write(Span<byte> buffer, long value)
+		[Obsolete("Use WriteUnsigned or WriteSigned instead")]
+		public static int Write(Span<byte> buffer, long value) => WriteUnsigned(buffer, (ulong)value);
+
+		/// <summary>
+		/// Write a variable-length unsigned integer.
+		/// </summary>
+		/// <param name="value">An unsigned integer to encode</param>
+		/// <param name="buffer">A buffer of at least 9 bytes to write the output to.</param>
+		/// <returns>The number of bytes used in the output</returns>
+		[Obsolete("Use WriteUnsigned or WriteSigned instead")]
+		public static int Write(Span<byte> buffer, ulong value) => WriteUnsigned(buffer, value);
+
+		/// <summary>
+		/// Write a variable-length unsigned integer.
+		/// </summary>
+		/// <param name="value">An unsigned integer to encode</param>
+		/// <param name="buffer">A buffer of at least 9 bytes to write the output to.</param>
+		/// <returns>The number of bytes used in the output</returns>
+		public static int WriteSigned(Span<byte> buffer, long value)
 		{
-			return Write(buffer, (ulong)value);
+			ulong unsignedValue = EncodeSigned(value);
+			return WriteUnsigned(buffer, unsignedValue);
 		}
 
 		/// <summary>
@@ -98,9 +177,25 @@ namespace EpicGames.Serialization
 		/// <param name="value">An unsigned integer to encode</param>
 		/// <param name="buffer">A buffer of at least 9 bytes to write the output to.</param>
 		/// <returns>The number of bytes used in the output</returns>
-		public static int Write(Span<byte> buffer, ulong value)
+		public static int WriteUnsigned(Span<byte> buffer, int value) => WriteUnsigned(buffer, (ulong)value);
+
+		/// <summary>
+		/// Write a variable-length unsigned integer.
+		/// </summary>
+		/// <param name="value">An unsigned integer to encode</param>
+		/// <param name="buffer">A buffer of at least 9 bytes to write the output to.</param>
+		/// <returns>The number of bytes used in the output</returns>
+		public static int WriteUnsigned(Span<byte> buffer, long value) => WriteUnsigned(buffer, (ulong)value);
+
+		/// <summary>
+		/// Write a variable-length unsigned integer.
+		/// </summary>
+		/// <param name="value">An unsigned integer to encode</param>
+		/// <param name="buffer">A buffer of at least 9 bytes to write the output to.</param>
+		/// <returns>The number of bytes used in the output</returns>
+		public static int WriteUnsigned(Span<byte> buffer, ulong value)
 		{
-			int byteCount = Measure(value);
+			int byteCount = MeasureUnsigned(value);
 
 			for (int idx = 1; idx < byteCount; idx++)
 			{
@@ -110,5 +205,25 @@ namespace EpicGames.Serialization
 			buffer[0] = (byte)((0xff << (9 - (int)byteCount)) | (byte)value);
 			return byteCount;
 		}
+
+		/// <summary>
+		/// Decode a signed VarInt value from an unsigned value
+		/// </summary>
+		/// <param name="value">Value to decode</param>
+		/// <returns>Decoded value</returns>
+		public static long DecodeSigned(ulong value)
+		{
+			return -(long)(value & 1) ^ (long)(value >> 1);
+		}
+
+		/// <summary>
+		/// Encode a signed VarInt value into an unsigned value
+		/// </summary>
+		/// <param name="value">Value to encode</param>
+		/// <returns>Encoded value</returns>
+		public static ulong EncodeSigned(long value)
+		{
+			return (ulong)((value >> 63) ^ (value << 1));
+		}
 	}
 }
