EpicGames.OIDC - Do not throw exceptions when we are unable to read the refresh token, instead we just treat it as not existing.

diff --git a/Engine/Source/Programs/Shared/EpicGames.OIDC/WindowsTokenStore.cs b/Engine/Source/Programs/Shared/EpicGames.OIDC/WindowsTokenStore.cs
--- a/Engine/Source/Programs/Shared/EpicGames.OIDC/WindowsTokenStore.cs
+++ b/Engine/Source/Programs/Shared/EpicGames.OIDC/WindowsTokenStore.cs
@@ -92,10 +92,23 @@ namespace EpicGames.OIDC
 				return false;
 			}
 
-			byte[] bytes = CryptProtectDataHelper.DoCryptUnprotectData(encryptedToken, $"OidcToken-{oidcProvider}", GetEntropy(oidcProvider));
-			refreshToken = Encoding.Unicode.GetString(bytes);
+			try
+			{
+				byte[] bytes = CryptProtectDataHelper.DoCryptUnprotectData(encryptedToken, $"OidcToken-{oidcProvider}", GetEntropy(oidcProvider));
+				refreshToken = Encoding.Unicode.GetString(bytes);
 
-			return true;
+				return true;
+			}
+			catch (Win32Exception e)
+			{
+				if (e.NativeErrorCode == 13) // data is invalid
+				{
+					// unable to decrypt the data, ignore it
+					refreshToken = "";
+					return false;
+				}
+				throw;
+			}
 		}
 
 		private static byte[] GetEntropy(string oidcProvider)
@@ -139,6 +152,7 @@ namespace EpicGames.OIDC
 		}
 	}
 
+#pragma warning disable IDE1006 // Pinvoke code doesnt use the same naming conventions as C#
 	static class CryptProtectDataHelper
 	{
 		
@@ -306,4 +320,5 @@ namespace EpicGames.OIDC
 			return buf;
 		}
 	}
+#pragma warning restore IDE1006 // Naming Styles
 }
\ No newline at end of file
