EpicGames.Core: Add a utility class to serialize data into a chunked byte sequence, similar to StringBuilder.

diff --git a/Engine/Source/Programs/Shared/EpicGames.Core/ByteArrayBuilder.cs b/Engine/Source/Programs/Shared/EpicGames.Core/ByteArrayBuilder.cs
new file mode 100644
--- /dev/null
+++ b/Engine/Source/Programs/Shared/EpicGames.Core/ByteArrayBuilder.cs
@@ -0,0 +1,177 @@
+// Copyright Epic Games, Inc. All Rights Reserved.
+
+using System;
+using System.Buffers;
+using System.Collections.Generic;
+
+namespace EpicGames.Core
+{
+	/// <summary>
+	/// Class for building byte sequences, similar to StringBuilder. Allocates memory in chunks to avoid copying data.
+	/// </summary>
+	public class ByteArrayBuilder
+	{
+		class Chunk
+		{
+			public readonly int RunningIndex;
+			public readonly byte[] Data;
+			public int Length;
+
+			public Chunk(int runningIndex, int size)
+			{
+				RunningIndex = runningIndex;
+				Data = new byte[size];
+			}
+
+			public ReadOnlySpan<byte> WrittenSpan => Data.AsSpan(0, Length);
+			public ReadOnlyMemory<byte> WrittenMemory => Data.AsMemory(0, Length);
+		}
+
+		class Segment : ReadOnlySequenceSegment<byte>
+		{
+			public Segment(Chunk chunk, Segment? next)
+			{
+				Memory = chunk.WrittenMemory;
+				RunningIndex = chunk.RunningIndex;
+				Next = next;
+			}
+		}
+
+		readonly List<Chunk> _chunks = new List<Chunk>();
+		readonly int _chunkSize;
+		Chunk _currentChunk;
+
+		/// <summary>
+		/// Length of the current sequence
+		/// </summary>
+		public int Length { get; private set; }
+
+		/// <summary>
+		/// Constructor
+		/// </summary>
+		/// <param name="initialSize">Size of the initial chunk</param>
+		/// <param name="chunkSize">Default size for subsequent chunks</param>
+		public ByteArrayBuilder(int initialSize = 4096, int chunkSize = 4096)
+		{
+			_currentChunk = new Chunk(0, initialSize);
+			_chunks.Add(_currentChunk);
+			_chunkSize = chunkSize;
+		}
+
+		/// <summary>
+		/// Appends a single byte to the sequence
+		/// </summary>
+		/// <param name="value">Byte to add</param>
+		public void WriteByte(byte value)
+		{
+			Span<byte> target = GetWritableSpan(1);
+			target[0] = value;
+		}
+
+		/// <summary>
+		/// Appends a span of bytes to the buffer
+		/// </summary>
+		/// <param name="span">Bytes to add</param>
+		public void Append(ReadOnlySpan<byte> span)
+		{
+			Span<byte> target = GetWritableSpan(span.Length);
+			span.CopyTo(target);
+		}
+
+		/// <summary>
+		/// Appends a sequence of bytes to the buffer
+		/// </summary>
+		/// <param name="sequence">Sequence to append</param>
+		public void Append(ReadOnlySequence<byte> sequence)
+		{
+			Span<byte> target = GetWritableSpan((int)sequence.Length);
+			sequence.CopyTo(target);
+		}
+
+		/// <summary>
+		/// Gets a span of bytes to write to, and increases the length of the buffer
+		/// </summary>
+		/// <param name="length">Length of the buffer</param>
+		/// <returns>Writable span of bytes</returns>
+		public Span<byte> GetWritableSpan(int length)
+		{
+			if (_currentChunk.Length + length > _currentChunk.Data.Length)
+			{
+				_currentChunk = new Chunk(_currentChunk.RunningIndex + _currentChunk.Length, Math.Max(length, _chunkSize));
+				_chunks.Add(_currentChunk);
+			}
+
+			Span<byte> span = _currentChunk.Data.AsSpan(_currentChunk.Length, length);
+			_currentChunk.Length += length;
+			Length += length;
+			return span;
+		}
+
+		/// <summary>
+		/// Gets a sequence representing the bytes that have been written so far
+		/// </summary>
+		/// <returns>Sequence of bytes</returns>
+		public ReadOnlySequence<byte> AsSequence() => AsSequence(0);
+
+		/// <summary>
+		/// Gets a sequence representing the bytes that have been written so far, starting at the given offset
+		/// </summary>
+		/// <param name="offset">Offset to start from</param>
+		/// <returns>Sequence of bytes</returns>
+		public ReadOnlySequence<byte> AsSequence(int offset)
+		{
+			if (offset == Length)
+			{
+				return ReadOnlySequence<byte>.Empty;
+			}
+			else
+			{
+				Segment last = new Segment(_chunks[^1], null);
+
+				Segment first = last;
+				for (int idx = _chunks.Count - 1; idx >= 0 && offset < _chunks[idx].RunningIndex + _chunks[idx].Length; idx--)
+				{
+					first = new Segment(_chunks[idx], first);
+				}
+
+				return new ReadOnlySequence<byte>(first, offset - (int)first.RunningIndex, last, last.Memory.Length);
+			}
+		}
+
+		/// <summary>
+		/// Gets a sequence representing the bytes that have been written so far, starting at the given offset and 
+		/// </summary>
+		/// <param name="offset">Offset to start from</param>
+		/// <param name="length">Length of the sequence to return</param>
+		/// <returns>Sequence of bytes</returns>
+		public ReadOnlySequence<byte> AsSequence(int offset, int length)
+		{
+			// TODO: could do a binary search for offset + length and work back from there
+			return AsSequence(offset).Slice(0, length);
+		}
+
+		/// <summary>
+		/// Copies the data to the given output span
+		/// </summary>
+		/// <param name="span">Span to write to</param>
+		public void CopyTo(Span<byte> span)
+		{
+			foreach (Chunk chunk in _chunks)
+			{
+				Span<byte> output = span.Slice(chunk.RunningIndex);
+				chunk.WrittenSpan.CopyTo(output);
+			}
+		}
+
+		/// <summary>
+		/// Create a byte array from the sequence
+		/// </summary>
+		/// <returns>Byte array containing the current buffer data</returns>
+		public byte[] ToByteArray()
+		{
+			byte[] data = new byte[Length];
+			CopyTo(data);
+			return data;
+		}
+	}
+}
