Fix Xmldoc warnings in EpicGames.Core prior to enabling tests in CI.

diff --git a/Engine/Source/Programs/Shared/EpicGames.Core/AssemblyUtils.cs b/Engine/Source/Programs/Shared/EpicGames.Core/AssemblyUtils.cs
--- a/Engine/Source/Programs/Shared/EpicGames.Core/AssemblyUtils.cs
+++ b/Engine/Source/Programs/Shared/EpicGames.Core/AssemblyUtils.cs
@@ -9,6 +9,9 @@ using System.Reflection;
 
 namespace EpicGames.Core
 {
+	/// <summary>
+	/// Utilities for dealing with assembly loading
+	/// </summary>
 	public static class AssemblyUtils
 	{
 		/// <summary>
@@ -96,6 +99,10 @@ namespace EpicGames.Core
 			}
 		}
 
+		/// <summary>
+		/// Adds a file to the cache
+		/// </summary>
+		/// <param name="assemblyPath"></param>
 		public static void AddFileToAssemblyCache(string assemblyPath)
 		{
 			// Ignore any reference assemblies
diff --git a/Engine/Source/Programs/Shared/EpicGames.Core/BinaryArchive.cs b/Engine/Source/Programs/Shared/EpicGames.Core/BinaryArchive.cs
--- a/Engine/Source/Programs/Shared/EpicGames.Core/BinaryArchive.cs
+++ b/Engine/Source/Programs/Shared/EpicGames.Core/BinaryArchive.cs
@@ -191,6 +191,11 @@ namespace EpicGames.Core
 			return info.ReadMethod(reader)!;
 		}
 
+		/// <summary>
+		/// Reads an object from the archive
+		/// </summary>
+		/// <param name="reader">Reader to deserialize from</param>
+		/// <returns>The deserialized object</returns>
 		public static object? ReadObject(this BinaryArchiveReader reader)
 		{
 			return reader.ReadUntypedObjectReference(() => ReadNewPolymorphicObject(reader));
diff --git a/Engine/Source/Programs/Shared/EpicGames.Core/BinaryArchiveAttributes.cs b/Engine/Source/Programs/Shared/EpicGames.Core/BinaryArchiveAttributes.cs
--- a/Engine/Source/Programs/Shared/EpicGames.Core/BinaryArchiveAttributes.cs
+++ b/Engine/Source/Programs/Shared/EpicGames.Core/BinaryArchiveAttributes.cs
@@ -10,6 +10,9 @@ namespace EpicGames.Core
 	[AttributeUsage(AttributeTargets.Property)]
 	public class BinaryIgnoreAttribute : Attribute
 	{
+		/// <summary>
+		/// Constructor
+		/// </summary>
 		public BinaryIgnoreAttribute()
 		{
 		}
@@ -21,6 +24,9 @@ namespace EpicGames.Core
 	[AttributeUsage(AttributeTargets.Class)]
 	public class BinarySerializableAttribute : Attribute
 	{
+		/// <summary>
+		/// Constructor
+		/// </summary>
 		public BinarySerializableAttribute()
 		{
 		}
diff --git a/Engine/Source/Programs/Shared/EpicGames.Core/BinaryArchiveReader.cs b/Engine/Source/Programs/Shared/EpicGames.Core/BinaryArchiveReader.cs
--- a/Engine/Source/Programs/Shared/EpicGames.Core/BinaryArchiveReader.cs
+++ b/Engine/Source/Programs/Shared/EpicGames.Core/BinaryArchiveReader.cs
@@ -510,7 +510,6 @@ namespace EpicGames.Core
 		/// Reads an object reference from the stream. Each object will only be serialized once using the supplied delegate; subsequent reads reference the original.
 		/// Since the reader only receives the object reference when the CreateObject delegate returns, it is not possible for the object to serialize a reference to itself.
 		/// </summary>
-		/// <typeparam name="T">Type of the object to read.</typeparam>
 		/// <param name="readObject">Delegate used to create an object instance. The object may not reference itself recursively.</param>
 		/// <returns>Object instance</returns>
 		public object? ReadUntypedObjectReference(Func<object?> readObject)
@@ -540,7 +539,6 @@ namespace EpicGames.Core
 		/// Reads an object reference from the stream. Each object will only be serialized once using the supplied delegate; subsequent reads reference the original.
 		/// Since the reader only receives the object reference when the CreateObject delegate returns, it is not possible for the object to serialize a reference to itself.
 		/// </summary>
-		/// <typeparam name="T">Type of the object to read.</typeparam>
 		/// <param name="readObject">Delegate used to create an object instance. The object may not reference itself recursively.</param>
 		/// <returns>Object instance</returns>
 		public object? ReadUntypedObjectReference(Func<BinaryArchiveReader, object?> readObject)
diff --git a/Engine/Source/Programs/Shared/EpicGames.Core/BinaryArchiveWriter.cs b/Engine/Source/Programs/Shared/EpicGames.Core/BinaryArchiveWriter.cs
--- a/Engine/Source/Programs/Shared/EpicGames.Core/BinaryArchiveWriter.cs
+++ b/Engine/Source/Programs/Shared/EpicGames.Core/BinaryArchiveWriter.cs
@@ -281,6 +281,7 @@ namespace EpicGames.Core
         /// Writes an array of primitive types to the output.
         /// </summary>
         /// <param name="data">Data to write. May be null.</param>
+		/// <param name="elementSize">Size of each element</param>
         private void WritePrimitiveArray<T>(T[]? data, int elementSize) where T : struct
         {
             if (data == null)
@@ -325,6 +326,7 @@ namespace EpicGames.Core
         /// Writes an array of primitive types to the output.
         /// </summary>
         /// <param name="data">Data to write. May be null.</param>
+		/// <param name="elementSize">Size of each element</param>
         private void WriteFixedSizePrimitiveArray<T>(T[] data, int elementSize) where T : struct
         {
             WriteBulkData(data, data.Length * elementSize);
diff --git a/Engine/Source/Programs/Shared/EpicGames.Core/BinaryReaderExtensions.cs b/Engine/Source/Programs/Shared/EpicGames.Core/BinaryReaderExtensions.cs
--- a/Engine/Source/Programs/Shared/EpicGames.Core/BinaryReaderExtensions.cs
+++ b/Engine/Source/Programs/Shared/EpicGames.Core/BinaryReaderExtensions.cs
@@ -6,6 +6,9 @@ using System.IO;
 
 namespace EpicGames.Core
 {
+	/// <summary>
+	/// Extension methods for <see cref="BinaryReader"/>
+	/// </summary>
 	public static class BinaryReaderExtensions
 	{
 		/// <summary>
@@ -78,7 +81,6 @@ namespace EpicGames.Core
 		/// </summary>
 		/// <typeparam name="T">The element type for the list</typeparam>
 		/// <param name="reader">Reader to read data from</param>
-		/// <param name="ReadElement">Delegate to call to serialize each element</param>
 		/// <returns>List of objects, as serialized. May be null.</returns>
 		public static List<T>? ReadList<T>(this BinaryReader reader) where T : class, IBinarySerializable
 		{
@@ -111,12 +113,13 @@ namespace EpicGames.Core
 		/// <summary>
 		/// Read a list of objects from a binary reader
 		/// </summary>
-		/// <typeparam name="T">The element type for the list</typeparam>
+		/// <typeparam name="TKey">Dictionary key type</typeparam>
+		/// <typeparam name="TValue">Dictionary value type</typeparam>
 		/// <param name="reader">Reader to read data from</param>
 		/// <param name="readKey">Delegate to call to serialize each key</param>
 		/// <param name="readValue">Delegate to call to serialize each value</param>
 		/// <returns>List of objects, as serialized. May be null.</returns>
-		public static Dictionary<TK, TV>? ReadDictionary<TK, TV>(this BinaryReader reader, Func<TK> readKey, Func<TV> readValue) where TK : notnull
+		public static Dictionary<TKey, TValue>? ReadDictionary<TKey, TValue>(this BinaryReader reader, Func<TKey> readKey, Func<TValue> readValue) where TKey : notnull
 		{
 			int numItems = reader.ReadInt32();
 			if (numItems < 0)
@@ -124,11 +127,11 @@ namespace EpicGames.Core
 				return null;
 			}
 
-			Dictionary<TK, TV> items = new Dictionary<TK, TV>(numItems);
+			Dictionary<TKey, TValue> items = new Dictionary<TKey, TValue>(numItems);
 			for (int idx = 0; idx < numItems; idx++)
 			{
-				TK key = readKey();
-				TV value = readValue();
+				TKey key = readKey();
+				TValue value = readValue();
 				items.Add(key, value);
 			}
 			return items;
diff --git a/Engine/Source/Programs/Shared/EpicGames.Core/BinaryWriterExtensions.cs b/Engine/Source/Programs/Shared/EpicGames.Core/BinaryWriterExtensions.cs
--- a/Engine/Source/Programs/Shared/EpicGames.Core/BinaryWriterExtensions.cs
+++ b/Engine/Source/Programs/Shared/EpicGames.Core/BinaryWriterExtensions.cs
@@ -13,6 +13,10 @@ namespace EpicGames.Core
 	/// </summary>
 	public interface IBinarySerializable
 	{
+		/// <summary>
+		/// Write an item
+		/// </summary>
+		/// <param name="writer"></param>
 		void Write(BinaryWriter writer);
 	}
 
@@ -122,14 +126,14 @@ namespace EpicGames.Core
 		/// <summary>
 		/// Write a dictionary to a binary writer
 		/// </summary>
-		/// <typeparam name="TK">The key type for the dictionary</typeparam>
-		/// <typeparam name="TV">The value type for the dictionary</typeparam>
-		/// <param name="Reader">Reader to read data from</param>
+		/// <typeparam name="TKey">The key type for the dictionary</typeparam>
+		/// <typeparam name="TValue">The value type for the dictionary</typeparam>
+		/// <param name="writer">Writer to write data to</param>
 		/// <param name="items">List of items to be written</param>
 		/// <param name="writeKey">Delegate to call to serialize each key</param>
-		/// <param name="writeKey">Delegate to call to serialize each value</param>
+		/// <param name="writeValue">Delegate to call to serialize each value</param>
 		/// <returns>Dictionary of objects, as serialized. May be null.</returns>
-		public static void Write<TK, TV>(this BinaryWriter writer, Dictionary<TK, TV> items, Action<TK> writeKey, Action<TV> writeValue) where TK : notnull
+		public static void Write<TKey, TValue>(this BinaryWriter writer, Dictionary<TKey, TValue> items, Action<TKey> writeKey, Action<TValue> writeValue) where TKey : notnull
 		{
 			if (items == null)
 			{
@@ -138,7 +142,7 @@ namespace EpicGames.Core
 			else
 			{
 				writer.Write(items.Count);
-				foreach (KeyValuePair<TK, TV> item in items)
+				foreach (KeyValuePair<TKey, TValue> item in items)
 				{
 					writeKey(item.Key);
 					writeValue(item.Value);
@@ -151,6 +155,7 @@ namespace EpicGames.Core
 		/// </summary>
 		/// <typeparam name="T">Type of the object</typeparam>
 		/// <param name="writer">Reader to read data from</param>
+		/// <param name="item">Item to write</param>
 		/// <param name="writeItem">Function to read the payload, if non-null</param>
 		/// <returns>Object instance or null</returns>
 		public static void WriteNullable<T>(this BinaryWriter writer, T item, Action writeItem) where T : class
diff --git a/Engine/Source/Programs/Shared/EpicGames.Core/BuzHash.cs b/Engine/Source/Programs/Shared/EpicGames.Core/BuzHash.cs
--- a/Engine/Source/Programs/Shared/EpicGames.Core/BuzHash.cs
+++ b/Engine/Source/Programs/Shared/EpicGames.Core/BuzHash.cs
@@ -81,7 +81,7 @@ namespace EpicGames.Core
 		/// <summary>
 		/// Add a byte to the window and update the hash
 		/// </summary>
-		/// <param name="Value">New value to append to the window</param>
+		/// <param name="span">New data to append to the window</param>
 		public void Add(ReadOnlySpan<byte> span)
 		{
 			_state = Add(_state, span);
@@ -127,7 +127,9 @@ namespace EpicGames.Core
 		/// <summary>
 		/// Removes a byte from the start of the hash window
 		/// </summary>
+		/// <param name="state">Current state</param>
 		/// <param name="value">Value at the start of the window</param>
+		/// <param name="count">Size of the window</param>
 		public static uint Sub(uint state, byte value, int count)
 		{
 			return state ^ Rol32(s_table[value], (int)(count - 1));
diff --git a/Engine/Source/Programs/Shared/EpicGames.Core/CityHash.cs b/Engine/Source/Programs/Shared/EpicGames.Core/CityHash.cs
--- a/Engine/Source/Programs/Shared/EpicGames.Core/CityHash.cs
+++ b/Engine/Source/Programs/Shared/EpicGames.Core/CityHash.cs
@@ -36,6 +36,7 @@ namespace EpicGames.Core
 #pragma warning disable IDE0054
 #pragma warning disable IDE0180
 #pragma warning disable IDE1006
+#pragma warning disable CS1591 // Missing documentation
 
 	public class CityHash
 	{
diff --git a/Engine/Source/Programs/Shared/EpicGames.Core/CommandHost.cs b/Engine/Source/Programs/Shared/EpicGames.Core/CommandHost.cs
--- a/Engine/Source/Programs/Shared/EpicGames.Core/CommandHost.cs
+++ b/Engine/Source/Programs/Shared/EpicGames.Core/CommandHost.cs
@@ -19,6 +19,7 @@ namespace EpicGames.Core
 		/// Configure this object with the given command line arguments
 		/// </summary>
 		/// <param name="arguments">Command line arguments</param>
+		/// <param name="logger">Logger for output</param>
 		void Configure(CommandLineArguments arguments, ILogger logger);
 
 		/// <summary>
diff --git a/Engine/Source/Programs/Shared/EpicGames.Core/CommandLineArguments.cs b/Engine/Source/Programs/Shared/EpicGames.Core/CommandLineArguments.cs
--- a/Engine/Source/Programs/Shared/EpicGames.Core/CommandLineArguments.cs
+++ b/Engine/Source/Programs/Shared/EpicGames.Core/CommandLineArguments.cs
@@ -293,7 +293,7 @@ namespace EpicGames.Core
 		/// <summary>
 		/// Gets the index of the numbered positional argument
 		/// </summary>
-		/// <param name="Index">Number of the positional argument</param>
+		/// <param name="num">Number of the positional argument</param>
 		/// <returns>Index of the positional argument</returns>
 		public int GetPositionalArgumentIndex(int num)
 		{
@@ -996,7 +996,7 @@ namespace EpicGames.Core
 		/// <summary>
 		/// Appends the given arguments to the current argument list
 		/// </summary>
-		/// <param name="Arguments">The arguments to add</param>
+		/// <param name="appendArguments">The arguments to add</param>
 		/// <returns>New argument list</returns>
 		public CommandLineArguments Append(IEnumerable<string> appendArguments)
 		{
@@ -1343,6 +1343,7 @@ namespace EpicGames.Core
 		/// <summary>
 		/// Takes a command line argument and adds quotes if necessary
 		/// </summary>
+		/// <param name="commandLine"></param>
 		/// <param name="argument">The command line argument</param>
 		/// <returns>The command line argument with quotes inserted to escape it if necessary</returns>
 		public static void Append(StringBuilder commandLine, string argument)
diff --git a/Engine/Source/Programs/Shared/EpicGames.Core/ContentHash.cs b/Engine/Source/Programs/Shared/EpicGames.Core/ContentHash.cs
--- a/Engine/Source/Programs/Shared/EpicGames.Core/ContentHash.cs
+++ b/Engine/Source/Programs/Shared/EpicGames.Core/ContentHash.cs
@@ -17,7 +17,14 @@ namespace EpicGames.Core
 	[JsonConverter(typeof(ContentHashJsonConverter))]
 	public class ContentHash : IEquatable<ContentHash>
 	{
+		/// <summary>
+		/// Length of an MD5 hash
+		/// </summary>
 		public const int LengthMD5 = 16;
+
+		/// <summary>
+		/// Length of a SHA1 hash
+		/// </summary>
 		public const int LengthSHA1 = 20;
 
 		/// <summary>
@@ -232,6 +239,7 @@ namespace EpicGames.Core
 		/// Parse a hash from a string
 		/// </summary>
 		/// <param name="text">Text to parse</param>
+		/// <param name="hash"></param>
 		/// <returns>Value of the hash</returns>
 		public static bool TryParse(string text, [NotNullWhen(true)] out ContentHash? hash)
 		{
diff --git a/Engine/Source/Programs/Shared/EpicGames.Core/CsProjBuildRecord.cs b/Engine/Source/Programs/Shared/EpicGames.Core/CsProjBuildRecord.cs
--- a/Engine/Source/Programs/Shared/EpicGames.Core/CsProjBuildRecord.cs
+++ b/Engine/Source/Programs/Shared/EpicGames.Core/CsProjBuildRecord.cs
@@ -6,44 +6,85 @@ using System.Text;
 
 namespace EpicGames.Core
 {
-	// Acceleration structure:
-	// used to encapsulate a full set of dependencies for an msbuild project - explicit and globbed
-	// These files are written to Intermediate/ScriptModules
+	/// <summary>
+	/// Acceleration structure:
+	/// used to encapsulate a full set of dependencies for an msbuild project - explicit and globbed
+	/// These files are written to Intermediate/ScriptModules
+	/// </summary>
 	public class CsProjBuildRecord
 	{
-		// Version number making it possible to quickly invalidate written records.
+		/// <summary>
+		/// Version number making it possible to quickly invalidate written records.
+		/// </summary>
 		public static readonly int CurrentVersion = 4;
+
+		/// <summary>
+		/// Version number for this record
+		/// </summary>
 		public int Version { get; set; } // what value does this get if deserialized from a file with no value for this field? 
 
-		// Path to the .csproj project file, relative to the location of the build record .json file
+		/// <summary>
+		/// Path to the .csproj project file, relative to the location of the build record .json file
+		/// </summary>
 		public string? ProjectPath { get; set; }
 
-		// The time that the target assembly was built (read from the file after the build)
+		/// <summary>
+		/// The time that the target assembly was built (read from the file after the build)
+		/// </summary>
 		public DateTime TargetBuildTime { get; set; }
 
 
 		// all following paths are relative to the project directory, the directory containing ProjectPath 
 
-		// assembly (dll) location
+		/// <summary>
+		/// assembly (dll) location
+		/// </summary>
 		public string? TargetPath { get; set; }
 
-		// Paths of referenced projects
+		/// <summary>
+		/// Paths of referenced projects
+		/// </summary>
 		public HashSet<string> ProjectReferences { get; set; } = new HashSet<string>();
 
-		// file dependencies from non-glob sources
+		/// <summary>
+		/// file dependencies from non-glob sources
+		/// </summary>
 		public HashSet<string> Dependencies { get; set; } = new HashSet<string>();
 
-		// file dependencies from globs
+		/// <summary>
+		/// file dependencies from globs
+		/// </summary>
 		public HashSet<string> GlobbedDependencies { get; set; } = new HashSet<string>();
 
+		/// <summary>
+		/// A glob pattern
+		/// </summary>
 		public class Glob
 		{
+			/// <summary>
+			/// Type of item
+			/// </summary>
 			public string? ItemType { get; set; }
+
+			/// <summary>
+			/// Paths to include
+			/// </summary>
 			public List<string>? Include { get; set; }
+
+			/// <summary>
+			/// Paths to exclude
+			/// </summary>
 			public List<string>? Exclude { get; set; }
+
+			/// <summary>
+			/// 
+			/// </summary>
 			public List<string>? Remove { get; set; }
 		}
 
+		/// <summary>
+		/// List of globs
+		/// </summary>
 		public List<Glob> Globs { get; set; } = new List<Glob>();
 	}
 }
diff --git a/Engine/Source/Programs/Shared/EpicGames.Core/CsProjectInfo.cs b/Engine/Source/Programs/Shared/EpicGames.Core/CsProjectInfo.cs
--- a/Engine/Source/Programs/Shared/EpicGames.Core/CsProjectInfo.cs
+++ b/Engine/Source/Programs/Shared/EpicGames.Core/CsProjectInfo.cs
@@ -49,6 +49,7 @@ namespace EpicGames.Core
 		/// Constructor
 		/// </summary>
 		/// <param name="inProperties">Initial mapping of property names to values</param>
+		/// <param name="inProjectPath"></param>
 		CsProjectInfo(Dictionary<string, string> inProperties, FileReference inProjectPath)
 		{
 			ProjectPath = inProjectPath;
@@ -228,7 +229,7 @@ namespace EpicGames.Core
 		/// Finds all content which will be copied into the output directory for this project. This includes content from any project references as "copy local" recursively (though MSBuild only traverses a single reference for actual binaries, in such cases)
 		/// </summary>
 		/// <param name="outputDir">The output directory</param>
-		/// <param name="BuildProducts">Receives the set of build products</param>
+		/// <param name="outputFiles">Receives the set of build products</param>
 		/// <param name="projectFileToInfo">Map of project file to information, to resolve build products from referenced projects copied locally</param>
 		private void FindCopiedContent(DirectoryReference outputDir, HashSet<FileReference> outputFiles, Dictionary<FileReference, CsProjectInfo> projectFileToInfo)
 		{
@@ -392,7 +393,7 @@ namespace EpicGames.Core
 		/// Parses a 'PropertyGroup' element.
 		/// </summary>
 		/// <param name="parentElement">The parent 'PropertyGroup' element</param>
-		/// <param name="Properties">Dictionary mapping property names to values</param>
+		/// <param name="projectInfo">Dictionary mapping property names to values</param>
 		static void ParsePropertyGroup(XmlElement parentElement, CsProjectInfo projectInfo)
 		{
 			// We need to know the overridden output type and output path for the selected configuration.
@@ -685,7 +686,7 @@ namespace EpicGames.Core
 		/// Evaluate whether the optional MSBuild condition on an XML element evaluates to true. Currently only supports 'ABC' == 'DEF' style expressions, but can be expanded as needed.
 		/// </summary>
 		/// <param name="element">The XML element to check</param>
-		/// <param name="Properties">Dictionary mapping from property names to values.</param>
+		/// <param name="projectInfo">Dictionary mapping from property names to values.</param>
 		/// <returns></returns>
 		static bool EvaluateCondition(XmlElement element, CsProjectInfo projectInfo)
 		{
diff --git a/Engine/Source/Programs/Shared/EpicGames.Core/DirectoryReference.cs b/Engine/Source/Programs/Shared/EpicGames.Core/DirectoryReference.cs
--- a/Engine/Source/Programs/Shared/EpicGames.Core/DirectoryReference.cs
+++ b/Engine/Source/Programs/Shared/EpicGames.Core/DirectoryReference.cs
@@ -18,6 +18,9 @@
 		/// </summary>
 		public enum Sanitize
 		{
+			/// <summary>
+			/// Dummy value
+			/// </summary>
 			None
 		}
 
@@ -43,7 +46,8 @@
 		/// Constructor for creating a directory object directly from two strings.
 		/// </summary>
 		/// <param name="fullName">The full, sanitized path name</param>
-		public DirectoryReference(string fullName, Sanitize _)
+		/// <param name="sanitize">Sanitize argument. Ignored.</param>
+		public DirectoryReference(string fullName, Sanitize sanitize)
 			: base(fullName)
 		{
 		}
diff --git a/Engine/Source/Programs/Shared/EpicGames.Core/FilePattern.cs b/Engine/Source/Programs/Shared/EpicGames.Core/FilePattern.cs
--- a/Engine/Source/Programs/Shared/EpicGames.Core/FilePattern.cs
+++ b/Engine/Source/Programs/Shared/EpicGames.Core/FilePattern.cs
@@ -8,18 +8,28 @@ using System.Text.RegularExpressions;
 
 namespace EpicGames.Core
 {
+	/// <summary>
+	/// Invalid file pattern exception
+	/// </summary>
 	public class FilePatternException : Exception
 	{
+		/// <summary>
+		/// Constructor
+		/// </summary>
 		public FilePatternException(string message)
 			: base(message)
 		{
 		}
 
+		/// <summary>
+		/// Constructor
+		/// </summary>
 		public FilePatternException(string format, params object[] args)
 			: base(String.Format(format, args))
 		{
 		}
 
+		/// <inheritdoc/>
 		public override string ToString()
 		{
 			return Message;
@@ -277,10 +287,8 @@ namespace EpicGames.Core
 		/// Creates a file mapping between a set of source patterns and a target pattern. All patterns should have a matching order and number of wildcards.
 		/// </summary>
 		/// <param name="files">Files to use for the mapping</param>
-		/// <param name="SourcePatterns">List of source patterns</param>
+		/// <param name="sourcePattern">List of source patterns</param>
 		/// <param name="targetPattern">Matching output pattern</param>
-		/// <param name="Filter">Filter to apply to source files</param>
-		/// <param name="TargetFileToSourceFile">Dictionary to receive a mapping from target file to source file. An exception is thrown if multiple source files map to one target file, or a source file is also used as a target file.</param>
 		public static Dictionary<FileReference, FileReference> CreateMapping(HashSet<FileReference>? files, ref FilePattern sourcePattern, ref FilePattern targetPattern)
 		{
 			// If the source pattern ends in a directory separator, or a set of input files are specified and it doesn't contain wildcards, treat it as a full directory match
diff --git a/Engine/Source/Programs/Shared/EpicGames.Core/FileReference.cs b/Engine/Source/Programs/Shared/EpicGames.Core/FileReference.cs
--- a/Engine/Source/Programs/Shared/EpicGames.Core/FileReference.cs
+++ b/Engine/Source/Programs/Shared/EpicGames.Core/FileReference.cs
@@ -22,6 +22,9 @@ namespace EpicGames.Core
 		/// </summary>
 		public enum Sanitize
 		{
+			/// <summary>
+			/// Dummy value
+			/// </summary>
 			None
 		}
 
@@ -51,7 +54,8 @@ namespace EpicGames.Core
 		/// Default constructor.
 		/// </summary>
 		/// <param name="fullName">The full sanitized path</param>
-		public FileReference(string fullName, Sanitize _)
+		/// <param name="sanitize">Sanitize argument. Ignored.</param>
+		public FileReference(string fullName, Sanitize sanitize)
 			: base(fullName)
 		{
 		}
@@ -429,6 +433,7 @@ namespace EpicGames.Core
 		/// Reads the contents of a file
 		/// </summary>
 		/// <param name="location">Location of the file</param>
+		/// <param name="cancellationToken">Cancellation token for the operation</param>
 		/// <returns>Byte array containing the contents of the file</returns>
 		public static Task<byte[]> ReadAllBytesAsync(FileReference location, CancellationToken cancellationToken = default)
 		{
@@ -480,6 +485,7 @@ namespace EpicGames.Core
 		/// Reads the contents of a file
 		/// </summary>
 		/// <param name="location">Location of the file</param>
+		/// <param name="cancellationToken">Cancellation token for the operation</param>
 		/// <returns>Contents of the file as a single string</returns>
 		public static Task<string> ReadAllTextAsync(FileReference location, CancellationToken cancellationToken = default)
 		{
@@ -491,6 +497,7 @@ namespace EpicGames.Core
 		/// </summary>
 		/// <param name="location">Location of the file</param>
 		/// <param name="encoding">Encoding of the file</param>
+		/// <param name="cancellationToken">Cancellation token for the operation</param>
 		/// <returns>Contents of the file as a single string</returns>
 		public static Task<string> ReadAllTextAsync(FileReference location, Encoding encoding, CancellationToken cancellationToken = default)
 		{
@@ -522,6 +529,7 @@ namespace EpicGames.Core
 		/// Reads the contents of a file
 		/// </summary>
 		/// <param name="location">Location of the file</param>
+		/// <param name="cancellationToken">Cancellation token for the operation</param>
 		/// <returns>String array containing the contents of the file</returns>
 		public static Task<string[]> ReadAllLinesAsync(FileReference location, CancellationToken cancellationToken = default)
 		{
@@ -533,6 +541,7 @@ namespace EpicGames.Core
 		/// </summary>
 		/// <param name="location">Location of the file</param>
 		/// <param name="encoding">The encoding to use when parsing the file</param>
+		/// <param name="cancellationToken">Cancellation token for the operation</param>
 		/// <returns>String array containing the contents of the file</returns>
 		public static Task<string[]> ReadAllLinesAsync(FileReference location, Encoding encoding, CancellationToken cancellationToken = default)
 		{
@@ -583,7 +592,7 @@ namespace EpicGames.Core
 		/// Sets the time that the file was last accessed.
 		/// </summary>
 		/// <param name="location">Location of the file.</param>
-		/// <param name="LastWriteTime">Last access time, in UTC time.</param>
+		/// <param name="lastWriteTimeUtc">Last access time, in UTC time.</param>
 		public static void SetLastAccessTimeUtc(FileReference location, DateTime lastWriteTimeUtc)
 		{
 			File.SetLastWriteTimeUtc(location.FullName, lastWriteTimeUtc);
@@ -604,6 +613,7 @@ namespace EpicGames.Core
 		/// </summary>
 		/// <param name="location">Location of the file</param>
 		/// <param name="contents">Contents of the file</param>
+		/// <param name="cancellationToken">Cancellation token for the operation</param>
 		public static Task WriteAllBytesAsync(FileReference location, byte[] contents, CancellationToken cancellationToken = default)
 		{
 			return File.WriteAllBytesAsync(location.FullName, contents, cancellationToken);
@@ -696,6 +706,7 @@ namespace EpicGames.Core
 		/// </summary>
 		/// <param name="location">Location of the file</param>
 		/// <param name="contents">Contents of the file</param>
+		/// <param name="cancellationToken">Cancellation token for the operation</param>
 		public static Task WriteAllLinesAsync(FileReference location, IEnumerable<string> contents, CancellationToken cancellationToken = default)
 		{
 			return File.WriteAllLinesAsync(location.FullName, contents, cancellationToken);
@@ -707,6 +718,7 @@ namespace EpicGames.Core
 		/// <param name="location">Location of the file</param>
 		/// <param name="contents">Contents of the file</param>
 		/// <param name="encoding">The encoding to use when parsing the file</param>
+		/// <param name="cancellationToken">Cancellation token for the operation</param>
 		public static Task WriteAllLinesAsync(FileReference location, IEnumerable<string> contents, Encoding encoding, CancellationToken cancellationToken = default)
 		{
 			return File.WriteAllLinesAsync(location.FullName, contents, encoding, cancellationToken);
@@ -717,6 +729,7 @@ namespace EpicGames.Core
 		/// </summary>
 		/// <param name="location">Location of the file</param>
 		/// <param name="contents">Contents of the file</param>
+		/// <param name="cancellationToken">Cancellation token for the operation</param>
 		public static Task WriteAllLinesAsync(FileReference location, string[] contents, CancellationToken cancellationToken = default)
 		{
 			return File.WriteAllLinesAsync(location.FullName, contents, cancellationToken);
@@ -728,6 +741,7 @@ namespace EpicGames.Core
 		/// <param name="location">Location of the file</param>
 		/// <param name="contents">Contents of the file</param>
 		/// <param name="encoding">The encoding to use when parsing the file</param>
+		/// <param name="cancellationToken">Cancellation token for the operation</param>
 		public static Task WriteAllLinesAsync(FileReference location, string[] contents, Encoding encoding, CancellationToken cancellationToken = default)
 		{
 			return File.WriteAllLinesAsync(location.FullName, contents, encoding, cancellationToken);
diff --git a/Engine/Source/Programs/Shared/EpicGames.Core/FileSet.cs b/Engine/Source/Programs/Shared/EpicGames.Core/FileSet.cs
--- a/Engine/Source/Programs/Shared/EpicGames.Core/FileSet.cs
+++ b/Engine/Source/Programs/Shared/EpicGames.Core/FileSet.cs
@@ -48,7 +48,8 @@ namespace EpicGames.Core
 		/// <summary>
 		/// Creates a file tree from a given set of files
 		/// </summary>
-		/// <param name="Files"></param>
+		/// <param name="directory">Base directory</param>
+		/// <param name="file">File to add</param>
 		/// <returns>Tree containing the given files</returns>
 		public static FileSet FromFile(DirectoryReference directory, string file)
 		{
@@ -68,7 +69,8 @@ namespace EpicGames.Core
 		/// <summary>
 		/// Creates a file tree from a given set of files
 		/// </summary>
-		/// <param name="Files"></param>
+		/// <param name="directory">Base directory for the file</param>
+		/// <param name="file">File to include</param>
 		/// <returns>Tree containing the given files</returns>
 		public static FileSet FromFile(DirectoryReference directory, FileReference file)
 		{
@@ -78,7 +80,8 @@ namespace EpicGames.Core
 		/// <summary>
 		/// Creates a file tree from a given set of files
 		/// </summary>
-		/// <param name="files"></param>
+		/// <param name="directory">Base directory for the fileset</param>
+		/// <param name="files">Files to include</param>
 		/// <returns>Tree containing the given files</returns>
 		public static FileSet FromFiles(DirectoryReference directory, IEnumerable<FileReference> files)
 		{
@@ -148,7 +151,7 @@ namespace EpicGames.Core
 		/// <summary>
 		/// Create a tree containing the exception of files with another tree
 		/// </summary>
-		/// <param name="Filter">Files to exclude from the filter</param>
+		/// <param name="rules">Files to exclude from the filter</param>
 		/// <returns></returns>
 		public FileSet Except(params string[] rules)
 		{
@@ -424,7 +427,6 @@ namespace EpicGames.Core
 	/// <summary>
 	/// File tree which includes only those files which match any given filter
 	/// </summary>
-	/// <typeparam name="T">Class containing information about a file</typeparam>
 	class FileSetFromFilter : FileSet
 	{
 		readonly FileSet _inner;
diff --git a/Engine/Source/Programs/Shared/EpicGames.Core/FileSystemReference.cs b/Engine/Source/Programs/Shared/EpicGames.Core/FileSystemReference.cs
--- a/Engine/Source/Programs/Shared/EpicGames.Core/FileSystemReference.cs
+++ b/Engine/Source/Programs/Shared/EpicGames.Core/FileSystemReference.cs
@@ -38,11 +38,11 @@ namespace EpicGames.Core
 			FullName = fullName;
 		}
 
+		static readonly ThreadLocal<StringBuilder> s_combineStringsStringBuilder = new ThreadLocal<StringBuilder>(() => new StringBuilder(260));
+
 		/// <summary>
 		/// Create a full path by concatenating multiple strings
 		/// </summary>
-		/// <returns></returns>
-		static readonly ThreadLocal<StringBuilder> s_combineStringsStringBuilder = new ThreadLocal<StringBuilder>(() => new StringBuilder(260));
 		protected static string CombineStrings(DirectoryReference baseDirectory, params string[] fragments)
 		{
 			// Get the initial string to append to, and strip any root directory suffix from it
@@ -230,13 +230,14 @@ namespace EpicGames.Core
 				return names.Any(x => ContainsName(x, baseDir.FullName.Length));
 			}
 		}
-		
+
+		static readonly ThreadLocal<StringBuilder> s_makeRelativeToStringBuilder = new ThreadLocal<StringBuilder>(() => new StringBuilder(260));
+
 		/// <summary>
 		/// Creates a relative path from the given base directory
 		/// </summary>
-		/// <param name="Directory">The directory to create a relative path from</param>
+		/// <param name="directory">The directory to create a relative path from</param>
 		/// <returns>A relative path from the given directory</returns>
-		static readonly ThreadLocal<StringBuilder> s_makeRelativeToStringBuilder = new ThreadLocal<StringBuilder>(() => new StringBuilder(260));
 		public string MakeRelativeTo(DirectoryReference directory)
 		{
 			StringBuilder result = s_makeRelativeToStringBuilder.Value!.Clear();
@@ -244,6 +245,11 @@ namespace EpicGames.Core
 			return result.ToString();
 		}
 
+		/// <summary>
+		/// Appens a relative path to a string builder
+		/// </summary>
+		/// <param name="directory"></param>
+		/// <param name="result"></param>
 		public void WriteRelativeTo(DirectoryReference directory, ref StringBuilder result)
 		{
 			// Find how much of the path is common between the two paths. This length does not include a trailing directory separator character.
@@ -323,6 +329,10 @@ namespace EpicGames.Core
 			}
 		}
 
+		/// <summary>
+		/// Normalize the path to using forward slashes
+		/// </summary>
+		/// <returns></returns>
 		public string ToNormalizedPath()
 		{
 			return FullName.Replace("\\", "/");
diff --git a/Engine/Source/Programs/Shared/EpicGames.Core/FileUtils.cs b/Engine/Source/Programs/Shared/EpicGames.Core/FileUtils.cs
--- a/Engine/Source/Programs/Shared/EpicGames.Core/FileUtils.cs
+++ b/Engine/Source/Programs/Shared/EpicGames.Core/FileUtils.cs
@@ -298,8 +298,8 @@ namespace EpicGames.Core
 		/// <summary>
 		/// Moves a file from one location to another. Creates the destination directory, and removes read-only files in the target location if necessary.
 		/// </summary>
-		/// <param name="SourceFileName">Path to the source file</param>
-		/// <param name="TargetFileName">Path to the target file</param>
+		/// <param name="sourceLocation">Path to the source file</param>
+		/// <param name="targetLocation">Path to the target file</param>
 		public static void ForceMoveFile(FileReference sourceLocation, FileReference targetLocation)
 		{
 			// Try to move the file into place
@@ -499,6 +499,10 @@ namespace EpicGames.Core
 		[DllImport("kernel32.dll", CharSet = CharSet.Unicode, SetLastError = true)]
 		internal static extern int GetFileAttributesW(string lpFileName);
 
+		/// <summary>
+		/// Force-delete a file (Windows only)
+		/// </summary>
+		/// <param name="fileName"></param>
 		public static void ForceDeleteFileWin32(string fileName)
 		{
 			// Try to delete the file normally
diff --git a/Engine/Source/Programs/Shared/EpicGames.Core/HelpUtils.cs b/Engine/Source/Programs/Shared/EpicGames.Core/HelpUtils.cs
--- a/Engine/Source/Programs/Shared/EpicGames.Core/HelpUtils.cs
+++ b/Engine/Source/Programs/Shared/EpicGames.Core/HelpUtils.cs
@@ -22,6 +22,7 @@ namespace EpicGames.Core
 		/// <summary>
 		/// Prints help for the given object type
 		/// </summary>
+		/// <param name="title"></param>
 		/// <param name="type">Type to print help for</param>
 		public static void PrintHelp(string title, Type type)
 		{
@@ -88,7 +89,6 @@ namespace EpicGames.Core
 		/// Prints a paragraph of text using word wrapping
 		/// </summary>
 		/// <param name="text">Text to print</param>
-		/// <param name="Logger">Logger implementation to write to</param>
 		public static void PrintParagraph(string text)
 		{
 			PrintParagraph(text, WindowWidth - 1);
@@ -157,6 +157,8 @@ namespace EpicGames.Core
 		/// <param name="items">List of parameters arranged as "-ParamName Param Description"</param>
 		/// <param name="indent">Indent from the left hand side</param>
 		/// <param name="minFirstColumnWidth">The minimum padding from the start of the param name to the start of the description (resizes with larger param names)</param>
+		/// <param name="maxWidth"></param>
+		/// <param name="lines"></param>
 		/// <returns>Sequence of formatted lines in the table</returns>
 		public static void FormatTable(IReadOnlyList<KeyValuePair<string, string>> items, int indent, int minFirstColumnWidth, int maxWidth, List<string> lines)
 		{
diff --git a/Engine/Source/Programs/Shared/EpicGames.Core/JsonLogEvent.cs b/Engine/Source/Programs/Shared/EpicGames.Core/JsonLogEvent.cs
--- a/Engine/Source/Programs/Shared/EpicGames.Core/JsonLogEvent.cs
+++ b/Engine/Source/Programs/Shared/EpicGames.Core/JsonLogEvent.cs
@@ -229,9 +229,9 @@ namespace EpicGames.Core
 		/// Formats an event as a string
 		/// </summary>
 		/// <param name="state"></param>
+		/// <param name="ex"></param>
 		/// <returns></returns>
-		/// <exception cref="NotImplementedException"></exception>
-		public static string Format(JsonLogEvent state, Exception? _) => LogEvent.Read(state.Data.Span).ToString();
+		public static string Format(JsonLogEvent state, Exception? ex) => LogEvent.Read(state.Data.Span).ToString();
 
 		/// <inheritdoc/>
 		public override string ToString() => Encoding.UTF8.GetString(Data.ToArray());
diff --git a/Engine/Source/Programs/Shared/EpicGames.Core/JsonSchema.cs b/Engine/Source/Programs/Shared/EpicGames.Core/JsonSchema.cs
--- a/Engine/Source/Programs/Shared/EpicGames.Core/JsonSchema.cs
+++ b/Engine/Source/Programs/Shared/EpicGames.Core/JsonSchema.cs
@@ -32,7 +32,7 @@ namespace EpicGames.Core
 		/// <inheritdoc cref="Utf8JsonWriter.WriteEndObject()"/>
 		void WriteEndArray();
 
-		/// <inheritdoc cref="Utf8JsonWriter.WriteBoolean(String, Bool)"/>
+		/// <inheritdoc cref="Utf8JsonWriter.WriteBoolean(String, Boolean)"/>
 		void WriteBoolean(string name, bool value);
 
 		/// <inheritdoc cref="Utf8JsonWriter.WriteString(String, String)"/>
@@ -316,6 +316,7 @@ namespace EpicGames.Core
 		/// Constructs a schema type from the given type object
 		/// </summary>
 		/// <param name="type"></param>
+		/// <param name="typeCache"></param>
 		/// <param name="xmlDoc"></param>
 		/// <returns></returns>
 		static JsonSchemaType CreateSchemaType(Type type, Dictionary<Type, JsonSchemaType> typeCache, XmlDocument? xmlDoc)
diff --git a/Engine/Source/Programs/Shared/EpicGames.Core/JsonSchemaType.cs b/Engine/Source/Programs/Shared/EpicGames.Core/JsonSchemaType.cs
--- a/Engine/Source/Programs/Shared/EpicGames.Core/JsonSchemaType.cs
+++ b/Engine/Source/Programs/Shared/EpicGames.Core/JsonSchemaType.cs
@@ -120,6 +120,7 @@ namespace EpicGames.Core
 		/// Constructor
 		/// </summary>
 		/// <param name="format"></param>
+		/// <param name="pattern"></param>
 		public JsonSchemaString(JsonSchemaStringFormat? format = null, string? pattern = null)
 		{
 			Format = format;
diff --git a/Engine/Source/Programs/Shared/EpicGames.Core/JsonTracer.cs b/Engine/Source/Programs/Shared/EpicGames.Core/JsonTracer.cs
--- a/Engine/Source/Programs/Shared/EpicGames.Core/JsonTracer.cs
+++ b/Engine/Source/Programs/Shared/EpicGames.Core/JsonTracer.cs
@@ -13,6 +13,8 @@ using OpenTracing.Propagation;
 using OpenTracing.Tag;
 using OpenTracing.Util;
 
+#pragma warning disable CS1591
+
 namespace EpicGames.Core
 {
 	public class JsonTracerSpanContext : ISpanContext
diff --git a/Engine/Source/Programs/Shared/EpicGames.Core/JsonWriter.cs b/Engine/Source/Programs/Shared/EpicGames.Core/JsonWriter.cs
--- a/Engine/Source/Programs/Shared/EpicGames.Core/JsonWriter.cs
+++ b/Engine/Source/Programs/Shared/EpicGames.Core/JsonWriter.cs
@@ -146,7 +146,6 @@ namespace EpicGames.Core
 		/// <summary>
 		/// Write the opening bracket for an unnamed array
 		/// </summary>
-		/// <param name="ArrayName">Name of the field</param>
 		public void WriteArrayStart()
 		{
 			WriteCommaNewline();
diff --git a/Engine/Source/Programs/Shared/EpicGames.Core/Log.cs b/Engine/Source/Programs/Shared/EpicGames.Core/Log.cs
--- a/Engine/Source/Programs/Shared/EpicGames.Core/Log.cs
+++ b/Engine/Source/Programs/Shared/EpicGames.Core/Log.cs
@@ -207,6 +207,7 @@ namespace EpicGames.Core
 		/// If a StartupTraceListener was in use, this function will copy its captured data to the log file(s)
 		/// and remove the startup listener from the list of registered listeners.
 		/// </summary>
+		/// <param name="name">Identifier for the writer</param>
 		/// <param name="outputFile">The file to write to</param>
 		/// <returns>The created trace listener</returns>
 		public static void AddFileWriter(string name, FileReference outputFile)
@@ -264,6 +265,7 @@ namespace EpicGames.Core
 		/// If a StartupTraceListener was in use, this function will copy its captured data to the log file(s)
 		/// and remove the startup listener from the list of registered listeners.
 		/// </summary>
+		/// <param name="name">Identifier for the writer</param>
 		/// <param name="outputFile">The file to write to</param>
 		/// <returns>The created trace listener</returns>
 		public static void AddFileWriterWithoutBackup(string name, FileReference outputFile)
@@ -754,7 +756,6 @@ namespace EpicGames.Core
 		/// <summary>
 		/// Updates the Pops the top status message from the stack. The mess
 		/// </summary>
-		/// <param name="Message"></param>
 		[Conditional("TRACE")]
 		public static void PopStatus()
 		{
@@ -1309,7 +1310,6 @@ namespace EpicGames.Core
 		/// <summary>
 		/// Updates the Pops the top status message from the stack. The mess
 		/// </summary>
-		/// <param name="Message"></param>
 		[Conditional("TRACE")]
 		public void PopStatus()
 		{
diff --git a/Engine/Source/Programs/Shared/EpicGames.Core/LogEvent.cs b/Engine/Source/Programs/Shared/EpicGames.Core/LogEvent.cs
--- a/Engine/Source/Programs/Shared/EpicGames.Core/LogEvent.cs
+++ b/Engine/Source/Programs/Shared/EpicGames.Core/LogEvent.cs
@@ -19,6 +19,7 @@ namespace EpicGames.Core
 	/// </summary>
 	public static class LogEventPropertyName
 	{
+#pragma warning disable CS1591 // Missing documentation
 		public static readonly Utf8String Time = new Utf8String("time");
 		public static readonly Utf8String Level = new Utf8String("level");
 		public static readonly Utf8String Id = new Utf8String("id");
@@ -35,6 +36,7 @@ namespace EpicGames.Core
 		public static readonly Utf8String Trace = new Utf8String("trace");
 		public static readonly Utf8String InnerException = new Utf8String("innerException");
 		public static readonly Utf8String InnerExceptions = new Utf8String("innerExceptions");
+#pragma warning restore CS1591 // Missing documentation
 	}
 
 	/// <summary>
@@ -407,12 +409,21 @@ namespace EpicGames.Core
 			}
 		}
 
+		/// <summary>
+		/// Create a new log event
+		/// </summary>
 		public static LogEvent Create(LogLevel level, string format, params object[] args)
 			=> Create(level, KnownLogEvents.None, null, format, args);
 
+		/// <summary>
+		/// Create a new log event
+		/// </summary>
 		public static LogEvent Create(LogLevel level, EventId eventId, string format, params object[] args)
 			=> Create(level, eventId, null, format, args);
 
+		/// <summary>
+		/// Create a new log event
+		/// </summary>
 		public static LogEvent Create(LogLevel level, EventId eventId, Exception? exception, string format, params object[] args)
 		{
 			Dictionary<string, object> properties = new Dictionary<string, object>();
diff --git a/Engine/Source/Programs/Shared/EpicGames.Core/LogEventBuilder.cs b/Engine/Source/Programs/Shared/EpicGames.Core/LogEventBuilder.cs
--- a/Engine/Source/Programs/Shared/EpicGames.Core/LogEventBuilder.cs
+++ b/Engine/Source/Programs/Shared/EpicGames.Core/LogEventBuilder.cs
@@ -227,6 +227,7 @@ namespace EpicGames.Core
 		/// Starts building a log event at the current cursor position
 		/// </summary>
 		/// <param name="cursor">The current cursor position</param>
+		/// <param name="lineCount">Number of lines to consume</param>
 		public LogEventBuilder(ILogCursor cursor, int lineCount = 1)
 		{
 			Current = cursor;
@@ -273,7 +274,7 @@ namespace EpicGames.Core
 		/// <param name="name">Name to use to identify the item in the format string</param>
 		/// <param name="offset">Offset within the line</param>
 		/// <param name="length">Length of the span</param>
-		/// <param name="Data">Data of the span</param>
+		/// <param name="value">Data of the span</param>
 		/// <returns>New span for the given range</returns>
 		public void Annotate(string name, int offset, int length, object? value = null)
 		{
@@ -298,6 +299,7 @@ namespace EpicGames.Core
 		/// </summary>
 		/// <param name="group">The match group</param>
 		/// <param name="name">Name to use to identify the item in the format string</param>
+		/// <param name="value">Optional value for the annotation</param>
 		public void Annotate(string name, Group group, object? value = null)
 		{
 			Annotate(name, group.Index, group.Length, value);
@@ -307,6 +309,7 @@ namespace EpicGames.Core
 		/// Adds a span naming a regex match group, using the name of the group
 		/// </summary>
 		/// <param name="group">The match group</param>
+		/// <param name="value">Optional value for the annotation</param>
 		public void Annotate(Group group, object? value = null)
 		{
 			Annotate(group.Name, group.Index, group.Length, value);
@@ -317,6 +320,7 @@ namespace EpicGames.Core
 		/// </summary>
 		/// <param name="group">The match group</param>
 		/// <param name="name">Name to use to identify the item in the format string</param>
+		/// <param name="value">Optional value for the annotation</param>
 		public bool TryAnnotate(string name, Group group, object? value = null)
 		{
 			if (group.Success)
@@ -331,6 +335,7 @@ namespace EpicGames.Core
 		/// Adds a span naming a regex match group, using the name of the group
 		/// </summary>
 		/// <param name="group">The match group</param>
+		/// <param name="value">Optional value for the annotation</param>
 		public bool TryAnnotate(Group group, object? value = null)
 		{
 			if (group.Success)
diff --git a/Engine/Source/Programs/Shared/EpicGames.Core/LogEventMarkup.cs b/Engine/Source/Programs/Shared/EpicGames.Core/LogEventMarkup.cs
--- a/Engine/Source/Programs/Shared/EpicGames.Core/LogEventMarkup.cs
+++ b/Engine/Source/Programs/Shared/EpicGames.Core/LogEventMarkup.cs
@@ -12,6 +12,7 @@ namespace EpicGames.Core
 	/// </summary>
 	public static class LogEventMarkup
 	{
+#pragma warning disable CS1591
 		public static LogValue Channel => new LogValue(LogValueType.Channel, "");
 		public static LogValue Severity => new LogValue(LogValueType.Severity, "");
 		public static LogValue Message => new LogValue(LogValueType.Message, "");
@@ -21,6 +22,7 @@ namespace EpicGames.Core
 		public static LogValue ErrorCode => new LogValue(LogValueType.ErrorCode, "");
 		public static LogValue ToolName => new LogValue(LogValueType.ToolName, "");
 		public static LogValue ScreenshotTest => new LogValue(LogValueType.ScreenshotTest, "");
+#pragma warning restore CS1591
 
 		static readonly char[] s_pathChars = { '/', '\\' };
 
diff --git a/Engine/Source/Programs/Shared/EpicGames.Core/LogEventParser.cs b/Engine/Source/Programs/Shared/EpicGames.Core/LogEventParser.cs
--- a/Engine/Source/Programs/Shared/EpicGames.Core/LogEventParser.cs
+++ b/Engine/Source/Programs/Shared/EpicGames.Core/LogEventParser.cs
@@ -18,27 +18,74 @@ namespace EpicGames.Core
 	/// </summary>
 	public enum LogEventPriority
 	{
+		/// <summary>
+		/// Unspecified priority
+		/// </summary>
 		None,
+
+		/// <summary>
+		/// Lowest confidence match
+		/// </summary>
 		Lowest,
+
+		/// <summary>
+		/// Low confidence match
+		/// </summary>
 		Low,
+
+		/// <summary>
+		/// Below normal confidence match
+		/// </summary>
 		BelowNormal,
+
+		/// <summary>
+		/// Normal confidence match
+		/// </summary>
 		Normal,
+
+		/// <summary>
+		/// Above normal confidence match
+		/// </summary>
 		AboveNormal,
+
+		/// <summary>
+		/// High confidence match
+		/// </summary>
 		High,
+
+		/// <summary>
+		/// Highest confidence match
+		/// </summary>
 		Highest,
 	}
 
+	/// <summary>
+	/// Information about a matched event
+	/// </summary>
 	public class LogEventMatch
 	{
+		/// <summary>
+		/// Confidence of the match
+		/// </summary>
 		public LogEventPriority Priority { get; }
+
+		/// <summary>
+		/// Matched events
+		/// </summary>
 		public List<LogEvent> Events { get; }
 
+		/// <summary>
+		/// Constructor
+		/// </summary>
 		public LogEventMatch(LogEventPriority priority, LogEvent logEvent)
 		{
 			Priority = priority;
 			Events = new List<LogEvent> { logEvent };
 		}
 
+		/// <summary>
+		/// Constructor
+		/// </summary>
 		public LogEventMatch(LogEventPriority priority, IEnumerable<LogEvent> events)
 		{
 			Priority = priority;
@@ -456,7 +503,7 @@ namespace EpicGames.Core
 		/// <summary>
 		/// Writes an event to the log
 		/// </summary>
-		/// <param name="Event">The event to write</param>
+		/// <param name="logEvents">The event to write</param>
 		protected virtual void WriteEvents(List<LogEvent> logEvents)
 		{
 			foreach (LogEvent logEvent in logEvents)
diff --git a/Engine/Source/Programs/Shared/EpicGames.Core/LogValue.cs b/Engine/Source/Programs/Shared/EpicGames.Core/LogValue.cs
--- a/Engine/Source/Programs/Shared/EpicGames.Core/LogValue.cs
+++ b/Engine/Source/Programs/Shared/EpicGames.Core/LogValue.cs
@@ -12,6 +12,7 @@ namespace EpicGames.Core
 	/// </summary>
 	public static class LogValueType
 	{
+#pragma warning disable CS1591
 		public static readonly Utf8String Asset = "Asset";
 		public static readonly Utf8String SourceFile = "SourceFile";
 		public static readonly Utf8String Object = "Object"; // Arbitrary structured object
@@ -25,6 +26,7 @@ namespace EpicGames.Core
 		public static readonly Utf8String ToolName = "ToolName";
 		public static readonly Utf8String ScreenshotTest = "ScreenshotTest";
 		public static readonly Utf8String DepotPath = "DepotPath";
+#pragma warning restore CS1591
 	}
 
 	/// <summary>
diff --git a/Engine/Source/Programs/Shared/EpicGames.Core/LogValueFormatter.cs b/Engine/Source/Programs/Shared/EpicGames.Core/LogValueFormatter.cs
--- a/Engine/Source/Programs/Shared/EpicGames.Core/LogValueFormatter.cs
+++ b/Engine/Source/Programs/Shared/EpicGames.Core/LogValueFormatter.cs
@@ -11,7 +11,6 @@ namespace EpicGames.Core
 	/// <summary>
 	/// Formatter for objects of a given type
 	/// </summary>
-	/// <typeparam name="T"></typeparam>
 	public interface ILogValueFormatter
 	{
 		/// <summary>
diff --git a/Engine/Source/Programs/Shared/EpicGames.Core/ManagedProcess.cs b/Engine/Source/Programs/Shared/EpicGames.Core/ManagedProcess.cs
--- a/Engine/Source/Programs/Shared/EpicGames.Core/ManagedProcess.cs
+++ b/Engine/Source/Programs/Shared/EpicGames.Core/ManagedProcess.cs
@@ -413,7 +413,7 @@ namespace EpicGames.Core
 		Task? _frameworkOutputTask;
 
 		/// <summary>
-		/// Merged output & error write stream for the framework child process. We use a channel rather than a pipe here to support cancellation, which is not supported by async pipes on Windows.
+		/// Merged output and error write stream for the framework child process. We use a channel rather than a pipe here to support cancellation, which is not supported by async pipes on Windows.
 		/// </summary>
 		Channel<ChannelBuffer>? _frameworkChannel;
 
@@ -518,6 +518,7 @@ namespace EpicGames.Core
 		/// <param name="environment">Environment variables for the new process. May be null, in which case the current process' environment is inherited</param>
 		/// <param name="input">Text to be passed via stdin to the new process. May be null.</param>
 		/// <param name="priority">Priority for the child process</param>
+		/// <param name="flags">Flags for the process</param>
 		public ManagedProcess(ManagedProcessGroup? group, string fileName, string commandLine, string? workingDirectory, IReadOnlyDictionary<string, string>? environment, byte[]? input, ProcessPriorityClass priority, ManagedProcessFlags flags = ManagedProcessFlags.MergeOutputPipes)
 			: this(group, fileName, commandLine, workingDirectory, environment, priority, flags)
 		{
@@ -537,6 +538,7 @@ namespace EpicGames.Core
 		/// <param name="workingDirectory">Working directory for the new process. May be null to use the current working directory.</param>
 		/// <param name="environment">Environment variables for the new process. May be null, in which case the current process' environment is inherited</param>
 		/// <param name="priority">Priority for the child process</param>
+		/// <param name="flags">Flags for the process</param>
 		public ManagedProcess(ManagedProcessGroup? group, string fileName, string commandLine, string? workingDirectory, IReadOnlyDictionary<string, string>? environment, ProcessPriorityClass priority, ManagedProcessFlags flags = ManagedProcessFlags.MergeOutputPipes)
 		{
 			// Create the child process
@@ -902,6 +904,7 @@ namespace EpicGames.Core
 		/// </summary>
 		/// <param name="source"></param>
 		/// <param name="target"></param>
+		/// <param name="cancellationToken">Cancellation token</param>
 		async Task CopyPipeAsync(Stream source, ChannelWriter<ChannelBuffer> target, CancellationToken cancellationToken)
 		{
 			for (; ; )
@@ -1001,6 +1004,7 @@ namespace EpicGames.Core
 		/// <param name="buffer">The buffer to receive the data</param>
 		/// <param name="offset">Offset within the buffer to write to</param>
 		/// <param name="count">Maximum number of bytes to read</param>
+		/// <param name="cancellationToken">Cancellation token</param>
 		/// <returns>Number of bytes read</returns>
 		public async ValueTask<int> ReadAsync(byte[] buffer, int offset, int count, CancellationToken cancellationToken)
 		{
@@ -1022,6 +1026,7 @@ namespace EpicGames.Core
 		/// Copy the process output to the given stream
 		/// </summary>
 		/// <param name="writeOutput">The output stream</param>
+		/// <param name="bufferSize">Size of the buffer for copying</param>
 		/// <param name="cancellationToken">Cancellation token</param>
 		/// <returns></returns>
 		public Task CopyToAsync(Action<byte[], int, int> writeOutput, int bufferSize, CancellationToken cancellationToken)
@@ -1038,6 +1043,7 @@ namespace EpicGames.Core
 		/// Copy the process output to the given stream
 		/// </summary>
 		/// <param name="writeOutputAsync">The output stream</param>
+		/// <param name="bufferSize">Size of the buffer for copying</param>
 		/// <param name="cancellationToken">Cancellation token</param>
 		/// <returns></returns>
 		public async Task CopyToAsync(Func<byte[], int, int, CancellationToken, Task> writeOutputAsync, int bufferSize, CancellationToken cancellationToken)
@@ -1189,6 +1195,9 @@ namespace EpicGames.Core
 			return await StdOutText!.ReadLineAsync();
 		}
 
+		/// <summary>
+		/// Total amount of time spent by the CPU
+		/// </summary>
 		public TimeSpan TotalProcessorTime
 		{
 			get
diff --git a/Engine/Source/Programs/Shared/EpicGames.Core/MessageTemplate.cs b/Engine/Source/Programs/Shared/EpicGames.Core/MessageTemplate.cs
--- a/Engine/Source/Programs/Shared/EpicGames.Core/MessageTemplate.cs
+++ b/Engine/Source/Programs/Shared/EpicGames.Core/MessageTemplate.cs
@@ -165,6 +165,7 @@ namespace EpicGames.Core
 		/// </summary>
 		/// <param name="format">Format string</param>
 		/// <param name="args">Argument list to parse</param>
+		/// <param name="properties"></param>
 		/// <returns></returns>
 		public static void ParsePropertyValues(string format, object[] args, Dictionary<string, object> properties)
 		{
diff --git a/Engine/Source/Programs/Shared/EpicGames.Core/ReadOnlyHashSet.cs b/Engine/Source/Programs/Shared/EpicGames.Core/ReadOnlyHashSet.cs
--- a/Engine/Source/Programs/Shared/EpicGames.Core/ReadOnlyHashSet.cs
+++ b/Engine/Source/Programs/Shared/EpicGames.Core/ReadOnlyHashSet.cs
@@ -28,7 +28,7 @@ namespace EpicGames.Core
 		/// <summary>
 		/// Constructor
 		/// </summary>
-		/// <param name="Collection">Elements for the hash set</param>
+		/// <param name="elements">Elements for the hash set</param>
 		public ReadOnlyHashSet(IEnumerable<T> elements)
 		{
 			_inner = new HashSet<T>(elements);
@@ -37,7 +37,7 @@ namespace EpicGames.Core
 		/// <summary>
 		/// Constructor
 		/// </summary>
-		/// <param name="Collection">Elements for the hash set</param>
+		/// <param name="elements">Elements for the hash set</param>
 		/// <param name="comparer">Comparer for elements in the set</param>
 		public ReadOnlyHashSet(IEnumerable<T> elements, IEqualityComparer<T> comparer)
 		{
diff --git a/Engine/Source/Programs/Shared/EpicGames.Core/RuntimePlatform.cs b/Engine/Source/Programs/Shared/EpicGames.Core/RuntimePlatform.cs
--- a/Engine/Source/Programs/Shared/EpicGames.Core/RuntimePlatform.cs
+++ b/Engine/Source/Programs/Shared/EpicGames.Core/RuntimePlatform.cs
@@ -4,6 +4,9 @@ using System.Runtime.InteropServices;
 
 namespace EpicGames.Core
 {
+	/// <summary>
+	/// Allows testing different platforsm
+	/// </summary>
 	public class RuntimePlatform
 	{
 		/// <summary>
@@ -21,9 +24,25 @@ namespace EpicGames.Core
 		/// </summary>
 		public static readonly bool IsWindows = RuntimeInformation.IsOSPlatform(OSPlatform.Windows);
 
+		/// <summary>
+		/// The platform type
+		/// </summary>
 		public enum Type
 		{
-			Windows, Linux, Mac
+			/// <summary>
+			/// Windows
+			/// </summary>
+			Windows, 
+			
+			/// <summary>
+			/// Linux
+			/// </summary>
+			Linux, 
+			
+			/// <summary>
+			/// Mac
+			/// </summary>
+			Mac
 		};
 
 		/// <summary>
diff --git a/Engine/Source/Programs/Shared/EpicGames.Core/ScopeProfiler.cs b/Engine/Source/Programs/Shared/EpicGames.Core/ScopeProfiler.cs
--- a/Engine/Source/Programs/Shared/EpicGames.Core/ScopeProfiler.cs
+++ b/Engine/Source/Programs/Shared/EpicGames.Core/ScopeProfiler.cs
@@ -11,10 +11,22 @@ namespace EpicGames.Core
 	/// </summary>
 	public class ScopeProfiler 
 	{
+		/// <summary>
+		/// Static instance of the profiler
+		/// </summary>
 		public static ScopeProfiler Instance = new ScopeProfiler();
 		
+		/// <summary>
+		/// 
+		/// </summary>
+		/// <param name="programName"></param>
+		/// <param name="hostAddress"></param>
+		/// <param name="maxThreadCount"></param>
 		public void InitializeAndStart(string programName, string hostAddress, int maxThreadCount) { }
 
+		/// <summary>
+		/// 
+		/// </summary>
 		public void StopAndShutdown() { }
 
 		/// <summary>
diff --git a/Engine/Source/Programs/Shared/EpicGames.Core/Sha1Hash.cs b/Engine/Source/Programs/Shared/EpicGames.Core/Sha1Hash.cs
--- a/Engine/Source/Programs/Shared/EpicGames.Core/Sha1Hash.cs
+++ b/Engine/Source/Programs/Shared/EpicGames.Core/Sha1Hash.cs
@@ -38,7 +38,7 @@ namespace EpicGames.Core
 		/// <summary>
 		/// Constructor
 		/// </summary>
-		/// <param name="Memory">Memory to construct from</param>
+		/// <param name="span">Memory to construct from</param>
 		public Sha1Hash(ReadOnlySpan<byte> span)
 			: this(BinaryPrimitives.ReadUInt64BigEndian(span), BinaryPrimitives.ReadUInt64BigEndian(span.Slice(8)), BinaryPrimitives.ReadUInt32BigEndian(span.Slice(16)))
 		{
@@ -47,7 +47,9 @@ namespace EpicGames.Core
 		/// <summary>
 		/// Constructor
 		/// </summary>
-		/// <param name="Memory">Memory to construct from</param>
+		/// <param name="a">First 64 bits of the hash</param>
+		/// <param name="b">Next 64 bits of the hash</param>
+		/// <param name="c">Final 32 bits of the hash</param>
 		public Sha1Hash(ulong a, ulong b, uint c)
 		{
 			_a = a;
@@ -126,6 +128,10 @@ namespace EpicGames.Core
 		/// <inheritdoc/>
 		public override string ToString() => StringUtils.FormatHexString(ToByteArray());
 
+		/// <summary>
+		/// Convert the hash to a byte array
+		/// </summary>
+		/// <returns></returns>
 		public byte[] ToByteArray()
 		{
 			byte[] data = new byte[NumBytes];
diff --git a/Engine/Source/Programs/Shared/EpicGames.Core/StringUtils.cs b/Engine/Source/Programs/Shared/EpicGames.Core/StringUtils.cs
--- a/Engine/Source/Programs/Shared/EpicGames.Core/StringUtils.cs
+++ b/Engine/Source/Programs/Shared/EpicGames.Core/StringUtils.cs
@@ -8,6 +8,9 @@ using System.Text;
 
 namespace EpicGames.Core
 {
+	/// <summary>
+	/// Utility methods for strings
+	/// </summary>
 	public static class StringUtils
 	{
 		/// <summary>
@@ -120,7 +123,7 @@ namespace EpicGames.Core
 			return result;
 		}
 
-		/// <inheritdoc cref="WordWrap(String, Int32, Int32, Int32)">
+		/// <inheritdoc cref="WordWrap(String, Int32, Int32, Int32)"/>
 		public static IEnumerable<string> WordWrap(string text, int maxWidth)
 		{
 			return WordWrap(text, 0, 0, maxWidth);
@@ -340,6 +343,7 @@ namespace EpicGames.Core
 		/// Parses a hexadecimal string into an array of bytes
 		/// </summary>
 		/// <param name="text">Text to parse</param>
+		/// <param name="outBytes">Receives the parsed string</param>
 		/// <returns></returns>
 		public static bool TryParseHexString(string text, [NotNullWhen(true)] out byte[]? outBytes)
 		{
@@ -368,6 +372,7 @@ namespace EpicGames.Core
 		/// Parses a hexadecimal string into an array of bytes
 		/// </summary>
 		/// <param name="text">Text to parse</param>
+		/// <param name="outBytes">Receives the parsed string</param>
 		/// <returns></returns>
 		public static bool TryParseHexString(ReadOnlySpan<byte> text, [NotNullWhen(true)] out byte[]? outBytes)
 		{
@@ -531,6 +536,7 @@ namespace EpicGames.Core
 		/// Converts a bytes string into bytes. E.g 1.5KB -> 1536
 		/// </summary>
 		/// <param name="bytesString"></param>
+		/// <param name="bytes">Receives the parsed bytes</param>
 		/// <returns></returns>
 		public static bool TryParseBytesString( string bytesString, out long? bytes )
 		{
@@ -612,7 +618,7 @@ namespace EpicGames.Core
 		/// </summary>
 		/// <param name="x">First string to compare</param>
 		/// <param name="y">Second string to compare</param>
-		/// <returns>Less than zero if X < Y, zero if X == Y, and greater than zero if X > y</returns>
+		/// <returns>Less than zero if X &lt; Y, zero if X == Y, and greater than zero if X &gt; y</returns>
 		public static int CompareIgnoreCaseUe(ReadOnlySpan<char> x, ReadOnlySpan<char> y)
 		{
 			int length = x.Length < y.Length ? x.Length : y.Length;
diff --git a/Engine/Source/Programs/Shared/EpicGames.Core/UHTTypes.cs b/Engine/Source/Programs/Shared/EpicGames.Core/UHTTypes.cs
--- a/Engine/Source/Programs/Shared/EpicGames.Core/UHTTypes.cs
+++ b/Engine/Source/Programs/Shared/EpicGames.Core/UHTTypes.cs
@@ -3,6 +3,8 @@
 using System.Collections.Generic;
 using System.Text.Json.Serialization;
 
+#pragma warning disable CS1591 // Missing documentation
+
 namespace EpicGames.Core
 {
 	/// <summary>
diff --git a/Engine/Source/Programs/Shared/EpicGames.Core/UnrealEngineTypes.cs b/Engine/Source/Programs/Shared/EpicGames.Core/UnrealEngineTypes.cs
--- a/Engine/Source/Programs/Shared/EpicGames.Core/UnrealEngineTypes.cs
+++ b/Engine/Source/Programs/Shared/EpicGames.Core/UnrealEngineTypes.cs
@@ -5,6 +5,7 @@ using System.Collections.Generic;
 using System.Reflection;
 
 #pragma warning disable IDE1006 // Naming Styles
+#pragma warning disable CS1591 // Missing documentation
 
 namespace EpicGames.Core
 {
@@ -236,7 +237,7 @@ namespace EpicGames.Core
 		HideDropDown = 0x04000000u,
 
 		/// <summary>
-		/// Class settings are saved to <AppData>/..../Blah.ini (as opposed to "DefaultConfig")
+		/// Class settings are saved to [AppData]/..../Blah.ini (as opposed to "DefaultConfig")
 		/// </summary>
 		GlobalUserConfig = 0x08000000u,
 
@@ -1145,7 +1146,7 @@ namespace EpicGames.Core
 		PersistentInstance = 0x0002000000000000,
 
 		/// <summary>
-		/// Property was parsed as a wrapper class like TSubclassOf<T>, FScriptInterface etc., rather than a USomething*
+		/// Property was parsed as a wrapper class like TSubclassOf&lt;T&gt;, FScriptInterface etc., rather than a USomething*
 		/// </summary>
 		UObjectWrapper = 0x0004000000000000,
 
@@ -1461,5 +1462,3 @@ namespace EpicGames.Core
 		}
 	}
 }
-
-#pragma warning restore IDE1006 // Naming Styles
diff --git a/Engine/Source/Programs/Shared/EpicGames.Core/Utf8String.cs b/Engine/Source/Programs/Shared/EpicGames.Core/Utf8String.cs
--- a/Engine/Source/Programs/Shared/EpicGames.Core/Utf8String.cs
+++ b/Engine/Source/Programs/Shared/EpicGames.Core/Utf8String.cs
@@ -370,6 +370,7 @@ namespace EpicGames.Core
 				return hash;
 			}
 
+			/// <inheritdoc/>
 			public override int Compare(ReadOnlySpan<byte> strA, ReadOnlySpan<byte> strB)
 			{
 				return strA.SequenceCompareTo(strB);
