Fix xmldoc warnings in EpicGames.Serialization.

diff --git a/Engine/Source/Programs/Shared/EpicGames.Serialization.Tests/EngineTests.cs b/Engine/Source/Programs/Shared/EpicGames.Serialization.Tests/EngineTests.cs
--- a/Engine/Source/Programs/Shared/EpicGames.Serialization.Tests/EngineTests.cs
+++ b/Engine/Source/Programs/Shared/EpicGames.Serialization.Tests/EngineTests.cs
@@ -325,7 +325,7 @@ namespace EpicGames.Serialization.Tests
 			{
 				CbField field = CbField.Empty;
 				TestFieldError(CbFieldType.Object, field, CbFieldError.TypeError);
-				CbField.MakeView(field).AsObject();
+				field.AsObject();
 			}
 
 			// Test FCbObjectView(ObjectWithName) and CreateIterator
@@ -425,7 +425,7 @@ namespace EpicGames.Serialization.Tests
 				byte[] payload = new byte[]{ 7, 3, intType, 1, intType, 2, intType, 3 };
 				CbField field = new CbField(payload, CbFieldType.Array);
 				TestField(CbFieldType.Array, field, new CbArray(payload, CbFieldType.Array));
-				CbArray array = CbArray.Clone(field.AsArray());
+				CbArray array = field.AsArray();
 				TestIntArray(array, 3, payload.Length);
 				TestIntArray(field.AsArray(), 3, payload.Length);
 				Assert.IsTrue(array.Equals(field.AsArray()));
@@ -437,7 +437,7 @@ namespace EpicGames.Serialization.Tests
 				byte[] payload = new byte[]{ 5, 3, intType, 1, 2, 3 };
 				CbField field = new CbField(payload, CbFieldType.UniformArray);
 				TestField(CbFieldType.UniformArray, field, new CbArray(payload, CbFieldType.UniformArray));
-				CbArray array = CbArray.Clone(field.AsArray());
+				CbArray array = field.AsArray();
 				TestIntArray(array, 3, payload.Length);
 				TestIntArray(field.AsArray(), 3, payload.Length);
 				Assert.IsTrue(array.Equals(field.AsArray()));
@@ -473,7 +473,7 @@ namespace EpicGames.Serialization.Tests
 			{
 				CbField field = new CbField();
 //				TestFieldError(CbFieldType.Array, Field, CbFieldError.TypeError);
-				CbField.MakeView(field).AsArray();
+				field.AsArray();
 			}
 
 			// Test CbArray(ArrayWithName) and CreateIterator
@@ -482,7 +482,7 @@ namespace EpicGames.Serialization.Tests
 				byte[] buffer = new byte[] { arrayType, 3, (byte)'K', (byte)'e', (byte)'y', 3, 1, (byte)(CbFieldType.IntegerPositive), 8 };
 				CbArray array = new CbArray(buffer);
 				Assert.AreEqual(array.GetSize(), 5);
-				CbArray arrayClone = CbArray.Clone(array);
+				CbArray arrayClone = array;
 				Assert.AreEqual(arrayClone.GetSize(), 5);
 				Assert.IsTrue(array.Equals(arrayClone));
 				Assert.AreEqual(arrayClone.GetHash(), array.GetHash());
@@ -552,7 +552,7 @@ namespace EpicGames.Serialization.Tests
 				CbField fieldView = new CbField(payload, CbFieldType.Binary);
 				TestField(CbFieldType.Binary, fieldView, (ReadOnlyMemory<byte>)payload.AsMemory(1, 3));
 
-				CbField field = CbField.Clone(fieldView);
+				CbField field = fieldView;
 				field.AsBinary();
 //				Assert.IsFalse(Field.GetOuterBuffer().IsNull());
 //				MoveTemp(Field).AsBinary();
diff --git a/Engine/Source/Programs/Shared/EpicGames.Serialization/CbConverter.cs b/Engine/Source/Programs/Shared/EpicGames.Serialization/CbConverter.cs
--- a/Engine/Source/Programs/Shared/EpicGames.Serialization/CbConverter.cs
+++ b/Engine/Source/Programs/Shared/EpicGames.Serialization/CbConverter.cs
@@ -31,7 +31,7 @@ namespace EpicGames.Serialization
 	}
 
 	/// <summary>
-	/// Base class for all converters. Deriving from <see cref="CbConverter{T}"/> is more efficient
+	/// Base class for all converters. Deriving from <see cref="ICbConverter{T}"/> is more efficient
 	/// </summary>
 	public interface ICbConverter
 	{
@@ -145,9 +145,9 @@ namespace EpicGames.Serialization
 		}
 
 		/// <summary>
-		/// Gets a <see cref="ICbConverterMethods"/> interface for the given type
+		/// Gets a <see cref="ICbConverterMethods"/> interface for the given property
 		/// </summary>
-		/// <param name="Type"></param>
+		/// <param name="property"></param>
 		/// <returns></returns>
 		public static ICbConverterMethods Get(PropertyInfo property)
 		{
diff --git a/Engine/Source/Programs/Shared/EpicGames.Serialization/CbField.cs b/Engine/Source/Programs/Shared/EpicGames.Serialization/CbField.cs
--- a/Engine/Source/Programs/Shared/EpicGames.Serialization/CbField.cs
+++ b/Engine/Source/Programs/Shared/EpicGames.Serialization/CbField.cs
@@ -314,7 +314,7 @@ namespace EpicGames.Serialization
 		/// <summary>
 		/// Use an object attachment as a hash
 		/// </summary>
-		/// <param name="Attachment">The attachment to convert</param>
+		/// <param name="hash">The attachment to convert</param>
 		public static implicit operator CbObjectAttachment(IoHash hash) => new CbObjectAttachment(hash);
 
 		/// <summary>
@@ -457,7 +457,7 @@ namespace EpicGames.Serialization
 		/// <summary>
 		/// Tests if the given field type is binary
 		/// </summary>
-		/// <param name="Type">Type to check</param>
+		/// <param name="type">Type to check</param>
 		/// <returns>True if the field is binary</returns>
 		public static bool IsBinary(CbFieldType type)
 		{
@@ -507,7 +507,7 @@ namespace EpicGames.Serialization
 		/// <summary>
 		/// Tests if the given field type is a compact binary attachment
 		/// </summary>
-		/// <param name="Type">Type to check</param>
+		/// <param name="type">Type to check</param>
 		/// <returns>True if the field is a compact binary attachment</returns>
 		public static bool IsObjectAttachment(CbFieldType type)
 		{
@@ -517,7 +517,7 @@ namespace EpicGames.Serialization
 		/// <summary>
 		/// Tests if the given field type is a binary attachment
 		/// </summary>
-		/// <param name="Type">Type to check</param>
+		/// <param name="type">Type to check</param>
 		/// <returns>True if the field is a binary attachment</returns>
 		public static bool IsBinaryAttachment(CbFieldType type)
 		{
@@ -547,7 +547,7 @@ namespace EpicGames.Serialization
 		/// <summary>
 		/// Tests if the given field type is a UUID
 		/// </summary>
-		/// <param name="Type">Type to check</param>
+		/// <param name="type">Type to check</param>
 		/// <returns>True if the field is a UUID</returns>
 		public static bool IsUuid(CbFieldType type)
 		{
@@ -557,7 +557,7 @@ namespace EpicGames.Serialization
 		/// <summary>
 		/// Tests if the given field type is a date/time
 		/// </summary>
-		/// <param name="Type">Type to check</param>
+		/// <param name="type">Type to check</param>
 		/// <returns>True if the field is a date/time</returns>
 		public static bool IsDateTime(CbFieldType type)
 		{
@@ -567,7 +567,7 @@ namespace EpicGames.Serialization
 		/// <summary>
 		/// Tests if the given field type is a timespan
 		/// </summary>
-		/// <param name="Type">Type to check</param>
+		/// <param name="type">Type to check</param>
 		/// <returns>True if the field is a timespan</returns>
 		public static bool IsTimeSpan(CbFieldType type)
 		{
@@ -577,7 +577,7 @@ namespace EpicGames.Serialization
 		/// <summary>
 		/// Tests if the given field type is a object id
 		/// </summary>
-		/// <param name="Type">Type to check</param>
+		/// <param name="type">Type to check</param>
 		/// <returns>True if the field is a object id</returns>
 		public static bool IsObjectId(CbFieldType type)
 		{
@@ -1071,7 +1071,7 @@ namespace EpicGames.Serialization
 		/// <summary>
 		/// Access the field as an integer, checking that it's in the correct range
 		/// </summary>
-		/// <param name="default"></param>
+		/// <param name="defaultValue"></param>
 		/// <param name="magnitudeBits"></param>
 		/// <param name="isSigned"></param>
 		/// <returns></returns>
@@ -1107,7 +1107,7 @@ namespace EpicGames.Serialization
 		/// <summary>
 		/// Access the field as a float. Returns the provided default on error.
 		/// </summary>
-		/// <param name="default">Default value</param>
+		/// <param name="defaultValue">Default value</param>
 		/// <returns>Value of the field</returns>
 		public float AsFloat(float defaultValue = 0.0f)
 		{
@@ -1145,7 +1145,7 @@ namespace EpicGames.Serialization
 		/// <summary>
 		/// Access the field as a double. Returns the provided default on error.
 		/// </summary>
-		/// <param name="default">Default value</param>
+		/// <param name="defaultValue">Default value</param>
 		/// <returns>Value of the field</returns>
 		public double AsDouble(double defaultValue)
 		{
@@ -1187,7 +1187,7 @@ namespace EpicGames.Serialization
 		/// <summary>
 		/// Access the field as a bool. Returns the provided default on error.
 		/// </summary>
-		/// <param name="default">Default value</param>
+		/// <param name="defaultValue">Default value</param>
 		/// <returns>Value of the field</returns>
 		public bool AsBool(bool defaultValue)
 		{
@@ -1214,7 +1214,7 @@ namespace EpicGames.Serialization
 		/// <summary>
 		/// Access the field as a hash referencing an object attachment. Returns the provided default on error.
 		/// </summary>
-		/// <param name="default">Default value</param>
+		/// <param name="defaultValue">Default value</param>
 		/// <returns>Value of the field</returns>
 		public CbObjectAttachment AsObjectAttachment(CbObjectAttachment defaultValue)
 		{
@@ -1233,14 +1233,13 @@ namespace EpicGames.Serialization
 		/// <summary>
 		/// Access the field as a hash referencing a binary attachment. Returns the provided default on error.
 		/// </summary>
-		/// <param name="Default">Default value</param>
 		/// <returns>Value of the field</returns>
 		public CbBinaryAttachment AsBinaryAttachment() => AsBinaryAttachment(CbBinaryAttachment.Zero);
 
 		/// <summary>
 		/// Access the field as a hash referencing a binary attachment. Returns the provided default on error.
 		/// </summary>
-		/// <param name="default">Default value</param>
+		/// <param name="defaultValue">Default value</param>
 		/// <returns>Value of the field</returns>
 		public CbBinaryAttachment AsBinaryAttachment(CbBinaryAttachment defaultValue)
 		{
@@ -1259,14 +1258,13 @@ namespace EpicGames.Serialization
 		/// <summary>
 		/// Access the field as a hash referencing an attachment. Returns the provided default on error.
 		/// </summary>
-		/// <param name="Default">Default value</param>
 		/// <returns>Value of the field</returns>
 		public IoHash AsAttachment() => AsAttachment(IoHash.Zero);
 
 		/// <summary>
 		/// Access the field as a hash referencing an attachment. Returns the provided default on error.
 		/// </summary>
-		/// <param name="default">Default value</param>
+		/// <param name="defaultValue">Default value</param>
 		/// <returns>Value of the field</returns>
 		public IoHash AsAttachment(IoHash defaultValue)
 		{
@@ -1291,7 +1289,7 @@ namespace EpicGames.Serialization
 		/// <summary>
 		/// Access the field as a hash referencing an attachment. Returns the provided default on error.
 		/// </summary>
-		/// <param name="default">Default value</param>
+		/// <param name="defaultValue">Default value</param>
 		/// <returns>Value of the field</returns>
 		public IoHash AsHash(IoHash defaultValue)
 		{
@@ -1310,7 +1308,7 @@ namespace EpicGames.Serialization
 		/// <summary>
 		/// Access the field as a UUID. Returns a nil UUID on error.
 		/// </summary>
-		/// <param name="default">Default value</param>
+		/// <param name="defaultValue">Default value</param>
 		/// <returns>Value of the field</returns>
 		public Guid AsUuid(Guid defaultValue = default)
 		{
@@ -1335,7 +1333,7 @@ namespace EpicGames.Serialization
 		/// <summary>
 		/// Reads a date time as number of ticks from the stream
 		/// </summary>
-		/// <param name="default"></param>
+		/// <param name="defaultValue"></param>
 		/// <returns></returns>
 		public long AsDateTimeTicks(long defaultValue = 0)
 		{
@@ -1363,7 +1361,7 @@ namespace EpicGames.Serialization
 		/// <summary>
 		/// Access the field as a DateTime.
 		/// </summary>
-		/// <param name="default"></param>
+		/// <param name="defaultValue"></param>
 		/// <returns></returns>
 		public DateTime AsDateTime(DateTime defaultValue)
 		{
@@ -1373,7 +1371,7 @@ namespace EpicGames.Serialization
 		/// <summary>
 		/// Reads a timespan as number of ticks from the stream
 		/// </summary>
-		/// <param name="default"></param>
+		/// <param name="defaultValue"></param>
 		/// <returns></returns>
 		public long AsTimeSpanTicks(long defaultValue = 0)
 		{
@@ -1392,14 +1390,14 @@ namespace EpicGames.Serialization
 		/// <summary>
 		/// Reads a timespan as number of ticks from the stream
 		/// </summary>
-		/// <param name="default"></param>
+		/// <param name="defaultValue"></param>
 		/// <returns></returns>
 		public TimeSpan AsTimeSpan(TimeSpan defaultValue = default) => new TimeSpan(AsTimeSpanTicks(defaultValue.Ticks));
 
 		/// <summary>
 		/// Access the field as a object id
 		/// </summary>
-		/// <param name="default"></param>
+		/// <param name="defaultValue"></param>
 		/// <returns></returns>
 		public ReadOnlyMemory<byte> AsObjectId(ReadOnlyMemory<byte> defaultValue = default)
 		{
@@ -1681,9 +1679,7 @@ namespace EpicGames.Serialization
 		/// <summary>
 		/// Returns the size of the value payload in bytes, which is the field excluding the type and name.
 		/// </summary>
-		/// <param name="Type"></param>
-		/// <param name="Payload"></param>
-		/// <returns></returns>
+		/// <returns>Size of the payload</returns>
 		public ulong GetPayloadSize()
 		{
 			switch (GetType())
@@ -1728,15 +1724,6 @@ namespace EpicGames.Serialization
 					return 0;
 			}
 		}
-
-		#region Mimic inheritance from TCbBufferFactory
-
-		public static CbField Clone(ReadOnlyMemory<byte> data) => Clone(new CbField(data));
-		public static CbField Clone(CbField other) => other;
-		public static CbField MakeView(ReadOnlyMemory<byte> data) => new CbField(data);
-		public static CbField MakeView(CbField other) => other;
-
-		#endregion
 	}
 
 	/// <summary>
@@ -1909,6 +1896,9 @@ namespace EpicGames.Serialization
 		}
 	}
 
+	/// <summary>
+	/// Enumerator for contents of a field
+	/// </summary>
 	public class CbFieldEnumerator : IEnumerator<CbField>
 	{
 		/// <summary>
@@ -2184,20 +2174,25 @@ namespace EpicGames.Serialization
 			return new CbFieldIterator(iterator._nextData, iterator._uniformType);
 		}
 
+		/// <inheritdoc/>
 		public override bool Equals(object? obj)
 		{
 			throw new NotImplementedException();
 		}
 
+		/// <inheritdoc/>
 		public override int GetHashCode()
 		{
 			throw new NotImplementedException();
 		}
 
+		/// <inheritdoc/>
 		public static bool operator ==(CbFieldIterator a, CbFieldIterator b)
 		{
 			return a.Current.Equals(b.Current);
 		}
+
+		/// <inheritdoc/>
 		public static bool operator !=(CbFieldIterator a, CbFieldIterator b)
 		{
 			return !a.Current.Equals(b.Current);
@@ -2395,15 +2390,6 @@ namespace EpicGames.Serialization
 		internal ReadOnlyMemory<byte> GetPayloadView() => _innerField.GetPayloadView();
 
 		#endregion
-
-		#region Mimic inheritance from TCbBufferFactory
-
-		public static CbArray Clone(ReadOnlyMemory<byte> data) => Clone(new CbArray(data));
-		public static CbArray Clone(CbArray other) => other;
-		public static CbArray MakeView(ReadOnlyMemory<byte> data) => new CbArray(data);
-		public static CbArray MakeView(CbArray other) => other;
-
-		#endregion
 	}
 
 	/// <summary>
@@ -2450,7 +2436,7 @@ namespace EpicGames.Serialization
 		/// <summary>
 		/// Constructor
 		/// </summary>
-		/// <param name="Buffer"></param>
+		/// <param name="field"></param>
 		private CbObject(CbField field)
 		{
 			_innerField = new CbField(field.Memory, field.TypeWithFlags);
@@ -2460,6 +2446,7 @@ namespace EpicGames.Serialization
 		/// Constructor
 		/// </summary>
 		/// <param name="buffer"></param>
+		/// <param name="fieldType">Explicit type of the data in buffer</param>
 		public CbObject(ReadOnlyMemory<byte> buffer, CbFieldType fieldType = CbFieldType.HasFieldType)
 		{
 			_innerField = new CbField(buffer, fieldType);
@@ -2503,7 +2490,10 @@ namespace EpicGames.Serialization
 		/// <returns>The matching field if found, otherwise a field with no value.</returns>
 		public CbField this[Utf8String name] => _innerField[name];
 
-		/// <inheritdoc cref="AsFieldView"/>
+		/// <summary>
+		/// Gets the underlying field for this object
+		/// </summary>
+		/// <returns></returns>
 		public CbField AsField() => _innerField;
 
 		/// <summary>
diff --git a/Engine/Source/Programs/Shared/EpicGames.Serialization/CbWriter.cs b/Engine/Source/Programs/Shared/EpicGames.Serialization/CbWriter.cs
--- a/Engine/Source/Programs/Shared/EpicGames.Serialization/CbWriter.cs
+++ b/Engine/Source/Programs/Shared/EpicGames.Serialization/CbWriter.cs
@@ -321,6 +321,7 @@ namespace EpicGames.Serialization
 		/// <summary>
 		/// Copies an entire field value to the output
 		/// </summary>
+		/// <param name="name">Name of the field</param>
 		/// <param name="field"></param>
 		public void WriteField(Utf8String name, CbField field)
 		{
diff --git a/Engine/Source/Programs/Shared/EpicGames.Serialization/Converters/CbFieldConverter.cs b/Engine/Source/Programs/Shared/EpicGames.Serialization/Converters/CbFieldConverter.cs
--- a/Engine/Source/Programs/Shared/EpicGames.Serialization/Converters/CbFieldConverter.cs
+++ b/Engine/Source/Programs/Shared/EpicGames.Serialization/Converters/CbFieldConverter.cs
@@ -7,7 +7,6 @@ namespace EpicGames.Serialization.Converters
 	/// <summary>
 	/// Converter for raw CbField types
 	/// </summary>
-	/// <typeparam name="T"></typeparam>
 	class CbFieldConverter : CbConverterBase<CbField>
 	{
 		/// <inheritdoc/>
@@ -32,7 +31,6 @@ namespace EpicGames.Serialization.Converters
 	/// <summary>
 	/// Converter for raw CbObject types
 	/// </summary>
-	/// <typeparam name="T"></typeparam>
 	class CbObjectConverter : CbConverterBase<CbObject>
 	{
 		/// <inheritdoc/>
