BuildGraph: Fix exception passing Bg* value types to methods. Need to construct a constant Bg* type rather than passing the wrapped value.

diff --git a/Engine/Source/Programs/AutomationTool/BuildGraph/BgNodeExecutor.cs b/Engine/Source/Programs/AutomationTool/BuildGraph/BgNodeExecutor.cs
--- a/Engine/Source/Programs/AutomationTool/BuildGraph/BgNodeExecutor.cs
+++ b/Engine/Source/Programs/AutomationTool/BuildGraph/BgNodeExecutor.cs
@@ -76,6 +76,10 @@ namespace AutomationTool
 				{
 					arguments[Idx] = Context;
 				}
+				else if (typeof(IBgExpr).IsAssignableFrom(ParameterType))
+				{
+					arguments[Idx] = BgType.Get(ParameterType).CreateConstant(Node.Arguments[Idx]!);
+				}
 				else
 				{
 					arguments[Idx] = Node.Arguments[Idx];
diff --git a/Engine/Source/Programs/Shared/EpicGames.BuildGraph/BgNodeSpec.cs b/Engine/Source/Programs/Shared/EpicGames.BuildGraph/BgNodeSpec.cs
--- a/Engine/Source/Programs/Shared/EpicGames.BuildGraph/BgNodeSpec.cs
+++ b/Engine/Source/Programs/Shared/EpicGames.BuildGraph/BgNodeSpec.cs
@@ -54,7 +54,7 @@ namespace EpicGames.BuildGraph
 	public class BgNodeSpec
 	{
 		internal BgMethod _method;
-		internal IBgExpr?[] _argumentExprs;
+		internal object?[] _arguments;
 		internal BgFileSet[] ResultExprs { get; }
 
 		/// <summary>
@@ -101,8 +101,8 @@ namespace EpicGames.BuildGraph
 
 			try
 			{
-				_argumentExprs = CreateArgumentExprs(call);
-				Name = GetDefaultNodeName(call.Method, _argumentExprs);
+				_arguments = CreateArgumentExprs(call);
+				Name = GetDefaultNodeName(call.Method, _arguments);
 				ResultExprs = CreateReturnExprs(call.Method, Name, call.Method.ReturnType);
 			}
 			catch (Exception ex)
@@ -132,9 +132,9 @@ namespace EpicGames.BuildGraph
 			return new BgNodeSpec<T>(call);
 		}
 
-		static IBgExpr?[] CreateArgumentExprs(MethodCallExpression call)
+		static object?[] CreateArgumentExprs(MethodCallExpression call)
 		{
-			IBgExpr?[] args = new IBgExpr?[call.Arguments.Count];
+			object?[] args = new object?[call.Arguments.Count];
 			for (int idx = 0; idx < call.Arguments.Count; idx++)
 			{
 				Expression expr = call.Arguments[idx];
@@ -148,16 +148,7 @@ namespace EpicGames.BuildGraph
 				else
 				{
 					Delegate compiled = Expression.Lambda(expr).Compile();
-
-					object? result = compiled.DynamicInvoke();
-					if (result is IBgExpr computable)
-					{
-						args[idx] = computable;
-					}
-					else
-					{
-						args[idx] = (BgString)(result?.ToString() ?? String.Empty);
-					}
+					args[idx] = compiled.DynamicInvoke();
 				}
 			}
 			return args;
@@ -224,7 +215,7 @@ namespace EpicGames.BuildGraph
 			return false;
 		}
 
-		static BgString GetDefaultNodeName(MethodInfo methodInfo, IBgExpr?[] args)
+		static BgString GetDefaultNodeName(MethodInfo methodInfo, object?[] args)
 		{
 			// Check if it's got an attribute override for the node name
 			BgNodeNameAttribute? nameAttr = methodInfo.GetCustomAttribute<BgNodeNameAttribute>();
@@ -238,7 +229,7 @@ namespace EpicGames.BuildGraph
 			}
 		}
 
-		static BgString GetNodeNameFromTemplate(string template, ParameterInfo[] parameters, IBgExpr?[] args)
+		static BgString GetNodeNameFromTemplate(string template, ParameterInfo[] parameters, object?[] args)
 		{
 			// Create a list of lazily computed string fragments which comprise the evaluated name
 			List<BgString> fragments = new List<BgString>();
@@ -272,10 +263,17 @@ namespace EpicGames.BuildGraph
 							throw new BgNodeException($"Unable to find parameter named {paramName} in {template}");
 						}
 
-						IBgExpr? arg = args[paramIdx];
+						object? arg = args[paramIdx];
 						if (arg != null)
 						{
-							fragments.Add(arg.ToBgString());
+							if (arg is IBgExpr expr)
+							{
+								fragments.Add(expr.ToBgString());
+							}
+							else
+							{
+								fragments.Add(arg.ToString() ?? String.Empty);
+							}
 						}
 
 						lastIdx = nextIdx = endIdx + 1;
@@ -357,8 +355,8 @@ namespace EpicGames.BuildGraph
 		{
 			HashSet<string> inputTags = new HashSet<string>();
 			inputTags.UnionWith(InputDependencies.ComputeTags(context));
-			inputTags.UnionWith(_argumentExprs.OfType<BgFileSet>().Select(x => x.ComputeTag(context)));
-			inputTags.UnionWith(_argumentExprs.OfType<BgList<BgFileSet>>().SelectMany(x => x.GetEnumerable(context)).Select(x => x.ComputeTag(context)));
+			inputTags.UnionWith(_arguments.OfType<BgFileSet>().Select(x => x.ComputeTag(context)));
+			inputTags.UnionWith(_arguments.OfType<BgList<BgFileSet>>().SelectMany(x => x.GetEnumerable(context)).Select(x => x.ComputeTag(context)));
 
 			HashSet<string> afterTags = new HashSet<string>(inputTags);
 			afterTags.UnionWith(AfterDependencies.ComputeTags(context));
@@ -371,7 +369,7 @@ namespace EpicGames.BuildGraph
 			BgNode[] afterNodes = afterDependencies.Select(x => x.ProducingNode).Distinct().ToArray();
 			bool runEarly = CanRunEarly.Compute(context);
 
-			List<object?> arguments = _argumentExprs.ConvertAll(x => x?.Compute(context));
+			List<object?> arguments = _arguments.ConvertAll(x => (x is IBgExpr expr) ? expr.Compute(context) : x);
 			BgExpressionNode node = new BgExpressionNode(name, _method, arguments, inputDependencies, outputNames, inputNodes, afterNodes, Array.Empty<FileReference>());
 			node.RunEarly = runEarly;
 
diff --git a/Engine/Source/Programs/Shared/EpicGames.BuildGraph/Expressions/IBgType.cs b/Engine/Source/Programs/Shared/EpicGames.BuildGraph/Expressions/IBgType.cs
--- a/Engine/Source/Programs/Shared/EpicGames.BuildGraph/Expressions/IBgType.cs
+++ b/Engine/Source/Programs/Shared/EpicGames.BuildGraph/Expressions/IBgType.cs
@@ -25,6 +25,13 @@ namespace EpicGames.BuildGraph.Expressions
 		/// <param name="text">The serialized value</param>
 		/// <returns>Instance of the expression</returns>
 		object DeserializeArgument(string text);
+
+		/// <summary>
+		/// Creates a value of the expression type
+		/// </summary>
+		/// <param name="value">Value to wrap</param>
+		/// <returns></returns>
+		object CreateConstant(object value);
 	}
 
 	/// <summary>
@@ -52,7 +59,7 @@ namespace EpicGames.BuildGraph.Expressions
 		/// </summary>
 		/// <param name="value">Value to wrap</param>
 		/// <returns></returns>
-		T CreateConstant(object value);
+		new T CreateConstant(object value);
 
 		/// <summary>
 		/// Creates a variable of the given type
@@ -78,6 +85,9 @@ namespace EpicGames.BuildGraph.Expressions
 		/// <inheritdoc/>
 		public abstract string SerializeArgument(T value, BgExprContext context);
 
+		/// <inheritdoc/>
+		object IBgType.CreateConstant(object value) => CreateConstant(value);
+
 		/// <inheritdoc/>
 		public abstract T CreateConstant(object value);
 
