Updating property flags with recent change

#rnx

diff --git a/Engine/Source/Programs/Shared/EpicGames.Core/UnrealEngineTypes.cs b/Engine/Source/Programs/Shared/EpicGames.Core/UnrealEngineTypes.cs
--- a/Engine/Source/Programs/Shared/EpicGames.Core/UnrealEngineTypes.cs
+++ b/Engine/Source/Programs/Shared/EpicGames.Core/UnrealEngineTypes.cs
@@ -977,7 +977,10 @@ namespace EpicGames.Core
 		/// </summary>
 		DisableEditOnTemplate = 0x0000000000000800,
 
-		// Unused = 0x0000000000001000,
+		/// <summary>
+		/// Object property can never be null
+		/// </summary>
+		NonNullable = 0x0000000000001000,
 
 		/// <summary>
 		/// Property is transient: shouldn't be saved or loaded, except for Blueprint CDOs.
