1) Fixed issue where UPARAM wasn't getting parsed as part of a template argument property type
2) Refactored how property specifiers are parsed due to above issue
3) Fixed issue where UClass derived types were not outputting correctly in event parameter lists
4) Added more diagnostic output to package files.

#rnx

diff --git a/Engine/Source/Programs/Shared/EpicGames.UHT/Exporters/CodeGen/UhtPackageCodeGeneratorCppFile.cs b/Engine/Source/Programs/Shared/EpicGames.UHT/Exporters/CodeGen/UhtPackageCodeGeneratorCppFile.cs
--- a/Engine/Source/Programs/Shared/EpicGames.UHT/Exporters/CodeGen/UhtPackageCodeGeneratorCppFile.cs
+++ b/Engine/Source/Programs/Shared/EpicGames.UHT/Exporters/CodeGen/UhtPackageCodeGeneratorCppFile.cs
@@ -91,6 +91,10 @@ namespace EpicGames.UHT.Exporters.CodeGen
 				if (this.Session.IncludeDebugOutput)
 				{
 					builder.Append("#if 0\r\n");
+					foreach (UhtHeaderFile headerFile in packageSortedHeaders)
+					{
+						builder.Append('\t').Append(headerFile.FilePath).Append("\r\n");
+					}
 					builder.Append(declarations);
 					builder.Append("#endif\r\n");
 				}
diff --git a/Engine/Source/Programs/Shared/EpicGames.UHT/Parsers/UhtFunctionParser.cs b/Engine/Source/Programs/Shared/EpicGames.UHT/Parsers/UhtFunctionParser.cs
--- a/Engine/Source/Programs/Shared/EpicGames.UHT/Parsers/UhtFunctionParser.cs
+++ b/Engine/Source/Programs/Shared/EpicGames.UHT/Parsers/UhtFunctionParser.cs
@@ -204,8 +204,8 @@ namespace EpicGames.UHT.Parsers
 					UhtProperty? returnValueProperty = null;
 					if (hasReturnValue)
 					{
-						topScope.HeaderParser.GetCachedPropertyParser(topScope).Parse(EPropertyFlags.None,
-							function.GetPropertyParseOptions(true), UhtParsePropertyDeclarationStyle.None, UhtPropertyCategory.Return,
+						topScope.HeaderParser.GetCachedPropertyParser().Parse(topScope, EPropertyFlags.None,
+							function.GetPropertyParseOptions(true), UhtPropertyCategory.Return,
 							(UhtParsingScope topScope, UhtProperty property, ref UhtToken nameToken, UhtLayoutMacroType layoutMacroType) =>
 							{
 								property.PropertyFlags |= EPropertyFlags.Parm | EPropertyFlags.OutParm | EPropertyFlags.ReturnParm;
@@ -429,8 +429,8 @@ namespace EpicGames.UHT.Parsers
 						// Get return type.  C++ style functions always have a return value type, even if it's void
 						UhtToken funcNameToken = new();
 						UhtProperty? returnValueProperty = null;
-						topScope.HeaderParser.GetCachedPropertyParser(topScope).Parse(EPropertyFlags.None,
-							function.GetPropertyParseOptions(true), UhtParsePropertyDeclarationStyle.None, UhtPropertyCategory.Return,
+						topScope.HeaderParser.GetCachedPropertyParser().Parse(topScope, EPropertyFlags.None,
+							function.GetPropertyParseOptions(true), UhtPropertyCategory.Return,
 							(UhtParsingScope topScope, UhtProperty property, ref UhtToken nameToken, UhtLayoutMacroType layoutMacroType) =>
 							{
 								property.PropertyFlags |= EPropertyFlags.Parm | EPropertyFlags.OutParm | EPropertyFlags.ReturnParm;
@@ -559,7 +559,7 @@ namespace EpicGames.UHT.Parsers
 
 			topScope.TokenReader.RequireList(')', ',', false, () =>
 			{
-				topScope.HeaderParser.GetCachedPropertyParser(topScope).Parse(disallowFlags, options, UhtParsePropertyDeclarationStyle.None, propertyCategory,
+				topScope.HeaderParser.GetCachedPropertyParser().Parse(topScope, disallowFlags, options, propertyCategory,
 					(UhtParsingScope topScope, UhtProperty property, ref UhtToken nameToken, UhtLayoutMacroType layoutMacroType) =>
 					{
 						property.PropertyFlags |= EPropertyFlags.Parm;
diff --git a/Engine/Source/Programs/Shared/EpicGames.UHT/Parsers/UhtHeaderFileParser.cs b/Engine/Source/Programs/Shared/EpicGames.UHT/Parsers/UhtHeaderFileParser.cs
--- a/Engine/Source/Programs/Shared/EpicGames.UHT/Parsers/UhtHeaderFileParser.cs
+++ b/Engine/Source/Programs/Shared/EpicGames.UHT/Parsers/UhtHeaderFileParser.cs
@@ -295,17 +295,12 @@ namespace EpicGames.UHT.Parsers
 		/// <summary>
 		/// Get the cached property parser
 		/// </summary>
-		/// <param name="topScope">Top scope requesting parser</param>
 		/// <returns>Property parser</returns>
-		public UhtPropertyParser GetCachedPropertyParser(UhtParsingScope topScope)
+		public UhtPropertyParser GetCachedPropertyParser()
 		{
 			if (this.PropertyParser == null)
 			{
-				this.PropertyParser = new UhtPropertyParser(topScope, topScope.TokenReader);
-			}
-			else
-			{
-				this.PropertyParser.Reset(topScope, topScope.TokenReader);
+				this.PropertyParser = new();
 			}
 			return this.PropertyParser;
 		}
diff --git a/Engine/Source/Programs/Shared/EpicGames.UHT/Parsers/UhtPropertyParser.cs b/Engine/Source/Programs/Shared/EpicGames.UHT/Parsers/UhtPropertyParser.cs
--- a/Engine/Source/Programs/Shared/EpicGames.UHT/Parsers/UhtPropertyParser.cs
+++ b/Engine/Source/Programs/Shared/EpicGames.UHT/Parsers/UhtPropertyParser.cs
@@ -109,23 +109,6 @@ namespace EpicGames.UHT.Parsers
 		}
 	}
 
-	/// <summary>
-	/// The property style of a variable declaration being parsed
-	/// </summary>
-	public enum UhtParsePropertyDeclarationStyle
-	{
-
-		/// <summary>
-		/// Property is a function return, parameters or this is a template argument
-		/// </summary>
-		None,
-
-		/// <summary>
-		/// Class or script structure property
-		/// </summary>
-		UPROPERTY
-	};
-
 	/// <summary>
 	/// Layout macro type
 	/// </summary>
@@ -316,7 +299,7 @@ namespace EpicGames.UHT.Parsers
 		/// <summary>
 		/// The property settings being parsed
 		/// </summary>
-		public UhtPropertySettings PropertySettings { get; set; }
+		public UhtPropertySettings PropertySettings { get; } = new UhtPropertySettings();
 
 		/// <summary>
 		/// If true, editor specifier seen
@@ -337,16 +320,44 @@ namespace EpicGames.UHT.Parsers
 		/// If true, blueprint getter specifier seen
 		/// </summary>
 		public bool SeenBlueprintGetterSpecifier { get; set; } = false;
+	}
+
+	/// <summary>
+	/// Helper class thread specified object cache
+	/// </summary>
+	public struct UhtThreadBorrower<T> : IDisposable where T : new()
+	{
+		private static readonly ThreadLocal<List<T>> s_tls = new(() => new());
+		private readonly T _instance;
+
+		/// <summary>
+		/// The borrowed instance
+		/// </summary>
+		public T Instance => _instance;
 
 		/// <summary>
-		/// Construct a new property specifier context
+		/// Request a context
 		/// </summary>
-		/// <param name="scope">Scope being parsed</param>
-		/// <param name="messageSite">Message site</param>
-		public UhtPropertySpecifierContext(UhtParsingScope scope, IUhtMessageSite messageSite) : base(scope, messageSite, UhtMetaData.Empty)
+		public UhtThreadBorrower()
 		{
-			this.PropertySettings = new UhtPropertySettings();
-			this.PropertySettings.Reset(scope.ScopeType, 1, UhtPropertyCategory.Member, 0);
+			List<T> cache = s_tls.Value!;
+			if (cache.Count == 0)
+			{
+				_instance = new();
+			}
+			else
+			{
+				_instance = cache[^1];
+				cache.RemoveAt(cache.Count - 1);
+			}
+		}
+
+		/// <summary>
+		/// Return the borrowed buffer to the cache
+		/// </summary>
+		public void Dispose()
+		{
+			s_tls.Value!.Add(_instance);
 		}
 	}
 
@@ -434,14 +445,11 @@ namespace EpicGames.UHT.Parsers
 	/// </summary>
 	public class UhtPropertyParser : IUhtMessageExtraContext
 	{
-		private readonly UhtPropertySpecifierContext _specifierContext;
+		private UhtPropertyCategory _category = UhtPropertyCategory.Member;
 		private UhtPropertyParseOptions _options;
-		private UhtParsePropertyDeclarationStyle _declarationStyle;
-
-		private IUhtTokenReader TokenReader => this._specifierContext.Scope.TokenReader;
-		private UhtParsingScope TopScope => this._specifierContext.Scope;
 
 		// Scratch pad variables used by actions
+		private IUhtTokenReader? _currentTokenReader = null;
 		private List<UhtToken> _currentTypeTokens = new();
 		private int _currentTemplateDepth = 0;
 
@@ -462,55 +470,46 @@ namespace EpicGames.UHT.Parsers
 		/// <summary>
 		/// Construct a new property parser
 		/// </summary>
-		/// <param name="scope">Current scope being parsed</param>
-		/// <param name="messageSite">Message site</param>
-		public UhtPropertyParser(UhtParsingScope scope, IUhtMessageSite messageSite)
+		public UhtPropertyParser()
 		{
-			this._specifierContext = new UhtPropertySpecifierContext(scope, messageSite);
 			this._gatherTypeTokensDelegate = GatherTypeTokens;
 		}
 
-		/// <summary>
-		/// Reset the property parser for a new parse
-		/// </summary>
-		/// <param name="scope">Current scope being parsed</param>
-		/// <param name="messageSite">Message site</param>
-		public void Reset(UhtParsingScope scope, IUhtMessageSite messageSite)
-		{
-			this._specifierContext.Scope = scope;
-			this._specifierContext.MessageSite = messageSite;
-		}
-
 		/// <summary>
 		/// Parse the property
 		/// </summary>
+		/// <param name="topScope">Current parsing scope</param>
 		/// <param name="disallowPropertyFlags">Flags to be disallowed</param>
 		/// <param name="options">Parsing options</param>
-		/// <param name="declarationStyle">Style of declaration</param>
 		/// <param name="category">Property category</param>
 		/// <param name="propertyDelegate">Delegate to be invoked after property has been parsed</param>
 		/// <returns>The property parser</returns>
-		public UhtPropertyParser Parse(EPropertyFlags disallowPropertyFlags, UhtPropertyParseOptions options, UhtParsePropertyDeclarationStyle declarationStyle, UhtPropertyCategory category, UhtPropertyDelegate propertyDelegate)
+		public UhtPropertyParser Parse(UhtParsingScope topScope, EPropertyFlags disallowPropertyFlags, UhtPropertyParseOptions options, UhtPropertyCategory category, UhtPropertyDelegate propertyDelegate)
 		{
-			// Reset the context and create the property
-			this._specifierContext.PropertySettings.Reset(this._specifierContext.Scope.ScopeType, 0, category, disallowPropertyFlags);
-			this._specifierContext.MetaData = this._specifierContext.PropertySettings.MetaData;
-			this._specifierContext.MetaNameIndex = UhtMetaData.IndexNone;
-			this._specifierContext.SeenEditSpecifier = false;
-			this._specifierContext.SeenBlueprintWriteSpecifier = false;
-			this._specifierContext.SeenBlueprintReadOnlySpecifier = false;
-			this._specifierContext.SeenBlueprintGetterSpecifier = false;
-
-			// Initialize the property
-
+			// Initialize the property context
+			using UhtThreadBorrower<UhtPropertySpecifierContext> borrower = new();
+			UhtPropertySpecifierContext specifierContext = borrower.Instance;
+			specifierContext.Type = topScope.ScopeType;
+			specifierContext.TokenReader = topScope.TokenReader;
+			specifierContext.AccessSpecifier = topScope.AccessSpecifier;
+			specifierContext.MessageSite = topScope.TokenReader;
+			specifierContext.PropertySettings.Reset(specifierContext.Type, 0, category, disallowPropertyFlags);
+			specifierContext.MetaData = specifierContext.PropertySettings.MetaData;
+			specifierContext.MetaNameIndex = UhtMetaData.IndexNone;
+			specifierContext.SeenEditSpecifier = false;
+			specifierContext.SeenBlueprintWriteSpecifier = false;
+			specifierContext.SeenBlueprintReadOnlySpecifier = false;
+			specifierContext.SeenBlueprintGetterSpecifier = false;
+
+			// Initialize the settings
 			this._options = options;
-			this._declarationStyle = declarationStyle;
-
+			this._category = category;
+			this._currentTokenReader = topScope.TokenReader;
 			this._currentTypeTokens = new List<UhtToken>();
 			this._currentTemplateDepth = 0;
 
 			using UhtMessageContext tokenContext = new(this);
-			ParseInternal(propertyDelegate);
+			ParseInternal(topScope, specifierContext, propertyDelegate);
 			return this;
 		}
 
@@ -521,7 +520,7 @@ namespace EpicGames.UHT.Parsers
 			get
 			{
 				Stack<object?> extraContext = new(1);
-				extraContext.Push(this._specifierContext.PropertySettings.PropertyCategory.GetHintText());
+				extraContext.Push(this._category.GetHintText());
 				return extraContext;
 			}
 		}
@@ -626,9 +625,27 @@ namespace EpicGames.UHT.Parsers
 
 			// Save the token reader state.  We need this to restore back to the start when invoking the resolve methods.
 			{
+				using UhtThreadBorrower<UhtPropertySpecifierContext> borrower = new();
+				UhtPropertySpecifierContext specifierContext = borrower.Instance;
+				specifierContext.Type = parentPropertySettings.Outer;
+				specifierContext.TokenReader = tokenReader;
+				specifierContext.AccessSpecifier = UhtAccessSpecifier.None;
+				specifierContext.MessageSite = tokenReader;
+				specifierContext.PropertySettings.Reset(parentPropertySettings, paramName.ToString(), tokenReader);
+				specifierContext.MetaData = specifierContext.PropertySettings.MetaData;
+				specifierContext.MetaNameIndex = UhtMetaData.IndexNone;
+				specifierContext.SeenEditSpecifier = false;
+				specifierContext.SeenBlueprintWriteSpecifier = false;
+				specifierContext.SeenBlueprintReadOnlySpecifier = false;
+				specifierContext.SeenBlueprintGetterSpecifier = false;
+
+				// Preparse any UPARAM 
+				PreParseTypeInternal(specifierContext, true);
+
+				// Save the state at this point
 				using UhtTokenSaveState saveState = new(tokenReader);
-				UhtPropertySettings propertySettings = new(parentPropertySettings, paramName.ToString(), tokenReader);
 
+				UhtPropertySettings propertySettings = specifierContext.PropertySettings;
 				UhtPropertyType propertyType = new();
 				while (!tokenReader.IsEOF)
 				{
@@ -758,110 +775,43 @@ namespace EpicGames.UHT.Parsers
 			return outProperty;
 		}
 
-		private void ParseInternal(UhtPropertyDelegate propertyDelegate)
+		private void ParseInternal(UhtParsingScope topScope, UhtPropertySpecifierContext specifierContext, UhtPropertyDelegate propertyDelegate)
 		{
-			this._specifierContext.PropertySettings.LineNumber = this.TokenReader.InputLine;
+			UhtPropertySettings propertySettings = specifierContext.PropertySettings;
+			IUhtTokenReader tokenReader = topScope.TokenReader;
 
-			bool isParamList = this._specifierContext.PropertySettings.PropertyCategory != UhtPropertyCategory.Member && this.TokenReader.TryOptional("UPARAM");
+			propertySettings.LineNumber = tokenReader.InputLine;
 
-			UhtSpecifierParser specifiers = UhtSpecifierParser.GetThreadInstance(this._specifierContext, "Variable",
-				isParamList ? this.TopScope.Session.GetSpecifierTable(UhtTableNames.PropertyArgument) : this.TopScope.Session.GetSpecifierTable(UhtTableNames.PropertyMember));
-			if (_declarationStyle == UhtParsePropertyDeclarationStyle.UPROPERTY || isParamList)
-			{
-				specifiers.ParseSpecifiers();
-			}
+			// Parse type information including UPARAM that might appear in template arguments
+			PreParseTypeInternal(specifierContext, false);
 
-			if (this._specifierContext.PropertySettings.PropertyCategory != UhtPropertyCategory.Member)
-			{
-				// const before the variable type support (only for params)
-				if (this.TokenReader.TryOptional("const"))
-				{
-					this._specifierContext.PropertySettings.PropertyFlags |= EPropertyFlags.ConstParm;
-					this._specifierContext.PropertySettings.MetaData.Add(UhtNames.NativeConst, "");
-				}
-			}
-
-			UhtCompilerDirective compilerDirective = this.TopScope.HeaderParser.GetCurrentCompositeCompilerDirective();
+			UhtCompilerDirective compilerDirective = topScope.HeaderParser.GetCurrentCompositeCompilerDirective();
 			if (compilerDirective.HasAnyFlags(UhtCompilerDirective.WithEditorOnlyData))
 			{
-				this._specifierContext.PropertySettings.PropertyFlags |= EPropertyFlags.EditorOnly;
+				propertySettings.PropertyFlags |= EPropertyFlags.EditorOnly;
 			}
 			else if (compilerDirective.HasAnyFlags(UhtCompilerDirective.WithEditor))
 			{
 				// Checking for this error is a bit tricky given legacy code.  
 				// 1) If already wrapped in WITH_EDITORONLY_DATA (see above), then we ignore the error via the else 
 				// 2) Ignore any module that is an editor module
-				UhtPackage package = this.TopScope.HeaderParser.HeaderFile.Package;
+				UhtPackage package = topScope.ScopeType.HeaderFile.Package;
 				UHTManifest.Module module = package.Module;
 				bool isEditorModule =
 					module.ModuleType == UHTModuleType.EngineEditor ||
 					module.ModuleType == UHTModuleType.GameEditor ||
 					module.ModuleType == UHTModuleType.EngineUncooked ||
 					module.ModuleType == UHTModuleType.GameUncooked;
-				if (this._specifierContext.PropertySettings.PropertyCategory == UhtPropertyCategory.Member && !isEditorModule)
+				if (propertySettings.PropertyCategory == UhtPropertyCategory.Member && !isEditorModule)
 				{
-					this.TokenReader.LogError("UProperties should not be wrapped by WITH_EDITOR, use WITH_EDITORONLY_DATA instead.");
+					tokenReader.LogError("UProperties should not be wrapped by WITH_EDITOR, use WITH_EDITORONLY_DATA instead.");
 				}
 			}
 
-			// Store the start and end positions of the parsed type
-			int startPos = this.TokenReader.InputPos;
-
-			// Process the specifiers
-			if (_declarationStyle == UhtParsePropertyDeclarationStyle.UPROPERTY || isParamList)
-			{
-				specifiers.ParseDeferred();
-			}
-
-			// If we saw a BlueprintGetter but did not see BlueprintSetter or 
-			// or BlueprintReadWrite then treat as BlueprintReadOnly
-			if (this._specifierContext.SeenBlueprintGetterSpecifier && !this._specifierContext.SeenBlueprintWriteSpecifier)
-			{
-				this._specifierContext.PropertySettings.PropertyFlags |= EPropertyFlags.BlueprintReadOnly;
-			}
-
-			if (this._specifierContext.PropertySettings.MetaData.ContainsKey(UhtNames.ExposeOnSpawn))
-			{
-				this._specifierContext.PropertySettings.PropertyFlags |= EPropertyFlags.ExposeOnSpawn;
-			}
-
-			UhtAccessSpecifier accessSpecifier = TopScope.AccessSpecifier;
-			if (accessSpecifier == UhtAccessSpecifier.Public || this._specifierContext.PropertySettings.PropertyCategory != UhtPropertyCategory.Member)
-			{
-				this._specifierContext.PropertySettings.PropertyFlags &= ~EPropertyFlags.Protected;
-				this._specifierContext.PropertySettings.PropertyExportFlags |= UhtPropertyExportFlags.Public;
-				this._specifierContext.PropertySettings.PropertyExportFlags &= ~(UhtPropertyExportFlags.Private | UhtPropertyExportFlags.Protected);
-
-				this._specifierContext.PropertySettings.PropertyFlags &= ~EPropertyFlags.NativeAccessSpecifiers;
-				this._specifierContext.PropertySettings.PropertyFlags |= EPropertyFlags.NativeAccessSpecifierPublic;
-			}
-			else if (accessSpecifier == UhtAccessSpecifier.Protected)
-			{
-				this._specifierContext.PropertySettings.PropertyFlags |= EPropertyFlags.Protected;
-				this._specifierContext.PropertySettings.PropertyExportFlags |= UhtPropertyExportFlags.Protected;
-				this._specifierContext.PropertySettings.PropertyExportFlags &= ~(UhtPropertyExportFlags.Public | UhtPropertyExportFlags.Private);
-
-				this._specifierContext.PropertySettings.PropertyFlags &= ~EPropertyFlags.NativeAccessSpecifiers;
-				this._specifierContext.PropertySettings.PropertyFlags |= EPropertyFlags.NativeAccessSpecifierProtected;
-			}
-			else if (accessSpecifier == UhtAccessSpecifier.Private)
-			{
-				this._specifierContext.PropertySettings.PropertyFlags &= ~EPropertyFlags.Protected;
-				this._specifierContext.PropertySettings.PropertyExportFlags |= UhtPropertyExportFlags.Private;
-				this._specifierContext.PropertySettings.PropertyExportFlags &= ~(UhtPropertyExportFlags.Public | UhtPropertyExportFlags.Protected);
-
-				this._specifierContext.PropertySettings.PropertyFlags &= ~EPropertyFlags.NativeAccessSpecifiers;
-				this._specifierContext.PropertySettings.PropertyFlags |= EPropertyFlags.NativeAccessSpecifierPrivate;
-			}
-			else
-			{
-				throw new UhtIceException("Unknown access level");
-			}
-
 			// Swallow inline keywords
-			if (this._specifierContext.PropertySettings.PropertyCategory == UhtPropertyCategory.Return)
+			if (propertySettings.PropertyCategory == UhtPropertyCategory.Return)
 			{
-				this.TokenReader
+				tokenReader
 					.Optional("inline")
 					.Optional("FORCENOINLINE")
 					.OptionalStartsWith("FORCEINLINE");
@@ -872,17 +822,17 @@ namespace EpicGames.UHT.Parsers
 			UhtLayoutMacroType layoutMacroType = UhtLayoutMacroType.None;
 			if (this._options.HasAnyFlags(UhtPropertyParseOptions.ParseLayoutMacro))
 			{
-				ref UhtToken layoutToken = ref this.TokenReader.PeekToken();
+				ref UhtToken layoutToken = ref tokenReader.PeekToken();
 				if (layoutToken.IsIdentifier())
 				{
 					if (s_layoutMacroTypes.TryGetValue(layoutToken.Value, out layoutMacroType))
 					{
-						this.TokenReader.ConsumeToken();
-						this.TokenReader.Require('(');
-						hasWrapperBrackets = this.TokenReader.TryOptional('(');
+						tokenReader.ConsumeToken();
+						tokenReader.Require('(');
+						hasWrapperBrackets = tokenReader.TryOptional('(');
 						if (layoutMacroType.IsEditorOnly())
 						{
-							this._specifierContext.PropertySettings.PropertyFlags |= EPropertyFlags.EditorOnly;
+							propertySettings.PropertyFlags |= EPropertyFlags.EditorOnly;
 						}
 					}
 				}
@@ -893,159 +843,257 @@ namespace EpicGames.UHT.Parsers
 				//@TODO - deprecate
 				else // if (LayoutToken.IsSymbol(';'))
 				{
-					this.TokenReader.ConsumeToken();
+					tokenReader.ConsumeToken();
 				}
 			}
 
 			//@TODO: Should flag as settable from a const context, but this is at least good enough to allow use for C++ land
-			this.TokenReader.Optional("mutable");
+			tokenReader.Optional("mutable");
 
 			// Gather the type tokens and possibly the property name.
-			this.TokenReader.While(_gatherTypeTokensDelegate);
+			tokenReader.While(_gatherTypeTokensDelegate);
 
 			// Verify we at least have one type
 			if (this._currentTypeTokens.Count < 1)
 			{
-				throw new UhtException(this.TokenReader, $"{this._specifierContext.PropertySettings.PropertyCategory.GetHintText()}: Missing variable type or name");
+				throw new UhtException(tokenReader, $"{propertySettings.PropertyCategory.GetHintText()}: Missing variable type or name");
 			}
 
 			// Consume the wrapper brackets.  This is just an extra set
 			if (hasWrapperBrackets)
 			{
-				this.TokenReader.Require(')');
+				tokenReader.Require(')');
 			}
 
 			// Check for any disallowed flags
-			if (this._specifierContext.PropertySettings.PropertyFlags.HasAnyFlags(this._specifierContext.PropertySettings.DisallowPropertyFlags))
+			if (propertySettings.PropertyFlags.HasAnyFlags(propertySettings.DisallowPropertyFlags))
 			{
-				this.TokenReader.LogError("Specified type modifiers not allowed here");
+				tokenReader.LogError("Specified type modifiers not allowed here");
 			}
 
 			if (this._options.HasAnyFlags(UhtPropertyParseOptions.AddModuleRelativePath))
 			{
-				UhtParsingScope.AddModuleRelativePathToMetaData(this._specifierContext.PropertySettings.MetaData, this._specifierContext.Scope.ScopeType.HeaderFile);
+				UhtParsingScope.AddModuleRelativePathToMetaData(propertySettings.MetaData, topScope.ScopeType.HeaderFile);
 			}
 
 			// Fetch the name of the property, bitfield and array size
 			if (layoutMacroType != UhtLayoutMacroType.None)
 			{
-				this.TokenReader.Require(',');
-				UhtToken nameToken = this.TokenReader.GetIdentifier();
+				tokenReader.Require(',');
+				UhtToken nameToken = tokenReader.GetIdentifier();
 				if (layoutMacroType.IsArray())
 				{
-					this.TokenReader.Require(',');
-					RequireArray(this._specifierContext.PropertySettings, ref nameToken, ')');
+					tokenReader.Require(',');
+					RequireArray(tokenReader, propertySettings, ref nameToken, ')');
 				}
 				else if (layoutMacroType.IsBitfield())
 				{
-					this.TokenReader.Require(',');
-					RequireBitfield(this._specifierContext.PropertySettings, ref nameToken);
+					tokenReader.Require(',');
+					RequireBitfield(tokenReader, propertySettings, ref nameToken);
 				}
 				else if (layoutMacroType.HasInitializer())
 				{
-					this.TokenReader.SkipBrackets('(', ')', 1);
+					tokenReader.SkipBrackets('(', ')', 1);
 				}
-				this.TokenReader.Require(')');
+				tokenReader.Require(')');
 
-				Finalize(ref nameToken, new ReadOnlyMemory<UhtToken>(this._currentTypeTokens.ToArray()), layoutMacroType, propertyDelegate);
+				Finalize(topScope, specifierContext, ref nameToken, new ReadOnlyMemory<UhtToken>(this._currentTypeTokens.ToArray()), layoutMacroType, propertyDelegate);
 			}
 			else if (this._options.HasAnyFlags(UhtPropertyParseOptions.List))
 			{
 				// Extract the property name from the types
 				if (this._currentTypeTokens.Count < 2 || !this._currentTypeTokens[^1].IsIdentifier())
 				{
-					throw new UhtException(this.TokenReader, $"{this._specifierContext.PropertySettings.PropertyCategory.GetHintText()}: Expected name");
+					throw new UhtException(tokenReader, $"{propertySettings.PropertyCategory.GetHintText()}: Expected name");
 				}
 				UhtToken nameToken = this._currentTypeTokens[^1];
 				this._currentTypeTokens.RemoveAt(this._currentTypeTokens.Count - 1);
-				CheckForOptionalParts(this._specifierContext.PropertySettings, ref nameToken);
+				CheckForOptionalParts(tokenReader, propertySettings, ref nameToken);
 
 				ReadOnlyMemory<UhtToken> typeTokens = new(this._currentTypeTokens.ToArray());
 
 				while (true)
 				{
-					UhtProperty _ = Finalize(ref nameToken, typeTokens, layoutMacroType, propertyDelegate);
+					UhtProperty _ = Finalize(topScope, specifierContext, ref nameToken, typeTokens, layoutMacroType, propertyDelegate);
 
 					// If we have reached the end
-					if (!this.TokenReader.TryOptional(','))
+					if (!tokenReader.TryOptional(','))
 					{
 						break;
 					}
 
 					// While we could continue parsing, the old UHT would flag this as an error.
-					throw new UhtException(TopScope.TokenReader, $"Comma delimited properties cannot be converted");
+					throw new UhtException(tokenReader, $"Comma delimited properties cannot be converted");
 
 					// we'll need any metadata tags we parsed later on when we call ConvertEOLCommentToTooltip() so the tags aren't clobbered
 					//this.SpecifierContext.PropertySettings.MetaData = Property.MetaData.Clone();
 					//this.SpecifierContext.PropertySettings.ResetTrailingSettings();
 
 					// Get the next property name
-					//NameToken = this.TokenReader.GetIdentifier();
+					//NameToken = tokenReader.GetIdentifier();
 					//CheckForOptionalParts(this.SpecifierContext.PropertySettings, ref NameToken);
 				}
 			}
 			else if (this._options.HasAnyFlags(UhtPropertyParseOptions.CommaSeparatedName))
 			{
-				this.TokenReader.Require(',');
-				UhtToken nameToken = this.TokenReader.GetIdentifier();
-				CheckForOptionalParts(this._specifierContext.PropertySettings, ref nameToken);
-				Finalize(ref nameToken, new ReadOnlyMemory<UhtToken>(this._currentTypeTokens.ToArray()), layoutMacroType, propertyDelegate);
+				tokenReader.Require(',');
+				UhtToken nameToken = tokenReader.GetIdentifier();
+				CheckForOptionalParts(tokenReader, propertySettings, ref nameToken);
+				Finalize(topScope, specifierContext, ref nameToken, new ReadOnlyMemory<UhtToken>(this._currentTypeTokens.ToArray()), layoutMacroType, propertyDelegate);
 			}
 			else if (this._options.HasAnyFlags(UhtPropertyParseOptions.FunctionNameIncluded))
 			{
 				UhtToken nameToken = this._currentTypeTokens[^1];
 				nameToken.Value = new StringView("Function");
-				if (CheckForOptionalParts(this._specifierContext.PropertySettings, ref nameToken))
+				if (CheckForOptionalParts(tokenReader, propertySettings, ref nameToken))
 				{
-					nameToken = this.TokenReader.GetIdentifier("function name");
+					nameToken = tokenReader.GetIdentifier("function name");
 				}
 				else
 				{
 					if (this._currentTypeTokens.Count < 2 || !this._currentTypeTokens[^1].IsIdentifier())
 					{
-						throw new UhtException(this.TokenReader, $"{this._specifierContext.PropertySettings.PropertyCategory.GetHintText()}: Expected name");
+						throw new UhtException(tokenReader, $"{propertySettings.PropertyCategory.GetHintText()}: Expected name");
 					}
 					nameToken = this._currentTypeTokens[^1];
 					this._currentTypeTokens.RemoveAt(this._currentTypeTokens.Count - 1);
 				}
-				Finalize(ref nameToken, new ReadOnlyMemory<UhtToken>(this._currentTypeTokens.ToArray()), layoutMacroType, propertyDelegate);
+				Finalize(topScope, specifierContext, ref nameToken, new ReadOnlyMemory<UhtToken>(this._currentTypeTokens.ToArray()), layoutMacroType, propertyDelegate);
 			}
 			else if (this._options.HasAnyFlags(UhtPropertyParseOptions.NameIncluded))
 			{
 				if (this._currentTypeTokens.Count < 2 || !this._currentTypeTokens[^1].IsIdentifier())
 				{
-					throw new UhtException(this.TokenReader, $"{this._specifierContext.PropertySettings.PropertyCategory.GetHintText()}: Expected name");
+					throw new UhtException(tokenReader, $"{propertySettings.PropertyCategory.GetHintText()}: Expected name");
 				}
 				UhtToken nameToken = this._currentTypeTokens[^1];
 				this._currentTypeTokens.RemoveAt(this._currentTypeTokens.Count - 1);
-				CheckForOptionalParts(this._specifierContext.PropertySettings, ref nameToken);
-				Finalize(ref nameToken, new ReadOnlyMemory<UhtToken>(this._currentTypeTokens.ToArray()), layoutMacroType, propertyDelegate);
+				CheckForOptionalParts(tokenReader, propertySettings, ref nameToken);
+				Finalize(topScope, specifierContext, ref nameToken, new ReadOnlyMemory<UhtToken>(this._currentTypeTokens.ToArray()), layoutMacroType, propertyDelegate);
 			}
 			else
 			{
 				UhtToken nameToken = new();
-				CheckForOptionalParts(this._specifierContext.PropertySettings, ref nameToken);
-				Finalize(ref nameToken, new ReadOnlyMemory<UhtToken>(this._currentTypeTokens.ToArray()), layoutMacroType, propertyDelegate);
+				CheckForOptionalParts(tokenReader, propertySettings, ref nameToken);
+				Finalize(topScope, specifierContext, ref nameToken, new ReadOnlyMemory<UhtToken>(this._currentTypeTokens.ToArray()), layoutMacroType, propertyDelegate);
+			}
+		}
+
+		/// <summary>
+		/// Parse the type elements excluding the type itself.
+		/// </summary>
+		/// <param name="specifierContext">Context of what is being parsed</param>
+		/// <param name="isTemplateArgument">If true, this is part of a template argument</param>
+		/// <exception cref="UhtIceException"></exception>
+		/// <exception cref="UhtException"></exception>
+		private static void PreParseTypeInternal(UhtPropertySpecifierContext specifierContext, bool isTemplateArgument)
+		{
+			UhtPropertySettings propertySettings = specifierContext.PropertySettings;
+			IUhtTokenReader tokenReader = specifierContext.TokenReader;
+			UhtSession session = specifierContext.Type.Session;
+
+			// We parse specifiers when:
+			//
+			// 1. This is the start of a member property (but not a template)
+			// 2. The UPARAM identifier is found
+			bool isMember = propertySettings.PropertyCategory == UhtPropertyCategory.Member;
+			bool parseSpecifiers = (isMember && !isTemplateArgument) || tokenReader.TryOptional("UPARAM");
+
+			UhtSpecifierParser specifiers = UhtSpecifierParser.GetThreadInstance(specifierContext, "Variable",
+				isMember ? session.GetSpecifierTable(UhtTableNames.PropertyMember) : session.GetSpecifierTable(UhtTableNames.PropertyArgument));
+			if (parseSpecifiers)
+			{
+				specifiers.ParseSpecifiers();
+			}
+
+			if (propertySettings.PropertyCategory != UhtPropertyCategory.Member && !isTemplateArgument)
+			{
+				// const before the variable type support (only for params)
+				if (tokenReader.TryOptional("const"))
+				{
+					propertySettings.PropertyFlags |= EPropertyFlags.ConstParm;
+					propertySettings.MetaData.Add(UhtNames.NativeConst, "");
+				}
+			}
+
+			// Process the specifiers
+			if (parseSpecifiers)
+			{
+				specifiers.ParseDeferred();
+			}
+
+			// If we saw a BlueprintGetter but did not see BlueprintSetter or 
+			// or BlueprintReadWrite then treat as BlueprintReadOnly
+			if (specifierContext.SeenBlueprintGetterSpecifier && !specifierContext.SeenBlueprintWriteSpecifier)
+			{
+				propertySettings.PropertyFlags |= EPropertyFlags.BlueprintReadOnly;
+			}
+
+			if (propertySettings.MetaData.ContainsKey(UhtNames.ExposeOnSpawn))
+			{
+				propertySettings.PropertyFlags |= EPropertyFlags.ExposeOnSpawn;
+			}
+
+			if (!isTemplateArgument)
+			{
+				UhtAccessSpecifier accessSpecifier = specifierContext.AccessSpecifier;
+				if (accessSpecifier == UhtAccessSpecifier.Public || propertySettings.PropertyCategory != UhtPropertyCategory.Member)
+				{
+					propertySettings.PropertyFlags &= ~EPropertyFlags.Protected;
+					propertySettings.PropertyExportFlags |= UhtPropertyExportFlags.Public;
+					propertySettings.PropertyExportFlags &= ~(UhtPropertyExportFlags.Private | UhtPropertyExportFlags.Protected);
+
+					propertySettings.PropertyFlags &= ~EPropertyFlags.NativeAccessSpecifiers;
+					propertySettings.PropertyFlags |= EPropertyFlags.NativeAccessSpecifierPublic;
+				}
+				else if (accessSpecifier == UhtAccessSpecifier.Protected)
+				{
+					propertySettings.PropertyFlags |= EPropertyFlags.Protected;
+					propertySettings.PropertyExportFlags |= UhtPropertyExportFlags.Protected;
+					propertySettings.PropertyExportFlags &= ~(UhtPropertyExportFlags.Public | UhtPropertyExportFlags.Private);
+
+					propertySettings.PropertyFlags &= ~EPropertyFlags.NativeAccessSpecifiers;
+					propertySettings.PropertyFlags |= EPropertyFlags.NativeAccessSpecifierProtected;
+				}
+				else if (accessSpecifier == UhtAccessSpecifier.Private)
+				{
+					propertySettings.PropertyFlags &= ~EPropertyFlags.Protected;
+					propertySettings.PropertyExportFlags |= UhtPropertyExportFlags.Private;
+					propertySettings.PropertyExportFlags &= ~(UhtPropertyExportFlags.Public | UhtPropertyExportFlags.Protected);
+
+					propertySettings.PropertyFlags &= ~EPropertyFlags.NativeAccessSpecifiers;
+					propertySettings.PropertyFlags |= EPropertyFlags.NativeAccessSpecifierPrivate;
+				}
+				else
+				{
+					throw new UhtIceException("Unknown access level");
+				}
 			}
 		}
 
 		/// <summary>
 		/// Finish creating the property
 		/// </summary>
+		/// <param name="topScope">Top most scope being parsed</param>
+		/// <param name="specifierContext">Context of the property being parsed</param>
 		/// <param name="nameToken">The name of the property</param>
 		/// <param name="typeTokens">Series of tokens that represent the type</param>
 		/// <param name="layoutMacroType">Optional layout macro type being parsed</param>
 		/// <param name="propertyDelegate">Delegate to invoke when processing has been completed</param>
 		/// <returns>The newly created property.  During the parsing phase, this will often be a temporary property if the type references engine types.</returns>
-		private UhtProperty Finalize(ref UhtToken nameToken, ReadOnlyMemory<UhtToken> typeTokens, UhtLayoutMacroType layoutMacroType, UhtPropertyDelegate propertyDelegate)
+		private UhtProperty Finalize(UhtParsingScope topScope, UhtPropertySpecifierContext specifierContext, ref UhtToken nameToken, 
+			ReadOnlyMemory<UhtToken> typeTokens, UhtLayoutMacroType layoutMacroType, UhtPropertyDelegate propertyDelegate)
 		{
-			this._specifierContext.PropertySettings.SourceName = this._specifierContext.PropertySettings.PropertyCategory == UhtPropertyCategory.Return ? "ReturnValue" : nameToken.Value.ToString();
+			UhtPropertySettings propertySettings = specifierContext.PropertySettings;
+			IUhtTokenReader tokenReader = specifierContext.TokenReader;
+
+			propertySettings.SourceName = propertySettings.PropertyCategory == UhtPropertyCategory.Return ? "ReturnValue" : nameToken.Value.ToString();
 
 			// Try to resolve the property using any immediate mode property types
-			UhtProperty? newProperty = ResolveProperty(UhtPropertyResolvePhase.Parsing, this._specifierContext.PropertySettings, this._specifierContext.PropertySettings.Outer.HeaderFile.Data.Memory, typeTokens);
+			UhtProperty? newProperty = ResolveProperty(UhtPropertyResolvePhase.Parsing, propertySettings, propertySettings.Outer.HeaderFile.Data.Memory, typeTokens);
 			if (newProperty == null)
 			{
-				newProperty = new UhtPreResolveProperty(this._specifierContext.PropertySettings, typeTokens);
+				newProperty = new UhtPreResolveProperty(propertySettings, typeTokens);
 			}
 
 			// Force the category in non-engine projects
@@ -1068,7 +1116,7 @@ namespace EpicGames.UHT.Parsers
 				{
 					if (deprecatedIndex != newProperty.SourceName.Length - 11)
 					{
-						this.TokenReader.LogError("Deprecated variables must end with _DEPRECATED");
+						tokenReader.LogError("Deprecated variables must end with _DEPRECATED");
 					}
 
 					// We allow deprecated properties in blueprints that have getters and setters assigned as they may be part of a backwards compatibility path
@@ -1078,12 +1126,12 @@ namespace EpicGames.UHT.Parsers
 
 					if (warnOnGetter)
 					{
-						this.TokenReader.LogWarning($"{newProperty.PropertyCategory.GetHintText()}: Deprecated property '{newProperty.SourceName}' should not be marked as blueprint visible without having a BlueprintGetter");
+						tokenReader.LogWarning($"{newProperty.PropertyCategory.GetHintText()}: Deprecated property '{newProperty.SourceName}' should not be marked as blueprint visible without having a BlueprintGetter");
 					}
 
 					if (warnOnSetter)
 					{
-						this.TokenReader.LogWarning($"{newProperty.PropertyCategory.GetHintText()}: Deprecated property '{newProperty.SourceName}' should not be marked as blueprint writable without having a BlueprintSetter");
+						tokenReader.LogWarning($"{newProperty.PropertyCategory.GetHintText()}: Deprecated property '{newProperty.SourceName}' should not be marked as blueprint writable without having a BlueprintSetter");
 					}
 
 					// Warn if a deprecated property is visible
@@ -1091,7 +1139,7 @@ namespace EpicGames.UHT.Parsers
 						(!blueprintVisible && newProperty.PropertyFlags.HasAnyFlags(EPropertyFlags.BlueprintReadOnly) &&
 						!newProperty.PropertyExportFlags.HasAnyFlags(UhtPropertyExportFlags.ImpliedBlueprintPure))) // Is BPRO, but not via Implied Flags and not caught by Getter/Setter path above
 					{
-						this.TokenReader.LogWarning($"{newProperty.PropertyCategory.GetHintText()}: Deprecated property '{newProperty.SourceName}' should not be marked as visible or editable");
+						tokenReader.LogWarning($"{newProperty.PropertyCategory.GetHintText()}: Deprecated property '{newProperty.SourceName}' should not be marked as visible or editable");
 					}
 
 					newProperty.PropertyFlags |= EPropertyFlags.Deprecated;
@@ -1102,62 +1150,64 @@ namespace EpicGames.UHT.Parsers
 			// Try gathering metadata for member fields
 			if (newProperty.PropertyCategory == UhtPropertyCategory.Member)
 			{
-				UhtSpecifierParser specifiers = UhtSpecifierParser.GetThreadInstance(this._specifierContext, newProperty.SourceName, this.TopScope.Session.GetSpecifierTable(UhtTableNames.PropertyMember));
+				UhtSpecifierParser specifiers = UhtSpecifierParser.GetThreadInstance(specifierContext, newProperty.SourceName, 
+					specifierContext.Type.Session.GetSpecifierTable(UhtTableNames.PropertyMember));
 				specifiers.ParseFieldMetaData();
-				this.TopScope.AddFormattedCommentsAsTooltipMetaData(newProperty);
+				topScope.AddFormattedCommentsAsTooltipMetaData(newProperty);
 			}
 
-			propertyDelegate(this.TopScope, newProperty, ref nameToken, layoutMacroType);
+			propertyDelegate(topScope, newProperty, ref nameToken, layoutMacroType);
 
 			// Void properties don't get added when they are the return value
 			if (newProperty.PropertyCategory != UhtPropertyCategory.Return || !this._options.HasAnyFlags(UhtPropertyParseOptions.DontAddReturn))
 			{
-				this.TopScope.ScopeType.AddChild(newProperty);
+				topScope.ScopeType.AddChild(newProperty);
 			}
 			return newProperty;
 		}
 
-		private bool CheckForOptionalParts(UhtPropertySettings propertySettings, ref UhtToken nameToken)
+		private static bool CheckForOptionalParts(IUhtTokenReader tokenReader, UhtPropertySettings propertySettings, ref UhtToken nameToken)
 		{
 			bool gotOptionalParts = false;
 
-			if (this.TokenReader.TryOptional('['))
+			if (tokenReader.TryOptional('['))
 			{
-				RequireArray(propertySettings, ref nameToken, ']');
-				this.TokenReader.Require(']');
+				RequireArray(tokenReader, propertySettings, ref nameToken, ']');
+				tokenReader.Require(']');
 				gotOptionalParts = true;
 			}
 
-			if (this.TokenReader.TryOptional(':'))
+			if (tokenReader.TryOptional(':'))
 			{
-				RequireBitfield(propertySettings, ref nameToken);
+				RequireBitfield(tokenReader, propertySettings, ref nameToken);
 				gotOptionalParts = true;
 			}
 			return gotOptionalParts;
 		}
 
-		private void RequireBitfield(UhtPropertySettings propertySettings, ref UhtToken nameToken)
+		private static void RequireBitfield(IUhtTokenReader tokenReader, UhtPropertySettings propertySettings, ref UhtToken nameToken)
 		{
-			if (!this.TokenReader.TryOptionalConstInt(out int bitfieldSize) || bitfieldSize != 1)
+			if (!tokenReader.TryOptionalConstInt(out int bitfieldSize) || bitfieldSize != 1)
 			{
-				throw new UhtException(this.TokenReader, $"Bad or missing bit field size for '{nameToken.Value}', must be 1.");
+				throw new UhtException(tokenReader, $"Bad or missing bit field size for '{nameToken.Value}', must be 1.");
 			}
 			propertySettings.IsBitfield = true;
 		}
 
-		private void RequireArray(UhtPropertySettings propertySettings, ref UhtToken nameToken, char terminator)
+		private static void RequireArray(IUhtTokenReader tokenReader, UhtPropertySettings propertySettings, ref UhtToken nameToken, char terminator)
 		{
 			// Ignore how the actual array dimensions are actually defined - we'll calculate those with the compiler anyway.
-			propertySettings.ArrayDimensions = this.TokenReader.GetRawString(terminator, UhtRawStringOptions.DontConsumeTerminator).ToString();
+			propertySettings.ArrayDimensions = tokenReader.GetRawString(terminator, UhtRawStringOptions.DontConsumeTerminator).ToString();
 			if (propertySettings.ArrayDimensions.Length == 0)
 			{
-				throw new UhtException(this.TokenReader, $"{propertySettings.PropertyCategory.GetHintText()} {nameToken.Value}: Missing array dimensions or terminating '{terminator}'");
+				throw new UhtException(tokenReader, $"{propertySettings.PropertyCategory.GetHintText()} {nameToken.Value}: Missing array dimensions or terminating '{terminator}'");
 			}
 		}
 
 		private bool GatherTypeTokens(ref UhtToken token)
 		{
-			if (this._currentTemplateDepth == 0 && token.IsSymbol() && (token.IsValue(',') || token.IsValue('(') || token.IsValue(')') || token.IsValue(';') || token.IsValue('[') || token.IsValue(':') || token.IsValue('=') || token.IsValue('{')))
+			if (this._currentTemplateDepth == 0 && token.IsSymbol() && (token.IsValue(',') || token.IsValue('(') || token.IsValue(')') || 
+				token.IsValue(';') || token.IsValue('[') || token.IsValue(':') || token.IsValue('=') || token.IsValue('{')))
 			{
 				return false;
 			}
@@ -1171,7 +1221,7 @@ namespace EpicGames.UHT.Parsers
 			{
 				if (this._currentTemplateDepth == 0)
 				{
-					throw new UhtTokenException(TokenReader, token, "',' or ')'");
+					throw new UhtTokenException(this._currentTokenReader!, token, "',' or ')'");
 				}
 				--this._currentTemplateDepth;
 			}
@@ -1190,7 +1240,7 @@ namespace EpicGames.UHT.Parsers
 		private static UhtParseResult UPROPERTYKeyword(UhtParsingScope topScope, UhtParsingScope actionScope, ref UhtToken token)
 		{
 			UhtPropertyParseOptions options = UhtPropertyParseOptions.ParseLayoutMacro | UhtPropertyParseOptions.List | UhtPropertyParseOptions.AddModuleRelativePath;
-			topScope.HeaderParser.GetCachedPropertyParser(topScope).Parse(EPropertyFlags.ParmFlags, options, UhtParsePropertyDeclarationStyle.UPROPERTY, UhtPropertyCategory.Member, s_propertyDelegate);
+			topScope.HeaderParser.GetCachedPropertyParser().Parse(topScope, EPropertyFlags.ParmFlags, options, UhtPropertyCategory.Member, s_propertyDelegate);
 			topScope.TokenReader.Require(';');
 			return UhtParseResult.Handled;
 		}
diff --git a/Engine/Source/Programs/Shared/EpicGames.UHT/Parsers/UhtSpecifierParser.cs b/Engine/Source/Programs/Shared/EpicGames.UHT/Parsers/UhtSpecifierParser.cs
--- a/Engine/Source/Programs/Shared/EpicGames.UHT/Parsers/UhtSpecifierParser.cs
+++ b/Engine/Source/Programs/Shared/EpicGames.UHT/Parsers/UhtSpecifierParser.cs
@@ -77,7 +77,7 @@ namespace EpicGames.UHT.Parsers
 		public UhtSpecifierParser(UhtSpecifierContext specifierContext, StringView context, UhtSpecifierTable table)
 		{
 			this._specifierContext = specifierContext;
-			this._tokenReader = specifierContext.Scope.TokenReader;
+			this._tokenReader = specifierContext.TokenReader;
 			this._context = context;
 			this._table = table;
 
@@ -96,7 +96,7 @@ namespace EpicGames.UHT.Parsers
 		public void Reset(UhtSpecifierContext specifierContext, StringView context, UhtSpecifierTable table)
 		{
 			this._specifierContext = specifierContext;
-			this._tokenReader = specifierContext.Scope.TokenReader;
+			this._tokenReader = specifierContext.TokenReader;
 			this._context = context;
 			this._table = table;
 			this._deferredSpecifiers.Clear();
@@ -122,7 +122,7 @@ namespace EpicGames.UHT.Parsers
 		/// <returns>Specifier parser</returns>
 		public UhtSpecifierParser ParseFieldMetaData()
 		{
-			this._tokenReader = this._specifierContext.Scope.TokenReader;
+			this._tokenReader = this._specifierContext.TokenReader;
 			this._isParsingFieldMetaData = true;
 
 			using UhtMessageContext tokenContext = new(this);
diff --git a/Engine/Source/Programs/Shared/EpicGames.UHT/Specifiers/UhtClassSpecifiers.cs b/Engine/Source/Programs/Shared/EpicGames.UHT/Specifiers/UhtClassSpecifiers.cs
--- a/Engine/Source/Programs/Shared/EpicGames.UHT/Specifiers/UhtClassSpecifiers.cs
+++ b/Engine/Source/Programs/Shared/EpicGames.UHT/Specifiers/UhtClassSpecifiers.cs
@@ -19,14 +19,14 @@ namespace EpicGames.UHT.Parsers
 		[UhtSpecifier(Extends = UhtTableNames.Class, ValueType = UhtSpecifierValueType.Legacy)]
 		private static void NoExportSpecifier(UhtSpecifierContext specifierContext)
 		{
-			UhtClass classObj = (UhtClass)specifierContext.Scope.ScopeType;
+			UhtClass classObj = (UhtClass)specifierContext.Type;
 			classObj.ClassExportFlags |= UhtClassExportFlags.NoExport;
 		}
 
 		[UhtSpecifier(Extends = UhtTableNames.Class, ValueType = UhtSpecifierValueType.Legacy)]
 		private static void IntrinsicSpecifier(UhtSpecifierContext specifierContext)
 		{
-			UhtClass classObj = (UhtClass)specifierContext.Scope.ScopeType;
+			UhtClass classObj = (UhtClass)specifierContext.Type;
 			classObj.AddClassFlags(EClassFlags.Intrinsic);
 		}
 
@@ -39,49 +39,49 @@ namespace EpicGames.UHT.Parsers
 		[UhtSpecifier(Extends = UhtTableNames.Class, ValueType = UhtSpecifierValueType.SingleString)]
 		private static void WithinSpecifier(UhtSpecifierContext specifierContext, StringView value)
 		{
-			UhtClass classObj = (UhtClass)specifierContext.Scope.ScopeType;
+			UhtClass classObj = (UhtClass)specifierContext.Type;
 			classObj.ClassWithinIdentifier = value.ToString();
 		}
 
 		[UhtSpecifier(Extends = UhtTableNames.Class, ValueType = UhtSpecifierValueType.Legacy)]
 		private static void EditInlineNewSpecifier(UhtSpecifierContext specifierContext)
 		{
-			UhtClass classObj = (UhtClass)specifierContext.Scope.ScopeType;
+			UhtClass classObj = (UhtClass)specifierContext.Type;
 			classObj.AddClassFlags(EClassFlags.EditInlineNew);
 		}
 
 		[UhtSpecifier(Extends = UhtTableNames.Class, ValueType = UhtSpecifierValueType.Legacy)]
 		private static void NotEditInlineNewSpecifier(UhtSpecifierContext specifierContext)
 		{
-			UhtClass classObj = (UhtClass)specifierContext.Scope.ScopeType;
+			UhtClass classObj = (UhtClass)specifierContext.Type;
 			classObj.RemoveClassFlags(EClassFlags.EditInlineNew);
 		}
 
 		[UhtSpecifier(Extends = UhtTableNames.Class, ValueType = UhtSpecifierValueType.Legacy)]
 		private static void PlaceableSpecifier(UhtSpecifierContext specifierContext)
 		{
-			UhtClass classObj = (UhtClass)specifierContext.Scope.ScopeType;
+			UhtClass classObj = (UhtClass)specifierContext.Type;
 			classObj.RemoveClassFlags(EClassFlags.NotPlaceable);
 		}
 
 		[UhtSpecifier(Extends = UhtTableNames.Class, ValueType = UhtSpecifierValueType.Legacy)]
 		private static void NotPlaceableSpecifier(UhtSpecifierContext specifierContext)
 		{
-			UhtClass classObj = (UhtClass)specifierContext.Scope.ScopeType;
+			UhtClass classObj = (UhtClass)specifierContext.Type;
 			classObj.AddClassFlags(EClassFlags.NotPlaceable);
 		}
 
 		[UhtSpecifier(Extends = UhtTableNames.Class, ValueType = UhtSpecifierValueType.Legacy)]
 		private static void DefaultToInstancedSpecifier(UhtSpecifierContext specifierContext)
 		{
-			UhtClass classObj = (UhtClass)specifierContext.Scope.ScopeType;
+			UhtClass classObj = (UhtClass)specifierContext.Type;
 			classObj.AddClassFlags(EClassFlags.DefaultToInstanced);
 		}
 
 		[UhtSpecifier(Extends = UhtTableNames.Class, ValueType = UhtSpecifierValueType.Legacy)]
 		private static void HideDropdownSpecifier(UhtSpecifierContext specifierContext)
 		{
-			UhtClass classObj = (UhtClass)specifierContext.Scope.ScopeType;
+			UhtClass classObj = (UhtClass)specifierContext.Type;
 			classObj.AddClassFlags(EClassFlags.HideDropDown);
 		}
 
@@ -89,7 +89,7 @@ namespace EpicGames.UHT.Parsers
 		private static void HiddenSpecifier(UhtSpecifierContext specifierContext)
 		{
 			// Prevents class from appearing in the editor class browser and edit inline menus.
-			UhtClass classObj = (UhtClass)specifierContext.Scope.ScopeType;
+			UhtClass classObj = (UhtClass)specifierContext.Type;
 			classObj.AddClassFlags(EClassFlags.Hidden);
 		}
 
@@ -102,56 +102,56 @@ namespace EpicGames.UHT.Parsers
 		[UhtSpecifier(Extends = UhtTableNames.Class, ValueType = UhtSpecifierValueType.Legacy)]
 		private static void MinimalAPISpecifier(UhtSpecifierContext specifierContext)
 		{
-			UhtClass classObj = (UhtClass)specifierContext.Scope.ScopeType;
+			UhtClass classObj = (UhtClass)specifierContext.Type;
 			classObj.AddClassFlags(EClassFlags.MinimalAPI);
 		}
 
 		[UhtSpecifier(Extends = UhtTableNames.Class, ValueType = UhtSpecifierValueType.Legacy)]
 		private static void ConstSpecifier(UhtSpecifierContext specifierContext)
 		{
-			UhtClass classObj = (UhtClass)specifierContext.Scope.ScopeType;
+			UhtClass classObj = (UhtClass)specifierContext.Type;
 			classObj.AddClassFlags(EClassFlags.Const);
 		}
 
 		[UhtSpecifier(Extends = UhtTableNames.Class, ValueType = UhtSpecifierValueType.Legacy)]
 		private static void PerObjectConfigSpecifier(UhtSpecifierContext specifierContext)
 		{
-			UhtClass classObj = (UhtClass)specifierContext.Scope.ScopeType;
+			UhtClass classObj = (UhtClass)specifierContext.Type;
 			classObj.AddClassFlags(EClassFlags.PerObjectConfig);
 		}
 
 		[UhtSpecifier(Extends = UhtTableNames.Class, ValueType = UhtSpecifierValueType.Legacy)]
 		private static void ConfigDoNotCheckDefaultsSpecifier(UhtSpecifierContext specifierContext)
 		{
-			UhtClass classObj = (UhtClass)specifierContext.Scope.ScopeType;
+			UhtClass classObj = (UhtClass)specifierContext.Type;
 			classObj.AddClassFlags(EClassFlags.ConfigDoNotCheckDefaults);
 		}
 
 		[UhtSpecifier(Extends = UhtTableNames.Class, ValueType = UhtSpecifierValueType.Legacy)]
 		private static void AbstractSpecifier(UhtSpecifierContext specifierContext)
 		{
-			UhtClass classObj = (UhtClass)specifierContext.Scope.ScopeType;
+			UhtClass classObj = (UhtClass)specifierContext.Type;
 			classObj.AddClassFlags(EClassFlags.Abstract);
 		}
 
 		[UhtSpecifier(Extends = UhtTableNames.Class, ValueType = UhtSpecifierValueType.Legacy)]
 		private static void DeprecatedSpecifier(UhtSpecifierContext specifierContext)
 		{
-			UhtClass classObj = (UhtClass)specifierContext.Scope.ScopeType;
+			UhtClass classObj = (UhtClass)specifierContext.Type;
 			classObj.AddClassFlags(EClassFlags.Deprecated | EClassFlags.NotPlaceable);
 		}
 
 		[UhtSpecifier(Extends = UhtTableNames.Class, ValueType = UhtSpecifierValueType.Legacy)]
 		private static void TransientSpecifier(UhtSpecifierContext specifierContext)
 		{
-			UhtClass classObj = (UhtClass)specifierContext.Scope.ScopeType;
+			UhtClass classObj = (UhtClass)specifierContext.Type;
 			classObj.AddClassFlags(EClassFlags.Transient);
 		}
 
 		[UhtSpecifier(Extends = UhtTableNames.Class, ValueType = UhtSpecifierValueType.Legacy)]
 		private static void NonTransientSpecifier(UhtSpecifierContext specifierContext)
 		{
-			UhtClass classObj = (UhtClass)specifierContext.Scope.ScopeType;
+			UhtClass classObj = (UhtClass)specifierContext.Type;
 			classObj.RemoveClassFlags(EClassFlags.Transient);
 		}
 
@@ -159,7 +159,7 @@ namespace EpicGames.UHT.Parsers
 		private static void OptionalSpecifier(UhtSpecifierContext specifierContext)
 		{
 			// Optional class
-			UhtClass classObj = (UhtClass)specifierContext.Scope.ScopeType;
+			UhtClass classObj = (UhtClass)specifierContext.Type;
 			classObj.AddClassFlags(EClassFlags.Optional);
 		}
 
@@ -167,14 +167,14 @@ namespace EpicGames.UHT.Parsers
 		private static void CustomConstructorSpecifier(UhtSpecifierContext specifierContext)
 		{
 			// we will not export a constructor for this class, assuming it is in the CPP block
-			UhtClass classObj = (UhtClass)specifierContext.Scope.ScopeType;
+			UhtClass classObj = (UhtClass)specifierContext.Type;
 			classObj.ClassExportFlags |= UhtClassExportFlags.HasCustomConstructor;
 		}
 
 		[UhtSpecifier(Extends = UhtTableNames.Class, ValueType = UhtSpecifierValueType.SingleString)]
 		private static void ConfigSpecifier(UhtSpecifierContext specifierContext, StringView value)
 		{
-			UhtClass classObj = (UhtClass)specifierContext.Scope.ScopeType;
+			UhtClass classObj = (UhtClass)specifierContext.Type;
 			classObj.Config = value.ToString();
 		}
 
@@ -182,7 +182,7 @@ namespace EpicGames.UHT.Parsers
 		private static void DefaultConfigSpecifier(UhtSpecifierContext specifierContext)
 		{
 			// Save object config only to Default INIs, never to local INIs.
-			UhtClass classObj = (UhtClass)specifierContext.Scope.ScopeType;
+			UhtClass classObj = (UhtClass)specifierContext.Type;
 			classObj.AddClassFlags(EClassFlags.DefaultConfig);
 		}
 
@@ -190,7 +190,7 @@ namespace EpicGames.UHT.Parsers
 		private static void GlobalUserConfigSpecifier(UhtSpecifierContext specifierContext)
 		{
 			// Save object config only to global user overrides, never to local INIs
-			UhtClass classObj = (UhtClass)specifierContext.Scope.ScopeType;
+			UhtClass classObj = (UhtClass)specifierContext.Type;
 			classObj.AddClassFlags(EClassFlags.GlobalUserConfig);
 		}
 
@@ -198,7 +198,7 @@ namespace EpicGames.UHT.Parsers
 		private static void ProjectUserConfigSpecifier(UhtSpecifierContext specifierContext)
 		{
 			// Save object config only to project user overrides, never to INIs that are checked in
-			UhtClass classObj = (UhtClass)specifierContext.Scope.ScopeType;
+			UhtClass classObj = (UhtClass)specifierContext.Type;
 			classObj.AddClassFlags(EClassFlags.ProjectUserConfig);
 		}
 
@@ -211,42 +211,42 @@ namespace EpicGames.UHT.Parsers
 		[UhtSpecifier(Extends = UhtTableNames.Class, ValueType = UhtSpecifierValueType.NonEmptyStringList)]
 		private static void ShowCategoriesSpecifier(UhtSpecifierContext specifierContext, List<StringView> value)
 		{
-			UhtClass classObj = (UhtClass)specifierContext.Scope.ScopeType;
+			UhtClass classObj = (UhtClass)specifierContext.Type;
 			classObj.ShowCategories.AddUniqueRange(value);
 		}
 
 		[UhtSpecifier(Extends = UhtTableNames.Class, ValueType = UhtSpecifierValueType.NonEmptyStringList)]
 		private static void HideCategoriesSpecifier(UhtSpecifierContext specifierContext, List<StringView> value)
 		{
-			UhtClass classObj = (UhtClass)specifierContext.Scope.ScopeType;
+			UhtClass classObj = (UhtClass)specifierContext.Type;
 			classObj.HideCategories.AddUniqueRange(value);
 		}
 
 		[UhtSpecifier(Extends = UhtTableNames.Class, ValueType = UhtSpecifierValueType.NonEmptyStringList)]
 		private static void ShowFunctionsSpecifier(UhtSpecifierContext specifierContext, List<StringView> value)
 		{
-			UhtClass classObj = (UhtClass)specifierContext.Scope.ScopeType;
+			UhtClass classObj = (UhtClass)specifierContext.Type;
 			classObj.ShowFunctions.AddUniqueRange(value);
 		}
 
 		[UhtSpecifier(Extends = UhtTableNames.Class, ValueType = UhtSpecifierValueType.NonEmptyStringList)]
 		private static void HideFunctionsSpecifier(UhtSpecifierContext specifierContext, List<StringView> value)
 		{
-			UhtClass classObj = (UhtClass)specifierContext.Scope.ScopeType;
+			UhtClass classObj = (UhtClass)specifierContext.Type;
 			classObj.HideFunctions.AddUniqueRange(value);
 		}
 
 		[UhtSpecifier(Extends = UhtTableNames.Class, ValueType = UhtSpecifierValueType.SingleString)]
 		private static void SparseClassDataTypesSpecifier(UhtSpecifierContext specifierContext, StringView value)
 		{
-			UhtClass classObj = (UhtClass)specifierContext.Scope.ScopeType;
+			UhtClass classObj = (UhtClass)specifierContext.Type;
 			classObj.SparseClassDataTypes.AddUnique(value.ToString());
 		}
 
 		[UhtSpecifier(Extends = UhtTableNames.Class, ValueType = UhtSpecifierValueType.NonEmptyStringList)]
 		private static void ClassGroupSpecifier(UhtSpecifierContext specifierContext, List<StringView> value)
 		{
-			UhtClass classObj = (UhtClass)specifierContext.Scope.ScopeType;
+			UhtClass classObj = (UhtClass)specifierContext.Type;
 			foreach (StringView element in value)
 			{
 				classObj.ClassGroupNames.Add(element.ToString());
@@ -256,7 +256,7 @@ namespace EpicGames.UHT.Parsers
 		[UhtSpecifier(Extends = UhtTableNames.Class, ValueType = UhtSpecifierValueType.NonEmptyStringList)]
 		private static void AutoExpandCategoriesSpecifier(UhtSpecifierContext specifierContext, List<StringView> value)
 		{
-			UhtClass classObj = (UhtClass)specifierContext.Scope.ScopeType;
+			UhtClass classObj = (UhtClass)specifierContext.Type;
 			classObj.AutoExpandCategories.AddUniqueRange(value);
 			classObj.AutoCollapseCategories.RemoveSwapRange(value);
 		}
@@ -264,7 +264,7 @@ namespace EpicGames.UHT.Parsers
 		[UhtSpecifier(Extends = UhtTableNames.Class, ValueType = UhtSpecifierValueType.NonEmptyStringList)]
 		private static void AutoCollapseCategoriesSpecifier(UhtSpecifierContext specifierContext, List<StringView> value)
 		{
-			UhtClass classObj = (UhtClass)specifierContext.Scope.ScopeType;
+			UhtClass classObj = (UhtClass)specifierContext.Type;
 			classObj.AutoCollapseCategories.AddUniqueRange(value);
 			classObj.AutoExpandCategories.RemoveSwapRange(value);
 		}
@@ -272,14 +272,14 @@ namespace EpicGames.UHT.Parsers
 		[UhtSpecifier(Extends = UhtTableNames.Class, ValueType = UhtSpecifierValueType.NonEmptyStringList)]
 		private static void PrioritizeCategoriesSpecifier(UhtSpecifierContext specifierContext, List<StringView> value)
 		{
-			UhtClass classObj = (UhtClass)specifierContext.Scope.ScopeType;
+			UhtClass classObj = (UhtClass)specifierContext.Type;
 			classObj.PrioritizeCategories.AddUniqueRange(value);
 		}
 
 		[UhtSpecifier(Extends = UhtTableNames.Class, ValueType = UhtSpecifierValueType.NonEmptyStringList)]
 		private static void DontAutoCollapseCategoriesSpecifier(UhtSpecifierContext specifierContext, List<StringView> value)
 		{
-			UhtClass classObj = (UhtClass)specifierContext.Scope.ScopeType;
+			UhtClass classObj = (UhtClass)specifierContext.Type;
 			classObj.AutoCollapseCategories.RemoveSwapRange(value);
 		}
 
@@ -287,7 +287,7 @@ namespace EpicGames.UHT.Parsers
 		private static void CollapseCategoriesSpecifier(UhtSpecifierContext specifierContext)
 		{
 			// Class' properties should not be shown categorized in the editor.
-			UhtClass classObj = (UhtClass)specifierContext.Scope.ScopeType;
+			UhtClass classObj = (UhtClass)specifierContext.Type;
 			classObj.AddClassFlags(EClassFlags.CollapseCategories);
 		}
 
@@ -295,7 +295,7 @@ namespace EpicGames.UHT.Parsers
 		private static void DontCollapseCategoriesSpecifier(UhtSpecifierContext specifierContext)
 		{
 			// Class' properties should be shown categorized in the editor.
-			UhtClass classObj = (UhtClass)specifierContext.Scope.ScopeType;
+			UhtClass classObj = (UhtClass)specifierContext.Type;
 			classObj.RemoveClassFlags(EClassFlags.CollapseCategories);
 		}
 
@@ -303,50 +303,50 @@ namespace EpicGames.UHT.Parsers
 		private static void AdvancedClassDisplaySpecifier(UhtSpecifierContext specifierContext)
 		{
 			// By default the class properties are shown in advanced sections in UI
-			UhtClass classObj = (UhtClass)specifierContext.Scope.ScopeType;
+			UhtClass classObj = (UhtClass)specifierContext.Type;
 			classObj.MetaData.Add(UhtNames.AdvancedClassDisplay, "true");
 		}
 
 		[UhtSpecifier(Extends = UhtTableNames.Class, ValueType = UhtSpecifierValueType.Legacy)]
 		private static void ConversionRootSpecifier(UhtSpecifierContext specifierContext)
 		{
-			UhtClass classObj = (UhtClass)specifierContext.Scope.ScopeType;
+			UhtClass classObj = (UhtClass)specifierContext.Type;
 			classObj.MetaData.Add(UhtNames.IsConversionRoot, "true");
 		}
 
 		[UhtSpecifier(Extends = UhtTableNames.Class, ValueType = UhtSpecifierValueType.Legacy)]
 		private static void NeedsDeferredDependencyLoadingSpecifier(UhtSpecifierContext specifierContext)
 		{
-			UhtClass classObj = (UhtClass)specifierContext.Scope.ScopeType;
+			UhtClass classObj = (UhtClass)specifierContext.Type;
 			classObj.AddClassFlags(EClassFlags.NeedsDeferredDependencyLoading);
 		}
 
 		[UhtSpecifier(Extends = UhtTableNames.Class, ValueType = UhtSpecifierValueType.Legacy)]
 		private static void MatchedSerializersSpecifier(UhtSpecifierContext specifierContext)
 		{
-			if (!specifierContext.Scope.HeaderParser.HeaderFile.IsNoExportTypes)
+			if (!specifierContext.Type.HeaderFile.IsNoExportTypes)
 			{
 				specifierContext.MessageSite.LogError("The 'MatchedSerializers' class specifier is only valid in the NoExportTypes.h file");
 			}
-			UhtClass classObj = (UhtClass)specifierContext.Scope.ScopeType;
+			UhtClass classObj = (UhtClass)specifierContext.Type;
 			classObj.AddClassFlags(EClassFlags.MatchedSerializers);
 		}
 
 		[UhtSpecifier(Extends = UhtTableNames.Class, ValueType = UhtSpecifierValueType.Legacy)]
 		private static void InterfaceSpecifier(UhtSpecifierContext specifierContext)
 		{
-			if (!specifierContext.Scope.HeaderParser.HeaderFile.IsNoExportTypes)
+			if (!specifierContext.Type.HeaderFile.IsNoExportTypes)
 			{
 				specifierContext.MessageSite.LogError("The 'Interface' class specifier is only valid in the NoExportTypes.h file");
 			}
-			UhtClass classObj = (UhtClass)specifierContext.Scope.ScopeType;
+			UhtClass classObj = (UhtClass)specifierContext.Type;
 			classObj.AddClassFlags(EClassFlags.Interface);
 		}
 
 		[UhtSpecifier(Extends = UhtTableNames.Class, ValueType = UhtSpecifierValueType.None)]
 		private static void CustomFieldNotifySpecifier(UhtSpecifierContext specifierContext)
 		{
-			UhtClass classObj = (UhtClass)specifierContext.Scope.ScopeType;
+			UhtClass classObj = (UhtClass)specifierContext.Type;
 			classObj.ClassExportFlags |= UhtClassExportFlags.HasCustomFieldNotify;
 		}
 	}
diff --git a/Engine/Source/Programs/Shared/EpicGames.UHT/Specifiers/UhtEnumSpecifiers.cs b/Engine/Source/Programs/Shared/EpicGames.UHT/Specifiers/UhtEnumSpecifiers.cs
--- a/Engine/Source/Programs/Shared/EpicGames.UHT/Specifiers/UhtEnumSpecifiers.cs
+++ b/Engine/Source/Programs/Shared/EpicGames.UHT/Specifiers/UhtEnumSpecifiers.cs
@@ -19,7 +19,7 @@ namespace EpicGames.UHT.Parsers
 		[UhtSpecifier(Extends = UhtTableNames.Enum, ValueType = UhtSpecifierValueType.Legacy)]
 		private static void FlagsSpecifier(UhtSpecifierContext specifierContext)
 		{
-			UhtEnum enumObj = (UhtEnum)specifierContext.Scope.ScopeType;
+			UhtEnum enumObj = (UhtEnum)specifierContext.Type;
 			enumObj.EnumFlags |= EEnumFlags.Flags;
 		}
 	}
diff --git a/Engine/Source/Programs/Shared/EpicGames.UHT/Specifiers/UhtFunctionSpecifiers.cs b/Engine/Source/Programs/Shared/EpicGames.UHT/Specifiers/UhtFunctionSpecifiers.cs
--- a/Engine/Source/Programs/Shared/EpicGames.UHT/Specifiers/UhtFunctionSpecifiers.cs
+++ b/Engine/Source/Programs/Shared/EpicGames.UHT/Specifiers/UhtFunctionSpecifiers.cs
@@ -22,28 +22,28 @@ namespace EpicGames.UHT.Parsers
 		[UhtSpecifier(Extends = UhtTableNames.Function, ValueType = UhtSpecifierValueType.Legacy)]
 		private static void BlueprintAuthorityOnlySpecifier(UhtSpecifierContext specifierContext)
 		{
-			UhtFunction function = (UhtFunction)specifierContext.Scope.ScopeType;
+			UhtFunction function = (UhtFunction)specifierContext.Type;
 			function.FunctionFlags |= EFunctionFlags.BlueprintAuthorityOnly;
 		}
 
 		[UhtSpecifier(Extends = UhtTableNames.Function, ValueType = UhtSpecifierValueType.Legacy)]
 		private static void BlueprintCallableSpecifier(UhtSpecifierContext specifierContext)
 		{
-			UhtFunction function = (UhtFunction)specifierContext.Scope.ScopeType;
+			UhtFunction function = (UhtFunction)specifierContext.Type;
 			function.FunctionFlags |= EFunctionFlags.BlueprintCallable;
 		}
 
 		[UhtSpecifier(Extends = UhtTableNames.Function, ValueType = UhtSpecifierValueType.Legacy)]
 		private static void BlueprintCosmeticSpecifier(UhtSpecifierContext specifierContext)
 		{
-			UhtFunction function = (UhtFunction)specifierContext.Scope.ScopeType;
+			UhtFunction function = (UhtFunction)specifierContext.Type;
 			function.FunctionFlags |= EFunctionFlags.BlueprintCosmetic;
 		}
 
 		[UhtSpecifier(Extends = UhtTableNames.Function, ValueType = UhtSpecifierValueType.Legacy)]
 		private static void BlueprintGetterSpecifier(UhtSpecifierContext specifierContext)
 		{
-			UhtFunction function = (UhtFunction)specifierContext.Scope.ScopeType;
+			UhtFunction function = (UhtFunction)specifierContext.Type;
 
 			if (function.FunctionFlags.HasAnyFlags(EFunctionFlags.Event))
 			{
@@ -59,7 +59,7 @@ namespace EpicGames.UHT.Parsers
 		[UhtSpecifier(Extends = UhtTableNames.Function, ValueType = UhtSpecifierValueType.Legacy)]
 		private static void BlueprintImplementableEventSpecifier(UhtSpecifierContext specifierContext)
 		{
-			UhtFunction function = (UhtFunction)specifierContext.Scope.ScopeType;
+			UhtFunction function = (UhtFunction)specifierContext.Type;
 
 			if (function.FunctionFlags.HasAnyFlags(EFunctionFlags.Net))
 			{
@@ -87,7 +87,7 @@ namespace EpicGames.UHT.Parsers
 		[UhtSpecifier(Extends = UhtTableNames.Function, ValueType = UhtSpecifierValueType.Legacy)]
 		private static void BlueprintNativeEventSpecifier(UhtSpecifierContext specifierContext)
 		{
-			UhtFunction function = (UhtFunction)specifierContext.Scope.ScopeType;
+			UhtFunction function = (UhtFunction)specifierContext.Type;
 
 			if (function.FunctionFlags.HasAnyFlags(EFunctionFlags.Net))
 			{
@@ -114,7 +114,7 @@ namespace EpicGames.UHT.Parsers
 		[UhtSpecifier(Extends = UhtTableNames.Function, ValueType = UhtSpecifierValueType.OptionalString)]
 		private static void BlueprintPureSpecifier(UhtSpecifierContext specifierContext, StringView? value)
 		{
-			UhtFunction function = (UhtFunction)specifierContext.Scope.ScopeType;
+			UhtFunction function = (UhtFunction)specifierContext.Type;
 
 			// This function can be called, and is also pure.
 			function.FunctionFlags |= EFunctionFlags.BlueprintCallable;
@@ -132,7 +132,7 @@ namespace EpicGames.UHT.Parsers
 		[UhtSpecifier(Extends = UhtTableNames.Function, ValueType = UhtSpecifierValueType.Legacy)]
 		private static void BlueprintSetterSpecifier(UhtSpecifierContext specifierContext)
 		{
-			UhtFunction function = (UhtFunction)specifierContext.Scope.ScopeType;
+			UhtFunction function = (UhtFunction)specifierContext.Type;
 
 			if (function.FunctionFlags.HasAnyFlags(EFunctionFlags.Event))
 			{
@@ -147,7 +147,7 @@ namespace EpicGames.UHT.Parsers
 		[UhtSpecifier(Extends = UhtTableNames.Function, ValueType = UhtSpecifierValueType.OptionalString)]
 		private static void ClientSpecifier(UhtSpecifierContext specifierContext, StringView? value)
 		{
-			UhtFunction function = (UhtFunction)specifierContext.Scope.ScopeType;
+			UhtFunction function = (UhtFunction)specifierContext.Type;
 
 			if (function.FunctionFlags.HasAnyFlags(EFunctionFlags.BlueprintEvent))
 			{
@@ -171,14 +171,14 @@ namespace EpicGames.UHT.Parsers
 		[UhtSpecifier(Extends = UhtTableNames.Function, ValueType = UhtSpecifierValueType.Legacy)]
 		private static void CustomThunkSpecifier(UhtSpecifierContext specifierContext)
 		{
-			UhtFunction function = (UhtFunction)specifierContext.Scope.ScopeType;
+			UhtFunction function = (UhtFunction)specifierContext.Type;
 			function.FunctionExportFlags |= UhtFunctionExportFlags.CustomThunk;
 		}
 
 		[UhtSpecifier(Extends = UhtTableNames.Function, ValueType = UhtSpecifierValueType.Legacy)]
 		private static void ExecSpecifier(UhtSpecifierContext specifierContext)
 		{
-			UhtFunction function = (UhtFunction)specifierContext.Scope.ScopeType;
+			UhtFunction function = (UhtFunction)specifierContext.Type;
 			function.FunctionFlags |= EFunctionFlags.Exec;
 			if (function.FunctionFlags.HasAnyFlags(EFunctionFlags.Net))
 			{
@@ -189,7 +189,7 @@ namespace EpicGames.UHT.Parsers
 		[UhtSpecifier(Extends = UhtTableNames.Function, ValueType = UhtSpecifierValueType.Legacy)]
 		private static void NetMulticastSpecifier(UhtSpecifierContext specifierContext)
 		{
-			UhtFunction function = (UhtFunction)specifierContext.Scope.ScopeType;
+			UhtFunction function = (UhtFunction)specifierContext.Type;
 			if (function.FunctionFlags.HasAnyFlags(EFunctionFlags.BlueprintEvent))
 			{
 				specifierContext.MessageSite.LogError("BlueprintImplementableEvent or BlueprintNativeEvent functions cannot be declared as Multicast");
@@ -202,21 +202,21 @@ namespace EpicGames.UHT.Parsers
 		[UhtSpecifier(Extends = UhtTableNames.Function, ValueType = UhtSpecifierValueType.Legacy)]
 		private static void ReliableSpecifier(UhtSpecifierContext specifierContext)
 		{
-			UhtFunction function = (UhtFunction)specifierContext.Scope.ScopeType;
+			UhtFunction function = (UhtFunction)specifierContext.Type;
 			function.FunctionFlags |= EFunctionFlags.NetReliable;
 		}
 
 		[UhtSpecifier(Extends = UhtTableNames.Function, ValueType = UhtSpecifierValueType.Legacy)]
 		private static void SealedEventSpecifier(UhtSpecifierContext specifierContext)
 		{
-			UhtFunction function = (UhtFunction)specifierContext.Scope.ScopeType;
+			UhtFunction function = (UhtFunction)specifierContext.Type;
 			function.FunctionExportFlags |= UhtFunctionExportFlags.SealedEvent;
 		}
 
 		[UhtSpecifier(Extends = UhtTableNames.Function, ValueType = UhtSpecifierValueType.OptionalString)]
 		private static void ServerSpecifier(UhtSpecifierContext specifierContext, StringView? value)
 		{
-			UhtFunction function = (UhtFunction)specifierContext.Scope.ScopeType;
+			UhtFunction function = (UhtFunction)specifierContext.Type;
 
 			if (function.FunctionFlags.HasAnyFlags(EFunctionFlags.BlueprintEvent))
 			{
@@ -240,7 +240,7 @@ namespace EpicGames.UHT.Parsers
 		[UhtSpecifier(Extends = UhtTableNames.Function, ValueType = UhtSpecifierValueType.OptionalEqualsKeyValuePairList)]
 		private static void ServiceRequestSpecifier(UhtSpecifierContext specifierContext, List<KeyValuePair<StringView, StringView>> value)
 		{
-			UhtFunction function = (UhtFunction)specifierContext.Scope.ScopeType;
+			UhtFunction function = (UhtFunction)specifierContext.Type;
 
 			if (function.FunctionFlags.HasAnyFlags(EFunctionFlags.BlueprintEvent))
 			{
@@ -263,7 +263,7 @@ namespace EpicGames.UHT.Parsers
 		[UhtSpecifier(Extends = UhtTableNames.Function, ValueType = UhtSpecifierValueType.OptionalEqualsKeyValuePairList)]
 		private static void ServiceResponseSpecifier(UhtSpecifierContext specifierContext, List<KeyValuePair<StringView, StringView>> value)
 		{
-			UhtFunction function = (UhtFunction)specifierContext.Scope.ScopeType;
+			UhtFunction function = (UhtFunction)specifierContext.Type;
 
 			if (function.FunctionFlags.HasAnyFlags(EFunctionFlags.BlueprintEvent))
 			{
@@ -286,14 +286,14 @@ namespace EpicGames.UHT.Parsers
 		[UhtSpecifier(Extends = UhtTableNames.Function, ValueType = UhtSpecifierValueType.Legacy)]
 		private static void UnreliableSpecifier(UhtSpecifierContext specifierContext)
 		{
-			UhtFunction function = (UhtFunction)specifierContext.Scope.ScopeType;
+			UhtFunction function = (UhtFunction)specifierContext.Type;
 			function.FunctionExportFlags |= UhtFunctionExportFlags.Unreliable;
 		}
 
 		[UhtSpecifier(Extends = UhtTableNames.Function, ValueType = UhtSpecifierValueType.StringList)]
 		private static void WithValidationSpecifier(UhtSpecifierContext specifierContext, List<StringView>? value)
 		{
-			UhtFunction function = (UhtFunction)specifierContext.Scope.ScopeType;
+			UhtFunction function = (UhtFunction)specifierContext.Type;
 			function.FunctionFlags |= EFunctionFlags.NetValidate;
 
 			if (value != null && value.Count > 0)
@@ -305,7 +305,7 @@ namespace EpicGames.UHT.Parsers
 		[UhtSpecifier(Extends = UhtTableNames.Function, ValueType = UhtSpecifierValueType.None)]
 		private static void FieldNotifySpecifier(UhtSpecifierContext specifierContext)
 		{
-			UhtFunction function = (UhtFunction)specifierContext.Scope.ScopeType;
+			UhtFunction function = (UhtFunction)specifierContext.Type;
 			function.FunctionExportFlags |= UhtFunctionExportFlags.FieldNotify;
 		}
 
@@ -389,8 +389,8 @@ namespace EpicGames.UHT.Parsers
 
 		private static void ParseNetServiceIdentifiers(UhtSpecifierContext specifierContext, List<KeyValuePair<StringView, StringView>> identifiers)
 		{
-			IUhtTokenReader tokenReader = specifierContext.Scope.TokenReader;
-			UhtFunction function = (UhtFunction)specifierContext.Scope.ScopeType;
+			IUhtTokenReader tokenReader = specifierContext.TokenReader;
+			UhtFunction function = (UhtFunction)specifierContext.Type;
 
 			foreach (KeyValuePair<StringView, StringView> kvp in identifiers)
 			{
diff --git a/Engine/Source/Programs/Shared/EpicGames.UHT/Specifiers/UhtInterfaceClassSpecifiers.cs b/Engine/Source/Programs/Shared/EpicGames.UHT/Specifiers/UhtInterfaceClassSpecifiers.cs
--- a/Engine/Source/Programs/Shared/EpicGames.UHT/Specifiers/UhtInterfaceClassSpecifiers.cs
+++ b/Engine/Source/Programs/Shared/EpicGames.UHT/Specifiers/UhtInterfaceClassSpecifiers.cs
@@ -21,14 +21,14 @@ namespace EpicGames.UHT.Parsers
 		[UhtSpecifier(Extends = UhtTableNames.Interface, ValueType = UhtSpecifierValueType.Legacy)]
 		private static void MinimalAPISpecifier(UhtSpecifierContext specifierContext)
 		{
-			UhtClass clasObj = (UhtClass)specifierContext.Scope.ScopeType;
+			UhtClass clasObj = (UhtClass)specifierContext.Type;
 			clasObj.ClassFlags |= EClassFlags.MinimalAPI;
 		}
 
 		[UhtSpecifier(Extends = UhtTableNames.Interface, ValueType = UhtSpecifierValueType.Legacy)]
 		private static void ConversionRootSpecifier(UhtSpecifierContext specifierContext)
 		{
-			UhtClass classObj = (UhtClass)specifierContext.Scope.ScopeType;
+			UhtClass classObj = (UhtClass)specifierContext.Type;
 			classObj.MetaData.Add(UhtNames.IsConversionRoot, "true");
 		}
 	}
diff --git a/Engine/Source/Programs/Shared/EpicGames.UHT/Specifiers/UhtPropertyMemberSpecifiers.cs b/Engine/Source/Programs/Shared/EpicGames.UHT/Specifiers/UhtPropertyMemberSpecifiers.cs
--- a/Engine/Source/Programs/Shared/EpicGames.UHT/Specifiers/UhtPropertyMemberSpecifiers.cs
+++ b/Engine/Source/Programs/Shared/EpicGames.UHT/Specifiers/UhtPropertyMemberSpecifiers.cs
@@ -99,7 +99,7 @@ namespace EpicGames.UHT.Parsers
 			}
 
 			bool allowPrivateAccess = context.MetaData.TryGetValue(UhtNames.AllowPrivateAccess, out string? privateAccessMD) && !privateAccessMD.Equals("false", StringComparison.OrdinalIgnoreCase);
-			if (specifierContext.Scope.AccessSpecifier == UhtAccessSpecifier.Private && !allowPrivateAccess)
+			if (specifierContext.AccessSpecifier == UhtAccessSpecifier.Private && !allowPrivateAccess)
 			{
 				context.MessageSite.LogError("BlueprintReadWrite should not be used on private members");
 			}
@@ -123,7 +123,7 @@ namespace EpicGames.UHT.Parsers
 			}
 
 			bool allowPrivateAccess = context.MetaData.TryGetValue(UhtNames.AllowPrivateAccess, out string? privateAccessMD) && !privateAccessMD.Equals("false", StringComparison.OrdinalIgnoreCase);
-			if (specifierContext.Scope.AccessSpecifier == UhtAccessSpecifier.Private && !allowPrivateAccess)
+			if (specifierContext.AccessSpecifier == UhtAccessSpecifier.Private && !allowPrivateAccess)
 			{
 				context.MessageSite.LogError("BlueprintReadOnly should not be used on private members");
 			}
diff --git a/Engine/Source/Programs/Shared/EpicGames.UHT/Specifiers/UhtScriptStructSpecifiers.cs b/Engine/Source/Programs/Shared/EpicGames.UHT/Specifiers/UhtScriptStructSpecifiers.cs
--- a/Engine/Source/Programs/Shared/EpicGames.UHT/Specifiers/UhtScriptStructSpecifiers.cs
+++ b/Engine/Source/Programs/Shared/EpicGames.UHT/Specifiers/UhtScriptStructSpecifiers.cs
@@ -18,7 +18,7 @@ namespace EpicGames.UHT.Parsers
 		[UhtSpecifier(Extends = UhtTableNames.ScriptStruct, ValueType = UhtSpecifierValueType.Legacy)]
 		private static void NoExportSpecifier(UhtSpecifierContext specifierContext)
 		{
-			UhtScriptStruct scriptStruct = (UhtScriptStruct)specifierContext.Scope.ScopeType;
+			UhtScriptStruct scriptStruct = (UhtScriptStruct)specifierContext.Type;
 			scriptStruct.ScriptStructFlags |= EStructFlags.NoExport;
 			scriptStruct.ScriptStructFlags &= ~EStructFlags.Native;
 		}
@@ -26,7 +26,7 @@ namespace EpicGames.UHT.Parsers
 		[UhtSpecifier(Extends = UhtTableNames.ScriptStruct, ValueType = UhtSpecifierValueType.Legacy)]
 		private static void AtomicSpecifier(UhtSpecifierContext specifierContext)
 		{
-			UhtScriptStruct scriptStruct = (UhtScriptStruct)specifierContext.Scope.ScopeType;
+			UhtScriptStruct scriptStruct = (UhtScriptStruct)specifierContext.Type;
 			scriptStruct.ScriptStructFlags |= EStructFlags.Atomic;
 		}
 
@@ -34,51 +34,51 @@ namespace EpicGames.UHT.Parsers
 		[SuppressMessage("CodeQuality", "IDE0051:Remove unused private members", Justification = "Attribute accessed method")]
 		private static void ImmutableSpecifier(UhtSpecifierContext specifierContext)
 		{
-			UhtScriptStruct scriptStruct = (UhtScriptStruct)specifierContext.Scope.ScopeType;
+			UhtScriptStruct scriptStruct = (UhtScriptStruct)specifierContext.Type;
 			scriptStruct.ScriptStructFlags |= EStructFlags.Atomic | EStructFlags.Immutable;
 		}
 
 		[UhtSpecifier(Extends = UhtTableNames.ScriptStruct, ValueType = UhtSpecifierValueType.Legacy)]
 		private static void HasDefaultsSpecifier(UhtSpecifierContext specifierContext)
 		{
-			if (!specifierContext.Scope.HeaderParser.HeaderFile.IsNoExportTypes)
+			if (!specifierContext.Type.HeaderFile.IsNoExportTypes)
 			{
 				specifierContext.MessageSite.LogError("The 'HasDefaults' struct specifier is only valid in the NoExportTypes.h file");
 			}
-			UhtScriptStruct scriptStruct = (UhtScriptStruct)specifierContext.Scope.ScopeType;
+			UhtScriptStruct scriptStruct = (UhtScriptStruct)specifierContext.Type;
 			scriptStruct.ScriptStructExportFlags |= UhtScriptStructExportFlags.HasDefaults;
 		}
 
 		[UhtSpecifier(Extends = UhtTableNames.ScriptStruct, ValueType = UhtSpecifierValueType.Legacy)]
 		private static void HasNoOpConstructorSpecifier(UhtSpecifierContext specifierContext)
 		{
-			if (!specifierContext.Scope.HeaderParser.HeaderFile.IsNoExportTypes)
+			if (!specifierContext.Type.HeaderFile.IsNoExportTypes)
 			{
 				specifierContext.MessageSite.LogError("The 'HasNoOpConstructor' struct specifier is only valid in the NoExportTypes.h file");
 			}
-			UhtScriptStruct scriptStruct = (UhtScriptStruct)specifierContext.Scope.ScopeType;
+			UhtScriptStruct scriptStruct = (UhtScriptStruct)specifierContext.Type;
 			scriptStruct.ScriptStructExportFlags |= UhtScriptStructExportFlags.HasNoOpConstructor;
 		}
 
 		[UhtSpecifier(Extends = UhtTableNames.ScriptStruct, ValueType = UhtSpecifierValueType.Legacy)]
 		private static void IsAlwaysAccessibleSpecifier(UhtSpecifierContext specifierContext)
 		{
-			if (!specifierContext.Scope.HeaderParser.HeaderFile.IsNoExportTypes)
+			if (!specifierContext.Type.HeaderFile.IsNoExportTypes)
 			{
 				specifierContext.MessageSite.LogError("The 'IsAlwaysAccessible' struct specifier is only valid in the NoExportTypes.h file");
 			}
-			UhtScriptStruct scriptStruct = (UhtScriptStruct)specifierContext.Scope.ScopeType;
+			UhtScriptStruct scriptStruct = (UhtScriptStruct)specifierContext.Type;
 			scriptStruct.ScriptStructExportFlags |= UhtScriptStructExportFlags.IsAlwaysAccessible;
 		}
 
 		[UhtSpecifier(Extends = UhtTableNames.ScriptStruct, ValueType = UhtSpecifierValueType.Legacy)]
 		private static void IsCoreTypeSpecifier(UhtSpecifierContext specifierContext)
 		{
-			if (!specifierContext.Scope.HeaderParser.HeaderFile.IsNoExportTypes)
+			if (!specifierContext.Type.HeaderFile.IsNoExportTypes)
 			{
 				specifierContext.MessageSite.LogError("The 'IsCoreType' struct specifier is only valid in the NoExportTypes.h file");
 			}
-			UhtScriptStruct scriptStruct = (UhtScriptStruct)specifierContext.Scope.ScopeType;
+			UhtScriptStruct scriptStruct = (UhtScriptStruct)specifierContext.Type;
 			scriptStruct.ScriptStructExportFlags |= UhtScriptStructExportFlags.IsCoreType;
 		}
 	}
diff --git a/Engine/Source/Programs/Shared/EpicGames.UHT/Tables/UhtSpecifierTable.cs b/Engine/Source/Programs/Shared/EpicGames.UHT/Tables/UhtSpecifierTable.cs
--- a/Engine/Source/Programs/Shared/EpicGames.UHT/Tables/UhtSpecifierTable.cs
+++ b/Engine/Source/Programs/Shared/EpicGames.UHT/Tables/UhtSpecifierTable.cs
@@ -5,6 +5,7 @@ using System.Collections.Generic;
 using System.Reflection;
 using EpicGames.Core;
 using EpicGames.UHT.Parsers; // It would be nice if we didn't need this here
+using EpicGames.UHT.Tokenizer;
 using EpicGames.UHT.Types;
 using EpicGames.UHT.Utils;
 
@@ -89,26 +90,41 @@ namespace EpicGames.UHT.Tables
 	/// </summary>
 	public class UhtSpecifierContext
 	{
+		private UhtType? _type = null;
+		private IUhtTokenReader? _tokenReader = null;
+		private IUhtMessageSite? _messageSite = null;
+		private UhtMetaData? _metaData = null;
 
 		/// <summary>
-		/// Current parsing scope (i.e. global, class, ...)
+		/// Get the type containing the specifiers.  For properties, this is the outer object and
+		/// not the property itself.
 		/// </summary>
-		public UhtParsingScope Scope { get; set; }
+		public UhtType Type { get => _type!; set => _type = value; }
+
+		/// <summary>
+		/// Return the currently active token reader
+		/// </summary>
+		public IUhtTokenReader TokenReader { get => _tokenReader!; set => _tokenReader = value; }
+
+		/// <summary>
+		/// Current access specifier
+		/// </summary>
+		public UhtAccessSpecifier AccessSpecifier { get; set; } = UhtAccessSpecifier.None;
 
 		/// <summary>
 		/// Message site for messages
 		/// </summary>
-		public IUhtMessageSite MessageSite { get; set; }
+		public IUhtMessageSite MessageSite { get => _messageSite!; set => _messageSite = value; }
 
 		/// <summary>
 		/// Meta data currently being parsed.
 		/// </summary>
-		public UhtMetaData MetaData { get; set; }
+		public UhtMetaData MetaData { get => _metaData!; set => _metaData = value; }
 
 		/// <summary>
 		/// Make data key index utilized by enumeration values
 		/// </summary>
-		public int MetaNameIndex { get; set; }
+		public int MetaNameIndex { get; set; } = UhtMetaData.IndexNone;
 
 		/// <summary>
 		/// Construct a new specifier context
@@ -119,11 +135,20 @@ namespace EpicGames.UHT.Tables
 		/// <param name="metaNameIndex"></param>
 		public UhtSpecifierContext(UhtParsingScope scope, IUhtMessageSite messageSite, UhtMetaData metaData, int metaNameIndex = UhtMetaData.IndexNone)
 		{
-			this.Scope = scope;
+			this.Type = scope.ScopeType;
+			this.TokenReader = scope.TokenReader;
+			this.AccessSpecifier = scope.AccessSpecifier;
 			this.MessageSite = messageSite;
 			this.MetaData = metaData;
 			this.MetaNameIndex = metaNameIndex;
 		}
+
+		/// <summary>
+		/// Construct an empty context.  Scope, MessageSite, and MetaData must be set at a later point
+		/// </summary>
+		public UhtSpecifierContext()
+		{
+		}
 	}
 
 	/// <summary>
@@ -397,7 +422,7 @@ namespace EpicGames.UHT.Tables
 		{
 			if (value != null)
 			{
-				specifierContext.Scope.TokenReader.LogDeprecation($"Specifier '{this.Name}' has a value which is unused, future versions of UnrealHeaderTool will flag this as an error.");
+				specifierContext.TokenReader.LogDeprecation($"Specifier '{this.Name}' has a value which is unused, future versions of UnrealHeaderTool will flag this as an error.");
 			}
 			this._delegate(specifierContext);
 			return UhtSpecifierDispatchResults.Known;
diff --git a/Engine/Source/Programs/Shared/EpicGames.UHT/Tokenizer/UhtTokenReplayReader.cs b/Engine/Source/Programs/Shared/EpicGames.UHT/Tokenizer/UhtTokenReplayReader.cs
--- a/Engine/Source/Programs/Shared/EpicGames.UHT/Tokenizer/UhtTokenReplayReader.cs
+++ b/Engine/Source/Programs/Shared/EpicGames.UHT/Tokenizer/UhtTokenReplayReader.cs
@@ -256,6 +256,10 @@ namespace EpicGames.UHT.Tokenizer
 			--this._savedStateCount;
 			this._currentTokenIndex = this._savedStates[this._savedStateCount].TokenIndex;
 			this._hasToken = this._savedStates[this._savedStateCount].HasToken;
+			if (this._hasToken)
+			{
+				this._currentToken = this._tokens.Span[this._currentTokenIndex];
+			}
 		}
 
 		/// <inheritdoc/>
diff --git a/Engine/Source/Programs/Shared/EpicGames.UHT/Types/Properties/UhtClassProperty.cs b/Engine/Source/Programs/Shared/EpicGames.UHT/Types/Properties/UhtClassProperty.cs
--- a/Engine/Source/Programs/Shared/EpicGames.UHT/Types/Properties/UhtClassProperty.cs
+++ b/Engine/Source/Programs/Shared/EpicGames.UHT/Types/Properties/UhtClassProperty.cs
@@ -105,7 +105,7 @@ namespace EpicGames.UHT.Types
 
 				case UhtPropertyTextType.EventParameterMember:
 				case UhtPropertyTextType.EventParameterFunctionMember:
-					if (this.PropertyFlags.HasAnyFlags(EPropertyFlags.ReturnParm) && !this.PropertyFlags.HasAnyFlags(EPropertyFlags.UObjectWrapper))
+					if (!this.PropertyFlags.HasAnyFlags(EPropertyFlags.UObjectWrapper))
 					{
 						builder.Append("UClass*");
 					}
diff --git a/Engine/Source/Programs/Shared/EpicGames.UHT/Types/UhtProperty.cs b/Engine/Source/Programs/Shared/EpicGames.UHT/Types/UhtProperty.cs
--- a/Engine/Source/Programs/Shared/EpicGames.UHT/Types/UhtProperty.cs
+++ b/Engine/Source/Programs/Shared/EpicGames.UHT/Types/UhtProperty.cs
@@ -734,7 +734,7 @@ namespace EpicGames.UHT.Types
 		/// <param name="parentPropertySettings">Parent container property</param>
 		/// <param name="sourceName">Name of the property</param>
 		/// <param name="messageSite">Message site used to construct meta data object</param>
-		public UhtPropertySettings(UhtPropertySettings parentPropertySettings, string sourceName, IUhtMessageSite messageSite)
+		public void Reset(UhtPropertySettings parentPropertySettings, string sourceName, IUhtMessageSite messageSite)
 		{
 			this.SourceName = sourceName;
 			this.EngineName = sourceName;
diff --git a/Engine/Source/Programs/UnrealHeaderTool/Private/CodeGenerator.cpp b/Engine/Source/Programs/UnrealHeaderTool/Private/CodeGenerator.cpp
--- a/Engine/Source/Programs/UnrealHeaderTool/Private/CodeGenerator.cpp
+++ b/Engine/Source/Programs/UnrealHeaderTool/Private/CodeGenerator.cpp
@@ -1892,7 +1892,7 @@ bool IsDelegateFunction(FUnrealFieldDefinitionInfo& FieldDef)
 	return FunctionDef != nullptr && FunctionDef->IsDelegateFunction();
 }
 
-void FNativeClassHeaderGenerator::ExportGeneratedPackageInitCode(FOutputDevice& Out, const TCHAR* InDeclarations, uint32 Hash)
+void FNativeClassHeaderGenerator::ExportGeneratedPackageInitCode(FOutputDevice& Out, const TCHAR* InDeclarations, const TArray<FGeneratedCPP*>& ExportedSorted, uint32 Hash)
 {
 	UPackage* Package = PackageDef.GetPackage();
 	const FString& SingletonName = GetPackageSingletonNameFuncAddr(PackageDef, nullptr);
@@ -1926,6 +1926,10 @@ void FNativeClassHeaderGenerator::ExportGeneratedPackageInitCode(FOutputDevice&
 	if (bIncludeDebugOutput)
 	{
 		Out.Log(TEXT("#if 0\r\n"));
+		for (FGeneratedCPP* gen : ExportedSorted)
+		{
+			Out.Logf(TEXT("\t%s\r\n"), *gen->SourceFile.GetFilename());
+		}
 		Out.Log(InDeclarations);
 		Out.Log(TEXT("#endif\r\n"));
 	}
@@ -6068,7 +6072,7 @@ void FNativeClassHeaderGenerator::Generate(
 						}
 					}
 
-					Generator.ExportGeneratedPackageInitCode(GeneratedFileInfo.GetGeneratedBody(), *GeneratedFunctionDeclarations, CombinedHash);
+					Generator.ExportGeneratedPackageInitCode(GeneratedFileInfo.GetGeneratedBody(), *GeneratedFunctionDeclarations, ExportedSorted, CombinedHash);
 					WriteSource(Module, GeneratedFileInfo, GeneratedFileInfo.GetGeneratedBody(), nullptr, TSet<FString>());
 				}
 			});
diff --git a/Engine/Source/Programs/UnrealHeaderTool/Private/NativeClassExporter.h b/Engine/Source/Programs/UnrealHeaderTool/Private/NativeClassExporter.h
--- a/Engine/Source/Programs/UnrealHeaderTool/Private/NativeClassExporter.h
+++ b/Engine/Source/Programs/UnrealHeaderTool/Private/NativeClassExporter.h
@@ -572,7 +572,7 @@ private:
 	 * @param	Out			The destination to write to.
 	 * @param	Package		Package to export code for.
 	**/
-	void ExportGeneratedPackageInitCode(FOutputDevice& Out, const TCHAR* InDeclarations, uint32 CRC);
+	void ExportGeneratedPackageInitCode(FOutputDevice& Out, const TCHAR* InDeclarations, const TArray<FGeneratedCPP*>& ExportedSorted, uint32 CRC);
 
 	/**
 	 * Function to output the C++ code necessary to set up the given array of properties
