Quick for for compile error when using VS 2019.

#rnx

diff --git a/Engine/Source/Programs/Shared/EpicGames.UHT/Parsers/UhtPropertyParser.cs b/Engine/Source/Programs/Shared/EpicGames.UHT/Parsers/UhtPropertyParser.cs
--- a/Engine/Source/Programs/Shared/EpicGames.UHT/Parsers/UhtPropertyParser.cs
+++ b/Engine/Source/Programs/Shared/EpicGames.UHT/Parsers/UhtPropertyParser.cs
@@ -338,7 +338,7 @@ namespace EpicGames.UHT.Parsers
 		/// <summary>
 		/// Request a context
 		/// </summary>
-		public UhtThreadBorrower()
+		public UhtThreadBorrower(bool _ = false) // argument needed for vs2019
 		{
 			List<T> cache = s_tls.Value!;
 			if (cache.Count == 0)
