Add extension methods to serialize VarInt values to ByteArrayBuilders.

diff --git a/Engine/Source/Programs/Shared/EpicGames.Serialization/VarInt.cs b/Engine/Source/Programs/Shared/EpicGames.Serialization/VarInt.cs
--- a/Engine/Source/Programs/Shared/EpicGames.Serialization/VarInt.cs
+++ b/Engine/Source/Programs/Shared/EpicGames.Serialization/VarInt.cs
@@ -226,4 +226,48 @@ namespace EpicGames.Serialization
 			return (ulong)((value >> 63) ^ (value << 1));
 		}
 	}
+
+	/// <summary>
+	/// Extension methods for writing VarInt values
+	/// </summary>
+	public static class VarIntExtensions
+	{
+		/// <summary>
+		/// Writes a signed VarInt to a byte array
+		/// </summary>
+		/// <param name="builder">Builder to write to</param>
+		/// <param name="value">Value to write</param>
+		public static void WriteSignedVarInt(this ByteArrayBuilder builder, int value)
+		{
+			WriteUnsignedVarInt(builder, VarInt.EncodeSigned(value));
+		}
+
+		/// <summary>
+		/// Writes a signed VarInt to a byte array
+		/// </summary>
+		/// <param name="builder">Builder to write to</param>
+		/// <param name="value">Value to write</param>
+		public static void WriteSignedVarInt(this ByteArrayBuilder builder, long value)
+		{
+			WriteUnsignedVarInt(builder, VarInt.EncodeSigned(value));
+		}
+
+		/// <summary>
+		/// Writes a unsigned VarInt to a byte array
+		/// </summary>
+		/// <param name="builder">Builder to write to</param>
+		/// <param name="value">Value to write</param>
+		public static void WriteUnsignedVarInt(this ByteArrayBuilder builder, int value) => WriteUnsignedVarInt(builder, (ulong)value);
+
+		/// <summary>
+		/// Writes a unsigned VarInt to a byte array
+		/// </summary>
+		/// <param name="builder">Builder to write to</param>
+		/// <param name="value">Value to write</param>
+		public static void WriteUnsignedVarInt(this ByteArrayBuilder builder, ulong value)
+		{
+			int length = VarInt.MeasureUnsigned(value);
+			VarInt.WriteUnsigned(builder.GetWritableSpan(length), value);
+		}
+	}
 }
