BuildGraph: Add bytecode compiler and interpreter to allow graph evaluation outside UAT.

diff --git a/Engine/Source/Programs/AutomationTool/BuildGraph/BgGraphBuilder.cs b/Engine/Source/Programs/AutomationTool/BuildGraph/BgGraphBuilder.cs
--- a/Engine/Source/Programs/AutomationTool/BuildGraph/BgGraphBuilder.cs
+++ b/Engine/Source/Programs/AutomationTool/BuildGraph/BgGraphBuilder.cs
@@ -1,6 +1,7 @@
 // Copyright Epic Games, Inc. All Rights Reserved.
 
 using EpicGames.BuildGraph;
+using EpicGames.BuildGraph.Expressions;
 using System;
 using System.Collections.Generic;
 using System.Text;
@@ -15,7 +16,7 @@ namespace AutomationTool
 		/// <summary>
 		/// Callback used to instantiate the graph
 		/// </summary>
-		/// <param name="Graph">The graph object</param>
-		public abstract void SetupGraph(BgGraphSpec Graph);
+		/// <param name="env">The graph context</param>
+		public abstract BgGraphSpec CreateGraph(BgEnvironment env);
 	}
 }
diff --git a/Engine/Source/Programs/AutomationTool/BuildGraph/BgNodeExecutor.cs b/Engine/Source/Programs/AutomationTool/BuildGraph/BgNodeExecutor.cs
--- a/Engine/Source/Programs/AutomationTool/BuildGraph/BgNodeExecutor.cs
+++ b/Engine/Source/Programs/AutomationTool/BuildGraph/BgNodeExecutor.cs
@@ -20,6 +20,8 @@ using EpicGames.Core;
 using Microsoft.Extensions.Logging;
 using OpenTracing;
 using OpenTracing.Util;
+using UnrealBuildBase;
+using UnrealBuildTool;
 
 #nullable enable
 
@@ -33,20 +35,46 @@ namespace AutomationTool
 		public abstract Task<bool> ExecuteAsync(JobContext Job, Dictionary<string, HashSet<FileReference>> TagNameToFileSet);
 	}
 
-	class BgExpressionNodeExecutor : BgNodeExecutor
+	class BgBytecodeNodeExecutor : BgNodeExecutor
 	{
-		readonly BgExpressionNode Node;
-		MethodInfo? method;
+		class BgContextImpl : BgContext
+		{
+			JobContext JobContext;
+
+			public BgContextImpl(JobContext JobContext, Dictionary<string, HashSet<FileReference>> TagNameToFileSet)
+				: base(TagNameToFileSet.ToDictionary(x => x.Key, x => FileSet.FromFiles(Unreal.RootDirectory, x.Value)))
+			{
+				this.JobContext = JobContext;
+			}
+
+			public override string Stream => CommandUtils.P4Enabled ? CommandUtils.P4Env.Branch : "";
+
+			public override int Change => CommandUtils.P4Enabled ? CommandUtils.P4Env.Changelist : 0;
+
+			public override int CodeChange => CommandUtils.P4Enabled ? CommandUtils.P4Env.CodeChangelist : 0;
+
+			public override (int Major, int Minor, int Patch) EngineVersion
+			{
+				get
+				{
+					ReadOnlyBuildVersion Current = ReadOnlyBuildVersion.Current;
+					return (Current.MajorVersion, Current.MinorVersion, Current.PatchVersion);
+				}
+			}
+
+			public override bool IsBuildMachine => CommandUtils.IsBuildMachine;
+		}
 
-		public BgExpressionNodeExecutor(BgExpressionNode node)
+		readonly BgBytecodeNode Node;
+
+		public BgBytecodeNodeExecutor(BgBytecodeNode node)
 		{
 			Node = node;
 		}
 
 		public bool Bind(ILogger logger)
 		{
-			method = Node.Method.Bind();
-			return method != null;
+			return true;
 		}
 
 		/// <summary>
@@ -57,10 +85,7 @@ namespace AutomationTool
 		/// <returns></returns>
 		public override async Task<bool> ExecuteAsync(JobContext Job, Dictionary<string, HashSet<FileReference>> TagNameToFileSet)
 		{
-			if (method == null)
-			{
-				throw new AutomationException("Node has not been bound");
-			}
+			MethodInfo method = Node.Method.Method;
 
 			HashSet<FileReference> BuildProducts = TagNameToFileSet[Node.DefaultOutput.TagName];
 
@@ -76,13 +101,13 @@ namespace AutomationTool
 				{
 					arguments[Idx] = Context;
 				}
-				else if (typeof(IBgExpr).IsAssignableFrom(ParameterType))
+				else if (typeof(BgExpr).IsAssignableFrom(ParameterType))
 				{
-					arguments[Idx] = BgType.Get(ParameterType).CreateConstant(Node.Arguments[Idx]!);
+					arguments[Idx] = BgType.Constant(ParameterType, Node.Arguments[Idx]);
 				}
 				else
 				{
-					arguments[Idx] = Node.Arguments[Idx];
+					arguments[Idx] = Node.Method.Arguments[Idx];
 				}
 			}
 
@@ -100,26 +125,21 @@ namespace AutomationTool
 				}
 
 				object?[] OutputValues;
-				if (Result is BgFileSet)
-				{
-					OutputValues = new[] { Result };
-				}
-				else if (Result is ITuple Tuple)
+				if (Result is ITuple Tuple)
 				{
 					OutputValues = Enumerable.Range(0, Tuple.Length).Select(x => Tuple[x]).ToArray();
 				}
 				else
 				{
-					OutputValues = Array.Empty<object>();
+					OutputValues = new[] { Result };
 				}
 
 				for (int Idx = 0; Idx < OutputValues.Length; Idx++)
 				{
-					if (OutputValues[Idx] is BgFileSet Token)
+					if (OutputValues[Idx] is BgFileSetOutputExpr FileSet)
 					{
-						string TagName = $"{Node.Name}${Idx}";
-						FileSet FileSet = Token.ComputeValue(Context.Context);
-						TagNameToFileSet[TagName] = new HashSet<FileReference>(FileSet.Flatten().Values);
+						string TagName = Node.Outputs[Idx + 1].TagName;
+						TagNameToFileSet[TagName] = new HashSet<FileReference>(FileSet.Value.Flatten().Values);
 						BuildProducts.UnionWith(TagNameToFileSet[TagName]);
 					}
 				}
diff --git a/Engine/Source/Programs/AutomationTool/BuildGraph/BuildGraph.cs b/Engine/Source/Programs/AutomationTool/BuildGraph/BuildGraph.cs
--- a/Engine/Source/Programs/AutomationTool/BuildGraph/BuildGraph.cs
+++ b/Engine/Source/Programs/AutomationTool/BuildGraph/BuildGraph.cs
@@ -103,6 +103,52 @@ namespace AutomationTool
 		}
 	}
 
+	/// <summary>
+	/// Environment for graph evaluation
+	/// </summary>
+	public class BgEnvironment
+	{
+		/// <summary>
+		/// The stream being built
+		/// </summary>
+		public BgString Stream { get; set; }
+
+		/// <summary>
+		/// Current changelist
+		/// </summary>
+		public BgInt Change { get; }
+
+		/// <summary>
+		/// Current code changelist
+		/// </summary>
+		public BgInt CodeChange { get; }
+
+		/// <summary>
+		/// Whether the graph is being run on a build machine
+		/// </summary>
+		public BgBool IsBuildMachine { get; }
+
+		/// <summary>
+		/// The current engine version
+		/// </summary>
+		public (BgInt Major, BgInt Minor, BgInt Patch) EngineVersion { get; }
+
+		/// <summary>
+		/// Constructor
+		/// </summary>
+		/// <param name="P4Env"></param>
+		internal BgEnvironment(P4Environment P4Env)
+		{
+			this.Stream = P4Env?.Branch ?? "Unknown";
+			this.Change = P4Env?.Changelist ?? 0;
+			this.CodeChange = P4Env?.CodeChangelist ?? 0;
+			this.IsBuildMachine = CommandUtils.IsBuildMachine;
+
+			ReadOnlyBuildVersion Version = ReadOnlyBuildVersion.Current;
+			this.EngineVersion = (Version.MajorVersion, Version.MinorVersion, Version.PatchVersion);
+		}
+	}
+
 	/// <summary>
 	/// Tool to execute build automation scripts for UE projects, which can be run locally or in parallel across a build farm (assuming synchronization and resource allocation implemented by a separate system).
 	///
@@ -162,26 +208,6 @@ namespace AutomationTool
 	[Help("WriteToSharedStorage", "Allow writing to shared storage. If not set, but -SharedStorageDir is specified, build products will read but not written")]
 	public class BuildGraph : BuildCommand
 	{
-		class BgGraphConstants : IBgGraphConstants
-		{
-			public BgString Stream { get; set; }
-			public BgInt Change { get; }
-			public BgInt CodeChange { get; }
-			public BgBool IsBuildMachine { get; }
-			public (BgInt Major, BgInt Minor, BgInt Patch) EngineVersion { get; }
-
-			public BgGraphConstants(P4Environment P4Env)
-			{
-				this.Stream = P4Env?.Branch ?? "Unknown";
-				this.Change = P4Env?.Changelist ?? 0;
-				this.CodeChange = P4Env?.CodeChangelist ?? 0;
-				this.IsBuildMachine = CommandUtils.IsBuildMachine;
-
-				ReadOnlyBuildVersion Version = ReadOnlyBuildVersion.Current;
-				this.EngineVersion = (Version.MajorVersion, Version.MinorVersion, Version.PatchVersion);
-			}
-		}
-
 		/// <summary>
 		/// Main entry point for the BuildGraph command
 		/// </summary>
@@ -211,7 +237,6 @@ namespace AutomationTool
 			string ReportName = ParseParamValue("ReportName", null);
 			string BranchOverride = ParseParamValue("Branch", null);
 
-
 			GraphPrintOptions PrintOptions = GraphPrintOptions.ShowCommandLineOptions;
 			if(ParseParam("ShowDeps"))
 			{
@@ -355,12 +380,14 @@ namespace AutomationTool
 					return ExitCode.Error_Unknown;
 				}
 
-				BgGraphSpec GraphSpec = new BgGraphSpec(new BgGraphConstants(P4Enabled? P4Env : null));
-
 				BgGraphBuilder Builder = (BgGraphBuilder)Activator.CreateInstance(BuilderType);
-				Builder.SetupGraph(GraphSpec);
+				BgGraphSpec GraphSpec = Builder.CreateGraph(new BgEnvironment(P4Enabled ? P4Env : null));
+
+				(byte[] Data, BgMethod[] Methods) = BgCompiler.Compile(GraphSpec);
 
-				Graph = GraphSpec.CreateGraph(Arguments);
+				BgInterpreter Interpreter = new BgInterpreter(Data, Methods, Arguments);
+				Interpreter.Disassemble(Logger);
+				Graph = (BgGraph)Interpreter.Evaluate();
 			}
 			else
 			{
@@ -593,7 +620,7 @@ namespace AutomationTool
 			if(SingleNode == null && (!bListOnly || bShowDiagnostics))
 			{
 				List<BgDiagnostic> Diagnostics = Graph.GetAllDiagnostics();
-				foreach(BgDiagnostic Diagnostic in Diagnostics)
+				foreach (BgDiagnostic Diagnostic in Diagnostics)
 				{
 					if(Diagnostic.Level == LogLevel.Information)
 					{
@@ -608,7 +635,7 @@ namespace AutomationTool
 						CommandUtils.LogError("{0}({1}): error: {2}", Diagnostic.File, Diagnostic.Line, Diagnostic.Message);
 					}
 				}
-				if(Diagnostics.Any(x => x.Level == LogLevel.Error))
+				if (Diagnostics.Any(x => x.Level == LogLevel.Error))
 				{
 					return ExitCode.Error_Unknown;
 				}
@@ -682,9 +709,9 @@ namespace AutomationTool
 				NodeToExecutor[Node] = executor;
 				return executor.Bind(NameToTask, TagNameToNodeOutput, Logger);
 			}
-			else if (Node is BgExpressionNode ExpressionNode)
+			else if (Node is BgBytecodeNode BytecodeNode)
 			{
-				BgExpressionNodeExecutor executor = new BgExpressionNodeExecutor(ExpressionNode);
+				BgBytecodeNodeExecutor executor = new BgBytecodeNodeExecutor(BytecodeNode);
 				NodeToExecutor[Node] = executor;
 				return executor.Bind(Logger);
 			}
diff --git a/Engine/Source/Programs/AutomationTool/BuildGraph/Scripts/InstalledBuild.cs b/Engine/Source/Programs/AutomationTool/BuildGraph/Scripts/InstalledBuild.cs
--- a/Engine/Source/Programs/AutomationTool/BuildGraph/Scripts/InstalledBuild.cs
+++ b/Engine/Source/Programs/AutomationTool/BuildGraph/Scripts/InstalledBuild.cs
@@ -81,14 +81,14 @@ namespace AutomationTool
 			return Projects;
 		}
 
-		public override void SetupGraph(BgGraphSpec Graph)
+		public override BgGraphSpec CreateGraph(BgEnvironment Context)
 		{
 			string RootDir = CommandUtils.CmdEnv.LocalRoot;
 
-			BgBoolOption HostPlatformOnly = Graph.AddOption("HostPlatformOnly", "A helper option to make an installed build for your host platform only, so that you don't have to disable each platform individually", false);
-			BgBoolOption HostPlatformEditorOnly = Graph.AddOption("HostPlatformEditorOnly", "A helper option to make an installed build for your host platform only, so that you don't have to disable each platform individually", false);
-			BgBoolOption AllPlatforms = Graph.AddOption("AllPlatforms", "Include all target platforms by default", false);
-			BgBoolOption CompileDatasmithPlugins = Graph.AddOption("CompileDatasmithPlugins", "If Datasmith plugins should be compiled on a separate node.", false);
+			BgBoolOption HostPlatformOnly = new BgBoolOption("HostPlatformOnly", "A helper option to make an installed build for your host platform only, so that you don't have to disable each platform individually", false);
+			BgBoolOption HostPlatformEditorOnly = new BgBoolOption("HostPlatformEditorOnly", "A helper option to make an installed build for your host platform only, so that you don't have to disable each platform individually", false);
+			BgBoolOption AllPlatforms = new BgBoolOption("AllPlatforms", "Include all target platforms by default", false);
+			BgBoolOption CompileDatasmithPlugins = new BgBoolOption("CompileDatasmithPlugins", "If Datasmith plugins should be compiled on a separate node.", false);
 
 			UnrealTargetPlatform CurrentHostPlatform = HostPlatform.Current.HostEditorPlatform;
 
@@ -100,115 +100,127 @@ namespace AutomationTool
 			BgBool DefaultWithIOS = !((CurrentHostPlatform != UnrealTargetPlatform.Mac) & !AllPlatforms);
 			BgBool DefaultWithHoloLens = !((CurrentHostPlatform != UnrealTargetPlatform.Win64) & !AllPlatforms);
 
-			BgBoolOption WithWin64 = Graph.AddOption("WithWin64", "Include the Win64 target platform", DefaultWithWin64);
-			BgBoolOption WithMac = Graph.AddOption("WithMac", "Include the Mac target platform", DefaultWithMac);
-			BgBoolOption WithAndroid = Graph.AddOption("WithAndroid", "Include the Android target platform", DefaultWithPlatform);
-			BgBoolOption WithIOS = Graph.AddOption("WithIOS", "Include the iOS target platform", DefaultWithIOS);
-			BgBoolOption WithTVOS = Graph.AddOption("WithTVOS", "Include the tvOS target platform", DefaultWithIOS);
-			BgBoolOption WithLinux = Graph.AddOption("WithLinux", "Include the Linux target platform", DefaultWithLinux);
-			BgBoolOption WithLinuxArm64 = Graph.AddOption("WithLinuxArm64", "Include the Linux AArch64 target platform", DefaultWithLinuxArm64);
-			BgBoolOption WithHoloLens = Graph.AddOption("WithHoloLens", "Include the HoloLens target platform", DefaultWithHoloLens);
+			BgBoolOption WithWin64 = new BgBoolOption("WithWin64", "Include the Win64 target platform", DefaultWithWin64);
+			BgBoolOption WithMac = new BgBoolOption("WithMac", "Include the Mac target platform", DefaultWithMac);
+			BgBoolOption WithAndroid = new BgBoolOption("WithAndroid", "Include the Android target platform", DefaultWithPlatform);
+			BgBoolOption WithIOS = new BgBoolOption("WithIOS", "Include the iOS target platform", DefaultWithIOS);
+			BgBoolOption WithTVOS = new BgBoolOption("WithTVOS", "Include the tvOS target platform", DefaultWithIOS);
+			BgBoolOption WithLinux = new BgBoolOption("WithLinux", "Include the Linux target platform", DefaultWithLinux);
+			BgBoolOption WithLinuxArm64 = new BgBoolOption("WithLinuxArm64", "Include the Linux AArch64 target platform", DefaultWithLinuxArm64);
+			BgBoolOption WithHoloLens = new BgBoolOption("WithHoloLens", "Include the HoloLens target platform", DefaultWithHoloLens);
 
-			BgBoolOption WithClient = Graph.AddOption("WithClient", "Include precompiled client targets", false);
-			BgBoolOption WithServer = Graph.AddOption("WithServer", "Include precompiled server targets", false);
-			BgBoolOption WithDDC = Graph.AddOption("WithDDC", "Build a standalone derived-data cache for the engine content and templates", true);
-			BgBoolOption HostPlatformDDCOnly = Graph.AddOption("HostPlatformDDCOnly", "Whether to include DDC for the host platform only", true);
-			BgBoolOption SignExecutables = Graph.AddOption("SignExecutables", "Sign the executables produced where signing is available", false);
+			BgBoolOption WithClient = new BgBoolOption("WithClient", "Include precompiled client targets", false);
+			BgBoolOption WithServer = new BgBoolOption("WithServer", "Include precompiled server targets", false);
+			BgBoolOption WithDDC = new BgBoolOption("WithDDC", "Build a standalone derived-data cache for the engine content and templates", true);
+			BgBoolOption HostPlatformDDCOnly = new BgBoolOption("HostPlatformDDCOnly", "Whether to include DDC for the host platform only", true);
+			BgBoolOption SignExecutables = new BgBoolOption("SignExecutables", "Sign the executables produced where signing is available", false);
 
-			BgStringOption AnalyticsTypeOverride = Graph.AddOption("AnalyticsTypeOverride", "Identifier for analytic events to send", "");
-			BgBool EmbedSrcSrvInfo = Graph.AddOption("EmbedSrcSrvInfo", "Whether to add Source indexing to Windows game apps so they can be added to a symbol server", false);
+			BgStrOption AnalyticsTypeOverride = new BgStrOption("AnalyticsTypeOverride", "Identifier for analytic events to send", "");
+			BgBool EmbedSrcSrvInfo = new BgBoolOption("EmbedSrcSrvInfo", "Whether to add Source indexing to Windows game apps so they can be added to a symbol server", false);
 
-			BgList<BgEnum<UnrealTargetConfiguration>> DefaultGameConfigurations = BgList<BgEnum<UnrealTargetConfiguration>>.Create(UnrealTargetConfiguration.DebugGame, UnrealTargetConfiguration.Development, UnrealTargetConfiguration.Shipping);
-			BgList<BgEnum<UnrealTargetConfiguration>> GameConfigurations = Graph.AddOption("GameConfigurations", "Which game configurations to include for packaged applications", DefaultGameConfigurations);
-			BgBoolOption WithFullDebugInfo = Graph.AddOption("WithFullDebugInfo", "Generate full debug info for binary editor and packaged application builds", false);
+			BgList<BgString> DefaultGameConfigurations = BgList<BgString>.Create(nameof(UnrealTargetConfiguration.DebugGame), nameof(UnrealTargetConfiguration.Development), nameof(UnrealTargetConfiguration.Shipping));
+			BgList<BgString> GameConfigurationStrings = new BgListOptionSpec("GameConfigurations", description: "Which game configurations to include for packaged applications", style: BgListOptionStyle.CheckList, values: DefaultGameConfigurations);
+			BgList<BgEnum<UnrealTargetConfiguration>> GameConfigurations = GameConfigurationStrings.Select(x => BgEnum<UnrealTargetConfiguration>.Parse(x));
 
-			BgStringOption BuiltDirectory = Graph.AddOption("BuiltDirectory", "Directory for outputting the built engine", RootDir + "/LocalBuilds/Engine");
+			BgBoolOption WithFullDebugInfo = new BgBoolOption("WithFullDebugInfo", "Generate full debug info for binary editor and packaged application builds", false);
 
-			BgStringOption CrashReporterAPIURL = Graph.AddOption("CrashReporterAPIURL", "The URL to use to talk to the CrashReporterClient API.", "");
-			BgStringOption CrashReporterAPIKey = Graph.AddOption("CrashReporterAPIKey", "The API key to use to talk to the CrashReporterClient API.", "");
+			BgStrOption BuiltDirectory = new BgStrOption("BuiltDirectory", "Directory for outputting the built engine", RootDir + "/LocalBuilds/Engine");
+
+			BgStrOption CrashReporterAPIURL = new BgStrOption("CrashReporterAPIURL", "The URL to use to talk to the CrashReporterClient API.", "");
+			BgStrOption CrashReporterAPIKey = new BgStrOption("CrashReporterAPIKey", "The API key to use to talk to the CrashReporterClient API.", "");
 
 			BgString CrashReporterCompileArgs = "";
-			CrashReporterCompileArgs = CrashReporterCompileArgs.IfThen(CrashReporterAPIURL != "" & CrashReporterAPIKey != "", BgString.Format("-define:CRC_TELEMETRY_URL=\"{0}\" -define:CRC_TELEMETRY_KEY_DEV=\"{1}\" -define:CRC_TELEMETRY_KEY_RELEASE=\"{1}\" -OverrideBuildEnvironment", CrashReporterAPIURL, CrashReporterAPIKey));
+			CrashReporterCompileArgs = CrashReporterCompileArgs.If(CrashReporterAPIURL != "" & CrashReporterAPIKey != "", BgString.Format("-define:CRC_TELEMETRY_URL=\"{0}\" -define:CRC_TELEMETRY_KEY_DEV=\"{1}\" -define:CRC_TELEMETRY_KEY_RELEASE=\"{1}\" -OverrideBuildEnvironment", CrashReporterAPIURL, CrashReporterAPIKey));
+
+			List<BgAggregateSpec> Aggregates = new List<BgAggregateSpec>();
 
 			/////// EDITORS ////////////////////////////////////////////////
 
 			//// Windows ////
-			BgAgentSpec EditorWin64 = Graph.AddAgent("Editor Win64").Type("Win64_Licensee");
+			BgAgentSpec EditorWin64 = new BgAgentSpec("Editor Win64", "Win64_Licensee");
 
 			BgNodeSpec VersionFilesNode = EditorWin64
-				.AddNode(x => UpdateVersionFilesAsync(x));
+				.AddNode(x => UpdateVersionFilesAsync(x))
+				.Construct();
 
 			BgNodeSpec<BgFileSet> WinUhtNode = EditorWin64
 				.AddNode(x => CompileUnrealHeaderToolAsync(x, UnrealTargetPlatform.Win64))
-				.Requires(VersionFilesNode);
+				.Requires(VersionFilesNode)
+				.Construct();
 
 			BgNodeSpec<BgFileSet> WinEditorNode = EditorWin64
 				.AddNode(x => CompileUnrealEditorWin64Async(x, CrashReporterCompileArgs, EmbedSrcSrvInfo, CompileDatasmithPlugins, WithFullDebugInfo, SignExecutables))
-				.Requires(WinUhtNode);
-			
-			Graph.AddAggregate("Win64 Editor", WinEditorNode, label: "Editors/Win64");
+				.Requires(WinUhtNode)
+				.Construct();
+
+			Aggregates.Add(new BgAggregateSpec("Win64 Editor", WinEditorNode, label: "Editors/Win64"));
 
 
 			/////// TARGET PLATFORMS ////////////////////////////////////////////////
 
 			//// Win64 ////
 
-			BgAgentSpec TargetWin64 = Graph.AddAgent("Target Win64").Type("Win64_Licensee");
+			BgAgentSpec TargetWin64 = new BgAgentSpec("Target Win64", "Win64_Licensee");
 
 			BgNodeSpec<BgFileSet> WinGame = TargetWin64
 				.AddNode(x => CompileUnrealGameWin64(x, GameConfigurations, EmbedSrcSrvInfo, WithFullDebugInfo, SignExecutables))
-				.Requires(WinUhtNode);
+				.Requires(WinUhtNode)
+				.Construct();
 
-			Graph.AddAggregate("TargetPlatforms_Win64", WinGame);
+			Aggregates.Add(new BgAggregateSpec("TargetPlatforms_Win64", WinGame));
 
 			/////// TOOLS //////////////////////////////////////////////////////////
 
 			//// Build Rules ////
 
-			BgAgentSpec BuildRules = Graph.AddAgent("BuildRules").Type("Win64_Licensee");
+			BgAgentSpec BuildRules = new BgAgentSpec("BuildRules", "Win64_Licensee");
 
 			BgNodeSpec RulesAssemblies = BuildRules
-				.AddNode(x => CompileRulesAssemblies(x));
+				.AddNode(x => CompileRulesAssemblies(x))
+				.Construct();
 
 			//// Win Tools ////
 
-			BgAgentSpec ToolsGroupWin64 = Graph.AddAgent("Tools Group Win64").Type("Win64_Licensee");
+			BgAgentSpec ToolsGroupWin64 = new BgAgentSpec("Tools Group Win64", "Win64_Licensee");
 
 			BgNodeSpec<BgFileSet> WinTools = ToolsGroupWin64
 				.AddNode(x => BuildToolsWin64Async(x, CrashReporterCompileArgs))
-				.Requires(WinUhtNode);
+				.Requires(WinUhtNode)
+				.Construct();
 
 			BgNodeSpec<BgFileSet> CsTools = ToolsGroupWin64
 				.AddNode(x => BuildToolsCSAsync(x, SignExecutables))
-				.Requires(VersionFilesNode);
+				.Requires(VersionFilesNode)
+				.Construct();
 
 
 			/////// DDC //////////////////////////////////////////////////////////
 
-			BgAgentSpec DDCGroupWin64 = Graph.AddAgent("DDC Group Win64").Type("Win64_Licensee");
+			BgAgentSpec DDCGroupWin64 = new BgAgentSpec("DDC Group Win64", "Win64_Licensee");
 
 			BgList<BgString> DDCPlatformsWin64 = BgList<BgString>.Create("WindowsEditor");
-			DDCPlatformsWin64 = DDCPlatformsWin64.AddIf(WithWin64, "Windows");
+			DDCPlatformsWin64 = DDCPlatformsWin64.If(WithWin64, x => x.Add("Windows"));
 
 			BgNodeSpec DdcNode = DDCGroupWin64
 				.AddNode(x => BuildDDCWin64Async(x, DDCPlatformsWin64, BgList<BgFileSet>.Create(WinUhtNode.Output, WinEditorNode.Output, WinTools.Output)))
-				.Requires(WinEditorNode, WinTools);
+				.Requires(WinEditorNode, WinTools)
+				.Construct();
 
 
 			/////// STAGING ///////
 
 			// Windows 
-			BgAgentSpec WinStageAgent = Graph.AddAgent("Installed Build Group Win64").Type("Win64_Licensee");
+			BgAgentSpec WinStageAgent = new BgAgentSpec("Installed Build Group Win64", "Win64_Licensee");
 
 			BgList<BgFileSet> WinInstalledFiles = BgList<BgFileSet>.Empty;
 			WinInstalledFiles = WinInstalledFiles.Add(WinUhtNode.Output);
 			WinInstalledFiles = WinInstalledFiles.Add(WinEditorNode.Output);
 			WinInstalledFiles = WinInstalledFiles.Add(WinTools.Output);
 			WinInstalledFiles = WinInstalledFiles.Add(CsTools.Output);
-			WinInstalledFiles = WinInstalledFiles.AddIf(WithWin64, WinGame.Output);
+			WinInstalledFiles = WinInstalledFiles.If(WithWin64, x => x.Add(WinGame.Output));
 
 			BgList<BgString> WinPlatforms = BgList<BgString>.Empty;
-			WinPlatforms = WinPlatforms.AddIf(WithWin64, "Win64");
+			WinPlatforms = WinPlatforms.If(WithWin64, x => x.Add("Win64"));
 
 			BgList<BgString> WinContentOnlyPlatforms = BgList<BgString>.Empty;
 
@@ -217,9 +229,12 @@ namespace AutomationTool
 
 			BgNodeSpec WinInstalledNode = WinStageAgent
 				.AddNode(x => MakeInstalledBuildWin64Async(x, WinInstalledFiles, WinFinalizeArgs, WinOutputDir))
-				.Requires(WinInstalledFiles);
+				.Requires(WinInstalledFiles)
+				.Construct();
+
+			Aggregates.Add(new BgAggregateSpec("HostPlatforms_Win64", WinInstalledNode, label: "Builds/Win64"));
 
-			Graph.AddAggregate("HostPlatforms_Win64", WinInstalledNode, label: "Builds/Win64");
+			return new BgGraphSpec(BgList<BgNodeSpec>.Empty, Aggregates);
 		}
 
 		/// <summary>
diff --git a/Engine/Source/Programs/AutomationTool/BuildGraph/Tasks/CallMethodTask.cs b/Engine/Source/Programs/AutomationTool/BuildGraph/Tasks/CallMethodTask.cs
deleted file mode 100644
diff --git a/Engine/Source/Programs/Shared/EpicGames.BuildGraph.Tests/GraphTests.cs b/Engine/Source/Programs/Shared/EpicGames.BuildGraph.Tests/GraphTests.cs
new file mode 100644
--- /dev/null
+++ b/Engine/Source/Programs/Shared/EpicGames.BuildGraph.Tests/GraphTests.cs
@@ -0,0 +1,127 @@
+// Copyright Epic Games, Inc. All Rights Reserved.
+
+using Microsoft.VisualStudio.TestTools.UnitTesting;
+using System;
+using System.Collections.Generic;
+using System.IO;
+using System.Linq;
+using System.Text;
+using System.Text.Json;
+using OpenTracing;
+using EpicGames.BuildGraph.Expressions;
+using Microsoft.Extensions.Logging;
+
+namespace EpicGames.BuildGraph.Tests
+{
+	/*
+	[TestClass]
+	public class GraphTests
+	{
+		object Evaluate(BgExpr expr)
+		{
+			byte[] data = BgBytecode.Compile(expr);
+			BgBytecodeInterpreter interpreter = new BgBytecodeInterpreter(data);
+			return interpreter.Evaluate();
+		}
+
+		[TestMethod]
+		public void NodeTest()
+		{
+			BgNodeSpec nodeSpec1 = new BgNodeSpec("Update Version Files");
+
+			BgNodeSpec nodeSpec2 = new BgNodeSpec("Compile ShooterGameEditor Win64", requires: BgList<BgFileSet>.Create(nodeSpec1));
+
+			BgNodeSpec nodeSpec3 = new BgNodeSpec("Cook ShooterGame Win64", requires: BgList<BgFileSet>.Create(nodeSpec2.DefaultOutput));
+
+			BgNode node3 = (BgNode)Evaluate(nodeSpec3);
+			Assert.AreEqual(node3.Name, "Cook ShooterGame Win64");
+			Assert.AreEqual(2, node3.InputDependencies.Count);
+			Assert.AreEqual(1, node3.Inputs.Count);
+
+			BgNode node2 = node3.Inputs[0].ProducingNode;
+			Assert.AreEqual(node2.Name, "Compile ShooterGameEditor Win64");
+			Assert.AreEqual(1, node2.InputDependencies.Count);
+			Assert.AreEqual(1, node2.Inputs.Count);
+
+			BgNode node1 = node2.Inputs[0].ProducingNode;
+			Assert.AreEqual(node1.Name, "Update Version Files");
+			Assert.AreEqual(0, node1.InputDependencies.Count);
+			Assert.AreEqual(0, node1.Inputs.Count);
+		}
+
+		[TestMethod]
+		public void AgentTest()
+		{
+			BgNodeSpec nodeSpec = new BgNodeSpec("Update Version Files");
+
+			BgAgentSpec agentSpec = new BgAgentSpec(BgList<BgString>.Create("win64", "incremental"));
+
+			BgAgent agent = (BgAgent)Evaluate(agentSpec);
+			Assert.AreEqual("TestAgent", agent.Name);
+
+			Assert.AreEqual(2, agent.PossibleTypes.Length);
+			Assert.AreEqual("win64", agent.PossibleTypes[0]);
+			Assert.AreEqual("incremental", agent.PossibleTypes[1]);
+
+			Assert.AreEqual(1, agent.Nodes.Count);
+
+			BgNode node = agent.Nodes[0];
+			Assert.AreEqual("Update Version Files", node.Name);
+		}
+
+		[TestMethod]
+		public void AggregateTest()
+		{
+			BgNodeSpec nodeSpec1 = new BgNodeSpec("Update Version Files");
+			BgNodeSpec nodeSpec2 = new BgNodeSpec("Compile UnrealEditor Win64");
+
+			BgAggregateSpec aggregateSpec = new BgAggregateSpec("All nodes", BgList<BgFileSet>.Create(nodeSpec1, nodeSpec2));
+
+			BgAggregate aggregate = (BgAggregate)Evaluate(aggregateSpec);
+			Assert.AreEqual("All nodes", aggregate.Name);
+
+			BgNode node1 = aggregate.RequiredNodes.OrderBy(x => x.Name).ElementAt(0);
+			Assert.AreEqual("Compile UnrealEditor Win64", node1.Name);
+
+			BgNode node2 = aggregate.RequiredNodes.OrderBy(x => x.Name).ElementAt(1);
+			Assert.AreEqual("Update Version Files", node2.Name);
+		}
+
+		[TestMethod]
+		public void LabelTest()
+		{
+			BgNodeSpec nodeSpec1 = new BgNodeSpec("Update Version Files");
+			BgNodeSpec nodeSpec2 = new BgNodeSpec("Compile UnrealEditor Win64");
+
+			BgLabelSpec labelSpec = new BgLabelSpec("name", "category", "ugsBadge", "ugsProject", "Code", BgList<BgNodeSpec>.Create(nodeSpec2), BgList<BgNodeSpec>.Create(nodeSpec1));
+
+			BgLabel label = (BgLabel)Evaluate(labelSpec);
+			Assert.AreEqual("name", label.DashboardName);
+			Assert.AreEqual("category", label.DashboardCategory);
+			Assert.AreEqual("ugsBadge", label.UgsBadge);
+			Assert.AreEqual("ugsProject", label.UgsProject);
+			Assert.AreEqual(BgLabelChange.Code, label.Change);
+			Assert.AreEqual(1, label.RequiredNodes.Count);
+			Assert.AreEqual(1, label.IncludedNodes.Count);
+
+			BgNode node1 = label.RequiredNodes.First();
+			Assert.AreEqual("Compile UnrealEditor Win64", node1.Name);
+
+			BgNode node2 = label.IncludedNodes.First();
+			Assert.AreEqual("Update Version Files", node2.Name);
+		}
+
+		[TestMethod]
+		public void DiagnosticTest()
+		{
+			BgDiagnosticSpec diagnosticSpec = new BgDiagnosticSpec(LogLevel.Error, "hello world", "file.cs", 123);
+
+			BgDiagnostic diagnostic = (BgDiagnostic)Evaluate(diagnosticSpec);
+			Assert.AreEqual(LogLevel.Error, diagnostic.Level);
+			Assert.AreEqual("hello world", diagnostic.Message);
+			Assert.AreEqual("file.cs", diagnostic.Location.File);
+			Assert.AreEqual(123, diagnostic.Location.LineNumber);
+		}
+	}
+	*/
+}
diff --git a/Engine/Source/Programs/Shared/EpicGames.BuildGraph/BgAgent.cs b/Engine/Source/Programs/Shared/EpicGames.BuildGraph/BgAgent.cs
--- a/Engine/Source/Programs/Shared/EpicGames.BuildGraph/BgAgent.cs
+++ b/Engine/Source/Programs/Shared/EpicGames.BuildGraph/BgAgent.cs
@@ -30,7 +30,7 @@ namespace EpicGames.BuildGraph
 		public List<BgNode> Nodes { get; set; } = new List<BgNode>();
 
 		/// <summary>
-		/// Diagnostics for this agent
+		/// Diagnostics to output if executing this agent
 		/// </summary>
 		public List<BgDiagnostic> Diagnostics { get; } = new List<BgDiagnostic>();
 
diff --git a/Engine/Source/Programs/Shared/EpicGames.BuildGraph/BgAgentSpec.cs b/Engine/Source/Programs/Shared/EpicGames.BuildGraph/BgAgentSpec.cs
--- a/Engine/Source/Programs/Shared/EpicGames.BuildGraph/BgAgentSpec.cs
+++ b/Engine/Source/Programs/Shared/EpicGames.BuildGraph/BgAgentSpec.cs
@@ -5,49 +5,15 @@ using System.Collections.Generic;
 using System.Linq.Expressions;
 using System.Threading.Tasks;
 using EpicGames.BuildGraph.Expressions;
+using EpicGames.Core;
 using Microsoft.Extensions.Logging;
 
 namespace EpicGames.BuildGraph
 {
-	/// <summary>
-	/// Configuration object for agents
-	/// </summary>
-	public class BgAgentConfig
-	{
-		/// <summary>
-		/// List of agent types to select from, in order of preference. The first agent type supported by a stream will be used.
-		/// </summary>
-		public BgList<BgString> Types { get; set; } = BgList<BgString>.Empty;
-
-		/// <summary>
-		/// List of diagnostics for this agent to run
-		/// </summary>
-		public BgList<BgDiagnosticSpec> Diagnostics { get; set; } = BgList<BgDiagnosticSpec>.Empty;
-
-		/// <summary>
-		/// Add agent types to the agent definition
-		/// </summary>
-		public BgAgentConfig AddTypes(params BgString[] typeNames)
-		{
-			Types = Types.Add(typeNames);
-			return this;
-		}
-
-		/// <summary>
-		/// Check that a certain condition is true before executing the node
-		/// </summary>
-		public BgAgentConfig Check(BgBool condition, LogLevel level, BgString message)
-		{
-			BgDiagnosticSpec diagnostic = new BgDiagnosticSpec(level, message);
-			Diagnostics = Diagnostics.AddIf(condition, diagnostic);
-			return this;
-		}
-	}
-
 	/// <summary>
 	/// Describes an agent that can execute execute build steps
 	/// </summary>
-	public class BgAgentSpec
+	public class BgAgentSpec : BgExpr
 	{
 		/// <summary>
 		/// Name of the agent
@@ -57,100 +23,35 @@ namespace EpicGames.BuildGraph
 		/// <summary>
 		/// List of agent types to select from, in order of preference. The first agent type supported by a stream will be used.
 		/// </summary>
-		public BgList<BgString> TypeNames { get; set; } = BgList<BgString>.Empty;
-
-		/// <summary>
-		/// List of diagnostics for this agent to run
-		/// </summary>
-		public BgList<BgDiagnosticSpec> Diagnostics { get; set; } = BgList<BgDiagnosticSpec>.Empty;
-
-		/// <summary>
-		/// Nodes that have been added to the agent
-		/// </summary>
-		List<BgNodeSpec> NodeSpecs { get; set; } = new List<BgNodeSpec>();
+		public BgList<BgString> Types { get; set; }
 
 		/// <summary>
-		/// Private constructor. Use <see cref="BgGraphSpec.AddAgent(BgString)"/> to construct an agent
+		/// Constructor
 		/// </summary>
-		/// <param name="name"></param>
-		internal BgAgentSpec(BgString name)
+		public BgAgentSpec(BgString name, BgString type)
+			: this(name, BgList<BgString>.Create(type))
 		{
-			Name = name;
-		}
-
-		/// <summary>
-		/// Adds a node to be executed on this agent
-		/// </summary>
-		/// <param name="function">Lambda expression which calls a static method to be used for this node</param>
-		/// <returns>New node spec instance</returns>
-		public BgNodeSpec AddNode(Expression<Func<BgContext, Task>> function)
-		{
-			BgNodeSpec nodeSpec = BgNodeSpec.Create(function);
-			NodeSpecs.Add(nodeSpec);
-			return nodeSpec;
-		}
-
-		/// <summary>
-		/// Adds a node to be executed on this agent
-		/// </summary>
-		/// <param name="function">Lambda expression which calls a static method to be used for this node</param>
-		/// <returns>New node spec instance</returns>
-		public BgNodeSpec<T> AddNode<T>(Expression<Func<BgContext, Task<T>>> function)
-		{
-			BgNodeSpec<T> nodeSpec = BgNodeSpec.Create<T>(function);
-			NodeSpecs.Add(nodeSpec);
-			return nodeSpec;
-		}
-
-		/// <summary>
-		/// Add agent types to the agent definition
-		/// </summary>
-		public BgAgentSpec Type(params BgString[] newTypeNames)
-		{
-			TypeNames = TypeNames.Add(newTypeNames);
-			return this;
 		}
 
 		/// <summary>
-		/// Check that a certain condition is true before executing the node
+		/// Constructor
 		/// </summary>
-		public BgAgentSpec WarnIf(BgBool condition, BgString message)
+		public BgAgentSpec(BgString name, BgList<BgString> types)
+			: base(BgExprFlags.ForceFragment)
 		{
-			BgDiagnosticSpec diagnostic = new BgDiagnosticSpec(LogLevel.Warning, message);
-			Diagnostics = Diagnostics.AddIf(condition, diagnostic);
-			return this;
+			Name = name;
+			Types = types;
 		}
 
-		/// <summary>
-		/// Check that a certain condition is true before executing the node
-		/// </summary>
-		public BgAgentSpec ErrorIf(BgBool condition, BgString message)
+		/// <inheritdoc/>
+		public override void Write(BgBytecodeWriter writer)
 		{
-			BgDiagnosticSpec diagnostic = new BgDiagnosticSpec(LogLevel.Error, message);
-			Diagnostics = Diagnostics.AddIf(condition, diagnostic);
-			return this;
+			writer.WriteOpcode(BgOpcode.Agent);
+			writer.WriteExpr(Name);
+			writer.WriteExpr(Types);
 		}
 
-		/// <summary>
-		/// Creates a concrete <see cref="BgAgent"/> object from the specification
-		/// </summary>
-		public void AddToGraph(BgExprContext context, BgGraph graph)
-		{
-			string[] types = TypeNames.Compute(context).ToArray();
-
-			BgAgent agent = new BgAgent(Name.Compute(context), types);
-			graph.NameToAgent.Add(agent.Name, agent);
-			graph.Agents.Add(agent);
-
-			foreach (BgDiagnosticSpec precondition in Diagnostics.GetEnumerable(context))
-			{
-				precondition.AddToGraph(context, graph, agent, null);
-			}
-
-			foreach (BgNodeSpec nodeSpec in NodeSpecs)
-			{
-				nodeSpec.AddToGraph(context, graph, agent);
-			}
-		}
+		/// <inheritdoc/>
+		public override BgString ToBgString() => "{Agent}";
 	}
 }
diff --git a/Engine/Source/Programs/Shared/EpicGames.BuildGraph/BgAggregate.cs b/Engine/Source/Programs/Shared/EpicGames.BuildGraph/BgAggregate.cs
--- a/Engine/Source/Programs/Shared/EpicGames.BuildGraph/BgAggregate.cs
+++ b/Engine/Source/Programs/Shared/EpicGames.BuildGraph/BgAggregate.cs
@@ -22,10 +22,10 @@ namespace EpicGames.BuildGraph
 		/// <summary>
 		/// Constructor
 		/// </summary>
-		/// <param name="inName">Name of this aggregate</param>
-		public BgAggregate(string inName)
+		/// <param name="name">Name of this aggregate</param>
+		public BgAggregate(string name)
 		{
-			Name = inName;
+			Name = name;
 		}
 
 		/// <summary>
@@ -37,4 +37,25 @@ namespace EpicGames.BuildGraph
 			return Name;
 		}
 	}
+
+	/// <summary>
+	/// Aggregate that was created from bytecode
+	/// </summary>
+	public class BgBytecodeAggregate : BgAggregate
+	{
+		/// <summary>
+		/// Labels to add this aggregate to
+		/// </summary>
+		public BgLabel? Label { get; }
+
+		/// <summary>
+		/// Constructor
+		/// </summary>
+		public BgBytecodeAggregate(string name, IEnumerable<BgNode> nodes, BgLabel? label)
+			: base(name)
+		{
+			RequiredNodes.UnionWith(nodes);
+			Label = label;
+		}
+	}
 }
diff --git a/Engine/Source/Programs/Shared/EpicGames.BuildGraph/BgAggregateSpec.cs b/Engine/Source/Programs/Shared/EpicGames.BuildGraph/BgAggregateSpec.cs
--- a/Engine/Source/Programs/Shared/EpicGames.BuildGraph/BgAggregateSpec.cs
+++ b/Engine/Source/Programs/Shared/EpicGames.BuildGraph/BgAggregateSpec.cs
@@ -1,6 +1,6 @@
 // Copyright Epic Games, Inc. All Rights Reserved.
 
-using System.Linq;
+using System;
 using EpicGames.BuildGraph.Expressions;
 
 namespace EpicGames.BuildGraph
@@ -8,7 +8,7 @@ namespace EpicGames.BuildGraph
 	/// <summary>
 	/// Specification for an aggregate target in the graph
 	/// </summary>
-	public class BgAggregateSpec
+	public class BgAggregateSpec : BgExpr
 	{
 		/// <summary>
 		/// Name of the aggregate
@@ -18,34 +18,58 @@ namespace EpicGames.BuildGraph
 		/// <summary>
 		/// Outputs required for the aggregate
 		/// </summary>
-		public BgList<BgFileSet> RequiredOutputs { get; set; }
+		public BgList<BgFileSet> Requires { get; }
 
 		/// <summary>
-		/// Internal constructor. Use <see cref="BgGraphSpec.AddAggregate(BgString, BgList{BgFileSet})"/> to create an aggregate.
+		/// Label to apply to this aggregate
 		/// </summary>
-		internal BgAggregateSpec(BgString name, BgList<BgFileSet> requiredOutputs)
+		public BgLabelSpec? Label { get; }
+
+		/// <summary>
+		/// Constructor.
+		/// </summary>
+		public BgAggregateSpec(BgString name, params BgFileSet[] requires)
+			: this(name, BgList<BgFileSet>.Create(requires))
+		{
+		}
+
+		/// <summary>
+		/// Constructor.
+		/// </summary>
+		public BgAggregateSpec(BgString name, params BgList<BgFileSet>[] requires)
+			: this(name, BgList<BgFileSet>.Concat(requires))
 		{
-			Name = name;
-			RequiredOutputs = requiredOutputs;
 		}
 
 		/// <summary>
-		/// Creates a concrete aggregate object from this specification.
+		/// Constructor.
 		/// </summary>
-		internal void AddToGraph(BgExprContext context, BgGraph graph)
+		public BgAggregateSpec(BgString name, BgList<BgFileSet> requires, string label)
+			: this(name, requires, new BgLabelSpec(label))
 		{
-			BgAggregate aggregate = new BgAggregate(Name.Compute(context));
-			aggregate.RequiredNodes.UnionWith(RequiredOutputs.ComputeTags(context).Select(x => graph.TagNameToNodeOutput[x].ProducingNode));
-			graph.NameToAggregate.Add(aggregate.Name, aggregate);
 		}
 
 		/// <summary>
-		/// Adds a set of dependencies to this aggregate
+		/// Constructor.
 		/// </summary>
-		/// <param name="tokens">List of token dependencies</param>
-		public void Requires(params BgFileSet[] tokens)
+		public BgAggregateSpec(BgString name, BgList<BgFileSet> requires, BgLabelSpec? label = null)
+			: base(BgExprFlags.None)
 		{
-			RequiredOutputs.Add(tokens);
+			Name = name;
+			Requires = requires;
+			Label = label;
 		}
+
+		/// <inheritdoc/>
+		public override void Write(BgBytecodeWriter writer)
+		{
+			writer.WriteOpcode(BgOpcode.Aggregate);
+			writer.WriteExpr(Name);
+			writer.WriteExpr(Requires);
+			writer.WriteExpr(Label ?? BgExpr.Null);
+		}
+
+		/// <inheritdoc/>
+		public override BgString ToBgString() => Name;
 	}
 }
diff --git a/Engine/Source/Programs/Shared/EpicGames.BuildGraph/BgBytecodeWriter.cs b/Engine/Source/Programs/Shared/EpicGames.BuildGraph/BgBytecodeWriter.cs
new file mode 100644
--- /dev/null
+++ b/Engine/Source/Programs/Shared/EpicGames.BuildGraph/BgBytecodeWriter.cs
@@ -0,0 +1,69 @@
+// Copyright Epic Games, Inc. All Rights Reserved.
+
+using System;
+using System.Buffers;
+using System.Collections.Generic;
+using System.Diagnostics;
+using System.Diagnostics.CodeAnalysis;
+using System.Reflection;
+using System.Runtime.CompilerServices;
+using System.Text;
+using EpicGames.Core;
+using EpicGames.Serialization;
+
+namespace EpicGames.BuildGraph
+{
+	/// <summary>
+	/// Helper class for writing BuildGraph bytecode to a buffer.
+	/// </summary>
+	public abstract class BgBytecodeWriter
+	{
+		/// <summary>
+		/// Writes an opcode to the output
+		/// </summary>
+		/// <param name="opcode"></param>
+		public abstract void WriteOpcode(BgOpcode opcode);
+
+		/// <summary>
+		/// Writes a string to the output
+		/// </summary>
+		/// <param name="str"></param>
+		public abstract void WriteString(string str);
+
+		/// <summary>
+		/// Writes a signed integer value to the output
+		/// </summary>
+		/// <param name="value"></param>
+		public abstract void WriteSignedInteger(long value);
+
+		/// <summary>
+		/// Writes an unsigned integer to the output
+		/// </summary>
+		/// <param name="value"></param>
+		public void WriteUnsignedInteger(int value) => WriteUnsignedInteger((ulong)value);
+
+		/// <summary>
+		/// Writes an unsigned integer to the output
+		/// </summary>
+		/// <param name="value"></param>
+		public abstract void WriteUnsignedInteger(ulong value);
+
+		/// <summary>
+		/// Writes an expression to the output buffer
+		/// </summary>
+		/// <param name="expr"></param>
+		public abstract void WriteExpr(BgExpr expr);
+
+		/// <summary>
+		/// Writes an expression as a standalone fragment, encoding just the fragment index into the output stream
+		/// </summary>
+		/// <param name="expr"></param>
+		public abstract void WriteExprAsFragment(BgExpr expr);
+
+		/// <summary>
+		/// Writes the method to be called for a particular node
+		/// </summary>
+		/// <param name="method">Method to be called</param>
+		public abstract void WriteMethod(BgMethod method);
+	}
+}
diff --git a/Engine/Source/Programs/Shared/EpicGames.BuildGraph/BgCompiler.cs b/Engine/Source/Programs/Shared/EpicGames.BuildGraph/BgCompiler.cs
new file mode 100644
--- /dev/null
+++ b/Engine/Source/Programs/Shared/EpicGames.BuildGraph/BgCompiler.cs
@@ -0,0 +1,220 @@
+// Copyright Epic Games, Inc. All Rights Reserved.
+
+using System;
+using System.Buffers;
+using System.Collections.Generic;
+using System.Diagnostics;
+using System.Diagnostics.CodeAnalysis;
+using System.Reflection;
+using System.Runtime.CompilerServices;
+using System.Text;
+using EpicGames.BuildGraph.Expressions;
+using EpicGames.Core;
+using EpicGames.Serialization;
+using Microsoft.Extensions.Logging;
+
+namespace EpicGames.BuildGraph
+{
+	/// <summary>
+	/// Version numbers for bytecode streams
+	/// </summary>
+	enum BgBytecodeVersion
+	{
+		Current = 0,
+	}
+
+	/// <summary>
+	/// Helper class for writing BuildGraph bytecode to a buffer.
+	/// </summary>
+	public static class BgCompiler
+	{
+		class ReferenceEqualityComparer : IEqualityComparer<object>
+		{
+			public static readonly ReferenceEqualityComparer Instance = new ReferenceEqualityComparer();
+
+			public new bool Equals([AllowNull] object x, [AllowNull] object y) => ReferenceEquals(x, y);
+
+			public int GetHashCode([DisallowNull] object obj) => RuntimeHelpers.GetHashCode(obj);
+		}
+
+		class FragmentCollector : BgBytecodeWriter
+		{
+			readonly HashSet<BgExpr> _uniqueExprs = new HashSet<BgExpr>(ReferenceEqualityComparer.Instance);
+			readonly List<BgExpr> _fragments;
+			readonly Dictionary<BgExpr, int> _exprToFragmentIndex;
+
+			public FragmentCollector(List<BgExpr> fragments, Dictionary<BgExpr, int> exprToFragmentIdx)
+			{
+				_fragments = fragments;
+				_exprToFragmentIndex = exprToFragmentIdx;
+			}
+
+			/// <inheritdoc/>
+			public override void WriteExpr(BgExpr expr)
+			{
+				if ((expr.Flags & BgExprFlags.ForceFragment) != 0 && false)
+				{
+					RegisterFragment(expr);
+				}
+				else if (_uniqueExprs.Add(expr) || (expr.Flags & BgExprFlags.NotInterned) != 0)
+				{
+					expr.Write(this);
+				}
+				else
+				{
+					RegisterFragment(expr);
+				}
+			}
+
+			/// <inheritdoc/>
+			public override void WriteExprAsFragment(BgExpr expr)
+			{
+				RegisterFragment(expr);
+
+				if (_uniqueExprs.Add(expr))
+				{
+					expr.Write(this);
+				}
+			}
+
+			/// <summary>
+			/// Registers an expression for compilation into a new fragment
+			/// </summary>
+			/// <param name="expr"></param>
+			void RegisterFragment(BgExpr expr)
+			{
+				if (!_exprToFragmentIndex.ContainsKey(expr))
+				{
+					int index = _exprToFragmentIndex.Count;
+					_fragments.Add(expr);
+					_exprToFragmentIndex[expr] = index;
+				}
+			}
+
+			/// <inheritdoc/>
+			public override void WriteOpcode(BgOpcode opcode) { }
+
+			/// <inheritdoc/>
+			public override void WriteString(string str) { }
+
+			/// <inheritdoc/>
+			public override void WriteSignedInteger(long value) { }
+
+			/// <inheritdoc/>
+			public override void WriteUnsignedInteger(ulong value) { }
+
+			/// <inheritdoc/>
+			public override void WriteMethod(BgMethod handler) { }
+		}
+
+		class ForwardWriter : BgBytecodeWriter
+		{
+			readonly ByteArrayBuilder _builder;
+			readonly Dictionary<BgExpr, int> _exprToFragmentIdx;
+			readonly List<BgMethod> _methods;
+
+			public ForwardWriter(ByteArrayBuilder builder, Dictionary<BgExpr, int> exprToFragmentIdx, List<BgMethod> handlers)
+			{
+				_builder = builder;
+				_exprToFragmentIdx = exprToFragmentIdx;
+				_methods = handlers;
+			}
+
+			/// <inheritdoc/>
+			public override void WriteExpr(BgExpr expr)
+			{
+				int index;
+				if (_exprToFragmentIdx.TryGetValue(expr, out index))
+				{
+					WriteOpcode(BgOpcode.Jump);
+					WriteUnsignedInteger((ulong)index);
+				}
+				else
+				{
+					expr.Write(this);
+				}
+			}
+
+			/// <inheritdoc/>
+			public override void WriteExprAsFragment(BgExpr expr)
+			{
+				WriteUnsignedInteger(_exprToFragmentIdx[expr]);
+			}
+
+			/// <inheritdoc/>
+			public override void WriteOpcode(BgOpcode opcode)
+			{
+				_builder.WriteByte((byte)opcode);
+			}
+
+			/// <inheritdoc/>
+			public override void WriteString(string str)
+			{
+				int textLength = Encoding.UTF8.GetByteCount(str);
+				int lengthLength = VarInt.MeasureUnsigned(textLength);
+
+				Span<byte> buffer = _builder.GetWritableSpan(lengthLength + textLength);
+				VarInt.WriteUnsigned(buffer, textLength);
+				Encoding.UTF8.GetBytes(str, buffer.Slice(lengthLength));
+			}
+
+			/// <inheritdoc/>
+			public override void WriteSignedInteger(long value) => _builder.WriteSignedVarInt(value);
+
+			/// <inheritdoc/>
+			public override void WriteUnsignedInteger(ulong value) => _builder.WriteUnsignedVarInt(value);
+
+			/// <inheritdoc/>
+			public override void WriteMethod(BgMethod method)
+			{
+				_builder.WriteUnsignedVarInt(_methods.Count);
+				_methods.Add(method);
+			}
+		}
+
+		/// <summary>
+		/// Compiles the given expression into bytecode
+		/// </summary>
+		/// <param name="expr">Expression to compile</param>
+		/// <returns>Compiled bytecode for the expression, suitable for passing to <see cref="BgInterpreter"/></returns>
+		public static (byte[], BgMethod[]) Compile(BgExpr expr)
+		{
+			List<BgExpr> fragments = new List<BgExpr>();
+			Dictionary<BgExpr, int> exprToFragmentIndex = new Dictionary<BgExpr, int>(ReferenceEqualityComparer.Instance);
+
+			// Figure out which expressions need to be compiled into separate fragments
+			FragmentCollector collector = new FragmentCollector(fragments, exprToFragmentIndex);
+			collector.WriteExprAsFragment(expr);
+
+			// Write all the fragments
+			ByteArrayBuilder code = new ByteArrayBuilder();
+			List<int> fragmentLengths = new List<int>(fragments.Count);
+
+			List<BgMethod> methods = new List<BgMethod>();
+
+			ForwardWriter writer = new ForwardWriter(code, exprToFragmentIndex, methods);
+			for (int idx = 0; idx < fragments.Count; idx++)
+			{
+				int fragmentOffset = code.Length;
+				fragments[idx].Write(writer);
+				fragmentLengths.Add(code.Length - fragmentOffset);
+			}
+
+			// Create the header
+			ByteArrayBuilder header = new ByteArrayBuilder();
+			header.WriteUnsignedVarInt((int)BgBytecodeVersion.Current);
+			header.WriteUnsignedVarInt(fragments.Count);
+			for (int idx = 0; idx < fragments.Count; idx++)
+			{
+				header.WriteUnsignedVarInt(fragmentLengths[idx]);
+			}
+
+			// Append them together
+			byte[] output = new byte[header.Length + code.Length];
+			header.CopyTo(output);
+			code.CopyTo(output.AsSpan(header.Length));
+
+			return (output, methods.ToArray());
+		}
+	}
+}
diff --git a/Engine/Source/Programs/Shared/EpicGames.BuildGraph/BgContext.cs b/Engine/Source/Programs/Shared/EpicGames.BuildGraph/BgContext.cs
--- a/Engine/Source/Programs/Shared/EpicGames.BuildGraph/BgContext.cs
+++ b/Engine/Source/Programs/Shared/EpicGames.BuildGraph/BgContext.cs
@@ -37,22 +37,19 @@ namespace EpicGames.BuildGraph
 		/// </summary>
 		public abstract bool IsBuildMachine { get; }
 
-		/// <summary>
-		/// Context for evaluating expressions
-		/// </summary>
-		public BgExprContext Context { get; }
-
 		/// <summary>
 		/// All outputs for the node
 		/// </summary>
 		public HashSet<FileReference> BuildProducts { get; } = new HashSet<FileReference>();
 
+		Dictionary<string, FileSet> _tagNameToFileSet;
+
 		/// <summary>
 		/// Constructor
 		/// </summary>
-		public BgContext(BgExprContext context)
+		public BgContext(Dictionary<string, FileSet> tagNameToFileSet)
 		{
-			Context = context;
+			_tagNameToFileSet = tagNameToFileSet;
 		}
 
 		/// <summary>
@@ -60,21 +57,21 @@ namespace EpicGames.BuildGraph
 		/// </summary>
 		/// <param name="expr">The boolean expression</param>
 		/// <returns>Value of the expression</returns>
-		public bool Get(BgBool expr) => expr.Compute(Context);
+		public bool Get(BgBool expr) => ((BgBoolConstantExpr)expr).Value;
 
 		/// <summary>
 		/// Resolve an integer expression to a value
 		/// </summary>
 		/// <param name="expr">The integer expression</param>
 		/// <returns>Value of the expression</returns>
-		public int Get(BgInt expr) => expr.Compute(Context);
+		public int Get(BgInt expr) => ((BgIntConstantExpr)expr).Value;
 
 		/// <summary>
 		/// Resolve a string expression to a value
 		/// </summary>
 		/// <param name="expr">The string expression</param>
 		/// <returns>Value of the expression</returns>
-		public string Get(BgString expr) => expr.Compute(Context);
+		public string Get(BgString expr) => ((BgStringConstantExpr)expr).Value;
 
 		/// <summary>
 		/// Resolves an enum expression to a value
@@ -82,7 +79,7 @@ namespace EpicGames.BuildGraph
 		/// <typeparam name="TEnum">The enum type</typeparam>
 		/// <param name="expr">Enum expression</param>
 		/// <returns>The enum value</returns>
-		public TEnum Get<TEnum>(BgEnum<TEnum> expr) where TEnum : struct => expr.Compute(Context);
+		public TEnum Get<TEnum>(BgEnum<TEnum> expr) where TEnum : struct => ((BgEnumConstantExpr<TEnum>)expr).Value;
 
 		/// <summary>
 		/// Resolve a list of enums to a value
@@ -90,34 +87,31 @@ namespace EpicGames.BuildGraph
 		/// <typeparam name="TEnum">The enum type</typeparam>
 		/// <param name="expr">Enum expression</param>
 		/// <returns>The enum value</returns>
-		public List<TEnum> Get<TEnum>(BgList<BgEnum<TEnum>> expr) where TEnum : struct => expr.GetEnumerable(Context).Select(x => Get(x)).ToList();
+		public List<TEnum> Get<TEnum>(BgList<BgEnum<TEnum>> expr) where TEnum : struct => ((BgListConstantExpr<BgEnum<TEnum>>)expr).Value.Select(x => (TEnum)x).ToList();
 
 		/// <summary>
 		/// Resolve a list of strings
 		/// </summary>
 		/// <param name="expr">List expression</param>
 		/// <returns></returns>
-		public List<string> Get(BgList<BgString> expr) => expr.Compute(Context);
+		public List<string> Get(BgList<BgString> expr) => (((BgListConstantExpr<BgString>)expr).Value).ConvertAll(x => (string)x);
 
 		/// <summary>
 		/// Resolve a file set
 		/// </summary>
 		/// <param name="fileSet">The token expression</param>
 		/// <returns>Set of files for the token</returns>
-		public FileSet Get(BgFileSet fileSet) => fileSet.ComputeValue(Context);
-
-		/// <summary>
-		/// Resolve a file set
-		/// </summary>
-		/// <param name="fileSets">The token expression</param>
-		/// <returns>Set of files for the token</returns>
-		public FileSet Get(BgList<BgFileSet> fileSets)
+		public FileSet Get(BgFileSet fileSet)
 		{
 			FileSet result = FileSet.Empty;
-			foreach (BgFileSet fileSet in fileSets.GetEnumerable(Context))
+
+			BgFileSetInputExpr input = (BgFileSetInputExpr)fileSet;
+			foreach (BgNodeOutput output in input.Outputs)
 			{
-				result += Get(fileSet);
+				FileSet set = _tagNameToFileSet[output.TagName];
+				result = FileSet.Union(result, set);
 			}
+
 			return result;
 		}
 
@@ -126,12 +120,12 @@ namespace EpicGames.BuildGraph
 		/// </summary>
 		/// <param name="fileSets">The token expression</param>
 		/// <returns>Set of files for the token</returns>
-		public FileSet Get(IEnumerable<BgFileSet> fileSets)
+		public FileSet Get(BgList<BgFileSet> fileSets)
 		{
 			FileSet result = FileSet.Empty;
-			foreach (BgFileSet fileSet in fileSets)
+			foreach (BgFileSet fileSet in ((BgListConstantExpr<BgFileSet>)fileSets).Value)
 			{
-				result += fileSet.ComputeValue(Context);
+				result += Get(fileSet);
 			}
 			return result;
 		}
diff --git a/Engine/Source/Programs/Shared/EpicGames.BuildGraph/BgDiagnosticSpec.cs b/Engine/Source/Programs/Shared/EpicGames.BuildGraph/BgDiagnosticSpec.cs
deleted file mode 100644
diff --git a/Engine/Source/Programs/Shared/EpicGames.BuildGraph/BgGraphSpec.cs b/Engine/Source/Programs/Shared/EpicGames.BuildGraph/BgGraphSpec.cs
--- a/Engine/Source/Programs/Shared/EpicGames.BuildGraph/BgGraphSpec.cs
+++ b/Engine/Source/Programs/Shared/EpicGames.BuildGraph/BgGraphSpec.cs
@@ -7,228 +7,72 @@ using Microsoft.Extensions.Logging;
 
 namespace EpicGames.BuildGraph
 {
-	/// <summary>
-	/// Constant values for evaluating the graph
-	/// </summary>
-	public interface IBgGraphConstants
-	{
-		/// <summary>
-		/// The current stream
-		/// </summary>
-		public BgString Stream { get; }
-
-		/// <summary>
-		/// The changelist being built
-		/// </summary>
-		public BgInt Change { get; }
-
-		/// <summary>
-		/// Current code changelist
-		/// </summary>
-		public BgInt CodeChange { get; }
-
-		/// <summary>
-		/// The current engnine version
-		/// </summary>
-		public (BgInt Major, BgInt Minor, BgInt Patch) EngineVersion { get; }
-	}
-
 	/// <summary>
 	/// Specification for a graph in fluent syntax
 	/// </summary>
-	public class BgGraphSpec
+	public class BgGraphSpec : BgExpr
 	{
 		/// <summary>
-		/// Constant properties for evaluating the graph
+		/// Nodes for the graph
 		/// </summary>
-		public IBgGraphConstants Constants { get; }
-
-		readonly List<IBgOption> _options = new List<IBgOption>();
-		readonly List<BgAgentSpec> _agentSpecs = new List<BgAgentSpec>();
-		readonly List<BgAggregateSpec> _aggregateSpecs = new List<BgAggregateSpec>();
-		readonly List<BgLabelSpec> _labelSpecs = new List<BgLabelSpec>();
-		BgList<BgDiagnosticSpec> _diagnosticSpecs = new List<BgDiagnosticSpec>();
+		public BgList<BgNodeSpec> Nodes { get; }
 
 		/// <summary>
-		/// Constructor
+		/// Aggregates for the graph
 		/// </summary>
-		/// <param name="constants">Constant values for use during graph evaluation</param>
-		public BgGraphSpec(IBgGraphConstants constants)
-		{
-			Constants = constants;
-		}
+		public BgList<BgAggregateSpec> Aggregates { get; }
 
 		/// <summary>
-		/// Adds a boolean option
-		/// </summary>
-		/// <param name="name">Name of the option. Must be unique.</param>
-		/// <param name="description">Description for the option</param>
-		/// <param name="defaultValue">Default value for the option.</param>
-		/// <returns>Boolean expression</returns>
-		public BgBoolOption AddOption(string name, BgString description, BgBool defaultValue)
-		{
-			BgBoolOption option = new BgBoolOption(name, description, defaultValue);
-			_options.Add(option);
-			return option;
-		}
-
-		/// <summary>
-		/// Adds an integer option
-		/// </summary>
-		/// <param name="name">Name of the option. Must be unique.</param>
-		/// <param name="description">Description for the option</param>
-		/// <param name="defaultValue">Default value for the option.</param>
-		/// <returns>Integer expression</returns>
-		public BgIntOption AddIntOption(string name, BgString description, BgInt defaultValue)
-		{
-			BgIntOption option = new BgIntOption(name, description, defaultValue);
-			_options.Add(option);
-			return option;
-		}
-
-		/// <summary>
-		/// Adds a string option
-		/// </summary>
-		/// <param name="name">Name of the option. Must be unique.</param>
-		/// <param name="description">Description for the option</param>
-		/// <param name="defaultValue">Default value for the option.</param>
-		/// <returns>String expression</returns>
-		public BgStringOption AddOption(string name, BgString description, BgString defaultValue)
-		{
-			BgStringOption option = new BgStringOption(name, description, defaultValue);
-			_options.Add(option);
-			return option;
-		}
-
-		/// <summary>
-		/// Adds an enum option
-		/// </summary>
-		/// <typeparam name="TEnum">Type of enum value</typeparam>
-		/// <param name="name">Name of the option. Must be unique.</param>
-		/// <param name="description">Description for the option</param>
-		/// <param name="defaultValue">Default value for the option.</param>
-		/// <returns>Enum option</returns>
-		public BgEnumOption<TEnum> AddOption<TEnum>(string name, BgString description, BgEnum<TEnum> defaultValue) where TEnum : struct
-		{
-			BgEnumOption<TEnum> option = new BgEnumOption<TEnum>(name, description, defaultValue);
-			_options.Add(option);
-			return option;
-		}
-
-		/// <summary>
-		/// Adds an option which may be a list of enum values
+		/// Constructor
 		/// </summary>
-		/// <typeparam name="TEnum">Type of enum value</typeparam>
-		/// <param name="name">Name of the option. Must be unique.</param>
-		/// <param name="description">Description for the option</param>
-		/// <param name="defaultValue">Default value for the option.</param>
-		/// <returns></returns>
-		public BgEnumListOption<TEnum> AddOption<TEnum>(string name, BgString description, BgList<BgEnum<TEnum>> defaultValue) where TEnum : struct
+		public BgGraphSpec(BgList<BgNodeSpec> nodes, BgList<BgAggregateSpec> aggregates)
+			: base(BgExprFlags.ForceFragment)
 		{
-			BgEnumListOption<TEnum> option = new BgEnumListOption<TEnum>(name, description, defaultValue);
-			_options.Add(option);
-			return option;
+			Nodes = nodes;
+			Aggregates = aggregates;
 		}
 
 		/// <summary>
-		/// Adds a new agent to the graph
+		/// Implicit conversion from a node spec
 		/// </summary>
-		/// <param name="name">Name of the agent</param>
-		/// <returns>Agent specification object</returns>
-		public BgAgentSpec AddAgent(BgString name)
+		public static implicit operator BgGraphSpec(BgNodeSpec node)
 		{
-			BgAgentSpec agent = new BgAgentSpec(name);
-			_agentSpecs.Add(agent);
-			return agent;
+			return new BgGraphSpec(node, BgList<BgAggregateSpec>.Empty);
 		}
 
 		/// <summary>
-		/// Adds a new aggregate to the graph
+		/// Implicit conversion from a list of node specs
 		/// </summary>
-		/// <param name="name">Name of the agent</param>
-		/// <param name="outputs">Initial set of tokens to include in the aggregate</param>
-		/// <returns>Agent specification object</returns>
-		public BgAggregateSpec AddAggregate(BgString name, BgList<BgFileSet> outputs)
+		public static implicit operator BgGraphSpec(BgList<BgNodeSpec> nodes)
 		{
-			BgAggregateSpec aggregate = new BgAggregateSpec(name, outputs);
-			_aggregateSpecs.Add(aggregate);
-			return aggregate;
+			return new BgGraphSpec(nodes, BgList<BgAggregateSpec>.Empty);
 		}
 
 		/// <summary>
-		/// Adds a new aggregate to the graph
+		/// Implicit conversion from an aggregate spec
 		/// </summary>
-		/// <param name="name">Name of the agent</param>
-		/// <param name="outputs">Initial set of tokens to include in the aggregate</param>
-		/// <param name="label"></param>
-		/// <returns>Agent specification object</returns>
-		public BgAggregateSpec AddAggregate(BgString name, BgList<BgFileSet> outputs, BgString label)
+		public static implicit operator BgGraphSpec(BgAggregateSpec aggregate)
 		{
-			AddLabel(config =>
-			{
-				config.DashboardName = label;
-				config.RequiredNodes = config.RequiredNodes.Add(outputs).Distinct();
-			});
-			return AddAggregate(name, outputs);
+			return new BgGraphSpec(BgList<BgNodeSpec>.Empty, aggregate);
 		}
 
 		/// <summary>
-		/// Adds a new label to the graph
+		/// Implicit conversion from a list of node specs
 		/// </summary>
-		/// <param name="configure">Callback for configuring the label</param>
-		/// <returns>Label specification object</returns>
-		public BgLabelSpec AddLabel(Action<BgLabelConfig> configure)
+		public static implicit operator BgGraphSpec(BgList<BgAggregateSpec> aggregates)
 		{
-			BgLabelConfig config = new BgLabelConfig();
-			configure(config);
-
-			BgLabelSpec labelSpec = new BgLabelSpec(config);
-			_labelSpecs.Add(labelSpec);
-
-			return labelSpec;
+			return new BgGraphSpec(BgList<BgNodeSpec>.Empty, aggregates);
 		}
 
-		/// <summary>
-		/// Adds a diagnostic to the graph
-		/// </summary>
-		/// <param name="condition">Condition for adding the diagnostic to the graph</param>
-		/// <param name="level"></param>
-		/// <param name="message"></param>
-		public void AddDiagnostic(BgBool condition, LogLevel level, BgString message)
+		/// <inheritdoc/>
+		public override void Write(BgBytecodeWriter writer)
 		{
-			BgDiagnosticSpec diagnostic = new BgDiagnosticSpec(level, message);
-			_diagnosticSpecs = _diagnosticSpecs.AddIf(condition, diagnostic);
+			writer.WriteOpcode(BgOpcode.Graph);
+			writer.WriteExpr(Nodes);
+			writer.WriteExpr(Aggregates);
 		}
 
-		/// <summary>
-		/// Creates a concrete graph from the specifications
-		/// </summary>
-		/// <param name="options"></param>
-		/// <returns>New graph instance</returns>
-		public BgGraph CreateGraph(Dictionary<string, string> options)
-		{
-			BgExprContext context = new BgExprContext();
-			context.Options = options;
-
-			BgGraph graph = new BgGraph();
-			foreach (BgAgentSpec agentSpec in _agentSpecs)
-			{
-				agentSpec.AddToGraph(context, graph);
-			}
-			foreach (BgAggregateSpec aggregateSpec in _aggregateSpecs)
-			{
-				aggregateSpec.AddToGraph(context, graph);
-			}
-			foreach (BgDiagnosticSpec diagnosticSpec in _diagnosticSpecs.GetEnumerable(context))
-			{
-				diagnosticSpec.AddToGraph(context, graph, null, null);
-			}
-			foreach (BgLabelSpec labelSpec in _labelSpecs)
-			{
-				labelSpec.AddToGraph(context, graph);
-			}
-
-			return graph;
-		}
+		/// <inheritdoc/>
+		public override BgString ToBgString() => "{Graph}";
 	}
 }
diff --git a/Engine/Source/Programs/Shared/EpicGames.BuildGraph/BgInterpreter.cs b/Engine/Source/Programs/Shared/EpicGames.BuildGraph/BgInterpreter.cs
new file mode 100644
--- /dev/null
+++ b/Engine/Source/Programs/Shared/EpicGames.BuildGraph/BgInterpreter.cs
@@ -0,0 +1,1118 @@
+// Copyright Epic Games, Inc. All Rights Reserved.
+
+using System;
+using System.Collections;
+using System.Collections.Generic;
+using System.Collections.Immutable;
+using System.Diagnostics;
+using System.IO;
+using System.Linq;
+using System.Text;
+using System.Text.RegularExpressions;
+using System.Xml;
+using EpicGames.BuildGraph.Expressions;
+using EpicGames.Core;
+using EpicGames.Serialization;
+using Microsoft.Extensions.Logging;
+
+namespace EpicGames.BuildGraph
+{
+	/// <summary>
+	/// Exception thrown by the runtime
+	/// </summary>
+	public sealed class BgBytecodeException : Exception
+	{
+		/// <summary>
+		/// Source file that the error was thrown from
+		/// </summary>
+		public string SourceFile { get; }
+
+		/// <summary>
+		/// Line number that threw the exception
+		/// </summary>
+		public int SourceLine { get; }
+
+		/// <summary>
+		/// Message to display
+		/// </summary>
+		public string Diagnostic { get; }
+
+		/// <summary>
+		/// Constructor
+		/// </summary>
+		/// <param name="sourceFile"></param>
+		/// <param name="sourceLine"></param>
+		/// <param name="diagnostic"></param>
+		public BgBytecodeException(string sourceFile, int sourceLine, string diagnostic)
+		{
+			SourceFile = sourceFile;
+			SourceLine = sourceLine;
+			Diagnostic = diagnostic;
+		}
+
+		/// <inheritdoc/>
+		public override string ToString() => $"{SourceFile}({SourceLine}): {Diagnostic}";
+	}
+
+	/// <summary>
+	/// Interprets compiled buildgraph bytecode
+	/// </summary>
+	public class BgInterpreter
+	{
+		class Frame
+		{
+			public int Offset { get; set; }
+			public IReadOnlyList<object> Arguments { get; }
+			public Dictionary<int, object> Objects { get; set; }
+
+			public Frame(Frame other)
+			{
+				Offset = other.Offset;
+				Arguments = other.Arguments;
+				Objects = other.Objects;
+			}
+
+			public Frame(int offset, IReadOnlyList<object> arguments)
+			{
+				Offset = offset;
+				Arguments = arguments;
+				Objects = new Dictionary<int, object>();
+			}
+		}
+
+		readonly byte[] _data;
+		readonly BgMethod[] _methods;
+		readonly IReadOnlyDictionary<string, string> _options;
+		readonly BgBytecodeVersion _version;
+		readonly int[] _fragments;
+
+		/// <summary>
+		/// Constructor
+		/// </summary>
+		/// <param name="data"></param>
+		/// <param name="lambdas">Handlers for each node type</param>
+		/// <param name="options">Options for evaluating the graph</param>
+		public BgInterpreter(byte[] data, BgMethod[] lambdas, IReadOnlyDictionary<string, string> options)
+		{
+			_data = data;
+			_methods = lambdas;
+			_options = options;
+
+			int offset = 0;
+
+			_version = (BgBytecodeVersion)VarInt.ReadUnsigned(data[offset..], out int bytesRead);
+			offset += bytesRead;
+
+			int fragmentCount = (int)VarInt.ReadUnsigned(data[offset..], out bytesRead);
+			offset += bytesRead;
+
+			int[] lengths = new int[fragmentCount];
+			for (int idx = 0; idx < fragmentCount; idx++)
+			{
+				lengths[idx] = (int)VarInt.ReadUnsigned(data[offset..], out bytesRead);
+				offset += bytesRead;
+			}
+
+			_fragments = new int[fragmentCount];
+			for (int idx = 0; idx < fragmentCount; idx++)
+			{
+				_fragments[idx] = offset;
+				offset += lengths[idx];
+			}
+		}
+
+		/// <summary>
+		/// Evaluates the graph
+		/// </summary>
+		public object Evaluate()
+		{
+			// Take the given parameters, evaluate the graph that was produced (creating a map of all named entites)
+			// Filter the map to the required target
+			// Evaluate it fully into a BgGraph
+			return Evaluate(new Frame(_fragments[0], Array.Empty<object>()));
+		}
+
+		object Evaluate(Frame frame)
+		{
+			BgOpcode opcode = ReadOpcode(frame);
+			switch (opcode)
+			{
+				#region Bool opcodes
+
+				case BgOpcode.BoolFalse:
+					return false;
+				case BgOpcode.BoolTrue:
+					return true;
+				case BgOpcode.BoolNot:
+					return !(bool)Evaluate(frame);
+				case BgOpcode.BoolAnd:
+					{
+						bool lhs = (bool)Evaluate(frame);
+						bool rhs = (bool)Evaluate(frame);
+						return lhs & rhs;
+					}
+				case BgOpcode.BoolOr:
+					{
+						bool lhs = (bool)Evaluate(frame);
+						bool rhs = (bool)Evaluate(frame);
+						return lhs | rhs;
+					}
+				case BgOpcode.BoolXor:
+					{
+						bool lhs = (bool)Evaluate(frame);
+						bool rhs = (bool)Evaluate(frame);
+						return lhs ^ rhs;
+					}
+				case BgOpcode.BoolEq:
+					{
+						bool lhs = (bool)Evaluate(frame);
+						bool rhs = (bool)Evaluate(frame);
+						return lhs == rhs;
+					}
+				case BgOpcode.BoolOption:
+					{
+						string name = (string)Evaluate(frame);
+						string label = (string)Evaluate(frame);
+						string description = (string)Evaluate(frame);
+						bool defaultValue = (bool)Evaluate(frame);
+
+						bool value = defaultValue;
+						if (_options.TryGetValue(name, out string? str))
+						{
+							value = Boolean.Parse(str);
+						}
+
+						return value;
+					}
+
+				#endregion
+				#region Integer opcodes
+
+				case BgOpcode.IntLiteral:
+					return ReadSignedInteger(frame);
+				case BgOpcode.IntEq:
+					{
+						int lhs = (int)Evaluate(frame);
+						int rhs = (int)Evaluate(frame);
+						return lhs == rhs;
+					}
+				case BgOpcode.IntLt:
+					{
+						int lhs = (int)Evaluate(frame);
+						int rhs = (int)Evaluate(frame);
+						return lhs < rhs;
+					}
+				case BgOpcode.IntGt:
+					{
+						int lhs = (int)Evaluate(frame);
+						int rhs = (int)Evaluate(frame);
+						return lhs > rhs;
+					}
+				case BgOpcode.IntAdd:
+					{
+						int lhs = (int)Evaluate(frame);
+						int rhs = (int)Evaluate(frame);
+						return lhs + rhs;
+					}
+				case BgOpcode.IntMultiply:
+					{
+						int lhs = (int)Evaluate(frame);
+						int rhs = (int)Evaluate(frame);
+						return lhs * rhs;
+					}
+				case BgOpcode.IntDivide:
+					{
+						int lhs = (int)Evaluate(frame);
+						int rhs = (int)Evaluate(frame);
+						return lhs / rhs;
+					}
+				case BgOpcode.IntModulo:
+					{
+						int lhs = (int)Evaluate(frame);
+						int rhs = (int)Evaluate(frame);
+						return lhs % rhs;
+					}
+				case BgOpcode.IntNegate:
+					return -(int)Evaluate(frame);
+
+				#endregion
+				#region String opcodes
+
+				case BgOpcode.StrEmpty:
+					return String.Empty;
+				case BgOpcode.StrLiteral:
+					return ReadString(frame);
+				case BgOpcode.StrCompare:
+					{
+						string lhs = (string)Evaluate(frame);
+						string rhs = (string)Evaluate(frame);
+						StringComparison comparison = (StringComparison)ReadUnsignedInteger(frame);
+						return String.Compare(lhs, rhs, comparison);
+					}
+				case BgOpcode.StrConcat:
+					{
+						string lhs = (string)Evaluate(frame);
+						string rhs = (string)Evaluate(frame);
+						return lhs + rhs;
+					}
+				case BgOpcode.StrFormat:
+					{
+						string format = (string)Evaluate(frame);
+
+						int count = (int)ReadUnsignedInteger(frame);
+						object[] arguments = new object[count];
+
+						for (int idx = 0; idx < count; idx++)
+						{
+							object argument = Evaluate(frame);
+							arguments[idx] = argument;
+						}
+
+						return String.Format(format, arguments);
+					}
+				case BgOpcode.StrSplit:
+					{
+						string source = (string)Evaluate(frame);
+						string separator = (string)Evaluate(frame);
+						return source.Split(separator, StringSplitOptions.RemoveEmptyEntries);
+					}
+				case BgOpcode.StrJoin:
+					{
+						string lhs = (string)Evaluate(frame);
+						IEnumerable<object> rhs = (IEnumerable<object>)Evaluate(frame);
+						return String.Join(lhs, rhs);
+					}
+				case BgOpcode.StrIsMatch:
+					{
+						string input = (string)Evaluate(frame);
+						string pattern = ReadString(frame);
+						return Regex.IsMatch(input, pattern);
+					}
+				case BgOpcode.StrReplace:
+					{
+						string input = (string)Evaluate(frame);
+						string pattern = ReadString(frame);
+						string replacement = (string)Evaluate(frame);
+						return Regex.Replace(input, pattern, replacement);
+					}
+				case BgOpcode.StrOption:
+					{
+						string name = (string)Evaluate(frame);
+						string label = (string)Evaluate(frame);
+						string description = (string)Evaluate(frame);
+						int style = (int)ReadUnsignedInteger(frame);
+						string defaultValue = (string)Evaluate(frame);
+						string pattern = (string)Evaluate(frame);
+						string patternFailed = (string)Evaluate(frame);
+						List<string> values = EvaluateList<string>(frame);
+						List<string> valueDescriptions = EvaluateList<string>(frame);
+
+						string value = defaultValue;
+						if (_options.TryGetValue(name, out string? option))
+						{
+							value = option;
+						}
+
+						return value;
+					}
+
+				#endregion
+				#region Enum opcodes
+
+				case BgOpcode.EnumConstant:
+					return ReadSignedInteger(frame);
+				case BgOpcode.EnumParse:
+					{
+						string name = (string)Evaluate(frame);
+						List<string> names = EvaluateList<string>(frame);
+						List<int> values = EvaluateList<int>(frame);
+
+						for (int idx = 0; idx < names.Count; idx++)
+						{
+							if (String.Equals(names[idx], name, StringComparison.OrdinalIgnoreCase))
+							{
+								return values[idx];
+							}
+						}
+
+						throw new InvalidDataException($"Unable to parse enum '{name}'");
+					}
+				case BgOpcode.EnumToString:
+					{
+						int value = (int)Evaluate(frame);
+						List<string> names = EvaluateList<string>(frame);
+						List<int> values = EvaluateList<int>(frame);
+
+						for (int idx = 0; idx < names.Count; idx++)
+						{
+							if (value == values[idx])
+							{
+								return names[idx];
+							}
+						}
+
+						return $"{value}";
+					}
+
+				#endregion
+				#region List opcodes
+
+				case BgOpcode.ListEmpty:
+					return Enumerable.Empty<object>();
+				case BgOpcode.ListPush:
+					{
+						IEnumerable<object> list = (IEnumerable<object>)Evaluate(frame);
+						object item = Evaluate(frame);
+
+						return list.Concat(new[] { item });
+					}
+				case BgOpcode.ListPushLazy:
+					{
+						IEnumerable<object> list = (IEnumerable<object>)Evaluate(frame);
+						int fragment = ReadFragment(frame);
+
+						IEnumerable<object> item = LazyEvaluateItem(new Frame(frame), fragment);
+						return list.Concat(item);
+					}
+				case BgOpcode.ListCount:
+					{
+						IEnumerable<object> list = (IEnumerable<object>)Evaluate(frame);
+						return list.Count();
+					}
+				case BgOpcode.ListElement:
+					{
+						IEnumerable<object> list = (IEnumerable<object>)Evaluate(frame);
+						int index = (int)Evaluate(frame);
+						return list.ElementAt(index);
+					}
+				case BgOpcode.ListConcat:
+					{
+						IEnumerable<object> lhs = (IEnumerable<object>)Evaluate(frame);
+						IEnumerable<object> rhs = (IEnumerable<object>)Evaluate(frame);
+						return Enumerable.Concat(lhs, rhs);
+					}
+				case BgOpcode.ListUnion:
+					{
+						IEnumerable<object> lhs = (IEnumerable<object>)Evaluate(frame);
+						IEnumerable<object> rhs = (IEnumerable<object>)Evaluate(frame);
+						return lhs.Union(rhs);
+					}
+				case BgOpcode.ListExcept:
+					{
+						IEnumerable<object> lhs = (IEnumerable<object>)Evaluate(frame);
+						IEnumerable<object> rhs = (IEnumerable<object>)Evaluate(frame);
+						return lhs.Except(rhs);
+					}
+				case BgOpcode.ListSelect:
+					{
+						IEnumerable<object> source = (IEnumerable<object>)Evaluate(frame);
+						int fragment = ReadFragment(frame);
+
+						return source.Select(x => Call(fragment, new object[] { x }));
+					}
+				case BgOpcode.ListWhere:
+					{
+						IEnumerable<object> source = (IEnumerable<object>)Evaluate(frame);
+						int fragment = ReadFragment(frame);
+
+						return source.Where(x => (bool)Call(fragment, new object[] { x }));
+					}
+				case BgOpcode.ListDistinct:
+					{
+						IEnumerable<object> source = (IEnumerable<object>)Evaluate(frame);
+						return source.Distinct();
+					}
+				case BgOpcode.ListContains:
+					{
+						IEnumerable<object> source = (IEnumerable<object>)Evaluate(frame);
+						object item = Evaluate(frame);
+						return source.Contains(item);
+					}
+				case BgOpcode.ListLazy:
+					{
+						int fragment = ReadFragment(frame);
+						return LazyEvaluateList(new Frame(frame), fragment);
+					}
+				case BgOpcode.ListOption:
+					{
+						string name = (string)Evaluate(frame);
+						string label = (string)Evaluate(frame);
+						string description = (string)Evaluate(frame);
+						BgListOptionStyle style = (BgListOptionStyle)ReadUnsignedInteger(frame);
+						string defaultValue = (string)Evaluate(frame);
+						string pattern = (string)Evaluate(frame);
+						string patternFailed = (string)Evaluate(frame);
+						List<string> values = EvaluateList<string>(frame);
+						List<string> valueDescriptions = EvaluateList<string>(frame);
+
+						string value = defaultValue;
+						if (_options.TryGetValue(name, out string? option))
+						{
+							value = option;
+						}
+
+						return value.Split('+', ';').Select(x => (object)x).ToList();
+					}
+
+				#endregion
+				#region Fileset opcodes
+
+				case BgOpcode.FileSetFromNode:
+					{
+						BgBytecodeNode node = (BgBytecodeNode)Evaluate(frame);
+						return node.InputDependencies.SelectMany(x => x.Outputs).Concat(node.Outputs).ToArray();
+					}
+				case BgOpcode.FileSetFromNodeOutput:
+					{
+						BgBytecodeNode node = (BgBytecodeNode)Evaluate(frame);
+						int outputIndex = (int)ReadUnsignedInteger(frame);
+						return new[] { node.Outputs[outputIndex] };
+					}
+
+				#endregion
+				#region Object opcodes
+
+				case BgOpcode.ObjEmpty:
+					return ImmutableDictionary<string, object>.Empty.WithComparers(StringComparer.Ordinal);
+				case BgOpcode.ObjGet:
+					{
+						ImmutableDictionary<string, object> obj = (ImmutableDictionary<string, object>)Evaluate(frame);
+						string name = ReadString(frame);
+						object defaultValue = Evaluate(frame);
+
+						object value;
+						if (!obj.TryGetValue(name, out value))
+						{
+							value = defaultValue;
+						}
+
+						return value;
+					}
+				case BgOpcode.ObjSet:
+					{
+						ImmutableDictionary<string, object> obj = (ImmutableDictionary<string, object>)Evaluate(frame);
+						string name = ReadString(frame);
+						object value = Evaluate(frame);
+						return obj.SetItem(name, value);
+					}
+
+				#endregion
+				#region Function opcodes
+
+				case BgOpcode.Call:
+					{
+						int count = (int)ReadUnsignedInteger(frame);
+
+						object[] arguments = new object[count];
+						for (int idx = 0; idx < count; idx++)
+						{
+							arguments[idx] = Evaluate(frame);
+						}
+
+						int function = (int)ReadUnsignedInteger(frame);
+						return Evaluate(new Frame(_fragments[function], arguments));
+					}
+				case BgOpcode.Argument:
+					{
+						int index = (int)ReadUnsignedInteger(frame);
+						return frame.Arguments[index];
+					}
+				case BgOpcode.Jump:
+					{
+						int fragment = (int)ReadUnsignedInteger(frame);
+						return Jump(frame, fragment);
+					}
+
+				#endregion
+				#region Generic opcodes
+
+				case BgOpcode.Choose:
+					{
+						bool condition = (bool)Evaluate(frame);
+
+						int fragmentIfTrue = (int)ReadUnsignedInteger(frame);
+						int fragmentIfFalse = (int)ReadUnsignedInteger(frame);
+
+						return Jump(frame, condition ? fragmentIfTrue : fragmentIfFalse);
+					}
+				case BgOpcode.Throw:
+					{
+						string sourceFile = ReadString(frame);
+						int sourceLine = (int)ReadUnsignedInteger(frame);
+						string message = (string)Evaluate(frame);
+						throw new BgBytecodeException(sourceFile, sourceLine, message);
+					}
+				case BgOpcode.Null:
+					return null!;
+
+				#endregion
+				#region Graph opcodes
+
+				case BgOpcode.Node:
+					{
+						string name = (string)Evaluate(frame);
+						BgAgent agent = (BgAgent)Evaluate(frame);
+						BgMethod method = _methods[(int)ReadUnsignedInteger(frame)];
+
+						int argumentCount = (int)ReadUnsignedInteger(frame);
+						object[] arguments = new object[argumentCount];
+
+						for (int idx = 0; idx < argumentCount; idx++)
+						{
+							arguments[idx] = Evaluate(frame);
+						}
+
+						List<BgNodeOutput> inputs = ReadNodeOutputList(frame);
+						List<BgNodeOutput> fences = ReadNodeOutputList(frame);
+						int taggedOutputCount = (int)ReadUnsignedInteger(frame);
+						bool runEarly = (bool)Evaluate(frame);
+						List<BgLabel> labels = ((IEnumerable<object>)Evaluate(frame)).Select(x => (BgLabel)x).ToList();
+
+						HashSet<BgNode> inputDependencies = new HashSet<BgNode>();
+						foreach (BgNodeOutput input in inputs)
+						{
+							inputDependencies.Add(input.ProducingNode);
+							inputDependencies.UnionWith(input.ProducingNode.InputDependencies);
+						}
+
+						BgNode[] orderDependencies = fences.Select(x => x.ProducingNode).Distinct().ToArray();
+
+						string[] outputNames = Enumerable.Range(0, taggedOutputCount).Select(x => BgNodeSpec.GetDefaultTagName(name, x)).ToArray();
+
+						BgBytecodeNode node = new BgBytecodeNode(agent, name, method, arguments, inputs, outputNames, inputDependencies.ToArray(), orderDependencies.ToArray(), Array.Empty<FileReference>(), labels);
+						return node;
+					}
+				case BgOpcode.Agent:
+					{
+						string name = (string)Evaluate(frame);
+						string[] types = ((IEnumerable<object>)Evaluate(frame)).Select(x => (string)x).ToArray();
+						return new BgAgent(name, types);
+					}
+				case BgOpcode.Aggregate:
+					{
+						string name = (string)Evaluate(frame);
+						List<BgNodeOutput> requires = ReadNodeOutputList(frame);
+						BgLabel? label = (BgLabel?)Evaluate(frame);
+
+						return new BgBytecodeAggregate(name, requires.Select(x => x.ProducingNode), label);
+					}
+				case BgOpcode.Label:
+					{
+						string? dashboardName = NullForEmptyString((string)Evaluate(frame));
+						string? dashboardCategory = NullForEmptyString((string)Evaluate(frame));
+						string? ugsBadge = NullForEmptyString((string)Evaluate(frame));
+						string? ugsProject = NullForEmptyString((string)Evaluate(frame));
+						string change = (string)Evaluate(frame);
+
+						BgLabelChange labelChange = BgLabelChange.Current;
+						if (change.Length > 0)
+						{
+							labelChange = Enum.Parse<BgLabelChange>(change);
+						}
+
+						return new BgLabel(dashboardName, dashboardCategory, ugsBadge, ugsProject, labelChange);
+					}
+				case BgOpcode.Graph:
+					{
+						List<BgBytecodeNode> nodes = EvaluateList<BgBytecodeNode>(frame);
+						List<BgBytecodeAggregate> aggregates = EvaluateList<BgBytecodeAggregate>(frame);
+
+						BgGraph graph = new BgGraph();
+						foreach (BgBytecodeAggregate aggregate in aggregates)
+						{
+							graph.NameToAggregate[aggregate.Name] = aggregate;
+							nodes.AddRange(aggregate.RequiredNodes.Select(x => (BgBytecodeNode)x));
+						}
+
+						HashSet<BgNode> uniqueNodes = new HashSet<BgNode>();
+						HashSet<BgAgent> uniqueAgents = new HashSet<BgAgent>();
+						foreach (BgBytecodeNode node in nodes)
+						{
+							RegisterNode(graph, node, uniqueNodes, uniqueAgents);
+						}
+
+						HashSet<BgLabel> labels = new HashSet<BgLabel>();
+						foreach (BgBytecodeNode node in nodes)
+						{
+							foreach (BgLabel label in node.Labels)
+							{
+								labels.Add(label);
+								label.RequiredNodes.Add(node);
+							}
+						}
+						foreach (BgBytecodeAggregate aggregate in aggregates)
+						{
+							if (aggregate.Label != null)
+							{
+								aggregate.Label.RequiredNodes.UnionWith(aggregate.RequiredNodes);
+								labels.Add(aggregate.Label);
+							}
+						}
+						graph.Labels.AddRange(labels);
+
+						return graph;
+					}
+				#endregion
+
+				default:
+					throw new InvalidDataException($"Invalid opcode: {opcode}");
+			}
+		}
+
+		void RegisterNode(BgGraph graph, BgBytecodeNode node, HashSet<BgNode> uniqueNodes, HashSet<BgAgent> uniqueAgents)
+		{
+			if (uniqueNodes.Add(node))
+			{
+				foreach (BgBytecodeNode inputNode in node.InputDependencies)
+				{
+					RegisterNode(graph, inputNode, uniqueNodes, uniqueAgents);
+				}
+
+				BgAgent agent = node.Agent;
+				if (uniqueAgents.Add(agent))
+				{
+					graph.Agents.Add(agent);
+					graph.NameToAgent.Add(agent.Name, agent);
+				}
+
+				agent.Nodes.Add(node);
+				graph.NameToNode.Add(node.Name, node);
+			}
+		}
+
+		static string? NullForEmptyString(string str) => (str.Length == 0) ? null : str;
+
+		List<T> EvaluateList<T>(Frame frame)
+		{
+			return ((IEnumerable<object>)Evaluate(frame)).Select(x => (T)x).ToList();
+		}
+
+		List<BgNodeOutput> ReadNodeOutputList(Frame frame)
+		{
+			IEnumerable<object> list = (IEnumerable<object>)Evaluate(frame);
+
+			List<BgNodeOutput> outputs = new List<BgNodeOutput>();
+			foreach (object item in list)
+			{
+				outputs.AddRange((BgNodeOutput[])item);
+			}
+			return outputs;
+		}
+
+		IEnumerable<object> LazyEvaluateItem(Frame frame, int fragment)
+		{
+			yield return Jump(frame, fragment);
+		}
+
+		IEnumerable<object> LazyEvaluateList(Frame frame, int fragment)
+		{
+			IEnumerable<object> result = (IEnumerable<object>)Jump(frame, fragment);
+			foreach (object item in result)
+			{
+				yield return item;
+			}
+		}
+
+		object Jump(Frame frame, int fragment)
+		{
+			object? result;
+			if (!frame.Objects.TryGetValue(fragment, out result))
+			{
+				int prevOffset = frame.Offset;
+
+				frame.Offset = _fragments[fragment];
+
+				result = Evaluate(frame);
+				frame.Objects.Add(fragment, result);
+
+				frame.Offset = prevOffset;
+			}
+			return result;
+		}
+
+		object Call(int fragment, object[] arguments)
+		{
+			int offset = _fragments[fragment];
+			return Evaluate(new Frame(offset, arguments));
+		}
+
+		int ReadFragment(Frame frame)
+		{
+			return (int)ReadUnsignedInteger(frame);
+		}
+
+		/// <summary>
+		/// Reads an opcode from the input stream
+		/// </summary>
+		/// <returns>The opcode that was read</returns>
+		BgOpcode ReadOpcode(Frame frame)
+		{
+			BgOpcode opcode = (BgOpcode)_data[frame.Offset++];
+			return opcode;
+		}
+
+		/// <summary>
+		/// Reads a string from the input stream
+		/// </summary>
+		string ReadString(Frame frame)
+		{
+			ReadOnlySpan<byte> buffer = _data[frame.Offset..];
+
+			int length = (int)VarInt.ReadUnsigned(buffer, out int bytesRead);
+			frame.Offset += bytesRead;
+
+			string text = Encoding.UTF8.GetString(buffer.Slice(bytesRead, length));
+			frame.Offset += length;
+
+			return text;
+		}
+
+		/// <summary>
+		/// Writes a signed integer value to the output
+		/// </summary>
+		/// <returns>the value that was read</returns>
+		int ReadSignedInteger(Frame frame)
+		{
+			ulong encoded = ReadUnsignedInteger(frame);
+			return DecodeSignedInteger(encoded);
+		}
+
+		/// <summary>
+		/// Read an unsigned integer value from the input
+		/// </summary>
+		/// <returns>The value that was read</returns>
+		ulong ReadUnsignedInteger(Frame frame)
+		{
+			ulong value = VarInt.ReadUnsigned(_data[frame.Offset..], out int bytesRead);
+			frame.Offset += bytesRead;
+			return value;
+		}
+
+		/// <summary>
+		/// Decode a signed integer using the lower bit for the sign flag, allowing us to encode it more efficiently as a <see cref="VarInt"/>
+		/// </summary>
+		/// <param name="value">Value to be decoded</param>
+		/// <returns>The decoded value</returns>
+		static int DecodeSignedInteger(ulong value)
+		{
+			if ((value & 1) != 0)
+			{
+				return -(int)(value >> 1);
+			}
+			else
+			{
+				return (int)(value >> 1);
+			}
+		}
+
+		/// <summary>
+		/// Disassemble the current script to a logger
+		/// </summary>
+		/// <param name="logger"></param>
+		public void Disassemble(ILogger logger)
+		{
+			logger.LogInformation("Version: {Version}", _version);
+			for (int idx = 0; idx < _fragments.Length; idx++)
+			{
+				logger.LogInformation("");
+				logger.LogInformation("Fragment {Idx}:", idx);
+				Disassemble(new Frame(_fragments[idx], Array.Empty<object>()), logger);
+			}
+		}
+
+		void Disassemble(Frame frame, ILogger logger)
+		{
+			BgOpcode opcode = Trace(frame, null, ReadOpcode, logger);
+			switch (opcode)
+			{
+				#region Bool opcodes
+
+				case BgOpcode.BoolFalse:
+				case BgOpcode.BoolTrue:
+					break;
+				case BgOpcode.BoolNot:
+					Disassemble(frame, logger);
+					break;
+				case BgOpcode.BoolAnd:
+				case BgOpcode.BoolOr:
+				case BgOpcode.BoolXor:
+				case BgOpcode.BoolEq:
+					Disassemble(frame, logger);
+					Disassemble(frame, logger);
+					break;
+				case BgOpcode.BoolOption:
+					Disassemble(frame, logger);
+					Disassemble(frame, logger);
+					Disassemble(frame, logger);
+					Disassemble(frame, logger);
+					break;
+
+				#endregion
+				#region Integer opcodes
+
+				case BgOpcode.IntLiteral:
+					Trace(frame, null, ReadSignedInteger, logger);
+					break;
+				case BgOpcode.IntEq:
+				case BgOpcode.IntLt:
+				case BgOpcode.IntGt:
+				case BgOpcode.IntAdd:
+				case BgOpcode.IntMultiply:
+				case BgOpcode.IntDivide:
+				case BgOpcode.IntModulo:
+					Disassemble(frame, logger);
+					Disassemble(frame, logger);
+					break;
+				case BgOpcode.IntNegate:
+					Disassemble(frame, logger);
+					break;
+
+				#endregion
+				#region String opcodes
+
+				case BgOpcode.StrEmpty:
+					break;
+				case BgOpcode.StrLiteral:
+					Trace(frame, ReadString, x => $"\"{x}\"", logger);
+					break;
+				case BgOpcode.StrCompare:
+					Disassemble(frame, logger);
+					Disassemble(frame, logger);
+					Trace(frame, "type", ReadUnsignedInteger, logger);
+					break;
+				case BgOpcode.StrConcat:
+				case BgOpcode.StrSplit:
+				case BgOpcode.StrJoin:
+					Disassemble(frame, logger);
+					Disassemble(frame, logger);
+					break;
+				case BgOpcode.StrFormat:
+					{
+						Disassemble(frame, logger);
+
+						int count = (int)Trace(frame, "count", ReadUnsignedInteger, logger);
+						for (int idx = 0; idx < count; idx++)
+						{
+							Disassemble(frame, logger);
+						}
+
+						break;
+					}
+				case BgOpcode.StrIsMatch:
+					Disassemble(frame, logger);
+					Trace(frame, "pattern", ReadString, logger);
+					break;
+				case BgOpcode.StrReplace:
+					Disassemble(frame, logger);
+					Trace(frame, "pattern", ReadString, logger);
+					Disassemble(frame, logger);
+					break;
+				case BgOpcode.StrOption:
+					Disassemble(frame, logger);
+					Disassemble(frame, logger);
+					Disassemble(frame, logger);
+					Disassemble(frame, logger);
+					Trace(frame, "style", ReadUnsignedInteger, logger);
+					Disassemble(frame, logger);
+					Disassemble(frame, logger);
+					Disassemble(frame, logger);
+					Disassemble(frame, logger);
+					Disassemble(frame, logger);
+					break;
+
+				#endregion
+				#region Enum opcodes
+
+				case BgOpcode.EnumConstant:
+					Trace(frame, "value", ReadSignedInteger, logger);
+					break;
+				case BgOpcode.EnumParse:
+				case BgOpcode.EnumToString:
+					Disassemble(frame, logger);
+					Disassemble(frame, logger);
+					Disassemble(frame, logger);
+					break;
+
+				#endregion
+				#region List opcodes
+
+				case BgOpcode.ListEmpty:
+					break;
+				case BgOpcode.ListPush:
+					Disassemble(frame, logger);
+					Disassemble(frame, logger);
+					break;
+				case BgOpcode.ListPushLazy:
+				case BgOpcode.ListSelect:
+				case BgOpcode.ListWhere:
+					Disassemble(frame, logger);
+					TraceFragment(frame, logger);
+					break;
+				case BgOpcode.ListCount:
+				case BgOpcode.ListDistinct:
+					Disassemble(frame, logger);
+					break;
+				case BgOpcode.ListElement:
+				case BgOpcode.ListConcat:
+				case BgOpcode.ListUnion:
+				case BgOpcode.ListExcept:
+				case BgOpcode.ListContains:
+					Disassemble(frame, logger);
+					Disassemble(frame, logger);
+					break;
+				case BgOpcode.ListLazy:
+					TraceFragment(frame, logger);
+					break;
+				case BgOpcode.ListOption:
+					Disassemble(frame, logger);
+					Disassemble(frame, logger);
+					Disassemble(frame, logger);
+					Trace(frame, "style", ReadUnsignedInteger, logger);
+					Disassemble(frame, logger);
+					Disassemble(frame, logger);
+					Disassemble(frame, logger);
+					Disassemble(frame, logger);
+					Disassemble(frame, logger);
+					break;
+
+				#endregion
+				#region Fileset opcodes
+
+				case BgOpcode.FileSetFromNode:
+					Disassemble(frame, logger);
+					break;
+				case BgOpcode.FileSetFromNodeOutput:
+					Disassemble(frame, logger);
+					TraceFragment(frame, logger);
+					break;
+
+				#endregion
+				#region Object opcodes
+
+				case BgOpcode.ObjEmpty:
+					break;
+				case BgOpcode.ObjGet:
+				case BgOpcode.ObjSet:
+					Disassemble(frame, logger);
+					Trace(frame, "name", ReadString, logger);
+					Disassemble(frame, logger);
+					break;
+
+				#endregion
+				#region Function opcodes
+
+				case BgOpcode.Call:
+					{
+						int count = (int)Trace(frame, "count", ReadUnsignedInteger, logger);
+						for (int idx = 0; idx < count; idx++)
+						{
+							Disassemble(frame, logger);
+						}
+
+						TraceFragment(frame, logger);
+						break;
+					}
+				case BgOpcode.Argument:
+					Trace(frame, "arg", ReadUnsignedInteger, logger);
+					break;
+				case BgOpcode.Jump:
+					TraceFragment(frame, logger);
+					break;
+
+				#endregion
+				#region Generic opcodes
+
+				case BgOpcode.Choose:
+					Disassemble(frame, logger);
+					TraceFragment(frame, logger);
+					TraceFragment(frame, logger);
+					break;
+				case BgOpcode.Throw:
+					Trace(frame, "file", ReadString, logger);
+					Trace(frame, "line", ReadUnsignedInteger, logger);
+					Trace(frame, "message", ReadString, logger);
+					break;
+				case BgOpcode.Null:
+					break;
+
+				#endregion
+				#region Graph opcodes
+
+				case BgOpcode.Node:
+					{
+						Disassemble(frame, logger);
+						Disassemble(frame, logger);
+						Trace(frame, "handler", ReadUnsignedInteger, logger);
+
+						int argCount = (int)Trace(frame, "args", ReadUnsignedInteger, logger);
+						for (int idx = 0; idx < argCount; idx++)
+						{
+							Disassemble(frame, logger);
+						}
+
+						Disassemble(frame, logger);
+						Disassemble(frame, logger);
+						Disassemble(frame, logger);
+						Disassemble(frame, logger);
+						break;
+					}
+				case BgOpcode.Agent:
+					Disassemble(frame, logger);
+					Disassemble(frame, logger);
+					break;
+				case BgOpcode.Aggregate:
+					Disassemble(frame, logger);
+					Disassemble(frame, logger);
+					break;
+				case BgOpcode.Label:
+					Disassemble(frame, logger);
+					Disassemble(frame, logger);
+					Disassemble(frame, logger);
+					Disassemble(frame, logger);
+					Disassemble(frame, logger);
+					Disassemble(frame, logger);
+					break;
+				case BgOpcode.Graph:
+					Disassemble(frame, logger);
+					Disassemble(frame, logger);
+					break;
+				#endregion
+
+				default:
+					throw new InvalidDataException($"Invalid opcode: {opcode}");
+			}
+		}
+
+		int TraceFragment(Frame frame, ILogger logger)
+		{
+			return Trace(frame, "-> Fragment", ReadFragment, logger);
+		}
+
+		T Trace<T>(Frame frame, string? type, Func<Frame, T> readValue, ILogger logger)
+		{
+			int offset = frame.Offset;
+			T value = readValue(frame);
+			int length = frame.Offset - offset;
+
+			string valueAndType = (type == null) ? $"{value}" : $"{type} {value}";
+			Trace<T>(offset, length, valueAndType, logger);
+			return value;
+		}
+
+		T Trace<T>(Frame frame, Func<Frame, T> readValue, Func<T, object> formatValue, ILogger logger)
+		{
+			int offset = frame.Offset;
+			T value = readValue(frame);
+			int length = frame.Offset - offset;
+
+			Trace<T>(offset, length, formatValue(value), logger);
+			return value;
+		}
+
+		void Trace<T>(int offset, int length, object value, ILogger logger)
+		{
+			string bytes = String.Join(" ", _data.AsSpan(offset, length).ToArray().Select(x => $"{x:x2}"));
+			logger.LogInformation("{Offset,6}: {Value,-20} {Bytes}", offset, value, bytes);
+		}
+	}
+}
diff --git a/Engine/Source/Programs/Shared/EpicGames.BuildGraph/BgLabelSpec.cs b/Engine/Source/Programs/Shared/EpicGames.BuildGraph/BgLabelSpec.cs
--- a/Engine/Source/Programs/Shared/EpicGames.BuildGraph/BgLabelSpec.cs
+++ b/Engine/Source/Programs/Shared/EpicGames.BuildGraph/BgLabelSpec.cs
@@ -7,79 +7,60 @@ using EpicGames.BuildGraph.Expressions;
 namespace EpicGames.BuildGraph
 {
 	/// <summary>
-	/// Configuration for a label
+	/// Specification for a label
 	/// </summary>
-	public class BgLabelConfig
+	public class BgLabelSpec : BgExpr
 	{
 		/// <summary>
 		/// Name of this badge
 		/// </summary>
-		public BgString? DashboardName { get; set; }
+		public BgString? DashboardName { get; }
 
 		/// <summary>
 		/// Category for this label
 		/// </summary>
-		public BgString? DashboardCategory { get; set; }
+		public BgString? DashboardCategory { get; }
 
 		/// <summary>
 		/// Name of the badge in UGS
 		/// </summary>
-		public BgString? UgsBadge { get; set; }
+		public BgString? UgsBadge { get; }
 
 		/// <summary>
 		/// Path to the project folder in UGS
 		/// </summary>
-		public BgString? UgsProject { get; set; }
+		public BgString? UgsProject { get; }
 
 		/// <summary>
 		/// Which change to show the badge for
 		/// </summary>
-		public BgString? Change { get; set; }
-
-		/// <summary>
-		/// Set of nodes that must be run for this label to be shown.
-		/// </summary>
-		public BgList<BgFileSet> RequiredNodes { get; set; } = BgList<BgFileSet>.Empty;
+		public BgString? Change { get; }
 
 		/// <summary>
-		/// Set of nodes that will be included in this label if present.
+		/// Constructor
 		/// </summary>
-		public BgList<BgFileSet> IncludedNodes { get; set; } = BgList<BgFileSet>.Empty;
-	}
-
-	/// <summary>
-	/// Specification for a label
-	/// </summary>
-	public class BgLabelSpec
-	{
-		private readonly BgLabelConfig _config;
-
-		internal BgLabelSpec(BgLabelConfig config)
+		public BgLabelSpec(BgString? name = null, BgString? category = null, BgString? ugsBadge = null, BgString? ugsProject = null, BgString? change = null)
+			: base(BgExprFlags.ForceFragment)
 		{
-			_config = config;
+			DashboardName = name;
+			DashboardCategory = category;
+			UgsBadge = ugsBadge;
+			UgsProject = ugsProject;
+			Change = change;
 		}
 
-		internal void AddToGraph(BgExprContext context, BgGraph graph)
+		/// <inheritdoc/>
+		public override void Write(BgBytecodeWriter writer)
 		{
-			string? dashboardName = _config.DashboardName?.Compute(context);
-			string? dashboardCategory = _config.DashboardCategory?.Compute(context);
-			string? ugsBadge = _config.UgsBadge?.Compute(context);
-			string? ugsProject = _config.UgsBadge?.Compute(context);
-
-			BgLabelChange labelChange;
-			if (_config.Change is null)
-			{
-				labelChange = BgLabelChange.Current;
-			}
-			else
-			{
-				labelChange = Enum.Parse<BgLabelChange>(_config.Change.Compute(context));
-			}
-
-			BgLabel label = new BgLabel(dashboardName, dashboardCategory, ugsBadge, ugsProject, labelChange);
-			label.RequiredNodes.UnionWith(_config.RequiredNodes.ComputeTags(context).Select(x => graph.TagNameToNodeOutput[x].ProducingNode));
-			label.IncludedNodes.UnionWith(_config.IncludedNodes.ComputeTags(context).Select(x => graph.TagNameToNodeOutput[x].ProducingNode));
-			graph.Labels.Add(label);
+			writer.WriteOpcode(BgOpcode.Label);
+			writer.WriteExpr(DashboardName ?? BgString.Empty);
+			writer.WriteExpr(DashboardCategory ?? BgString.Empty);
+			writer.WriteExpr(UgsBadge ?? BgString.Empty);
+			writer.WriteExpr(UgsProject ?? BgString.Empty);
+			writer.WriteExpr(Change ?? BgString.Empty);
 		}
+
+		/// <inheritdoc/>
+		public override BgString ToBgString() => "{Label}";
 	}
 }
diff --git a/Engine/Source/Programs/Shared/EpicGames.BuildGraph/BgNode.cs b/Engine/Source/Programs/Shared/EpicGames.BuildGraph/BgNode.cs
--- a/Engine/Source/Programs/Shared/EpicGames.BuildGraph/BgNode.cs
+++ b/Engine/Source/Programs/Shared/EpicGames.BuildGraph/BgNode.cs
@@ -43,67 +43,10 @@ namespace EpicGames.BuildGraph
 		/// <returns>The name of this output</returns>
 		public override string ToString()
 		{
-			return String.Format("{0}: {1}", ProducingNode.Name, TagName);
+			return String.Format("{0} [{1}]", TagName, ProducingNode.Name);
 		}
 	}
 
-	/// <summary>
-	/// Method to be executed for a node
-	/// </summary>
-	public class BgMethod
-	{
-		/// <summary>
-		/// Full name of the class containing the method to execute
-		/// </summary>
-		public string ClassName { get; }
-
-		/// <summary>
-		/// Name of the method to execute
-		/// </summary>
-		public string MethodName { get; }
-
-		/// <summary>
-		/// Constructor
-		/// </summary>
-		public BgMethod(string className, string methodName)
-		{
-			ClassName = className;
-			MethodName = methodName;
-		}
-
-		/// <summary>
-		/// Constructor
-		/// </summary>
-		public BgMethod(MethodInfo method)
-			: this(method.DeclaringType!.AssemblyQualifiedName!, method.Name)
-		{
-		}
-
-		/// <summary>
-		/// Bind this method name to a method instance
-		/// </summary>
-		/// <returns>The resolved method instance</returns>
-		public MethodInfo Bind()
-		{
-			Type? Type = Type.GetType(ClassName);
-			if (Type == null)
-			{
-				throw new BgNodeException($"Unable to find class '{ClassName}'");
-			}
-
-			MethodInfo? Method = Type.GetMethod(MethodName, BindingFlags.Public | BindingFlags.NonPublic | BindingFlags.Static);
-			if (Method == null)
-			{
-				throw new BgNodeException($"Unable to find method '{Type.FullName}.{MethodName}'");
-			}
-
-			return Method;
-		}
-
-		/// <inheritdoc/>
-		public override string ToString() => $"{MethodName}";
-	}
-
 	/// <summary>
 	/// Defines a node, a container for tasks and the smallest unit of execution that can be run as part of a build graph.
 	/// </summary>
@@ -165,7 +108,7 @@ namespace EpicGames.BuildGraph
 		public Dictionary<string, string> Annotations { get; } = new Dictionary<string, string>(StringComparer.OrdinalIgnoreCase);
 
 		/// <summary>
-		/// Diagnostics for this node
+		/// Diagnostics to output if executing this node
 		/// </summary>
 		public List<BgDiagnostic> Diagnostics { get; } = new List<BgDiagnostic>();
 
@@ -239,26 +182,38 @@ namespace EpicGames.BuildGraph
 	/// <summary>
 	/// Node constructed from a bytecode expression
 	/// </summary>
-	public class BgExpressionNode : BgNode
+	public class BgBytecodeNode : BgNode
 	{
 		/// <summary>
-		/// The method to execute for this node
+		/// Agent declaring this node
+		/// </summary>
+		public BgAgent Agent { get; }
+
+		/// <summary>
+		/// Offset in the bytecode of this node. Used to lookup method calling information.
 		/// </summary>
 		public BgMethod Method { get; }
 
 		/// <summary>
-		/// Arguments for invoking the method
+		/// Expression arguments to the method
+		/// </summary>
+		public IReadOnlyList<object> Arguments { get; }
+
+		/// <summary>
+		/// Labels to add this node to
 		/// </summary>
-		public IReadOnlyList<object?> Arguments { get; }
+		public IReadOnlyList<BgLabel> Labels { get; }
 
 		/// <summary>
 		/// Constructor
 		/// </summary>
-		public BgExpressionNode(string name, BgMethod method, IReadOnlyList<object?> arguments, IReadOnlyList<BgNodeOutput> inputs, IReadOnlyList<string> outputNames, IReadOnlyList<BgNode> inputDependencies, IReadOnlyList<BgNode> orderDependencies, IReadOnlyList<FileReference> requiredTokens)
+		public BgBytecodeNode(BgAgent agent, string name, BgMethod method, IReadOnlyList<object> arguments, IReadOnlyList<BgNodeOutput> inputs, IReadOnlyList<string> outputNames, IReadOnlyList<BgNode> inputDependencies, IReadOnlyList<BgNode> orderDependencies, IReadOnlyList<FileReference> requiredTokens, IReadOnlyList<BgLabel> labels)
 			: base(name, inputs, outputNames, inputDependencies, orderDependencies, requiredTokens)
 		{
+			Agent = agent;
 			Method = method;
 			Arguments = arguments;
+			Labels = labels;
 		}
 	}
 }
diff --git a/Engine/Source/Programs/Shared/EpicGames.BuildGraph/BgNodeSpec.cs b/Engine/Source/Programs/Shared/EpicGames.BuildGraph/BgNodeSpec.cs
--- a/Engine/Source/Programs/Shared/EpicGames.BuildGraph/BgNodeSpec.cs
+++ b/Engine/Source/Programs/Shared/EpicGames.BuildGraph/BgNodeSpec.cs
@@ -9,7 +9,6 @@ using System.Text;
 using System.Threading.Tasks;
 using EpicGames.BuildGraph.Expressions;
 using EpicGames.Core;
-using Microsoft.Extensions.Logging;
 
 namespace EpicGames.BuildGraph
 {
@@ -49,159 +48,319 @@ namespace EpicGames.BuildGraph
 	}
 
 	/// <summary>
-	/// Specification for a node to execute
+	/// Information about the method bound to execute a node
 	/// </summary>
-	public class BgNodeSpec
+	public class BgMethod
 	{
-		internal BgMethod _method;
-		internal object?[] _arguments;
-		internal BgFileSet[] ResultExprs { get; }
+		/// <summary>
+		/// Method to call
+		/// </summary>
+		public MethodInfo Method { get; }
+
+		/// <summary>
+		/// Arguments to the method
+		/// </summary>
+		public IReadOnlyList<object?> Arguments { get; }
+
+		/// <summary>
+		/// Constructor
+		/// </summary>
+		public BgMethod(MethodInfo method, IReadOnlyList<object?> arguments)
+		{
+			Method = method;
+			Arguments = arguments;
+		}
+	}
 
+	/// <summary>
+	/// Specification for a node to execute
+	/// </summary>
+	public class BgNodeSpec : BgExpr
+	{
 		/// <summary>
 		/// Name of the node
 		/// </summary>
-		public BgString Name { get; set; }
+		public BgString Name { get; }
 
 		/// <summary>
-		/// Tokens for inputs of this node
+		/// Method to execute
 		/// </summary>
-		public BgList<BgFileSet> InputDependencies { get; set; } = BgList<BgFileSet>.Empty;
+		protected BgMethod Method { get; }
 
 		/// <summary>
-		/// Weak dependency on outputs that must be generated for the node to run, without making those dependencies inputs.
+		/// Expression arguments from the bound function.
+		/// </summary>
+		protected BgExpr[] ExpressionArguments { get; }
+
+		/// <summary>
+		/// The default output of this node. Includes all other outputs. 
+		/// </summary>
+		public BgFileSet DefaultOutput { get; }
+
+		/// <summary>
+		/// Explicitly tagged outputs from this node
+		/// </summary>
+		public IReadOnlyList<BgExpr> TaggedOutputs { get; protected set; } = Array.Empty<BgExpr>();
+
+		/// <summary>
+		/// Agent for the node to be run on
 		/// </summary>
-		public BgList<BgFileSet> AfterDependencies { get; set; } = BgList<BgFileSet>.Empty;
+		public BgAgentSpec Agent { get; }
 
 		/// <summary>
-		/// Token for the output of this node
+		/// Tokens for inputs of this node
 		/// </summary>
-		BgFileSet Output { get; }
+		public BgList<BgFileSet> Inputs { get; }
 
 		/// <summary>
-		/// All inputs and outputs of this node
+		/// Weak dependency on outputs that must be generated for the node to run, without making those dependencies inputs.
 		/// </summary>
-		BgList<BgFileSet> InputsAndOutputs => InputDependencies.Add(Output);
+		public BgList<BgFileSet> Fences { get; }
 
 		/// <summary>
 		/// Whether this node should start running as soon as its dependencies are ready, even if other nodes in the same agent are not.
 		/// </summary>
-		public BgBool CanRunEarly { get; set; } = BgBool.False;
+		public BgBool RunEarly { get; } = BgBool.False;
 
 		/// <summary>
-		/// Diagnostics for running this node
+		/// Labels that this node contributes to
 		/// </summary>
-		public BgList<BgDiagnosticSpec> Diagnostics { get; set; } = BgList<BgDiagnosticSpec>.Empty;
+		public BgList<BgLabelSpec> Labels { get; }
 
 		/// <summary>
 		/// Constructor
 		/// </summary>
-		internal BgNodeSpec(MethodCallExpression call)
+		public BgNodeSpec(BgString name, BgMethod method, BgAgentSpec agent, BgList<BgFileSet> inputs, BgList<BgFileSet> fences, BgBool runEarly, BgList<BgLabelSpec> labels)
+			: base(BgExprFlags.ForceFragment)
 		{
-			_method = new BgMethod(call.Method);
+			Name = name;
 
-			try
+			ParameterInfo[] parameters = method.Method.GetParameters();
+			ExpressionArguments = new BgExpr[parameters.Length];
+
+			object?[] constantArguments = new object?[method.Arguments.Count];
+			for (int idx = 0; idx < parameters.Length; idx++)
 			{
-				_arguments = CreateArgumentExprs(call);
-				Name = GetDefaultNodeName(call.Method, _arguments);
-				ResultExprs = CreateReturnExprs(call.Method, Name, call.Method.ReturnType);
+				if (typeof(BgExpr).IsAssignableFrom(parameters[idx].ParameterType))
+				{
+					ExpressionArguments[idx] = (BgExpr)method.Arguments[idx]!;
+					constantArguments[idx] = null;
+				}
+				else
+				{
+					ExpressionArguments[idx] = BgBool.False;
+					constantArguments[idx] = method.Arguments[idx];
+				}
 			}
-			catch (Exception ex)
+			Method = new BgMethod(method.Method, constantArguments);
+
+			Agent = agent;
+			DefaultOutput = new BgFileSetFromNodeOutputExpr(this, 0);
+			Inputs = inputs;
+			Fences = fences;
+			RunEarly = runEarly;
+			Labels = labels;
+		}
+
+		int written = 0;
+
+		/// <inheritdoc/>
+		public override void Write(BgBytecodeWriter writer)
+		{
+			writer.WriteOpcode(BgOpcode.Node);
+
+			writer.WriteExpr(Name ?? BgString.Empty);
+			writer.WriteExpr(Agent);
+			writer.WriteMethod(Method);
+
+			writer.WriteUnsignedInteger((uint)ExpressionArguments.Length);
+			foreach (BgExpr argument in ExpressionArguments)
 			{
-				ExceptionUtils.AddContext(ex, $"while calling method {_method}");
-				throw;
+				writer.WriteExpr(argument);
 			}
 
-			Output = new BgFileSetTagFromStringExpr(BgString.Format("#{0}", Name));
+			writer.WriteExpr(Inputs);
+			writer.WriteExpr(Fences);
+			writer.WriteUnsignedInteger(TaggedOutputs.Count);
+			writer.WriteExpr(RunEarly);
+			writer.WriteExpr(Labels);
+			written++;
 		}
 
 		/// <summary>
-		/// Creates a node specification
+		/// Gets the default tag name for the numbered output index
 		/// </summary>
-		internal static BgNodeSpec Create(Expression<Func<BgContext, Task>> action)
+		/// <param name="name">Name of the node</param>
+		/// <param name="index">Index of the output. Index zero is the default, others are explicit.</param>
+		/// <returns></returns>
+		internal static string GetDefaultTagName(string name, int index)
 		{
-			MethodCallExpression call = (MethodCallExpression)action.Body;
-			return new BgNodeSpec(call);
+			return $"#{name}${index}";
 		}
 
 		/// <summary>
-		/// Creates a node specification
+		/// Implicit conversion to a fileset
 		/// </summary>
-		internal static BgNodeSpec<T> Create<T>(Expression<Func<BgContext, Task<T>>> function)
+		/// <param name="node"></param>
+		public static implicit operator BgFileSet(BgNodeSpec node)
 		{
-			MethodCallExpression call = (MethodCallExpression)function.Body;
-			return new BgNodeSpec<T>(call);
+			return new BgFileSetFromNodeExpr(node);
 		}
 
-		static object?[] CreateArgumentExprs(MethodCallExpression call)
+		/// <summary>
+		/// Implicit conversion to a fileset
+		/// </summary>
+		/// <param name="node"></param>
+		public static implicit operator BgList<BgFileSet>(BgNodeSpec node)
 		{
-			object?[] args = new object?[call.Arguments.Count];
-			for (int idx = 0; idx < call.Arguments.Count; idx++)
-			{
-				Expression expr = call.Arguments[idx];
-				if (expr is ParameterExpression parameterExpr)
-				{
-					if (parameterExpr.Type != typeof(BgContext))
-					{
-						throw new BgNodeException($"Unable to determine type of parameter '{parameterExpr.Name}'");
-					}
-				}
-				else
-				{
-					Delegate compiled = Expression.Lambda(expr).Compile();
-					args[idx] = compiled.DynamicInvoke();
-				}
-			}
-			return args;
+			return (BgFileSet)node;
 		}
 
-		static BgFileSet[] CreateReturnExprs(MethodInfo methodInfo, BgString name, Type type)
+		/// <inheritdoc/>
+		public override BgString ToBgString() => Name ?? BgString.Empty;
+	}
+
+	/// <summary>
+	/// Nodespec with a typed return value
+	/// </summary>
+	/// <typeparam name="T"></typeparam>
+	public class BgNodeSpec<T> : BgNodeSpec
+	{
+		/// <summary>
+		/// Output from this node
+		/// </summary>
+		public T Output { get; }
+
+		/// <summary>
+		/// Constructor
+		/// </summary>
+		public BgNodeSpec(BgString name, BgMethod handler, BgAgentSpec agent, BgList<BgFileSet> inputs, BgList<BgFileSet> fences, BgBool runEarly, BgList<BgLabelSpec> labels)
+			: base(name, handler, agent, inputs, fences, runEarly, labels)
+		{
+			Output = CreateOutput();
+		}
+
+		T CreateOutput()
 		{
-			if (type == typeof(Task))
+			Type type = typeof(T);
+			if (IsValueTuple(type))
 			{
-				return Array.Empty<BgFileSet>();
+				BgExpr[] outputs = CreateOutputExprs(type.GetGenericArguments());
+				TaggedOutputs = outputs;
+				return (T)Activator.CreateInstance(type, outputs)!;
 			}
-			else if (type.IsGenericType && type.GetGenericTypeDefinition() == typeof(Task<>))
+			else
 			{
-				return CreateReturnExprsInner(methodInfo, name, type.GetGenericArguments()[0]);
+				BgExpr[] outputs = CreateOutputExprs(new[] { type });
+				TaggedOutputs = outputs;
+				return (T)(object)outputs[0];
 			}
-			else
+		}
+
+		BgExpr[] CreateOutputExprs(Type[] types)
+		{
+			BgExpr[] outputs = new BgExpr[types.Length];
+			for (int idx = 0; idx < types.Length; idx++)
 			{
-				throw new NotImplementedException();
+				outputs[idx] = CreateOutputExpr(types[idx], idx);
 			}
+			return outputs;
 		}
 
-		static BgFileSet[] CreateReturnExprsInner(MethodInfo methodInfo, BgString name, Type type)
+		BgExpr CreateOutputExpr(Type type, int index)
 		{
-			Type[] outputTypes;
-			if (IsValueTuple(type))
+			if (type == typeof(BgFileSet))
 			{
-				outputTypes = type.GetGenericArguments();
+				return new BgFileSetFromNodeOutputExpr(this, index);
 			}
 			else
 			{
-				outputTypes = new[] { type };
+				throw new NotImplementedException();
 			}
+		}
 
-			BgFileSet[] outputExprs = new BgFileSet[outputTypes.Length];
-			for (int idx = 0; idx < outputTypes.Length; idx++)
+		internal static bool IsValueTuple(Type returnType)
+		{
+			if (returnType.IsGenericType)
 			{
-				outputExprs[idx] = CreateOutputExpr(methodInfo, name, idx, outputTypes[idx]);
+				Type genericType = returnType.GetGenericTypeDefinition();
+				if (genericType.FullName != null && genericType.FullName.StartsWith("System.ValueTuple`", StringComparison.Ordinal))
+				{
+					return true;
+				}
 			}
-			return outputExprs;
+			return false;
 		}
+	}
+
+	/// <summary>
+	/// Class that allows constructing a nodespec
+	/// </summary>
+	public class BgNodeSpecBuilder
+	{
+		/// <inheritdoc cref="BgNodeSpec.Name"/>
+		public BgString Name { get; }
+
+		/// <inheritdoc cref="BgNodeSpec.Method"/>
+		protected BgMethod Handler { get; }
+
+		/// <inheritdoc cref="BgNodeSpec.Agent"/>
+		public BgAgentSpec Agent { get; }
 
-		static BgFileSet CreateOutputExpr(MethodInfo methodInfo, BgString name, int index, Type type)
+		/// <inheritdoc cref="BgNodeSpec.Inputs"/>
+		public BgList<BgFileSet> Inputs { get; set; } = BgList<BgFileSet>.Empty;
+
+		/// <inheritdoc cref="BgNodeSpec.Fences"/>
+		public BgList<BgFileSet> Fences { get; set; } = BgList<BgFileSet>.Empty;
+
+		/// <inheritdoc cref="BgNodeSpec.RunEarly"/>
+		public BgBool RunEarly { get; set; } = BgBool.False;
+
+		/// <inheritdoc cref="BgNodeSpec.Labels"/>
+		public BgList<BgLabelSpec> Labels { get; set; } = BgList<BgLabelSpec>.Empty;
+
+		/// <summary>
+		/// Constructor
+		/// </summary>
+		public BgNodeSpecBuilder(BgString? name, MethodCallExpression call, BgAgentSpec agent)
 		{
-			if (type == typeof(BgFileSet))
+			try
 			{
-				return new BgFileSetTagFromStringExpr(BgString.Format("#{0}${1}", name, index));
+				object?[] arguments = CreateArgumentExprs(call);
+
+				Name = name ?? GetDefaultNodeName(call.Method, arguments);
+				Handler = new BgMethod(call.Method, arguments);
+				Agent = agent;
 			}
-			else
+			catch (Exception ex)
 			{
-				throw new BgNodeException($"Unsupported return type for {methodInfo.Name}: {type.Name}");
+				ExceptionUtils.AddContext(ex, $"while calling method {call.Method}");
+				throw;
 			}
 		}
 
+		static object?[] CreateArgumentExprs(MethodCallExpression call)
+		{
+			object?[] args = new object?[call.Arguments.Count];
+			for (int idx = 0; idx < call.Arguments.Count; idx++)
+			{
+				Expression argumentExpr = call.Arguments[idx];
+				if (argumentExpr is ParameterExpression parameterExpr)
+				{
+					if (parameterExpr.Type != typeof(BgContext))
+					{
+						throw new BgNodeException($"Unable to determine type of parameter '{parameterExpr.Name}'");
+					}
+				}
+				else
+				{
+					Delegate compiled = Expression.Lambda(argumentExpr).Compile();
+					args[idx] = compiled.DynamicInvoke();
+				}
+			}
+			return args;
+		}
+
 		internal static bool IsValueTuple(Type returnType)
 		{
 			if (returnType.IsGenericType)
@@ -266,7 +425,7 @@ namespace EpicGames.BuildGraph
 						object? arg = args[paramIdx];
 						if (arg != null)
 						{
-							if (arg is IBgExpr expr)
+							if (arg is BgExpr expr)
 							{
 								fragments.Add(expr.ToBgString());
 							}
@@ -290,7 +449,14 @@ namespace EpicGames.BuildGraph
 			}
 			fragments.Add(template.Substring(lastIdx, template.Length - lastIdx));
 
-			return BgString.Join(BgString.Empty, fragments);
+			if (fragments.Count == 1)
+			{
+				return fragments[0];
+			}
+			else
+			{
+				return BgString.Join(BgString.Empty, fragments);
+			}
 		}
 
 		/// <summary>
@@ -332,160 +498,115 @@ namespace EpicGames.BuildGraph
 		}
 
 		/// <summary>
-		/// Gets the signature for a method
+		/// Construct a nodespec from the current parameters
 		/// </summary>
-		public static string GetSignature(MethodInfo methodInfo)
+		/// <returns>New nodespec</returns>
+		public BgNodeSpec Construct()
 		{
-			StringBuilder arguments = new StringBuilder();
-			foreach (ParameterInfo parameterInfo in methodInfo.GetParameters())
-			{
-				arguments.AppendLine(parameterInfo.ParameterType.FullName);
-			}
-			return $"{methodInfo.Name}:{Digest.Compute<Sha1>(arguments.ToString())}";
+			return new BgNodeSpec(Name, Handler, Agent, Inputs, Fences, RunEarly, Labels);
 		}
+	}
 
+	/// <summary>
+	/// Specification for a node to execute that returns one or more tagged outputs
+	/// </summary>
+	/// <typeparam name="T">Return type from the node. Must consist of latent Bg* types (eg. BgToken)</typeparam>
+	public class BgNodeSpecBuilder<T> : BgNodeSpecBuilder
+	{
 		/// <summary>
-		/// Creates a concrete node
+		/// Constructor
 		/// </summary>
-		/// <param name="context"></param>
-		/// <param name="graph"></param>
-		/// <param name="agent"></param>
-		/// <returns></returns>
-		internal void AddToGraph(BgExprContext context, BgGraph graph, BgAgent agent)
+		internal BgNodeSpecBuilder(BgString? name, MethodCallExpression call, BgAgentSpec agent)
+			: base(name, call, agent)
 		{
-			HashSet<string> inputTags = new HashSet<string>();
-			inputTags.UnionWith(InputDependencies.ComputeTags(context));
-			inputTags.UnionWith(_arguments.OfType<BgFileSet>().Select(x => x.ComputeTag(context)));
-			inputTags.UnionWith(_arguments.OfType<BgList<BgFileSet>>().SelectMany(x => x.GetEnumerable(context)).Select(x => x.ComputeTag(context)));
-
-			HashSet<string> afterTags = new HashSet<string>(inputTags);
-			afterTags.UnionWith(AfterDependencies.ComputeTags(context));
-
-			string name = Name.Compute(context);
-			BgNodeOutput[] inputDependencies = inputTags.Select(x => graph.TagNameToNodeOutput[x]).Distinct().ToArray();
-			BgNodeOutput[] afterDependencies = afterTags.Select(x => graph.TagNameToNodeOutput[x]).Distinct().ToArray();
-			string[] outputNames = Array.ConvertAll(ResultExprs, x => x.ComputeTag(context));
-			BgNode[] inputNodes = inputDependencies.Select(x => x.ProducingNode).Distinct().ToArray();
-			BgNode[] afterNodes = afterDependencies.Select(x => x.ProducingNode).Distinct().ToArray();
-			bool runEarly = CanRunEarly.Compute(context);
-
-			List<object?> arguments = _arguments.ConvertAll(x => (x is IBgExpr expr) ? expr.Compute(context) : x);
-			BgExpressionNode node = new BgExpressionNode(name, _method, arguments, inputDependencies, outputNames, inputNodes, afterNodes, Array.Empty<FileReference>());
-			node.RunEarly = runEarly;
-
-			agent.Nodes.Add(node);
-			graph.NameToNode.Add(name, node);
-
-			foreach (BgDiagnosticSpec precondition in Diagnostics.GetEnumerable(context))
-			{
-				precondition.AddToGraph(context, graph, agent, node);
-			}
-
-			foreach (BgNodeOutput output in node.Outputs)
-			{
-				graph.TagNameToNodeOutput.Add(output.TagName, output);
-			}
 		}
 
 		/// <summary>
-		/// Allows using the 
+		/// Construct a nodespec from the current parameters
 		/// </summary>
-		/// <param name="nodeSpec"></param>
-		public static implicit operator BgList<BgFileSet>(BgNodeSpec nodeSpec)
+		/// <returns>New nodespec</returns>
+		public new BgNodeSpec<T> Construct()
 		{
-			return nodeSpec.InputsAndOutputs.Add(nodeSpec.Output);
+			return new BgNodeSpec<T>(Name, Handler, Agent, Inputs, Fences, RunEarly, Labels);
 		}
 	}
 
 	/// <summary>
-	/// Specification for a node to execute that returns one or more tagged outputs
+	/// Extension methods for BgNode types
 	/// </summary>
-	/// <typeparam name="T">Return type from the node. Must consist of latent Bg* types (eg. BgToken)</typeparam>
-	public class BgNodeSpec<T> : BgNodeSpec
+	public static class BgNodeExtensions
 	{
 		/// <summary>
-		/// Output value from this node
+		/// Creates a node builder for the given agent
 		/// </summary>
-		public T Output { get; }
+		/// <param name="agent">Agent to run the node</param>
+		/// <param name="func">Function to execute</param>
+		/// <returns>Node builder</returns>
+		public static BgNodeSpecBuilder AddNode(this BgAgentSpec agent, Expression<Func<BgContext, Task>> func)
+		{
+			MethodCallExpression call = (MethodCallExpression)func.Body;
+			return new BgNodeSpecBuilder(null, call, agent);
+		}
 
 		/// <summary>
-		/// Constructor
+		/// Creates a node builder for the given agent
 		/// </summary>
-		internal BgNodeSpec(MethodCallExpression call)
-			: base(call)
+		/// <param name="agent">Agent to run the node</param>
+		/// <param name="func">Function to execute</param>
+		/// <returns>Node builder</returns>
+		public static BgNodeSpecBuilder<T> AddNode<T>(this BgAgentSpec agent, Expression<Func<BgContext, Task<T>>> func)
 		{
-			Type returnType = call.Method.ReturnType;
-			if (IsValueTuple(returnType))
-			{
-				Output = (T)Activator.CreateInstance(returnType, ResultExprs)!;
-			}
-			else
-			{
-				Output = (T)(object)ResultExprs[0];
-			}
+			MethodCallExpression call = (MethodCallExpression)func.Body;
+			return new BgNodeSpecBuilder<T>(null, call, agent);
 		}
-	}
 
-	/// <summary>
-	/// Extension methods for BgNode types
-	/// </summary>
-	public static class BgNodeExtensions
-	{
 		/// <summary>
 		/// Add dependencies onto other nodes or outputs. Outputs from the given tokens will be copied to the current machine before execution of the node.
 		/// </summary>
-		/// <param name="nodeSpec">The node specification</param>
-		/// <param name="tokens">Tokens to add dependencies on</param>
+		/// <param name="builder">The node builder</param>
+		/// <param name="inputs">Files to add as inputs</param>
 		/// <returns>The current node spec, to allow chaining calls</returns>
-		public static T Requires<T>(this T nodeSpec, params BgList<BgFileSet>[] tokens) where T : BgNodeSpec
+		public static T Requires<T>(this T builder, params BgNodeSpec[] inputs) where T : BgNodeSpecBuilder
 		{
-			foreach (BgList<BgFileSet> tokenSet in tokens)
-			{
-				nodeSpec.InputDependencies = nodeSpec.InputDependencies.Add(tokenSet);
-			}
-			return nodeSpec;
+			builder.Inputs = builder.Inputs.Add(inputs.Select(x => (BgFileSet)x));
+			return builder;
 		}
 
 		/// <summary>
-		/// Add weak dependencies onto other nodes or outputs. The producing nodes must complete successfully if they are part of the graph, but outputs from them will not be 
-		/// transferred to the machine running this node.
+		/// Add dependencies onto other nodes or outputs. Outputs from the given tokens will be copied to the current machine before execution of the node.
 		/// </summary>
-		/// <param name="nodeSpec">The node specification</param>
-		/// <param name="tokens">Tokens to add dependencies on</param>
+		/// <param name="builder">The node builder</param>
+		/// <param name="inputs">Files to add as inputs</param>
 		/// <returns>The current node spec, to allow chaining calls</returns>
-		public static T After<T>(this T nodeSpec, params BgList<BgFileSet>[] tokens) where T : BgNodeSpec
+		public static T Requires<T>(this T builder, params BgFileSet[] inputs) where T : BgNodeSpecBuilder
 		{
-			foreach (BgList<BgFileSet> tokenSet in tokens)
-			{
-				nodeSpec.AfterDependencies = nodeSpec.AfterDependencies.Add(tokenSet);
-			}
-			return nodeSpec;
+			builder.Inputs = builder.Inputs.Add(inputs);
+			return builder;
 		}
 
 		/// <summary>
-		/// Adds a warning when executing this node
+		/// Add dependencies onto other nodes or outputs. Outputs from the given tokens will be copied to the current machine before execution of the node.
 		/// </summary>
-		/// <param name="nodeSpec">The node specification</param>
-		/// <param name="condition">Condition for writing the message</param>
-		/// <param name="message">Message to output</param>
+		/// <param name="builder">The node builder</param>
+		/// <param name="inputs">Files to add as inputs</param>
 		/// <returns>The current node spec, to allow chaining calls</returns>
-		public static BgNodeSpec WarnIf<T>(this T nodeSpec, BgBool condition, BgString message) where T : BgNodeSpec
+		public static T Requires<T>(this T builder, BgList<BgFileSet> inputs) where T : BgNodeSpecBuilder
 		{
-			nodeSpec.Diagnostics = nodeSpec.Diagnostics.AddIf(condition, new BgDiagnosticSpec(LogLevel.Warning, message));
-			return nodeSpec;
+			builder.Inputs = builder.Inputs.Add(inputs);
+			return builder;
 		}
 
 		/// <summary>
-		/// Adds an error when executing this node
+		/// Add weak dependencies onto other nodes or outputs. The producing nodes must complete successfully if they are part of the graph, but outputs from them will not be 
+		/// transferred to the machine running this node.
 		/// </summary>
-		/// <param name="nodeSpec">The node specification</param>
-		/// <param name="condition">Condition for writing the message</param>
-		/// <param name="message">Message to output</param>
+		/// <param name="builder">The node builder</param>
+		/// <param name="inputs">Files to add as inputs</param>
 		/// <returns>The current node spec, to allow chaining calls</returns>
-		public static BgNodeSpec ErrorIf<T>(this T nodeSpec, BgBool condition, BgString message) where T : BgNodeSpec
+		public static T After<T>(this T builder, params BgFileSet[] inputs) where T : BgNodeSpecBuilder
 		{
-			nodeSpec.Diagnostics = nodeSpec.Diagnostics.AddIf(condition, new BgDiagnosticSpec(LogLevel.Error, message));
-			return nodeSpec;
+			builder.Fences = builder.Fences.Add(inputs);
+			return builder;
 		}
 	}
 }
diff --git a/Engine/Source/Programs/Shared/EpicGames.BuildGraph/BgOpcode.cs b/Engine/Source/Programs/Shared/EpicGames.BuildGraph/BgOpcode.cs
new file mode 100644
--- /dev/null
+++ b/Engine/Source/Programs/Shared/EpicGames.BuildGraph/BgOpcode.cs
@@ -0,0 +1,368 @@
+// Copyright Epic Games, Inc. All Rights Reserved.
+
+using System;
+using EpicGames.Serialization;
+
+namespace EpicGames.BuildGraph
+{
+	/// <summary>
+	/// Opcodes for graph expressions. Each opcode is followed by 1-3 operands
+	/// </summary>
+	public enum BgOpcode
+	{
+		#region Boolean opcodes
+
+		/// <summary>
+		/// Pushes a boolean 'false' value onto the stack
+		/// </summary>
+		BoolFalse = 0x00,
+
+		/// <summary>
+		/// Pushes a boolean 'true' value onto the stack
+		/// </summary>
+		BoolTrue = 0x01,
+
+		/// <summary>
+		/// Logical NOT
+		/// </summary>
+		BoolNot = 0x02,
+
+		/// <summary>
+		/// Logical AND
+		/// </summary>
+		BoolAnd = 0x03,
+		
+		/// <summary>
+		/// Logical OR
+		/// </summary>
+		BoolOr = 0x04,
+
+		/// <summary>
+		/// Logical XOR
+		/// </summary>
+		BoolXor = 0x05,
+
+		/// <summary>
+		/// Tests whether two boolean values are equal
+		/// </summary>
+		BoolEq = 0x06,
+
+		/// <summary>
+		/// A boolean option
+		/// </summary>
+		BoolOption = 0x07,
+
+		/// <summary>
+		/// Converts a bool to a string
+		/// </summary>
+		BoolToString = 0x08,
+
+		#endregion
+
+		#region Integer opcodes
+
+		/// <summary>
+		/// An integer literal. Opcode is followed by a 32-bit little-endian integer.
+		/// </summary>
+		IntLiteral = 0x10,
+
+		/// <summary>
+		/// Tests whether two integers are equal
+		/// </summary>
+		IntEq = 0x11,
+
+		/// <summary>
+		/// Pops an integer from the stack, tests whether it is less than zero, and pushes a bool onto the stack with the result. 
+		/// </summary>
+		IntLt = 0x12,
+
+		/// <summary>
+		/// Pops an integer from the stack, tests whether it is greater than zero, and pushes a bool onto the stack with the result. 
+		/// </summary>
+		IntGt = 0x13,
+
+		/// <summary>
+		/// Adds two integers together
+		/// </summary>
+		IntAdd = 0x14,
+
+		/// <summary>
+		/// Multiplies two integers
+		/// </summary>
+		IntMultiply = 0x15,
+
+		/// <summary>
+		/// Divides one integer by another
+		/// </summary>
+		IntDivide = 0x16,
+
+		/// <summary>
+		/// Computes the modulo of one integer with another
+		/// </summary>
+		IntModulo = 0x17,
+
+		/// <summary>
+		/// Negates an integer
+		/// </summary>
+		IntNegate = 0x18,
+
+		/// <summary>
+		/// An integer option between two values (see <see cref="BgIntOption"/>)
+		/// </summary>
+		IntOption = 0x19,
+
+		/// <summary>
+		/// Converts a bool to a string
+		/// </summary>
+		IntToString = 0x1a,
+
+		#endregion
+
+		#region String opcodes
+
+		/// <summary>
+		/// An empty string literal
+		/// </summary>
+		StrEmpty = 0x20,
+
+		/// <summary>
+		/// A string literal. Opcode is followed by a UTF-8 encoded string with length.
+		/// </summary>
+		StrLiteral = 0x21,
+
+		/// <summary>
+		/// Compares two strings for equality, using a <see cref="StringComparison"/> value encoded into the byte stream as an unsigned <see cref="VarInt"/>.
+		/// </summary>
+		StrCompare = 0x22,
+
+		/// <summary>
+		/// Concatenates two strings and returns the result
+		/// </summary>
+		StrConcat = 0x23,
+
+		/// <summary>
+		/// Format a string, similar to <see cref="String.Format(String, Object?[])"/>
+		/// </summary>
+		StrFormat = 0x24,
+
+		/// <summary>
+		/// Joins a list of strings (first argument) with a separator (second argument)
+		/// </summary>
+		StrJoin = 0x25,
+
+		/// <summary>
+		/// Splits a string by a delimiter
+		/// </summary>
+		StrSplit = 0x26,
+
+		/// <summary>
+		/// Tests whether a string (first argument) matches a regular expression (second argument)
+		/// </summary>
+		StrIsMatch = 0x27,
+
+		/// <summary>
+		/// Returns a string with all ocurrences of the second argument in the first argument replaced with the third argument
+		/// </summary>
+		StrReplace = 0x28,
+
+		/// <summary>
+		/// A string option (see <see cref="BgStrOption"/>)
+		/// </summary>
+		StrOption = 0x29,
+
+		#endregion
+
+		#region
+
+		/// <summary>
+		/// Literal enum value, as an integer.
+		/// </summary>
+		EnumConstant = 0x30,
+
+		/// <summary>
+		/// Parses a string as an enum
+		/// </summary>
+		EnumParse = 0x31,
+
+		/// <summary>
+		/// Converts an enum to a string
+		/// </summary>
+		EnumToString = 0x32,
+
+		#endregion
+
+		#region List opcodes
+
+		/// <summary>
+		/// Creates an empty list
+		/// </summary>
+		ListEmpty = 0x40,
+
+		/// <summary>
+		/// Adds an item to a list
+		/// </summary>
+		ListPush = 0x41,
+
+		/// <summary>
+		/// Adds a lazily evaluated item to a list
+		/// </summary>
+		ListPushLazy = 0x42,
+
+		/// <summary>
+		/// Gets the length of a list
+		/// </summary>
+		ListCount = 0x43,
+
+		/// <summary>
+		/// Gets the element of a list at an index. Second argument is an integer expression.
+		/// </summary>
+		ListElement = 0x44,
+
+		/// <summary>
+		/// Concatenates two lists
+		/// </summary>
+		ListConcat = 0x45,
+
+		/// <summary>
+		/// Creates the union of two lists
+		/// </summary>
+		ListUnion = 0x46,
+
+		/// <summary>
+		/// Creates the set excluding another set of items
+		/// </summary>
+		ListExcept = 0x47,
+
+		/// <summary>
+		/// Call a function on each element of a list, returning the transformed list
+		/// </summary>
+		ListSelect = 0x48,
+
+		/// <summary>
+		/// Select elements from a list using a predicate
+		/// </summary>
+		ListWhere = 0x49,
+
+		/// <summary>
+		/// Selects all the unique entries in the list. The comparer for the list type is determined by the type of the first element.
+		/// </summary>
+		ListDistinct = 0x4a,
+
+		/// <summary>
+		/// Determines whether a list contains a particular item
+		/// </summary>
+		ListContains = 0x4b,
+
+		/// <summary>
+		/// Indicates that the list should only be evaluated when the first item is enumerated from it.
+		/// </summary>
+		ListLazy = 0x4c,
+
+		/// <summary>
+		/// An option providing a list of strings (see <see cref="BgListOptionSpec"/>
+		/// </summary>
+		ListOption = 0x4d,
+
+		#endregion
+
+		#region Fileset opcodes
+
+		/// <summary>
+		/// Fileset constructed from all the inputs and outputs of a node
+		/// </summary>
+		FileSetFromNode = 0x50,
+
+		/// <summary>
+		/// Fileset for an individual numbered output from a node
+		/// </summary>
+		FileSetFromNodeOutput = 0x51,
+
+		#endregion
+
+		#region Object opcodes
+
+		/// <summary>
+		/// An empty object
+		/// </summary>
+		ObjEmpty = 0x60,
+
+		/// <summary>
+		/// Gets the value of a property
+		/// </summary>
+		ObjGet = 0x61,
+
+		/// <summary>
+		/// Sets the value of a property
+		/// </summary>
+		ObjSet = 0x62,
+
+		#endregion
+
+		#region Function opcodes
+
+		/// <summary>
+		/// Takes an operand indicating the number of arguments, and offset of the function to call. Pops the arguments from the evaluation stack, pushes them onto the function stack, and jumps to the given offset.
+		/// </summary>
+		Call = 0xd0,
+
+		/// <summary>
+		/// Fetches a numbered argument from the functions stack frame.
+		/// </summary>
+		Argument = 0xd1,
+
+		/// <summary>
+		/// Jumps to another fragment without creating a new stack frame
+		/// </summary>
+		Jump = 0xd2,
+
+		#endregion
+
+		#region Generic opcodes
+
+		/// <summary>
+		/// Chooses between two operands based on a boolean parameter.
+		/// </summary>
+		Choose = 0xe0,
+
+		/// <summary>
+		/// Throws an exception
+		/// </summary>
+		Throw = 0xe1,
+
+		/// <summary>
+		/// Null value for an optional expression
+		/// </summary>
+		Null = 0xe2,
+
+		#endregion
+
+		#region Graph structures
+
+		/// <summary>
+		/// Defines a node (see <see cref="BgNodeSpec"/>)
+		/// </summary>
+		Node = 0xf0,
+
+		/// <summary>
+		/// Defines an agent to execute work on (see <see cref="BgAgentSpec"/>)
+		/// </summary>
+		Agent = 0xf1,
+		
+		/// <summary>
+		/// Defines a label (see <see cref="BgLabelSpec"/>)
+		/// </summary>
+		Label = 0xf2,
+
+		/// <summary>
+		/// Defines an aggregate (see <see cref="BgAggregateSpec"/>)
+		/// </summary>
+		Aggregate = 0xf3,
+
+		/// <summary>
+		/// Defines a graph (see <see cref="BgGraphSpec"/>)
+		/// </summary>
+		Graph = 0xf4,
+
+		#endregion
+	}
+}
diff --git a/Engine/Source/Programs/Shared/EpicGames.BuildGraph/BgOptionSpec.cs b/Engine/Source/Programs/Shared/EpicGames.BuildGraph/BgOptionSpec.cs
--- a/Engine/Source/Programs/Shared/EpicGames.BuildGraph/BgOptionSpec.cs
+++ b/Engine/Source/Programs/Shared/EpicGames.BuildGraph/BgOptionSpec.cs
@@ -9,222 +9,168 @@ using EpicGames.BuildGraph.Expressions;
 namespace EpicGames.BuildGraph
 {
 	/// <summary>
-	/// Exception thrown if an option fails validation
+	/// A boolean option expression
 	/// </summary>
-	sealed class BgOptionValidationException : Exception
+	public class BgBoolOption : BgBool
 	{
 		/// <summary>
-		/// Constructor
+		/// Name of the option
 		/// </summary>
-		/// <param name="message"></param>
-		public BgOptionValidationException(string message)
-			: base(message)
-		{
-		}
-	}
+		public BgString Name { get; }
 
-	/// <summary>
-	/// Base class for option configuration
-	/// </summary>
-	public interface IBgOption : IBgExpr
-	{
 		/// <summary>
-		/// Name of the option
+		/// Label to display next to the option
 		/// </summary>
-		string Name { get; }
+		public BgString? Label { get; }
 
 		/// <summary>
-		/// Label to show against the option in the UI
+		/// Help text to display for the user
 		/// </summary>
-		BgString? Label { get; set; }
+		public BgString? Description { get; }
 
 		/// <summary>
-		/// Description for the option
+		/// Default value for the option
 		/// </summary>
-		BgString Description { get; set; }
-	}
-
-	/// <summary>
-	/// A boolean option expression
-	/// </summary>
-	public class BgBoolOption : BgBool, IBgOption
-	{
-		/// <inheritdoc/>
-		public string Name { get; }
-
-		/// <inheritdoc/>
-		public BgString? Label { get; set; }
-
-		/// <inheritdoc/>
-		public BgString Description { get; set; }
+		public BgBool? DefaultValue { get; }
 
 		/// <summary>
-		/// Default value for the option
+		/// Constructor
 		/// </summary>
-		public BgBool DefaultValue { get; set; }
+		public BgBoolOption(BgString name, BgString? description = null, BgBool? defaultValue = null)
+			: this(name, null, description, defaultValue)
+		{
+		}
 
 		/// <summary>
 		/// Constructor
 		/// </summary>
-		internal BgBoolOption(string name, BgString description, BgBool defaultValue)
+		public BgBoolOption(BgString name, BgString? label, BgString? description, BgBool? defaultValue)
+			: base(BgExprFlags.None)
 		{
 			Name = name;
+			Label = label;
 			Description = description;
 			DefaultValue = defaultValue;
 		}
 
 		/// <inheritdoc/>
-		public override bool Compute(BgExprContext context)
+		public override void Write(BgBytecodeWriter writer)
 		{
-			string? value;
-			if (context.Options.TryGetValue(Name, out value))
-			{
-				bool boolValue;
-				if (!Boolean.TryParse(value, out boolValue))
-				{
-					throw new BgOptionValidationException($"Argument for {Name} is not a valid bool ({value})");
-				}
-				return boolValue;
-			}
-			return DefaultValue.Compute(context);
+			writer.WriteOpcode(BgOpcode.BoolOption);
+			writer.WriteExpr(Name);
+			writer.WriteExpr(Label ?? BgString.Empty);
+			writer.WriteExpr(Description ?? BgString.Empty);
+			writer.WriteExpr(DefaultValue ?? BgBool.False);
 		}
 	}
 
 	/// <summary>
 	/// An integer option expression
 	/// </summary>
-	public class BgIntOption : BgInt, IBgOption
+	public class BgIntOption : BgInt
 	{
-		/// <inheritdoc/>
-		public string Name { get; }
+		/// <summary>
+		/// Name of the option
+		/// </summary>
+		public BgString Name { get; }
 
-		/// <inheritdoc/>
-		public BgString? Label { get; set; }
+		/// <summary>
+		/// Label to display next to the option
+		/// </summary>
+		public BgString? Label { get; }
 
-		/// <inheritdoc/>
-		public BgString Description { get; set; }
+		/// <summary>
+		/// Help text to display for the user
+		/// </summary>
+		public BgString? Description { get; }
 
 		/// <summary>
 		/// Default value for the option
 		/// </summary>
-		public BgInt DefaultValue { get; set; }
+		public BgInt? DefaultValue { get; }
 
 		/// <summary>
 		/// Minimum allowed value
 		/// </summary>
-		public BgInt? MinValue { get; set; }
+		public BgInt? MinValue { get; }
 
 		/// <summary>
 		/// Maximum allowed value
 		/// </summary>
-		public BgInt? MaxValue { get; set; }
+		public BgInt? MaxValue { get; }
 
 		/// <summary>
 		/// Constructor
 		/// </summary>
-		internal BgIntOption(string name, BgString description, BgInt defaultValue)
+		public BgIntOption(string name, BgString? description = null, BgInt? defaultValue = null, BgInt? minValue = null, BgInt? maxValue = null, BgString? label = null)
+			: base(BgExprFlags.None)
 		{
 			Name = name;
+			Label = label;
 			Description = description;
 			DefaultValue = defaultValue;
+			MinValue = minValue;
+			MaxValue = maxValue;
 		}
 
 		/// <inheritdoc/>
-		public override int Compute(BgExprContext context)
+		public override void Write(BgBytecodeWriter writer)
 		{
-			string? value;
-			if (context.Options.TryGetValue(Name, out value))
-			{
-				int intValue;
-				if (!Int32.TryParse(value, out intValue))
-				{
-					throw new BgOptionValidationException($"Argument for '{Name}' is not a valid integer");
-				}
-				if (!(MinValue is null))
-				{
-					int intMinValue = MinValue.Compute(context);
-					if (intValue < intMinValue)
-					{
-						throw new BgOptionValidationException($"Argument for '{Name}' is less than the allowed minimum ({intValue} < {intMinValue})");
-					}
-				}
-				if (!(MaxValue is null))
-				{
-					int intMaxValue = MaxValue.Compute(context);
-					if (intValue > intMaxValue)
-					{
-						throw new BgOptionValidationException($"Argument for '{Name}' is greater than the allowed maximum ({intValue} > {intMaxValue})");
-					}
-				}
-			}
-			return DefaultValue.Compute(context);
+			writer.WriteOpcode(BgOpcode.IntOption);
+			writer.WriteExpr(Name);
+			writer.WriteExpr(Label ?? BgString.Empty);
+			writer.WriteExpr(Description ?? BgString.Empty);
+			writer.WriteExpr(DefaultValue ?? (BgInt)0);
+			writer.WriteExpr(MinValue ?? (BgInt)(-1));
+			writer.WriteExpr(MaxValue ?? (BgInt)(-1));
 		}
 	}
 
 	/// <summary>
-	/// A string option expression
+	/// Style for a string option
 	/// </summary>
-	public class BgEnumOption<TEnum> : BgEnum<TEnum>, IBgOption where TEnum : struct
+	public enum BgStrOptionStyle
 	{
-		/// <inheritdoc/>
-		public string Name { get; }
-
-		/// <inheritdoc/>
-		public BgString? Label { get; set; }
-
-		/// <inheritdoc/>
-		public BgString Description { get; set; }
-
 		/// <summary>
-		/// Default value for the option
+		/// Free-form text entry
 		/// </summary>
-		public BgEnum<TEnum> DefaultValue { get; set; }
+		Text,
 
 		/// <summary>
-		/// Constructor
+		/// List of options
 		/// </summary>
-		internal BgEnumOption(string name, BgString description, BgEnum<TEnum> defaultValue)
-		{
-			Name = name;
-			Description = description;
-			DefaultValue = defaultValue;
-		}
-
-		/// <inheritdoc/>
-		public override TEnum Compute(BgExprContext context)
-		{
-			string? value;
-			if (context.Options.TryGetValue(Name, out value))
-			{
-				TEnum enumValue;
-				if (!Enum.TryParse<TEnum>(value, true, out enumValue))
-				{
-					throw new BgOptionValidationException($"Argument '{Name}' is not a valid value for {typeof(TEnum).Name}");
-				}
-				return enumValue;
-			}
-			return DefaultValue.Compute(context);
-		}
+		DropList,
 	}
 
 	/// <summary>
 	/// A string option expression
 	/// </summary>
-	public class BgStringOption : BgString, IBgOption
+	public class BgStrOption : BgString
 	{
-		/// <inheritdoc/>
-		public string Name { get; }
+		/// <summary>
+		/// Name of the option
+		/// </summary>
+		public BgString Name { get; }
 
-		/// <inheritdoc/>
-		public BgString? Label { get; set; }
+		/// <summary>
+		/// Label to display next to the option
+		/// </summary>
+		public BgString? Label { get; }
 
-		/// <inheritdoc/>
-		public BgString Description { get; set; }
+		/// <summary>
+		/// Help text to display for the user
+		/// </summary>
+		public BgString? Description { get; }
 
 		/// <summary>
 		/// Default value for the option
 		/// </summary>
-		public BgString DefaultValue { get; set; }
+		public BgString? DefaultValue { get; set; }
+
+		/// <summary>
+		/// Style for this option
+		/// </summary>
+		public BgStrOptionStyle Style { get; }
 
 		/// <summary>
 		/// Regex for validating values for the option
@@ -237,87 +183,144 @@ namespace EpicGames.BuildGraph
 		public BgString? PatternFailed { get; set; }
 
 		/// <summary>
-		/// Allowed values of the option
+		/// List of values to choose from
 		/// </summary>
-		public BgList<BgString>? Enum { get; set; }
+		public BgList<BgString>? Values { get; set; }
+
+		/// <summary>
+		/// Matching list of descriptions for each value
+		/// </summary>
+		public BgList<BgString>? ValueDescriptions { get; set; }
 
 		/// <summary>
 		/// Constructor
 		/// </summary>
-		internal BgStringOption(string name, BgString description, BgString defaultValue)
+		public BgStrOption(string name, BgString? label = null, BgString? description = null, BgString? defaultValue = null, BgStrOptionStyle style = BgStrOptionStyle.Text, BgString? pattern = null, BgString? patternFailed = null, BgList<BgString>? values = null, BgList<BgString>? valueDescriptions = null)
+			: base(BgExprFlags.None)
 		{
 			Name = name;
+			Label = label;
 			Description = description;
+			Style = style;
 			DefaultValue = defaultValue;
+			Pattern = pattern;
+			PatternFailed = patternFailed;
+			Values = values;
+			ValueDescriptions = valueDescriptions;
 		}
 
 		/// <inheritdoc/>
-		public override string Compute(BgExprContext context)
+		public override void Write(BgBytecodeWriter writer)
 		{
-			string? value;
-			if (context.Options.TryGetValue(Name, out value))
-			{
-				if (!(Pattern is null))
-				{
-					string patternValue = Pattern.Compute(context);
-					if (!Regex.IsMatch(value, patternValue))
-					{
-						string patternFailedValue = PatternFailed?.Compute(context) ?? $"Argument '{Name}' does not match the required pattern: '{patternValue}'";
-						throw new BgOptionValidationException(patternFailedValue);
-					}
-				}
-				if (!(Enum is null))
-				{
-					List<string> enumValues = Enum.Compute(context);
-					if (!enumValues.Any(x => x.Equals(value, StringComparison.OrdinalIgnoreCase)))
-					{
-						throw new BgOptionValidationException($"Argument '{Name}' is invalid");
-					}
-				}
-				return value;
-			}
-			return DefaultValue.Compute(context);
+			writer.WriteOpcode(BgOpcode.StrOption);
+			writer.WriteExpr(Name);
+			writer.WriteExpr(Label ?? BgString.Empty);
+			writer.WriteExpr(Description ?? BgString.Empty);
+			writer.WriteUnsignedInteger((int)Style);
+			writer.WriteExpr(DefaultValue ?? BgString.Empty);
+			writer.WriteExpr(Pattern ?? BgString.Empty);
+			writer.WriteExpr(PatternFailed ?? BgString.Empty);
+			writer.WriteExpr(Values ?? BgList<BgString>.Empty);
+			writer.WriteExpr(ValueDescriptions ?? BgList<BgString>.Empty);
 		}
 	}
 
+	/// <summary>
+	/// Style for a list option
+	/// </summary>
+	public enum BgListOptionStyle
+	{
+		/// <summary>
+		/// List of checkboxes
+		/// </summary>
+		CheckList = 0,
+
+		/// <summary>
+		/// Tag picker
+		/// </summary>
+		TagPicker = 1,
+	}
+
 	/// <summary>
 	/// A list option expression
 	/// </summary>
-	public class BgEnumListOption<TEnum> : BgList<BgEnum<TEnum>>, IBgOption where TEnum : struct
+	public class BgListOptionSpec : BgList<BgString>
 	{
-		/// <inheritdoc/>
-		public string Name { get; }
+		/// <summary>
+		/// Name of the option
+		/// </summary>
+		public BgString Name { get; }
 
-		/// <inheritdoc/>
-		public BgString? Label { get; set; }
+		/// <summary>
+		/// Label to display next to the option
+		/// </summary>
+		public BgString? Label { get; }
 
-		/// <inheritdoc/>
-		public BgString Description { get; set; }
+		/// <summary>
+		/// Help text to display for the user
+		/// </summary>
+		public BgString? Description { get; }
+
+		/// <summary>
+		/// Style for this list box
+		/// </summary>
+		public BgListOptionStyle Style { get; }
 
 		/// <summary>
 		/// Default value for the option
 		/// </summary>
-		public BgList<BgEnum<TEnum>> DefaultValue { get; set; }
+		public BgString? DefaultValue { get; set; }
+
+		/// <summary>
+		/// Regex for validating values for the option
+		/// </summary>
+		public BgString? Pattern { get; set; }
+
+		/// <summary>
+		/// Message to display if validation fails
+		/// </summary>
+		public BgString? PatternFailed { get; set; }
+
+		/// <summary>
+		/// List of values to choose from
+		/// </summary>
+		public BgList<BgString>? Values { get; set; }
+
+		/// <summary>
+		/// Matching list of descriptions for each value
+		/// </summary>
+		public BgList<BgString>? ValueDescriptions { get; set; }
 
 		/// <summary>
 		/// Constructor
 		/// </summary>
-		internal BgEnumListOption(string name, BgString description, BgList<BgEnum<TEnum>> defaultValue)
+		public BgListOptionSpec(string name, BgString? label = null, BgString? description = null, BgString? defaultValue = null, BgListOptionStyle style = BgListOptionStyle.CheckList, BgString? pattern = null, BgString? patternFailed = null, BgList<BgString>? values = null, BgList<BgString>? valueDescriptions = null)
+			: base(BgExprFlags.None)
 		{
 			Name = name;
+			Label = label;
 			Description = description;
 			DefaultValue = defaultValue;
+			Style = style;
+			Pattern = pattern;
+			PatternFailed = patternFailed;
+			Values = values;
+			ValueDescriptions = valueDescriptions;
 		}
 
 		/// <inheritdoc/>
-		public override IEnumerable<BgEnum<TEnum>> GetEnumerable(BgExprContext context)
+		public override void Write(BgBytecodeWriter writer)
 		{
-			BgList<BgEnum<TEnum>> value = DefaultValue;
-			if (context.Options.TryGetValue(Name, out string? valueText))
-			{
-				value = BgType.Get<BgList<BgEnum<TEnum>>>().DeserializeArgument(valueText);
-			}
-			return value.GetEnumerable(context);
+			writer.WriteOpcode(BgOpcode.ListOption);
+			writer.WriteExpr(Name);
+			writer.WriteExpr(Label ?? BgString.Empty);
+			writer.WriteExpr(Description ?? BgString.Empty);
+			writer.WriteUnsignedInteger((int)Style);
+			writer.WriteExpr(DefaultValue ?? BgString.Empty);
+			writer.WriteExpr(Pattern ?? BgString.Empty);
+			writer.WriteExpr(PatternFailed ?? BgString.Empty);
+			writer.WriteExpr(Values ?? BgList<BgString>.Empty);
+			writer.WriteExpr(ValueDescriptions ?? BgList<BgString>.Empty);
 		}
 	}
 }
diff --git a/Engine/Source/Programs/Shared/EpicGames.BuildGraph/EpicGames.BuildGraph.csproj b/Engine/Source/Programs/Shared/EpicGames.BuildGraph/EpicGames.BuildGraph.csproj
--- a/Engine/Source/Programs/Shared/EpicGames.BuildGraph/EpicGames.BuildGraph.csproj
+++ b/Engine/Source/Programs/Shared/EpicGames.BuildGraph/EpicGames.BuildGraph.csproj
@@ -38,5 +38,6 @@
     <ProjectReference Include="..\..\Shared\EpicGames.Core\EpicGames.Core.csproj" />
   </ItemGroup>
   <ItemGroup>
+    <ProjectReference Include="..\EpicGames.Serialization\EpicGames.Serialization.csproj" />
   </ItemGroup>
 </Project>
\ No newline at end of file
diff --git a/Engine/Source/Programs/Shared/EpicGames.BuildGraph/Expressions/BgBool.cs b/Engine/Source/Programs/Shared/EpicGames.BuildGraph/Expressions/BgBool.cs
--- a/Engine/Source/Programs/Shared/EpicGames.BuildGraph/Expressions/BgBool.cs
+++ b/Engine/Source/Programs/Shared/EpicGames.BuildGraph/Expressions/BgBool.cs
@@ -1,6 +1,8 @@
 // Copyright Epic Games, Inc. All Rights Reserved.
 
 using System;
+using System.Collections.Generic;
+using System.Linq;
 
 namespace EpicGames.BuildGraph.Expressions
 {
@@ -8,7 +10,7 @@ namespace EpicGames.BuildGraph.Expressions
 	/// Abstract base class for expressions returning a boolean value 
 	/// </summary>
 	[BgType(typeof(BgBoolType))]
-	public abstract class BgBool : IBgExpr<BgBool>
+	public abstract class BgBool : BgExpr
 	{
 		/// <summary>
 		/// Constant value for false
@@ -20,8 +22,14 @@ namespace EpicGames.BuildGraph.Expressions
 		/// </summary>
 		public static BgBool True { get; } = true;
 
-		/// <inheritdoc/>
-		public BgBool IfThen(BgBool condition, BgBool valueIfTrue) => new BgBoolChooseExpr(condition, valueIfTrue, this);
+		/// <summary>
+		/// Constructor
+		/// </summary>
+		/// <param name="flags">Flags for this expression</param>
+		public BgBool(BgExprFlags flags)
+			: base(flags)
+		{
+		}
 
 		/// <summary>
 		/// Implict conversion operator from a boolean literal
@@ -35,22 +43,19 @@ namespace EpicGames.BuildGraph.Expressions
 		public static BgBool operator !(BgBool inner) => new BgBoolNotExpr(inner);
 
 		/// <inheritdoc/>
-		public static BgBool operator &(BgBool lhs, BgBool rhs) => new BgBoolAndExpr(lhs, rhs);
+		public static BgBool operator &(BgBool lhs, BgBool rhs) => new BgBoolBinaryExpr(BgOpcode.BoolAnd, lhs, rhs);
 
 		/// <inheritdoc/>
-		public static BgBool operator |(BgBool lhs, BgBool rhs) => new BgBoolOrExpr(lhs, rhs);
+		public static BgBool operator |(BgBool lhs, BgBool rhs) => new BgBoolBinaryExpr(BgOpcode.BoolOr, lhs, rhs);
 
 		/// <inheritdoc/>
-		public static BgBool operator ==(BgBool lhs, BgBool rhs) => new BgBoolEqExpr(lhs, rhs);
+		public static BgBool operator ^(BgBool lhs, BgBool rhs) => new BgBoolBinaryExpr(BgOpcode.BoolXor, lhs, rhs);
 
 		/// <inheritdoc/>
-		public static BgBool operator !=(BgBool lhs, BgBool rhs) => new BgBoolNotExpr(new BgBoolEqExpr(lhs, rhs));
+		public static BgBool operator ==(BgBool lhs, BgBool rhs) => new BgBoolBinaryExpr(BgOpcode.BoolEq, lhs, rhs);
 
 		/// <inheritdoc/>
-		object IBgExpr.Compute(BgExprContext context) => Compute(context);
-
-		/// <inheritdoc/>
-		public abstract bool Compute(BgExprContext context);
+		public static BgBool operator !=(BgBool lhs, BgBool rhs) => !(lhs == rhs);
 
 		/// <inheritdoc/>
 		public sealed override bool Equals(object? obj) => throw new InvalidOperationException();
@@ -59,25 +64,19 @@ namespace EpicGames.BuildGraph.Expressions
 		public sealed override int GetHashCode() => throw new InvalidOperationException();
 
 		/// <inheritdoc/>
-		public BgString ToBgString() => new BgBoolFormatExpr(this);
+		public override BgString ToBgString() => new BgBoolToBgStringExpr(this);
 	}
 
 	/// <summary>
 	/// Type traits for a <see cref="BgBool"/>
 	/// </summary>
-	class BgBoolType : BgTypeBase<BgBool>
+	class BgBoolType : BgType<BgBool>
 	{
 		/// <inheritdoc/>
-		public override BgBool DeserializeArgument(string text) => Boolean.Parse(text);
-
-		/// <inheritdoc/>
-		public override string SerializeArgument(BgBool value, BgExprContext context) => value.Compute(context).ToString();
-
-		/// <inheritdoc/>
-		public override BgBool CreateConstant(object value) => new BgBoolConstantExpr((bool)value);
+		public override BgBool Constant(object value) => new BgBoolConstantExpr((bool)value);
 
 		/// <inheritdoc/>
-		public override IBgExprVariable<BgBool> CreateVariable() => new BgBoolVariableExpr();
+		public override BgBool Wrap(BgExpr expr) => new BgBoolWrappedExpr(expr);
 	}
 
 	#region Expression classes
@@ -87,100 +86,81 @@ namespace EpicGames.BuildGraph.Expressions
 		public BgBool Inner { get; }
 
 		public BgBoolNotExpr(BgBool inner)
+			: base(inner.Flags & BgExprFlags.Eager)
 		{
 			Inner = inner;
 		}
 
-		public override bool Compute(BgExprContext context) => !Inner.Compute(context);
-	}
-
-	class BgBoolEqExpr : BgBool
-	{
-		public BgBool Lhs { get; }
-		public BgBool Rhs { get; }
-
-		public BgBoolEqExpr(BgBool lhs, BgBool rhs)
+		public override void Write(BgBytecodeWriter writer)
 		{
-			Lhs = lhs;
-			Rhs = rhs;
+			writer.WriteOpcode(BgOpcode.BoolNot);
+			writer.WriteExpr(Inner);
 		}
-
-		public override bool Compute(BgExprContext context) => Lhs.Compute(context) == Rhs.Compute(context);
 	}
 
-	class BgBoolAndExpr : BgBool
+	class BgBoolBinaryExpr : BgBool
 	{
+		public BgOpcode Opcode { get; }
 		public BgBool Lhs { get; }
 		public BgBool Rhs { get; }
 
-		public BgBoolAndExpr(BgBool lhs, BgBool rhs)
+		public BgBoolBinaryExpr(BgOpcode opcode, BgBool lhs, BgBool rhs)
+			: base(lhs.Flags & rhs.Flags & BgExprFlags.Eager)
 		{
+			Opcode = opcode;
 			Lhs = lhs;
 			Rhs = rhs;
 		}
 
-		public override bool Compute(BgExprContext context) => Lhs.Compute(context) && Rhs.Compute(context);
-	}
-
-	class BgBoolOrExpr : BgBool
-	{
-		public BgBool Lhs { get; }
-		public BgBool Rhs { get; }
-
-		public BgBoolOrExpr(BgBool lhs, BgBool rhs)
+		public override void Write(BgBytecodeWriter writer)
 		{
-			Lhs = lhs;
-			Rhs = rhs;
+			writer.WriteOpcode(Opcode);
+			writer.WriteExpr(Lhs);
+			writer.WriteExpr(Rhs);
 		}
-
-		public override bool Compute(BgExprContext context) => Lhs.Compute(context) || Rhs.Compute(context);
 	}
 
-	class BgBoolChooseExpr : BgBool
+	class BgBoolConstantExpr : BgBool
 	{
-		public BgBool Condition { get; }
-		public BgBool ValueIfTrue { get; }
-		public BgBool ValueIfFalse { get; }
+		public bool Value { get; }
 
-		public BgBoolChooseExpr(BgBool condition, BgBool valueIfTrue, BgBool valueIfFalse)
+		public BgBoolConstantExpr(bool value)
+			: base(BgExprFlags.NotInterned | BgExprFlags.Eager)
 		{
-			Condition = condition;
-			ValueIfTrue = valueIfTrue;
-			ValueIfFalse = valueIfFalse;
+			Value = value;
 		}
 
-		public override bool Compute(BgExprContext context) => Condition.Compute(context) ? ValueIfTrue.Compute(context) : ValueIfFalse.Compute(context);
+		public override void Write(BgBytecodeWriter writer) => writer.WriteOpcode(Value? BgOpcode.BoolTrue : BgOpcode.BoolFalse);
 	}
 
-	class BgBoolVariableExpr : BgBool, IBgExprVariable<BgBool>
+	class BgBoolWrappedExpr : BgBool
 	{
-		public BgBool Value { get; set; } = BgBool.False;
+		public BgExpr Expr { get; }
 
-		public override bool Compute(BgExprContext context) => Value.Compute(context);
-	}
-
-	class BgBoolConstantExpr : BgBool
-	{
-		public bool Value { get; }
-
-		public BgBoolConstantExpr(bool value)
+		public BgBoolWrappedExpr(BgExpr expr)
+			: base(expr.Flags)
 		{
-			Value = value;
+			Expr = expr;
 		}
 
-		public override bool Compute(BgExprContext context) => Value;
+		public override void Write(BgBytecodeWriter writer) => Expr.Write(writer);
 	}
 
-	class BgBoolFormatExpr : BgString
+	class BgBoolToBgStringExpr : BgString
 	{
-		public BgBool Value { get; }
+		public BgExpr Expr { get; }
 
-		public BgBoolFormatExpr(BgBool value)
+		public BgBoolToBgStringExpr(BgExpr expr)
+			: base(expr.Flags & BgExprFlags.Eager)
 		{
-			Value = value;
+			Expr = expr;
 		}
 
-		public override string Compute(BgExprContext context) => Value.Compute(context).ToString();
+		public override void Write(BgBytecodeWriter writer)
+		{
+			writer.WriteOpcode(BgOpcode.BoolToString);
+			writer.WriteExpr(Expr);
+		}
 	}
 
 	#endregion
diff --git a/Engine/Source/Programs/Shared/EpicGames.BuildGraph/Expressions/BgEnum.cs b/Engine/Source/Programs/Shared/EpicGames.BuildGraph/Expressions/BgEnum.cs
--- a/Engine/Source/Programs/Shared/EpicGames.BuildGraph/Expressions/BgEnum.cs
+++ b/Engine/Source/Programs/Shared/EpicGames.BuildGraph/Expressions/BgEnum.cs
@@ -1,29 +1,37 @@
 // Copyright Epic Games, Inc. All Rights Reserved.
 
 using System;
+using System.Linq;
 
 namespace EpicGames.BuildGraph.Expressions
 {
-	/// <summary>
-	/// Abstract base class for expressions returning an enum. Use <see cref="BgEnum{TEnum}"/> for a strongly typed version.
-	/// </summary>
-	public interface IBgEnum : IBgExpr
-	{
-		/// <summary>
-		/// The enum type
-		/// </summary>
-		Type EnumType { get; }
-	}
-
 	/// <summary>
 	/// Abstract base class for expressions returning a string value 
 	/// </summary>
 	[BgType(typeof(BgEnumType<>))]
-	public abstract class BgEnum<TEnum> : IBgEnum, IBgExpr<BgEnum<TEnum>> where TEnum : struct
+	public abstract class BgEnum<TEnum> : BgExpr where TEnum : struct
 	{
 		/// <inheritdoc/>
 		public Type EnumType => typeof(TEnum);
 
+		/// <summary>
+		/// Names of this enum
+		/// </summary>
+		public static BgList<BgString> Names { get; } = BgList<BgString>.Create(Enum.GetNames(typeof(TEnum)).Select(x => (BgString)x));
+
+		/// <summary>
+		/// Values of this enum
+		/// </summary>
+		public static BgList<BgInt> Values { get; } = BgList<BgInt>.Create(((TEnum[])Enum.GetValues(typeof(TEnum))).Select(x => (BgInt)(int)(object)x));
+
+		/// <summary>
+		/// Constructor
+		/// </summary>
+		public BgEnum(BgExprFlags flags)
+			: base(flags)
+		{
+		}
+
 		/// <summary>
 		/// Implicit conversion from a regular enum type
 		/// </summary>
@@ -35,122 +43,105 @@ namespace EpicGames.BuildGraph.Expressions
 		/// <summary>
 		/// Explicit conversion from a string value
 		/// </summary>
-		public static explicit operator BgEnum<TEnum>(BgString value)
+		public static BgEnum<TEnum> Parse(BgString value)
 		{
 			return new BgEnumParseExpr<TEnum>(value);
 		}
 
 		/// <inheritdoc/>
-		public BgEnum<TEnum> IfThen(BgBool condition, BgEnum<TEnum> valueIfTrue) => new BgEnumChooseExpr<TEnum>(condition, valueIfTrue, this);
-
-		/// <inheritdoc/>
-		object IBgExpr.Compute(BgExprContext context) => Compute(context);
-
-		/// <inheritdoc/>
-		public abstract TEnum Compute(BgExprContext context);
-
-		/// <inheritdoc/>
-		public BgString ToBgString() => new BgEnumFormatExpr<TEnum>(this);
+		public override BgString ToBgString() => new BgEnumToStringExpr<TEnum>(this);
 	}
 
 	/// <summary>
 	/// Type traits for a <see cref="BgEnum{TEnum}"/>
 	/// </summary>
-	class BgEnumType<TEnum> : BgTypeBase<BgEnum<TEnum>> where TEnum : struct
+	class BgEnumType<TEnum> : BgType<BgEnum<TEnum>> where TEnum : struct
 	{
 		/// <inheritdoc/>
-		public override BgEnum<TEnum> DeserializeArgument(string text) => Enum.Parse<TEnum>(text);
-
-		/// <inheritdoc/>
-		public override string SerializeArgument(BgEnum<TEnum> value, BgExprContext context) => value.Compute(context).ToString() ?? String.Empty;
-
-		/// <inheritdoc/>
-		public override BgEnum<TEnum> CreateConstant(object value) => new BgEnumConstantExpr<TEnum>((TEnum)value);
+		public override BgEnum<TEnum> Constant(object value)
+		{
+			return new BgEnumConstantExpr<TEnum>((TEnum)value);
+		}
 
 		/// <inheritdoc/>
-		public override IBgExprVariable<BgEnum<TEnum>> CreateVariable() => throw new NotImplementedException();
+		public override BgEnum<TEnum> Wrap(BgExpr expr)
+		{
+			throw new NotImplementedException();
+		}
 	}
 
 	#region Expression classes
 
-	class BgEnumParseExpr<TEnum> : BgEnum<TEnum> where TEnum : struct
+	class BgEnumConstantExpr<TEnum> : BgEnum<TEnum> where TEnum : struct
 	{
-		public BgString Value { get; }
+		public TEnum Value { get; }
 
-		public BgEnumParseExpr(BgString value)
+		public BgEnumConstantExpr(TEnum value)
+			: base(BgExprFlags.NotInterned | BgExprFlags.Eager)
 		{
 			Value = value;
 		}
 
-		public override TEnum Compute(BgExprContext context) => Enum.Parse<TEnum>(Value.Compute(context));
-
-		public override string ToString()
+		public override void Write(BgBytecodeWriter writer)
 		{
-			return $"Parse<{typeof(TEnum).Name}>({Value})";
+			writer.WriteOpcode(BgOpcode.EnumConstant);
+			writer.WriteSignedInteger((int)(object)Value);
 		}
 	}
 
-	class BgEnumConstantExpr<TEnum> : BgEnum<TEnum> where TEnum : struct
+	class BgEnumParseExpr<TEnum> : BgEnum<TEnum> where TEnum : struct
 	{
-		public TEnum Value { get; }
+		public BgString Value { get; }
 
-		public BgEnumConstantExpr(TEnum value)
+		public BgEnumParseExpr(BgString value)
+			: base(value.Flags & BgExprFlags.Eager)
 		{
 			Value = value;
 		}
 
-		public override TEnum Compute(BgExprContext context) => Value;
-
-		public override string ToString()
+		public override void Write(BgBytecodeWriter writer)
 		{
-			return $"\"{Value}\"";
+			writer.WriteOpcode(BgOpcode.EnumParse);
+			writer.WriteExpr(Value);
+			writer.WriteExpr(Names);
+			writer.WriteExpr(Values);
 		}
 	}
 
-	class BgEnumChooseExpr<TEnum> : BgEnum<TEnum> where TEnum : struct
+	class BgEnumToStringExpr<TEnum> : BgString where TEnum : struct
 	{
-		public BgBool Condition { get; }
-		public BgEnum<TEnum> ValueIfTrue { get; }
-		public BgEnum<TEnum> ValueIfFalse { get; }
+		public BgEnum<TEnum> Expr { get; }
 
-		public BgEnumChooseExpr(BgBool condition, BgEnum<TEnum> valueIfTrue, BgEnum<TEnum> valueIfFalse)
+		public BgEnumToStringExpr(BgEnum<TEnum> expr)
+			: base(expr.Flags & BgExprFlags.Eager)
 		{
-			Condition = condition;
-			ValueIfTrue = valueIfTrue;
-			ValueIfFalse = valueIfFalse;
+			Expr = expr;
 		}
 
-		public override TEnum Compute(BgExprContext context) => Condition.Compute(context) ? ValueIfTrue.Compute(context) : ValueIfFalse.Compute(context);
-
-		public override string ToString()
+		public override void Write(BgBytecodeWriter writer)
 		{
-			return $"Choose({Condition}, {ValueIfTrue}, {ValueIfFalse});";
+			writer.WriteOpcode(BgOpcode.EnumToString);
+			writer.WriteExpr(Expr);
+			writer.WriteExpr(BgEnum<TEnum>.Names);
+			writer.WriteExpr(BgEnum<TEnum>.Values);
 		}
 	}
 
-	class BgEnumFormatExpr<TEnum> : BgString where TEnum : struct
+	class BgEnumWrappedExpr<TEnum> : BgEnum<TEnum> where TEnum : struct
 	{
-		public BgEnum<TEnum> Value { get; }
+		public BgExpr Expr { get; }
 
-		public BgEnumFormatExpr(BgEnum<TEnum> value)
+		public BgEnumWrappedExpr(BgExpr expr)
+			: base(expr.Flags)
 		{
-			Value = value;
+			Expr = expr;
 		}
 
-		public override string Compute(BgExprContext context) => Value.Compute(context).ToString() ?? String.Empty;
-
-		public override string ToString()
+		public override void Write(BgBytecodeWriter writer)
 		{
-			return $"Format({Value})";
+			Expr.Write(writer);
 		}
 	}
 
-	class BgEnumVariableExpr<TEnum> : BgEnum<TEnum>, IBgExprVariable<BgEnum<TEnum>> where TEnum : struct
-	{
-		public BgEnum<TEnum> Value { get; set; } = null!;
-
-		public override TEnum Compute(BgExprContext context) => Value.Compute(context);
-	}
-
 	#endregion
 }
diff --git a/Engine/Source/Programs/Shared/EpicGames.BuildGraph/Expressions/BgExpr.cs b/Engine/Source/Programs/Shared/EpicGames.BuildGraph/Expressions/BgExpr.cs
new file mode 100644
--- /dev/null
+++ b/Engine/Source/Programs/Shared/EpicGames.BuildGraph/Expressions/BgExpr.cs
@@ -0,0 +1,197 @@
+// Copyright Epic Games, Inc. All Rights Reserved.
+
+using System;
+using System.Collections.Generic;
+using System.Runtime.CompilerServices;
+using EpicGames.BuildGraph.Expressions;
+using EpicGames.Core;
+using Microsoft.Extensions.Logging;
+
+namespace EpicGames.BuildGraph
+{
+	/// <summary>
+	/// Flags for an expression
+	/// </summary>
+	[Flags]
+	public enum BgExprFlags
+	{
+		/// <summary>
+		/// No flags set
+		/// </summary>
+		None = 0,
+
+		/// <summary>
+		/// Indicates that the expression should never be interned (ie. encoded to a separate fragment), and will always be duplicated in the bytecode. 
+		/// Trivial constants with short encodings typically have this flag set.
+		/// </summary>
+		NotInterned = 1,
+
+		/// <summary>
+		/// Always eagerly evaluate this expression
+		/// </summary>
+		Eager = 2,
+
+		/// <summary>
+		/// Force this expression to be stored in a separate fragment. Can help improves readability of the disassembly output.
+		/// </summary>
+		ForceFragment = 4,
+	}
+
+	/// <summary>
+	/// Base class for computable expressions
+	/// </summary>
+	public abstract class BgExpr
+	{
+		/// <summary>
+		/// Null value
+		/// </summary>
+		public static BgExpr Null { get; } = new BgNullExpr();
+
+		/// <summary>
+		/// Flags for this expression
+		/// </summary>
+		public BgExprFlags Flags { get; }
+
+		/// <summary>
+		/// Constructor
+		/// </summary>
+		/// <param name="flags"></param>
+		public BgExpr(BgExprFlags flags)
+		{
+			Flags = flags;
+		}
+
+		/// <summary>
+		/// Throws an exception
+		/// </summary>
+		/// <typeparam name="T">Type of the expression to masquerade as</typeparam>
+		/// <param name="message">Message to display</param>
+		/// <param name="sourcePath">Path to the source file declaring this diagnostic. Automatically set by the runtime.</param>
+		/// <param name="sourceLine">Line number in the source file declaring this diagnostic. Automatically set by the runtime.</param>
+		public static T Throw<T>(BgString message, [CallerFilePath] string sourcePath = "", [CallerLineNumber] int sourceLine = 0) where T : BgExpr
+		{
+			return BgType.Wrap<T>(new BgThrowExpr(sourcePath, sourceLine, message));
+		}
+
+		/// <summary>
+		/// Chooses between two values based on a condition
+		/// </summary>
+		/// <typeparam name="T">Type of the expression to choose between</typeparam>
+		/// <param name="condition">Condition to check</param>
+		/// <param name="valueIfTrue">Value to return if the condition is true</param>
+		/// <param name="valueIfFalse">Value to return if the condition is false</param>
+		/// <returns>The chosen value</returns>
+		public static T Choose<T>(BgBool condition, T valueIfTrue, T valueIfFalse) where T : BgExpr
+		{
+			return BgType.Wrap<T>(new BgChooseExpr<T>(condition, valueIfTrue, valueIfFalse));
+		}
+
+		/// <summary>
+		/// Serialize the expression to an output stream
+		/// </summary>
+		/// <param name="writer">Writer for output data</param>
+		public abstract void Write(BgBytecodeWriter writer);
+
+		/// <summary>
+		/// Convert the value of the expression to a string
+		/// </summary>
+		public abstract BgString ToBgString();
+	}
+
+	class BgChooseExpr<T> : BgExpr where T : BgExpr
+	{
+		public BgBool Condition { get; }
+		public T ValueIfTrue { get; }
+		public T ValueIfFalse { get; }
+
+		public BgChooseExpr(BgBool condition, T valueIfTrue, T valueIfFalse)
+			: base(BgExprFlags.None)
+		{
+			Condition = condition;
+			ValueIfTrue = valueIfTrue;
+			ValueIfFalse = valueIfFalse;
+		}
+
+		public override void Write(BgBytecodeWriter writer)
+		{
+			writer.WriteOpcode(BgOpcode.Choose);
+			writer.WriteExpr(Condition);
+			writer.WriteExprAsFragment(ValueIfTrue);
+			writer.WriteExprAsFragment(ValueIfFalse);
+		}
+
+		public override BgString ToBgString() => throw new InvalidOperationException();
+	}
+
+	class BgThrowExpr : BgExpr
+	{
+		public string SourcePath { get; }
+		public int SourceLine { get; }
+		public BgString Message { get; }
+
+		public BgThrowExpr(string sourcePath, int sourceLine, BgString message)
+			: base(BgExprFlags.None)
+		{
+			SourcePath = sourcePath;
+			SourceLine = sourceLine;
+			Message = message;
+		}
+
+		public override void Write(BgBytecodeWriter writer)
+		{
+			writer.WriteOpcode(BgOpcode.Throw);
+			writer.WriteString(SourcePath);
+			writer.WriteUnsignedInteger(SourceLine);
+			writer.WriteExpr(Message);
+		}
+
+		public override BgString ToBgString() => Message;
+	}
+
+	class BgNullExpr : BgExpr
+	{
+		public BgNullExpr()
+			: base(BgExprFlags.NotInterned)
+		{
+		}
+
+		public override void Write(BgBytecodeWriter writer)
+		{
+			writer.WriteOpcode(BgOpcode.Null);
+		}
+
+		public override BgString ToBgString() => "null";
+	}
+
+	/// <summary>
+	/// Extension methods for expressions
+	/// </summary>
+	public static class BgExprExtensions
+	{
+		/// <summary>
+		/// Chose an expression if a condition evaluates to true
+		/// </summary>
+		/// <typeparam name="T">Type of the expression</typeparam>
+		/// <param name="expr">The expression value</param>
+		/// <param name="condition">Condition to test</param>
+		/// <param name="value">Value to return if the condition is true</param>
+		/// <returns>New expression</returns>
+		public static T If<T>(this T expr, BgBool condition, T value) where T : BgExpr
+		{
+			return BgExpr.Choose(condition, value, expr);
+		}
+
+		/// <summary>
+		/// Chose an expression if a condition evaluates to true
+		/// </summary>
+		/// <typeparam name="T">Type of the expression</typeparam>
+		/// <param name="expr">The expression value</param>
+		/// <param name="condition">Condition to test</param>
+		/// <param name="func">Function to apply to the expression if the condition is true</param>
+		/// <returns>New expression</returns>
+		public static T If<T>(this T expr, BgBool condition, Func<T, T> func) where T : BgExpr
+		{
+			return BgExpr.Choose(condition, func(expr), expr);
+		}
+	}
+}
diff --git a/Engine/Source/Programs/Shared/EpicGames.BuildGraph/Expressions/BgFileSet.cs b/Engine/Source/Programs/Shared/EpicGames.BuildGraph/Expressions/BgFileSet.cs
--- a/Engine/Source/Programs/Shared/EpicGames.BuildGraph/Expressions/BgFileSet.cs
+++ b/Engine/Source/Programs/Shared/EpicGames.BuildGraph/Expressions/BgFileSet.cs
@@ -1,6 +1,7 @@
 // Copyright Epic Games, Inc. All Rights Reserved.
 
 using System;
+using System.Collections.Generic;
 using EpicGames.Core;
 
 namespace EpicGames.BuildGraph.Expressions
@@ -9,139 +10,140 @@ namespace EpicGames.BuildGraph.Expressions
 	/// Represents a placeholder for the output from a node, which can be exchanged for the artifacts produced by a node at runtime
 	/// </summary>
 	[BgType(typeof(BgFileSetType))]
-	public abstract class BgFileSet : IBgExpr<BgFileSet>
+	public abstract class BgFileSet : BgExpr
 	{
-		/// <inheritdoc/>
-		public BgFileSet IfThen(BgBool condition, BgFileSet valueIfTrue) => new BgFileSetIfThenExpr(condition, valueIfTrue, this);
-
-		/// <inheritdoc/>
-		object IBgExpr.Compute(BgExprContext context) => Compute(context);
-
-		/// <inheritdoc/>
-		public abstract BgFileSet Compute(BgExprContext context);
-
 		/// <summary>
-		/// Gets the tag name for this fileset
+		/// Constant empty fileset
 		/// </summary>
-		/// <returns></returns>
-		public string ComputeTag(BgExprContext context) => ((BgFileSetTagExpr)Compute(context)).Name;
+		public static BgFileSet Empty { get; } = new BgFileSetOutputExpr(FileSet.Empty);
 
 		/// <summary>
-		/// Gets the tag name for this fileset
+		/// Constructor
 		/// </summary>
-		/// <returns></returns>
-		public FileSet ComputeValue(BgExprContext context) => ((BgFileSetValueExpr)Compute(context)).FileSet;
+		/// <param name="flags"></param>
+		public BgFileSet(BgExprFlags flags)
+			: base(flags)
+		{
+		}
 
 		/// <summary>
-		/// Implicit conversion from a file set to a functional file set
+		/// Implicit conversion from a regular fileset
 		/// </summary>
 		/// <param name="fileSet"></param>
-		public static implicit operator BgFileSet(FileSet fileSet) => new BgFileSetValueExpr(fileSet);
+		public static implicit operator BgFileSet(FileSet fileSet)
+		{
+			return new BgFileSetOutputExpr(fileSet);
+		}
 
 		/// <inheritdoc/>
-		public BgString ToBgString() => new BgFileSetToStringExpr(this);
+		public override BgString ToBgString() => "{FileSet}";
 	}
 
 	/// <summary>
 	/// Traits for a <see cref="BgFileSet"/>
 	/// </summary>
-	class BgFileSetType : BgTypeBase<BgFileSet>
+	class BgFileSetType : BgType<BgFileSet>
 	{
 		/// <inheritdoc/>
-		public override BgFileSet DeserializeArgument(string text) => new BgFileSetTagExpr(text);
-
-		/// <inheritdoc/>
-		public override string SerializeArgument(BgFileSet value, BgExprContext context) => ((BgFileSetTagExpr)value.Compute(context)).Name;
+		public override BgFileSet Constant(object value) => new BgFileSetInputExpr((BgNodeOutput[])value);
 
 		/// <inheritdoc/>
-		public override BgFileSet CreateConstant(object value) => new BgFileSetTagExpr(((BgFileSetTagExpr)value).Name);
-
-		/// <inheritdoc/>
-		public override IBgExprVariable<BgFileSet> CreateVariable() => throw new NotImplementedException();
+		public override BgFileSet Wrap(BgExpr expr) => new BgFileSetWrappedExpr(expr);
 	}
 
 	#region Expression classes
 
-	class BgFileSetIfThenExpr : BgFileSet
+	class BgFileSetInputExpr : BgFileSet
 	{
-		public BgBool Condition { get; }
-		public BgFileSet ValueIfTrue { get; }
-		public BgFileSet ValueIfFalse { get; }
+		public IReadOnlyList<BgNodeOutput> Outputs { get; }
 
-		public BgFileSetIfThenExpr(BgBool condition, BgFileSet valueIfTrue, BgFileSet valueIfFalse)
+		public BgFileSetInputExpr(BgNodeOutput[] outputs)
+			: base(BgExprFlags.None)
 		{
-			Condition = condition;
-			ValueIfTrue = valueIfTrue;
-			ValueIfFalse = valueIfFalse;
+			Outputs = outputs;
 		}
 
-		public override BgFileSet Compute(BgExprContext context) => Condition.Compute(context) ? ValueIfTrue.Compute(context) : ValueIfFalse.Compute(context);
+		public override void Write(BgBytecodeWriter writer)
+		{
+			throw new NotImplementedException();
+		}
 	}
 
-	class BgFileSetToStringExpr : BgString
+	/// <summary>
+	/// /
+	/// </summary>
+	public class BgFileSetOutputExpr : BgFileSet
 	{
-		public BgFileSet FileSet { get; }
+		/// <summary>
+		/// 
+		/// </summary>
+		public FileSet Value { get; }
 
-		internal BgFileSetToStringExpr(BgFileSet token)
+		/// <summary>
+		/// 
+		/// </summary>
+		/// <param name="value"></param>
+		public BgFileSetOutputExpr(FileSet value)
+			: base(BgExprFlags.NotInterned)
 		{
-			FileSet = token;
+			Value = value;
 		}
 
-		public override string Compute(BgExprContext context)
+		/// <inheritdoc/>
+		public override void Write(BgBytecodeWriter writer)
 		{
-			BgFileSet fileSetValue = FileSet.Compute(context);
-			if (fileSetValue is BgFileSetTagExpr tagExpr)
-			{
-				return tagExpr.Name;
-			}
-			else
-			{
-				return "{...}";
-			}
+			throw new NotImplementedException();
 		}
 	}
 
-	class BgFileSetTagExpr : BgFileSet
+	class BgFileSetFromNodeExpr : BgFileSet
 	{
-		public string Name { get; set; }
+		public BgNodeSpec Node { get; }
 
-		public BgFileSetTagExpr(string name)
+		public BgFileSetFromNodeExpr(BgNodeSpec node)
+			: base(BgExprFlags.NotInterned)
 		{
-			Name = name;
+			Node = node;
 		}
 
-		public override BgFileSet Compute(BgExprContext context) => this;
+		public override void Write(BgBytecodeWriter writer)
+		{
+			writer.WriteOpcode(BgOpcode.FileSetFromNode);
+			writer.WriteExpr(Node);
+		}
 	}
 
-	class BgFileSetTagFromStringExpr : BgFileSet
+	class BgFileSetFromNodeOutputExpr : BgFileSet
 	{
-		public BgString Name { get; }
+		public BgNodeSpec Node { get; }
+		public int OutputIndex { get; }
 
-		public BgFileSetTagFromStringExpr(BgString name)
+		public BgFileSetFromNodeOutputExpr(BgNodeSpec node, int outputIndex)
+			: base(BgExprFlags.NotInterned)
 		{
-			Name = name;
+			Node = node;
+			OutputIndex = outputIndex;
 		}
 
-		public override BgFileSet Compute(BgExprContext context) => new BgFileSetTagExpr(Name.Compute(context));
-	}
-
-	class BgFileSetVariableExpr : BgFileSet, IBgExprVariable<BgFileSet>
-	{
-		public BgFileSet Value { get; set; } = new BgFileSetValueExpr(FileSet.Empty);
-
-		public override BgFileSet Compute(BgExprContext context) => Value.Compute(context);
+		public override void Write(BgBytecodeWriter writer)
+		{
+			writer.WriteOpcode(BgOpcode.FileSetFromNodeOutput);
+			writer.WriteExpr(Node);
+			writer.WriteUnsignedInteger(OutputIndex);
+		}
 	}
 
-	class BgFileSetValueExpr : BgFileSet
+	class BgFileSetWrappedExpr : BgFileSet
 	{
-		public FileSet FileSet { get; }
+		public BgExpr Expr { get; }
 
-		public BgFileSetValueExpr(FileSet fileSet)
+		public BgFileSetWrappedExpr(BgExpr expr)
+			: base(expr.Flags)
 		{
-			FileSet = fileSet;
+			Expr = expr;
 		}
 
-		public override BgFileSet Compute(BgExprContext context) => this;
+		public override void Write(BgBytecodeWriter writer) => Expr.Write(writer);
 	}
 
 	#endregion
diff --git a/Engine/Source/Programs/Shared/EpicGames.BuildGraph/Expressions/BgFunc.cs b/Engine/Source/Programs/Shared/EpicGames.BuildGraph/Expressions/BgFunc.cs
new file mode 100644
--- /dev/null
+++ b/Engine/Source/Programs/Shared/EpicGames.BuildGraph/Expressions/BgFunc.cs
@@ -0,0 +1,172 @@
+// Copyright Epic Games, Inc. All Rights Reserved.
+
+using System;
+using System.Collections.Generic;
+using System.Linq;
+
+namespace EpicGames.BuildGraph.Expressions
+{
+	/// <summary>
+	/// Abstract base class for function expressions
+	/// </summary>
+	public abstract class BgFunc
+	{
+		/// <summary>
+		/// The function expression.
+		/// </summary>
+		public BgExpr Body { get; }
+
+		/// <summary>
+		/// Constructor
+		/// </summary>
+		/// <param name="body">Expression to be evaluated</param>
+		protected BgFunc(BgExpr body)
+		{
+			// Evaluate the body twice, and determine which expressions are captured outside this context
+			// Lexical scope can include those expressions too - if they are declared locally, they will not be duplicated
+			Body = body;
+		}
+	}
+
+	/// <summary>
+	/// A function taking no arguments
+	/// </summary>
+	/// <typeparam name="TResult">Result type</typeparam>
+	public class BgFunc<TResult> : BgFunc
+		where TResult : BgExpr
+	{
+		/// <summary>
+		/// Constructor
+		/// </summary>
+		/// <param name="func">Function to construct from</param>
+		public BgFunc(Func<TResult> func)
+			: base(func())
+		{
+		}
+
+		/// <summary>
+		/// Implicit conversion from a regular C# function
+		/// </summary>
+		public static implicit operator BgFunc<TResult>(Func<TResult> func) => new BgFunc<TResult>(func);
+
+		/// <summary>
+		/// Call the function with the given arguments
+		/// </summary>
+		/// <returns>Result from the function</returns>
+		public TResult Call() => BgType.Wrap<TResult>(new BgFuncCallExpr<TResult>(this));
+	}
+
+	/// <summary>
+	/// A function taking a single argument
+	/// </summary>
+	/// <typeparam name="TArg">Type of the function argument</typeparam>
+	/// <typeparam name="TResult">Result type</typeparam>
+	public class BgFunc<TArg, TResult> : BgFunc
+		where TArg : BgExpr
+		where TResult : BgExpr
+	{
+		/// <summary>
+		/// Constructor
+		/// </summary>
+		/// <param name="func">Function to construct from</param>
+		public BgFunc(Func<TArg, TResult> func)
+			: base(func(BgType.Wrap<TArg>(new BgFuncArgumentExpr(0))))
+		{
+		}
+
+		/// <summary>
+		/// Implicit conversion from a regular C# function
+		/// </summary>
+		public static implicit operator BgFunc<TArg, TResult>(Func<TArg, TResult> func) => new BgFunc<TArg, TResult>(func);
+
+		/// <summary>
+		/// Call the function with the given arguments
+		/// </summary>
+		/// <param name="arg">Argument to pass to the function</param>
+		/// <returns>Result from the function</returns>
+		public TResult Call(TArg arg) => BgType.Wrap<TResult>(new BgFuncCallExpr<TResult>(this, arg));
+	}
+
+	/// <summary>
+	/// A function taking two arguments
+	/// </summary>
+	/// <typeparam name="TArg1">Type of the first function argument</typeparam>
+	/// <typeparam name="TArg2">Type of the second function argument</typeparam>
+	/// <typeparam name="TResult">Result type</typeparam>
+	public class BgFunc<TArg1, TArg2, TResult> : BgFunc
+		where TArg1 : BgExpr
+		where TArg2 : BgExpr
+		where TResult : BgExpr
+	{
+		/// <summary>
+		/// Constructor
+		/// </summary>
+		/// <param name="func">Function to construct from</param>
+		public BgFunc(Func<TArg1, TArg2, TResult> func)
+			: base(func(BgType.Wrap<TArg1>(new BgFuncArgumentExpr(0)), BgType.Wrap<TArg2>(new BgFuncArgumentExpr(1))))
+		{
+		}
+
+		/// <summary>
+		/// Implicit conversion from a regular C# function
+		/// </summary>
+		public static implicit operator BgFunc<TArg1, TArg2, TResult>(Func<TArg1, TArg2, TResult> func) => new BgFunc<TArg1, TArg2, TResult>(func);
+
+		/// <summary>
+		/// Call the function with the given arguments
+		/// </summary>
+		/// <param name="arg1">First argument to the function</param>
+		/// <param name="arg2">Second argument to the function</param>
+		/// <returns>Result from the function</returns>
+		public TResult Call(TArg1 arg1, TArg2 arg2) => BgType.Wrap<TResult>(new BgFuncCallExpr<TResult>(this, arg1, arg2));
+	}
+
+	#region Expression classes
+
+	class BgFuncArgumentExpr : BgExpr
+	{
+		uint Index { get; }
+
+		public BgFuncArgumentExpr(uint index)
+			: base(BgExprFlags.NotInterned)
+		{
+			Index = index;
+		}
+
+		public override void Write(BgBytecodeWriter writer)
+		{
+			writer.WriteOpcode(BgOpcode.Argument);
+			writer.WriteUnsignedInteger(Index);
+		}
+
+		public override BgString ToBgString() => "{Function}";
+	}
+
+	class BgFuncCallExpr<T> : BgExpr where T : BgExpr
+	{
+		readonly BgFunc Function;
+		readonly BgExpr[] Arguments;
+
+		public BgFuncCallExpr(BgFunc function, params BgExpr[] arguments)
+			: base(BgExprFlags.None)
+		{
+			Function = function;
+			Arguments = arguments;
+		}
+
+		public override void Write(BgBytecodeWriter writer)
+		{
+			writer.WriteOpcode(BgOpcode.Call);
+			writer.WriteUnsignedInteger((uint)Arguments.Length);
+			foreach (BgExpr argument in Arguments)
+			{
+				argument.Write(writer);
+			}
+			writer.WriteExprAsFragment(Function.Body);
+		}
+
+		public override BgString ToBgString() => throw new InvalidOperationException();
+	}
+
+	#endregion
+}
diff --git a/Engine/Source/Programs/Shared/EpicGames.BuildGraph/Expressions/BgInt.cs b/Engine/Source/Programs/Shared/EpicGames.BuildGraph/Expressions/BgInt.cs
--- a/Engine/Source/Programs/Shared/EpicGames.BuildGraph/Expressions/BgInt.cs
+++ b/Engine/Source/Programs/Shared/EpicGames.BuildGraph/Expressions/BgInt.cs
@@ -1,14 +1,16 @@
 // Copyright Epic Games, Inc. All Rights Reserved.
 
 using System;
+using System.Collections.Generic;
+using System.Linq;
 
 namespace EpicGames.BuildGraph.Expressions
 {
 	/// <summary>
 	/// Abstract base class for expressions returning a 32-bit integer value 
 	/// </summary>
-	[BgType(typeof(BgIntTraits))]
-	public abstract class BgInt : IBgExpr<BgInt>
+	[BgType(typeof(BgIntType))]
+	public abstract class BgInt : BgExpr
 	{
 		/// <summary>
 		/// Implicit conversion from an integer value
@@ -18,38 +20,50 @@ namespace EpicGames.BuildGraph.Expressions
 			return new BgIntConstantExpr(value);
 		}
 
+		/// <summary>
+		/// Constructor
+		/// </summary>
+		/// <param name="flags">Flags for this expression</param>
+		public BgInt(BgExprFlags flags)
+			: base(flags)
+		{
+		}
+
+		/// <inheritdoc/>
+		public static BgInt operator -(BgInt value) => new BgIntUnaryExpr(BgOpcode.IntNegate, value);
+
 		/// <inheritdoc/>
-		public static BgInt operator +(BgInt lhs, BgInt rhs) => new BgIntAddExpr(lhs, rhs);
+		public static BgInt operator +(BgInt lhs, BgInt rhs) => new BgIntBinaryExpr(BgOpcode.IntAdd, lhs, rhs);
 
 		/// <inheritdoc/>
-		public static BgInt operator -(BgInt lhs, BgInt rhs) => new BgIntAddExpr(lhs, new BgIntNegateExpr(rhs));
+		public static BgInt operator -(BgInt lhs, BgInt rhs) => new BgIntBinaryExpr(BgOpcode.IntAdd, lhs, -rhs);
 
 		/// <inheritdoc/>
-		public static BgBool operator <(BgInt lhs, BgInt rhs) => new BgIntLtExpr(lhs, rhs);
+		public static BgInt operator *(BgInt lhs, BgInt rhs) => new BgIntBinaryExpr(BgOpcode.IntMultiply, lhs, rhs);
 
 		/// <inheritdoc/>
-		public static BgBool operator >(BgInt lhs, BgInt rhs) => new BgIntGtExpr(lhs, rhs);
+		public static BgInt operator /(BgInt lhs, BgInt rhs) => new BgIntBinaryExpr(BgOpcode.IntDivide, lhs, rhs);
 
 		/// <inheritdoc/>
-		public static BgBool operator ==(BgInt lhs, BgInt rhs) => new BgIntEqExpr(lhs, rhs);
+		public static BgInt operator %(BgInt lhs, BgInt rhs) => new BgIntBinaryExpr(BgOpcode.IntModulo, lhs, rhs);
 
 		/// <inheritdoc/>
-		public static BgBool operator !=(BgInt lhs, BgInt rhs) => new BgBoolNotExpr(new BgIntEqExpr(lhs, rhs));
+		public static BgBool operator <(BgInt lhs, BgInt rhs) => new BgIntTestExpr(BgOpcode.IntLt, lhs, rhs);
 
 		/// <inheritdoc/>
-		public static BgBool operator <=(BgInt lhs, BgInt rhs) => new BgBoolNotExpr(new BgIntGtExpr(lhs, rhs));
+		public static BgBool operator >(BgInt lhs, BgInt rhs) => new BgIntTestExpr(BgOpcode.IntGt, lhs, rhs);
 
 		/// <inheritdoc/>
-		public static BgBool operator >=(BgInt lhs, BgInt rhs) => new BgBoolNotExpr(new BgIntLtExpr(lhs, rhs));
+		public static BgBool operator ==(BgInt lhs, BgInt rhs) => new BgIntTestExpr(BgOpcode.IntEq, lhs, rhs);
 
 		/// <inheritdoc/>
-		object IBgExpr.Compute(BgExprContext context) => Compute(context);
+		public static BgBool operator !=(BgInt lhs, BgInt rhs) => !(lhs == rhs);
 
 		/// <inheritdoc/>
-		public abstract int Compute(BgExprContext context);
+		public static BgBool operator <=(BgInt lhs, BgInt rhs) => !(lhs > rhs);
 
 		/// <inheritdoc/>
-		public BgInt IfThen(BgBool condition, BgInt valueIfTrue) => new BgIntChooseExpr(condition, valueIfTrue, this);
+		public static BgBool operator >=(BgInt lhs, BgInt rhs) => !(lhs < rhs);
 
 		/// <inheritdoc/>
 		public override bool Equals(object? obj) => throw new InvalidOperationException();
@@ -58,95 +72,87 @@ namespace EpicGames.BuildGraph.Expressions
 		public override int GetHashCode() => throw new InvalidOperationException();
 
 		/// <inheritdoc/>
-		public BgString ToBgString() => new BgIntFormatExpr(this);
+		public override BgString ToBgString()
+		{
+			throw new NotImplementedException();
+		}
 	}
 
 	/// <summary>
 	/// Traits for a <see cref="BgInt"/>
 	/// </summary>
-	class BgIntTraits : BgTypeBase<BgInt>
+	class BgIntType : BgType<BgInt>
 	{
 		/// <inheritdoc/>
-		public override BgInt DeserializeArgument(string value) => (BgInt)Int32.Parse(value);
+		public override BgInt Constant(object value) => new BgIntConstantExpr((int)value);
 
 		/// <inheritdoc/>
-		public override string SerializeArgument(BgInt value, BgExprContext context) => value.Compute(context).ToString();
-
-		/// <inheritdoc/>
-		public override BgInt CreateConstant(object value) => new BgIntConstantExpr((int)value);
-
-		/// <inheritdoc/>
-		public override IBgExprVariable<BgInt> CreateVariable() => new BgIntVariableExpr();
+		public override BgInt Wrap(BgExpr expr) => new BgIntWrappedExpr(expr);
 	}
 
 	#region Expression classes
 
-	class BgIntEqExpr : BgBool
+	class BgIntUnaryExpr : BgInt
 	{
-		public BgInt Lhs { get; }
-		public BgInt Rhs { get; }
+		public BgOpcode Opcode { get; }
+		public BgInt Value { get; }
 
-		public BgIntEqExpr(BgInt lhs, BgInt rhs)
+		public BgIntUnaryExpr(BgOpcode opcode, BgInt value)
+			: base(value.Flags & BgExprFlags.Eager)
 		{
-			Lhs = lhs;
-			Rhs = rhs;
+			Opcode = opcode;
+			Value = value;
 		}
 
-		public override bool Compute(BgExprContext context) => Lhs.Compute(context) == Rhs.Compute(context);
-	}
-
-	class BgIntLtExpr : BgBool
-	{
-		public BgInt Lhs { get; }
-		public BgInt Rhs { get; }
-
-		public BgIntLtExpr(BgInt lhs, BgInt rhs)
+		public override void Write(BgBytecodeWriter writer)
 		{
-			Lhs = lhs;
-			Rhs = rhs;
+			writer.WriteOpcode(Opcode);
+			writer.WriteExpr(Value);
 		}
-
-		public override bool Compute(BgExprContext context) => Lhs.Compute(context) < Rhs.Compute(context);
 	}
 
-	class BgIntGtExpr : BgBool
+	class BgIntBinaryExpr : BgInt
 	{
+		public BgOpcode Opcode { get; }
 		public BgInt Lhs { get; }
 		public BgInt Rhs { get; }
 
-		public BgIntGtExpr(BgInt lhs, BgInt rhs)
+		public BgIntBinaryExpr(BgOpcode opcode, BgInt lhs, BgInt rhs)
+			: base(lhs.Flags & rhs.Flags & BgExprFlags.Eager)
 		{
+			Opcode = opcode;
 			Lhs = lhs;
 			Rhs = rhs;
 		}
 
-		public override bool Compute(BgExprContext context) => Lhs.Compute(context) > Rhs.Compute(context);
+		public override void Write(BgBytecodeWriter writer)
+		{
+			writer.WriteOpcode(Opcode);
+			writer.WriteExpr(Lhs);
+			writer.WriteExpr(Rhs);
+		}
 	}
 
-	class BgIntAddExpr : BgInt
+	class BgIntTestExpr : BgBool
 	{
+		public BgOpcode Opcode { get; }
 		public BgInt Lhs { get; }
 		public BgInt Rhs { get; }
 
-		public BgIntAddExpr(BgInt lhs, BgInt rhs)
+		public BgIntTestExpr(BgOpcode opcode, BgInt lhs, BgInt rhs)
+			: base(lhs.Flags & rhs.Flags & BgExprFlags.Eager)
 		{
+			Opcode = opcode;
 			Lhs = lhs;
 			Rhs = rhs;
 		}
 
-		public override int Compute(BgExprContext context) => Lhs.Compute(context) + Rhs.Compute(context);
-	}
-
-	class BgIntNegateExpr : BgInt
-	{
-		public BgInt Inner { get; }
-
-		public BgIntNegateExpr(BgInt inner)
+		public override void Write(BgBytecodeWriter writer)
 		{
-			Inner = inner;
+			writer.WriteOpcode(Opcode);
+			writer.WriteExpr(Lhs);
+			writer.WriteExpr(Rhs);
 		}
-
-		public override int Compute(BgExprContext context) => -Inner.Compute(context);
 	}
 
 	class BgIntChooseExpr : BgInt
@@ -156,13 +162,20 @@ namespace EpicGames.BuildGraph.Expressions
 		public BgInt ValueIfFalse { get; }
 
 		public BgIntChooseExpr(BgBool condition, BgInt valueIfTrue, BgInt valueIfFalse)
+			: base(BgExprFlags.None)
 		{
 			Condition = condition;
 			ValueIfTrue = valueIfTrue;
 			ValueIfFalse = valueIfFalse;
 		}
 
-		public override int Compute(BgExprContext context) => Condition.Compute(context) ? ValueIfTrue.Compute(context) : ValueIfFalse.Compute(context);
+		public override void Write(BgBytecodeWriter writer)
+		{
+			writer.WriteOpcode(BgOpcode.Choose);
+			writer.WriteExpr(Condition);
+			writer.WriteExpr(ValueIfTrue);
+			writer.WriteExpr(ValueIfFalse);
+		}
 	}
 
 	class BgIntConstantExpr : BgInt
@@ -170,30 +183,46 @@ namespace EpicGames.BuildGraph.Expressions
 		public int Value { get; }
 
 		public BgIntConstantExpr(int value)
+			: base(BgExprFlags.NotInterned | BgExprFlags.Eager)
 		{
 			Value = value;
 		}
 
-		public override int Compute(BgExprContext context) => Value;
+		public override void Write(BgBytecodeWriter writer)
+		{
+			writer.WriteOpcode(BgOpcode.IntLiteral);
+			writer.WriteSignedInteger(Value);
+		}
 	}
 
-	class BgIntVariableExpr : BgInt, IBgExprVariable<BgInt>
+	class BgIntWrappedExpr : BgInt
 	{
-		public BgInt Value { get; set; } = 0;
+		BgExpr Expr { get; }
 
-		public override int Compute(BgExprContext context) => Value.Compute(context);
+		public BgIntWrappedExpr(BgExpr expr)
+			: base(expr.Flags)
+		{
+			Expr = expr;
+		}
+
+		public override void Write(BgBytecodeWriter writer) => Expr.Write(writer);
 	}
 
-	class BgIntFormatExpr : BgString
+	class BgIntToBgStringExpr : BgString
 	{
-		public BgInt Value { get; }
+		BgInt Expr { get; }
 
-		public BgIntFormatExpr(BgInt value)
+		public BgIntToBgStringExpr(BgInt expr)
+			: base(expr.Flags & BgExprFlags.Eager)
 		{
-			Value = value;
+			Expr = expr;
 		}
 
-		public override string Compute(BgExprContext context) => Value.Compute(context).ToString();
+		public override void Write(BgBytecodeWriter writer)
+		{
+			writer.WriteOpcode(BgOpcode.IntToString);
+			writer.WriteExpr(Expr);
+		}
 	}
 
 	#endregion
diff --git a/Engine/Source/Programs/Shared/EpicGames.BuildGraph/Expressions/BgList.cs b/Engine/Source/Programs/Shared/EpicGames.BuildGraph/Expressions/BgList.cs
--- a/Engine/Source/Programs/Shared/EpicGames.BuildGraph/Expressions/BgList.cs
+++ b/Engine/Source/Programs/Shared/EpicGames.BuildGraph/Expressions/BgList.cs
@@ -6,22 +6,11 @@ using System.Linq;
 
 namespace EpicGames.BuildGraph.Expressions
 {
-	/// <summary>
-	/// Interface for all list types
-	/// </summary>
-	public interface IBgList : IBgExpr
-	{
-		/// <summary>
-		/// Type of elements in this list
-		/// </summary>
-		Type ElementType { get; }
-	}
-
 	/// <summary>
 	/// Abstract base class for expressions returning an immutable list of values
 	/// </summary>
 	[BgType(typeof(BgListType<>))]
-	public abstract class BgList<T> : IBgExpr<BgList<T>>, IBgList where T : IBgExpr<T>
+	public abstract class BgList<T> : BgExpr where T : BgExpr
 	{
 		/// <inheritdoc/>
 		public Type ElementType => typeof(T);
@@ -29,74 +18,124 @@ namespace EpicGames.BuildGraph.Expressions
 		/// <summary>
 		/// Constant representation of an empty list
 		/// </summary>
-		public static BgList<T> Empty { get; } = new BgListConstantExpr<T>(Array.Empty<T>());
+		public static BgList<T> Empty { get; } = new BgListEmptyExpr<T>();
+
+		/// <summary>
+		/// Constructor
+		/// </summary>
+		/// <param name="flags">Flags for this expression</param>
+		public BgList(BgExprFlags flags)
+			: base(flags)
+		{
+		}
+
+		/// <inheritdoc/>
+		public T this[BgInt index] => BgType.Wrap<T>(new BgListElementExpr<T>(this, index));
 
 		/// <summary>
 		/// Implicit conversion operator from a single value
 		/// </summary>
-		public static implicit operator BgList<T>(T value) => new BgListConstantExpr<T>(new[] { value });
+		public static implicit operator BgList<T>(T item) => Create(item);
 
 		/// <summary>
 		/// Implicit conversion operator from an array of values
 		/// </summary>
-		public static implicit operator BgList<T>(T[] value) => new BgListConstantExpr<T>(value);
+		public static implicit operator BgList<T>(T[] items) => Create(items);
 
 		/// <summary>
 		/// Implicit conversion operator from a list of values
 		/// </summary>
-		public static implicit operator BgList<T>(List<T> value) => new BgListConstantExpr<T>(value);
+		public static implicit operator BgList<T>(List<T> items) => Create(items);
 
 		/// <summary>
 		/// Crates a list from an array of values
 		/// </summary>
 		/// <param name="items">Sequence to construct from</param>
 		/// <returns></returns>
-		public static BgList<T> Create(IEnumerable<T> items) => new BgListConstantExpr<T>(items);
+		public static BgList<T> Create(IEnumerable<T> items) => Empty.Add(items);
 
 		/// <summary>
 		/// Crates a list from an array of values
 		/// </summary>
 		/// <param name="items">Sequence to construct from</param>
 		/// <returns></returns>
-		public static BgList<T> Create(params T[] items) => new BgListConstantExpr<T>(items);
+		public static BgList<T> Create(params T[] items) => Empty.Add(items);
 
 		/// <summary>
-		/// Concatenates mutiple lists together
+		/// Concatenates two lists together
 		/// </summary>
-		/// <param name="sources"></param>
+		/// <param name="lhs"></param>
+		/// <param name="rhs"></param>
 		/// <returns></returns>
-		public static BgList<T> Concat(params BgList<T>[] sources) => new BgListConcatExpr<T>(sources);
+		public static BgList<T> Concat(BgList<T> lhs, BgList<T> rhs) => new BgListConcatExpr<T>(lhs, rhs);
+
+		/// <summary>
+		/// Concatenates two lists together
+		/// </summary>
+		/// <param name="others"></param>
+		/// <returns></returns>
+		public static BgList<T> Concat(params BgList<T>[] others)
+		{
+			BgList<T> result = BgList<T>.Empty;
+			for (int idx = 0; idx < others.Length; idx++)
+			{
+				result = Concat(result, others[idx]);
+			}
+			return result;
+		}
+
+		/// <summary>
+		/// Gets the length of this list
+		/// </summary>
+		public BgInt Count => new BgListCountExpr<T>(this);
+
+		/// <summary>
+		/// Adds an item to the end of the list, returning the new list
+		/// </summary>
+		/// <param name="item">Items to add</param>
+		/// <returns>New list containing the given items</returns>
+		public BgList<T> Add(T item) => new BgListPushExpr<T>(this, item);
 
 		/// <summary>
 		/// Adds items to the end of the list, returning the new list
 		/// </summary>
 		/// <param name="items">Items to add</param>
 		/// <returns>New list containing the given items</returns>
-		public BgList<T> Add(params T[] items) => new BgListConcatExpr<T>(this, items);
+		public BgList<T> Add(params T[] items) => Add((IEnumerable<T>)items);
 
 		/// <inheritdoc cref="Add(T[])"/>
-		public BgList<T> Add(params BgList<T>[] items) => new BgListConcatExpr<T>(this, items);
+		public BgList<T> Add(BgList<T> list) => Concat(this, list);
 
 		/// <inheritdoc cref="Add(T[])"/>
-		public BgList<T> AddIf(BgBool condition, params T[] items) => IfThen(condition, Add(items));
+		public BgList<T> Add(IEnumerable<T> items)
+		{
+			BgList<T> list = this;
+			foreach (T item in items)
+			{
+				list = list.Add(item);
+			}
+			return list;
+		}
 
-		/// <inheritdoc cref="Add(T[])"/>
-		public BgList<T> AddIf(BgBool condition, params BgList<T>[] items) => IfThen(condition, Add(items));
+		/// <inheritdoc cref="Union(BgList{T})"/>
+		public BgList<T> Union(params T[] items) => Union(Create(items));
+
+		/// <inheritdoc cref="Union(BgList{T})"/>
+		public BgList<T> Union(IEnumerable<T> items) => Union(Create(items));
 
 		/// <summary>
-		/// Removes the given items from this list
+		/// Creates the union of this list with another
 		/// </summary>
-		/// <param name="items">Items to remove</param>
-		/// <returns>New list without the given items</returns>
-		public BgList<T> Except(BgList<T> items) => new BgListExceptExpr<T>(this, items);
+		/// <param name="items">Items to add</param>
+		/// <returns>Union with the given items</returns>
+		public BgList<T> Union(BgList<T> items) => new BgListUnionExpr<T>(this, items);
 
 		/// <summary>
 		/// Removes the given items from this list
 		/// </summary>
-		/// <param name="condition">Condition to remove the items</param>
 		/// <param name="items">Items to remove</param>
 		/// <returns>New list without the given items</returns>
-		public BgList<T> ExceptIf(BgBool condition, BgList<T> items) => IfThen(condition, Except(items));
+		public BgList<T> Except(BgList<T> items) => new BgListExceptExpr<T>(this, items);
 
 		/// <summary>
 		/// Removes any duplicate items from the list. The first item in the list is retained in its original order.
@@ -105,7 +144,10 @@ namespace EpicGames.BuildGraph.Expressions
 		public BgList<T> Distinct() => new BgListDistinctExpr<T>(this);
 
 		/// <inheritdoc cref="Enumerable.Select{TSource, TResult}(IEnumerable{TSource}, Func{TSource, TResult})"/>
-		public BgList<TResult> Select<TResult>(Func<T, TResult> function) where TResult : IBgExpr<TResult> => new BgListSelectExpr<T, TResult>(this, function);
+		public BgList<TResult> Select<TResult>(Func<T, TResult> function) where TResult : BgExpr => new BgListSelectExpr<T, TResult>(this, function);
+
+		/// <inheritdoc cref="Enumerable.Select{TSource, TResult}(IEnumerable{TSource}, Func{TSource, TResult})"/>
+		public BgList<TResult> Select<TResult>(BgFunc<T, TResult> function) where TResult : BgExpr => new BgListSelectExpr<T, TResult>(this, function);
 
 		/// <inheritdoc cref="Enumerable.Where{TSource}(IEnumerable{TSource}, Func{TSource, Boolean})"/>
 		public BgList<T> Where(Func<T, BgBool> predicate) => new BgListWhereExpr<T>(this, predicate);
@@ -113,273 +155,292 @@ namespace EpicGames.BuildGraph.Expressions
 		/// <inheritdoc cref="Enumerable.Contains{TSource}(IEnumerable{TSource}, TSource)"/>
 		public BgBool Contains(T item) => new BgListContainsExpr<T>(this, item);
 
-		/// <inheritdoc/>
-		object IBgExpr.Compute(BgExprContext context) => GetEnumerable(context);
-
-		/// <inheritdoc/>
-		public abstract IEnumerable<T> GetEnumerable(BgExprContext context);
-
-		/// <inheritdoc/>
-		public BgList<T> IfThen(BgBool condition, BgList<T> valueIfTrue) => new BgListChooseExpr<T>(condition, valueIfTrue, this);
+		/// <summary>
+		/// Creates a lazily evaluated copy of this list, if it's not constant
+		/// </summary>
+		/// <returns></returns>
+		public BgList<T> Lazy() => new BgListLazyEval<T>(this);
 
 		/// <inheritdoc/>
-		public BgString ToBgString() => new BgListFormatExpr<T>(this);
+		public override BgString ToBgString() => "{List}";
 	}
 
 	/// <summary>
 	/// Traits for a <see cref="BgList{T}"/>
 	/// </summary>
-	class BgListType<T> : BgTypeBase<BgList<T>> where T : IBgExpr<T>
+	class BgListType<T> : BgType<BgList<T>> where T : BgExpr
 	{
 		/// <inheritdoc/>
-		public override BgList<T> DeserializeArgument(string text)
-		{
-			IBgType<T> converter = BgType.Get<T>();
-			return BgList<T>.Create(text.Split(';').Select(x => converter.DeserializeArgument(x)));
-		}
+		public override BgList<T> Constant(object value) => new BgListConstantExpr<T>((List<object>)value);
 
 		/// <inheritdoc/>
-		public override string SerializeArgument(BgList<T> value, BgExprContext context)
-		{
-			IBgType<T> converter = BgType.Get<T>();
-			return String.Join(";", value.GetEnumerable(context).Select(x => converter.SerializeArgument(x, context)));
-		}
+		public override BgList<T> Wrap(BgExpr expr) => new BgListWrappedExpr<T>(expr);
+	}
 
-		/// <inheritdoc/>
-		public override BgList<T> CreateConstant(object value)
+	#region Expression classes
+
+	class BgListEmptyExpr<T> : BgList<T> where T : BgExpr
+	{
+		public BgListEmptyExpr()
+			: base(BgExprFlags.NotInterned)
 		{
-			return new BgListConstantExpr<T>((IEnumerable<T>)value);
 		}
 
-		/// <inheritdoc/>
-		public override IBgExprVariable<BgList<T>> CreateVariable()
+		public override void Write(BgBytecodeWriter writer)
 		{
-			throw new NotImplementedException();
+			writer.WriteOpcode(BgOpcode.ListEmpty);
 		}
 	}
 
-	/// <summary>
-	/// Extension methods for lists
-	/// </summary>
-	public static class BgListExtensions
+	class BgListConstantExpr<T> : BgList<T> where T : BgExpr
 	{
-		internal static List<string> Compute(this BgList<BgString> list, BgExprContext context)
-		{
-			return list.GetEnumerable(context).Select(x => x.Compute(context)).ToList();
-		}
+		public List<object> Value { get; }
 
-		internal static List<BgFileSet> Compute(this BgList<BgFileSet> list, BgExprContext context)
+		public BgListConstantExpr(List<object> value)
+			: base(BgExprFlags.NotInterned)
 		{
-			return list.GetEnumerable(context).Select(x => x.Compute(context)).ToList();
+			Value = value;
 		}
 
-		internal static List<string> ComputeTags(this BgList<BgFileSet> list, BgExprContext context)
+		public override void Write(BgBytecodeWriter writer)
 		{
-			return list.GetEnumerable(context).Select(x => x.ComputeTag(context)).ToList();
+			throw new NotImplementedException();
 		}
 	}
 
-	#region Expression classes
-
-	class BgListConcatExpr<T> : BgList<T> where T : IBgExpr<T>
+	class BgListPushExpr<T> : BgList<T> where T : BgExpr
 	{
-		public BgList<T>[] Sources { get; }
-
-		public BgListConcatExpr(params BgList<T>[] sources)
-		{
-			Sources = sources;
-		}
+		public BgList<T> List { get; }
+		public T Item { get; }
 
-		public BgListConcatExpr(BgList<T> firstSource, BgList<T>[] otherSources)
+		public BgListPushExpr(BgList<T> list, T item)
+			: base(BgExprFlags.None)
 		{
-			Sources = new BgList<T>[otherSources.Length + 1];
-			Sources[0] = firstSource;
-			Array.Copy(otherSources, 0, Sources, 1, otherSources.Length);
+			List = list;
+			Item = item;
 		}
 
-		public override IEnumerable<T> GetEnumerable(BgExprContext context)
+		public override void Write(BgBytecodeWriter writer)
 		{
-			return Sources.SelectMany(x => x.GetEnumerable(context));
+			if ((Item.Flags & BgExprFlags.Eager) != 0)
+			{
+				writer.WriteOpcode(BgOpcode.ListPush);
+				writer.WriteExpr(List);
+				writer.WriteExpr(Item);
+			}
+			else
+			{
+				writer.WriteOpcode(BgOpcode.ListPushLazy);
+				writer.WriteExpr(List);
+				writer.WriteExprAsFragment(Item);
+			}
 		}
 	}
 
-	class BgListChooseExpr<T> : BgList<T> where T : IBgExpr<T>
+	class BgListCountExpr<T> : BgInt where T : BgExpr
 	{
-		public BgBool Condition { get; }
-		public BgList<T> ValueIfTrue { get; }
-		public BgList<T> ValueIfFalse { get; }
+		public BgList<T> List { get; }
 
-		public BgListChooseExpr(BgBool condition, BgList<T> valueIfTrue, BgList<T> valueIfFalse)
+		public BgListCountExpr(BgList<T> list)
+			: base(BgExprFlags.None)
 		{
-			Condition = condition;
-			ValueIfTrue = valueIfTrue;
-			ValueIfFalse = valueIfFalse;
+			List = list;
 		}
 
-		public override IEnumerable<T> GetEnumerable(BgExprContext context) => Condition.Compute(context) ? ValueIfTrue.GetEnumerable(context) : ValueIfFalse.GetEnumerable(context);
+		public override void Write(BgBytecodeWriter writer)
+		{
+			writer.WriteOpcode(BgOpcode.ListCount);
+			writer.WriteExpr(List);
+		}
 	}
 
-	class BgListConstantExpr<T> : BgList<T> where T : IBgExpr<T>
+	class BgListConcatExpr<T> : BgList<T> where T : BgExpr
 	{
-		public IEnumerable<T> Value { get; }
+		public BgList<T> Lhs { get; }
+		public BgList<T> Rhs { get; }
 
-		public BgListConstantExpr(IEnumerable<T> value)
+		public BgListConcatExpr(BgList<T> lhs, BgList<T> rhs)
+			: base(BgExprFlags.None)
 		{
-			Value = value;
+			Lhs = lhs;
+			Rhs = rhs.Lazy();
 		}
 
-		public override IEnumerable<T> GetEnumerable(BgExprContext context) => Value;
+		public override void Write(BgBytecodeWriter writer)
+		{
+			writer.WriteOpcode(BgOpcode.ListConcat);
+			writer.WriteExpr(Lhs);
+			writer.WriteExpr(Rhs);
+		}
 	}
 
-	class BgListFormatExpr<T> : BgString where T : IBgExpr<T>
+	class BgListElementExpr<T> : BgExpr where T : BgExpr
 	{
-		public BgList<T> Inner { get; }
+		public BgList<T> List;
+		public BgInt Index;
 
-		public BgListFormatExpr(BgList<T> inner)
+		public BgListElementExpr(BgList<T> list, BgInt index)
+			: base(BgExprFlags.None)
 		{
-			Inner = inner;
+			List = list;
+			Index = index;
 		}
 
-		public override string Compute(BgExprContext context) => String.Join(";", Inner.GetEnumerable(context).Select(x => x.ToBgString().Compute(context)));
-	}
-
-	class BgListVariableExpr<T> : BgList<T>, IBgExprVariable<BgList<T>> where T : IBgExpr<T>
-	{
-		public BgList<T> Value { get; set; } = BgList<T>.Empty;
+		public override void Write(BgBytecodeWriter writer)
+		{
+			writer.WriteOpcode(BgOpcode.ListElement);
+			writer.WriteExpr(List);
+			writer.WriteExpr(Index);
+		}
 
-		public override IEnumerable<T> GetEnumerable(BgExprContext context) => Value.GetEnumerable(context);
+		public override BgString ToBgString() => throw new InvalidOperationException();
 	}
 
-	class BgListDistinctExpr<T> : BgList<T> where T : IBgExpr<T>
+	class BgListDistinctExpr<T> : BgList<T> where T : BgExpr
 	{
 		public BgList<T> Source { get; }
 
 		public BgListDistinctExpr(BgList<T> source)
+			: base(BgExprFlags.None)
 		{
 			Source = source;
 		}
 
-		public override IEnumerable<T> GetEnumerable(BgExprContext context)
+		public override void Write(BgBytecodeWriter writer)
 		{
-			HashSet<object> values = new HashSet<object>();
-			foreach (T item in Source.GetEnumerable(context))
-			{
-				object itemValue = item.Compute(context);
-				if (values.Add(itemValue))
-				{
-					yield return item;
-				}
-			}
+			writer.WriteOpcode(BgOpcode.ListDistinct);
+			writer.WriteExpr(Source);
 		}
 	}
 
-	class BgFunc<TIn, TOut> where TIn : IBgExpr<TIn> where TOut : IBgExpr<TOut>
+	class BgListSelectExpr<TIn, TOut> : BgList<TOut> where TIn : BgExpr where TOut : BgExpr
 	{
-		public IBgExprVariable<TIn> ArgVar { get; }
-		public IBgExpr<TOut> FuncExpr { get; }
+		public BgList<TIn> Source { get; }
+		public BgFunc<TIn, TOut> Function { get; }
 
-		public BgFunc(IBgExprVariable<TIn> argVar, IBgExpr<TOut> funcExpr)
+		public BgListSelectExpr(BgList<TIn> source, BgFunc<TIn, TOut> function)
+			: base(BgExprFlags.None)
 		{
-			ArgVar = argVar;
-			FuncExpr = funcExpr;
+			Source = source;
+			Function = function;
 		}
 
-		public static implicit operator BgFunc<TIn, TOut>(Func<TIn, TOut> function)
+		public override void Write(BgBytecodeWriter writer)
 		{
-			IBgExprVariable<TIn> argVar = BgType.Get<TIn>().CreateVariable();
-			return new BgFunc<TIn, TOut>(argVar, function((TIn)argVar));
+			writer.WriteOpcode(BgOpcode.ListSelect);
+			writer.WriteExpr(Source);
+			writer.WriteExprAsFragment(Function.Body);
 		}
+	}
 
-		public object Compute(BgExprContext context, TIn argument)
+	class BgListWhereExpr<T> : BgList<T> where T : BgExpr
+	{
+		public BgList<T> Source { get; }
+		public BgFunc<T, BgBool> Predicate { get; }
+
+		public BgListWhereExpr(BgList<T> source, BgFunc<T, BgBool> predicate)
+			: base(BgExprFlags.None)
 		{
-			ArgVar.Value = argument;
-			return FuncExpr.Compute(context);
+			Source = source;
+			Predicate = predicate;
 		}
 
-		public BgString ToBgString()
+		public override void Write(BgBytecodeWriter writer)
 		{
-			throw new NotImplementedException();
+			writer.WriteOpcode(BgOpcode.ListWhere);
+			writer.WriteExpr(Source);
+			writer.WriteExprAsFragment(Predicate.Body);
 		}
 	}
 
-	class BgListSelectExpr<TIn, TOut> : BgList<TOut> where TIn : IBgExpr<TIn> where TOut : IBgExpr<TOut>
+	class BgListUnionExpr<T> : BgList<T> where T : BgExpr
 	{
-		public BgList<TIn> Source { get; }
-		public BgFunc<TIn, TOut> Function { get; }
+		public BgList<T> Input { get; }
+		public BgList<T> Other { get; }
 
-		public BgListSelectExpr(BgList<TIn> source, BgFunc<TIn, TOut> function)
+		public BgListUnionExpr(BgList<T> input, BgList<T> other)
+			: base(BgExprFlags.None)
 		{
-			Source = source;
-			Function = function;
+			Input = input;
+			Other = other.Lazy();
 		}
 
-		public override IEnumerable<TOut> GetEnumerable(BgExprContext context)
+		public override void Write(BgBytecodeWriter writer)
 		{
-			return Source.GetEnumerable(context).Select(x => (TOut)Function.Compute(context, x));
+			writer.WriteOpcode(BgOpcode.ListUnion);
+			writer.WriteExpr(Input);
+			writer.WriteExpr(Other);
 		}
 	}
 
-	class BgListWhereExpr<T> : BgList<T> where T : IBgExpr<T>
+	class BgListExceptExpr<T> : BgList<T> where T : BgExpr
 	{
-		public BgList<T> Source { get; }
-		public BgFunc<T, BgBool> Predicate { get; }
+		public BgList<T> Input { get; }
+		public BgList<T> Other { get; }
 
-		public BgListWhereExpr(BgList<T> source, BgFunc<T, BgBool> predicate)
+		public BgListExceptExpr(BgList<T> input, BgList<T> other)
+			: base(BgExprFlags.None)
 		{
-			Source = source;
-			Predicate = predicate;
+			Input = input;
+			Other = other;
 		}
 
-		public override IEnumerable<T> GetEnumerable(BgExprContext context)
+		public override void Write(BgBytecodeWriter writer)
 		{
-			return Source.GetEnumerable(context).Where(x => (bool)Predicate.Compute(context, x));
+			writer.WriteOpcode(BgOpcode.ListExcept);
+			writer.WriteExpr(Input);
+			writer.WriteExpr(Other);
 		}
 	}
 
-	class BgListExceptExpr<T> : BgList<T> where T : IBgExpr<T>
+	class BgListContainsExpr<T> : BgBool where T : BgExpr
 	{
-		public BgList<T> SourceList { get; }
-		public BgList<T> ExceptList { get; }
+		public BgList<T> List { get; }
+		public T Item { get; }
 
-		public BgListExceptExpr(BgList<T> source, BgList<T> except)
+		public BgListContainsExpr(BgList<T> list, T item)
+			: base(BgExprFlags.None)
 		{
-			SourceList = source;
-			ExceptList = except;
+			List = list;
+			Item = item;
 		}
 
-		public override IEnumerable<T> GetEnumerable(BgExprContext context)
+		public override void Write(BgBytecodeWriter writer)
 		{
-			HashSet<object> exceptValues = new HashSet<object>();
-			exceptValues.UnionWith(ExceptList.GetEnumerable(context).Select(x => x.Compute(context)));
-
-			IBgType<T> type = BgType.Get<T>();
-			foreach (T sourceItem in SourceList.GetEnumerable(context))
-			{
-				object sourceValue = sourceItem.Compute(context);
-				if (!exceptValues.Contains(sourceValue))
-				{
-					yield return type.CreateConstant(sourceValue);
-				}
-			}
+			writer.WriteOpcode(BgOpcode.ListContains);
+			writer.WriteExpr(List);
+			writer.WriteExpr(Item);
 		}
 	}
 
-	class BgListContainsExpr<T> : BgBool where T : IBgExpr<T>
+	class BgListLazyEval<T> : BgList<T> where T : BgExpr
 	{
-		public BgList<T> Source { get; }
-		public T Item { get; }
+		BgList<T> List { get; }
 
-		public BgListContainsExpr(BgList<T> source, T item)
+		public BgListLazyEval(BgList<T> list)
+			: base(BgExprFlags.None)
 		{
-			Source = source;
-			Item = item;
+			List = list;
+		}
+
+		public override void Write(BgBytecodeWriter writer)
+		{
+			writer.WriteOpcode(BgOpcode.ListLazy);
+			writer.WriteExprAsFragment(List);
 		}
+	}
+
+	class BgListWrappedExpr<T> : BgList<T> where T : BgExpr
+	{
+		BgExpr Expr { get; }
 
-		public override bool Compute(BgExprContext context)
+		public BgListWrappedExpr(BgExpr expr)
+			: base(expr.Flags)
 		{
-			object value = Item.Compute(context);
-			return Source.GetEnumerable(context).Any(x => x.Compute(context).Equals(value));
+			Expr = expr;
 		}
+
+		public override void Write(BgBytecodeWriter writer) => Expr.Write(writer);
 	}
 
 	#endregion
diff --git a/Engine/Source/Programs/Shared/EpicGames.BuildGraph/Expressions/BgObject.cs b/Engine/Source/Programs/Shared/EpicGames.BuildGraph/Expressions/BgObject.cs
new file mode 100644
--- /dev/null
+++ b/Engine/Source/Programs/Shared/EpicGames.BuildGraph/Expressions/BgObject.cs
@@ -0,0 +1,187 @@
+// Copyright Epic Games, Inc. All Rights Reserved.
+
+using System;
+using System.Collections.Generic;
+using System.Collections.Immutable;
+using System.Linq;
+using System.Linq.Expressions;
+using System.Reflection;
+
+namespace EpicGames.BuildGraph.Expressions
+{
+	/// <summary>
+	/// Abstract base class for expressions returning an object
+	/// </summary>
+	/// <typeparam name="T">The native type which mirrors this object</typeparam>
+	[BgType(typeof(BgObjectType<>))]
+	public abstract class BgObject<T> : BgExpr where T : new()
+	{
+		/// <summary>
+		/// Default instance of the type, for determining default values
+		/// </summary>
+		internal static T DefaultInstance { get; } = new T();
+
+		/// <summary>
+		/// Constant representation of an empty object
+		/// </summary>
+		public static BgObject<T> Empty { get; } = new BgObjectEmptyExpr<T>();
+
+		/// <summary>
+		/// Constructor
+		/// </summary>
+		/// <param name="flags">Flags for this expression</param>
+		public BgObject(BgExprFlags flags)
+			: base(flags)
+		{
+		}
+
+		/// <summary>
+		/// Sets the value of a property in the object
+		/// </summary>
+		/// <param name="name">Name of the field</param>
+		/// <param name="value">Value for the field</param>
+		/// <returns>New object with the field set</returns>
+		public BgObject<T> Set(string name, BgExpr value) => new BgObjectSetExpr<T>(this, name, value);
+
+		/// <summary>
+		/// Sets the value of a property in the object
+		/// </summary>
+		/// <param name="property">Name of the field</param>
+		/// <param name="value">Value for the field</param>
+		/// <returns>New object with the field set</returns>
+		public BgObject<T> Set<TExpr, TNative>(Expression<Func<T, TNative>> property, TExpr value) where TExpr : BgExpr
+		{
+			MemberExpression member = ((MemberExpression)property.Body);
+			PropertyInfo propertyInfo = (PropertyInfo)member.Member;
+			return Set(propertyInfo.Name, value);
+		}
+
+		/// <summary>
+		/// Gets the value of a field in the object
+		/// </summary>
+		/// <param name="name">Name of the field</param>
+		/// <param name="defaultValue">Default value for the field, if it's not defined</param>
+		/// <returns>Value of the field</returns>
+		public TExpr Get<TExpr>(string name, TExpr defaultValue) where TExpr : BgExpr => BgType.Wrap<TExpr>(new BgObjectGetExpr<T>(this, name, defaultValue));
+
+		/// <summary>
+		/// Gets the value of a field in the object
+		/// </summary>
+		/// <param name="property">Expression indicating the property to retrieve</param>
+		/// <returns>Value of the field</returns>
+		public TExpr Get<TExpr, TNative>(Expression<Func<T, TNative>> property) where TExpr : BgExpr
+		{
+			MemberExpression member = ((MemberExpression)property.Body);
+			PropertyInfo propertyInfo = (PropertyInfo)member.Member;
+			TExpr defaultValue = BgType.Constant<TExpr>(propertyInfo.GetValue(DefaultInstance)!);
+			return Get(propertyInfo.Name, defaultValue);
+		}
+
+		/// <inheritdoc/>
+		public override BgString ToBgString() => "{Object}";
+	}
+
+	/// <summary>
+	/// Traits for a <see cref="BgList{T}"/>
+	/// </summary>
+	class BgObjectType<T> : BgType<BgObject<T>> where T : new()
+	{
+		/// <inheritdoc/>
+		public override BgObject<T> Constant(object value) => new BgObjectConstantExpr<T>((ImmutableDictionary<string, object>)value);
+
+		/// <inheritdoc/>
+		public override BgObject<T> Wrap(BgExpr expr) => new BgObjectWrappedExpr<T>(expr);
+	}
+
+	#region Expression classes
+
+	class BgObjectEmptyExpr<T> : BgObject<T> where T : new()
+	{
+		public BgObjectEmptyExpr()
+			: base(BgExprFlags.NotInterned)
+		{
+		}
+
+		public override void Write(BgBytecodeWriter writer)
+		{
+			writer.WriteOpcode(BgOpcode.ObjEmpty);
+		}
+	}
+
+	class BgObjectConstantExpr<T> : BgObject<T> where T : new()
+	{
+		public ImmutableDictionary<string, object> Fields { get; }
+
+		public BgObjectConstantExpr(ImmutableDictionary<string, object> fields)
+			: base(BgExprFlags.NotInterned)
+		{
+			Fields = fields;
+		}
+
+		public override void Write(BgBytecodeWriter writer)
+		{
+			throw new NotImplementedException();
+		}
+	}
+
+	class BgObjectSetExpr<T> : BgObject<T> where T : new()
+	{
+		public BgObject<T> Object { get; }
+		public string Name { get; }
+		public BgExpr Value { get; }
+
+		public BgObjectSetExpr(BgObject<T> obj, string name, BgExpr value)
+			: base(BgExprFlags.NotInterned)
+		{
+			Object = obj;
+			Name = name;
+			Value = value;
+		}
+
+		public override void Write(BgBytecodeWriter writer)
+		{
+			writer.WriteOpcode(BgOpcode.ObjSet);
+			writer.WriteExpr(Object);
+			writer.WriteString(Name);
+			writer.WriteExpr(Value);
+		}
+	}
+
+	class BgObjectGetExpr<T> : BgObject<T> where T : new()
+	{
+		public BgObject<T> Object { get; }
+		public string Name { get; }
+		public BgExpr DefaultValue { get; }
+
+		public BgObjectGetExpr(BgObject<T> obj, string name, BgExpr defaultValue)
+			: base(BgExprFlags.NotInterned)
+		{
+			Object = obj;
+			Name = name;
+			DefaultValue = defaultValue;
+		}
+
+		public override void Write(BgBytecodeWriter writer)
+		{
+			writer.WriteOpcode(BgOpcode.ObjGet);
+			writer.WriteExpr(Object);
+			writer.WriteString(Name);
+			writer.WriteExpr(DefaultValue);
+		}
+	}
+
+	class BgObjectWrappedExpr<T> : BgObject<T> where T : new()
+	{
+		BgExpr Expr { get; }
+
+		public BgObjectWrappedExpr(BgExpr expr)
+			: base(expr.Flags)
+		{
+			Expr = expr;
+		}
+
+		public override void Write(BgBytecodeWriter writer) => Expr.Write(writer);
+	}
+
+	#endregion
+}
diff --git a/Engine/Source/Programs/Shared/EpicGames.BuildGraph/Expressions/BgString.cs b/Engine/Source/Programs/Shared/EpicGames.BuildGraph/Expressions/BgString.cs
--- a/Engine/Source/Programs/Shared/EpicGames.BuildGraph/Expressions/BgString.cs
+++ b/Engine/Source/Programs/Shared/EpicGames.BuildGraph/Expressions/BgString.cs
@@ -10,13 +10,22 @@ namespace EpicGames.BuildGraph.Expressions
 	/// <summary>
 	/// Abstract base class for expressions returning a string value 
 	/// </summary>
-	[BgType(typeof(BgStringTraits))]
-	public abstract class BgString : IBgExpr<BgString>
+	[BgType(typeof(BgStringType))]
+	public abstract class BgString : BgExpr
 	{
 		/// <summary>
 		/// Constant value for an empty string
 		/// </summary>
-		public static BgString Empty { get; } = new BgStringConstantExpr(String.Empty);
+		public static BgString Empty { get; } = new BgStringEmptyExpr();
+
+		/// <summary>
+		/// Constructor
+		/// </summary>
+		/// <param name="flags">Flags for this expression</param>
+		public BgString(BgExprFlags flags)
+			: base(flags)
+		{
+		}
 
 		/// <summary>
 		/// Implicit conversion from a regular string type
@@ -35,8 +44,11 @@ namespace EpicGames.BuildGraph.Expressions
 		/// <inheritdoc cref="String.Join{T}(String?, IEnumerable{T})"/>
 		public static BgString Join(BgString separator, BgList<BgString> values) => new BgStringJoinExpr(separator, values);
 
+		/// <inheritdoc cref="String.Split(String?, StringSplitOptions)"/>
+		public BgList<BgString> Split(BgString separator) => new BgStringSplitExpr(this, separator);
+
 		/// <inheritdoc cref="String.Format(String, Object?[])"/>
-		public static BgString Format(string format, params object?[] args) => new BgStringFormatExpr(format, args);
+		public static BgString Format(string format, params BgExpr[] args) => new BgStringFormatExpr(format, args);
 
 		/// <inheritdoc/>
 		public static BgString operator +(BgString lhs, BgString rhs) => new BgStringConcatExpr(lhs, rhs);
@@ -52,25 +64,11 @@ namespace EpicGames.BuildGraph.Expressions
 		/// </summary>
 		public BgString Append(BgString other) => this + other;
 
-		/// <summary>
-		/// Appens another string to this one if a condition is true
-		/// </summary>
-		public BgString AppendIf(BgBool condition, BgString other) => IfThen(condition, this + other);
-
-		/// <inheritdoc/>
-		public BgString IfThen(BgBool condition, BgString valueIfTrue) => new BgStringChooseExpr(condition, valueIfTrue, this);
-
-		/// <inheritdoc/>
-		public BgBool IsMatch(BgString input, string pattern) => new BgStringIsMatchExpr(input, pattern);
-
-		/// <inheritdoc/>
-		public BgString Replace(BgString input, string pattern, string replace) => new BgStringReplaceExpr(input, pattern, replace);
-
 		/// <inheritdoc/>
-		object IBgExpr.Compute(BgExprContext context) => Compute(context);
+		public BgBool Match(string pattern) => new BgStringMatchExpr(this, pattern);
 
 		/// <inheritdoc/>
-		public abstract string Compute(BgExprContext context);
+		public BgString Replace(string pattern, BgString replace) => new BgStringReplaceExpr(this, pattern, replace);
 
 		/// <inheritdoc/>
 		public override bool Equals(object? obj) => throw new InvalidOperationException();
@@ -79,25 +77,19 @@ namespace EpicGames.BuildGraph.Expressions
 		public override int GetHashCode() => throw new InvalidOperationException();
 
 		/// <inheritdoc/>
-		public BgString ToBgString() => this;
+		public override BgString ToBgString() => this;
 	}
 
 	/// <summary>
 	/// Traits implementation for <see cref="BgString"/>
 	/// </summary>
-	class BgStringTraits : BgTypeBase<BgString>
+	class BgStringType : BgType<BgString>
 	{
 		/// <inheritdoc/>
-		public override BgString DeserializeArgument(string text) => text;
+		public override BgString Constant(object value) => new BgStringConstantExpr((string)value);
 
 		/// <inheritdoc/>
-		public override string SerializeArgument(BgString value, BgExprContext context) => value.Compute(context);
-
-		/// <inheritdoc/>
-		public override BgString CreateConstant(object value) => new BgStringConstantExpr((string)value);
-
-		/// <inheritdoc/>
-		public override IBgExprVariable<BgString> CreateVariable() => new BgStringVariableExpr();
+		public override BgString Wrap(BgExpr expr) => new BgStringWrappedExpr(expr);
 	}
 
 	#region Expression classes
@@ -109,13 +101,20 @@ namespace EpicGames.BuildGraph.Expressions
 		public StringComparison Comparison { get; }
 
 		public BgStringCompareExpr(BgString lhs, BgString rhs, StringComparison comparison)
+			: base(lhs.Flags & rhs.Flags & BgExprFlags.Eager)
 		{
 			Lhs = lhs;
 			Rhs = rhs;
 			Comparison = comparison;
 		}
 
-		public override int Compute(BgExprContext context) => String.Compare(Lhs.Compute(context), Rhs.Compute(context), Comparison);
+		public override void Write(BgBytecodeWriter writer)
+		{
+			writer.WriteOpcode(BgOpcode.StrCompare);
+			writer.WriteExpr(Lhs);
+			writer.WriteExpr(Rhs);
+			writer.WriteUnsignedInteger((uint)Comparison);
+		}
 
 		public override string ToString() => $"Compare({Lhs}, {Rhs})";
 	}
@@ -126,12 +125,18 @@ namespace EpicGames.BuildGraph.Expressions
 		public BgString Rhs { get; }
 
 		public BgStringConcatExpr(BgString lhs, BgString rhs)
+			: base(lhs.Flags & rhs.Flags & BgExprFlags.Eager)
 		{
 			Lhs = lhs;
 			Rhs = rhs;
 		}
 
-		public override string Compute(BgExprContext context) => Lhs.Compute(context) + Rhs.Compute(context);
+		public override void Write(BgBytecodeWriter writer)
+		{
+			writer.WriteOpcode(BgOpcode.StrConcat);
+			writer.WriteExpr(Lhs);
+			writer.WriteExpr(Rhs);
+		}
 
 		public override string ToString() => $"Concat({Lhs}, {Rhs})";
 	}
@@ -142,101 +147,122 @@ namespace EpicGames.BuildGraph.Expressions
 		public BgList<BgString> Values { get; }
 
 		public BgStringJoinExpr(BgString separator, BgList<BgString> values)
+			: base(BgExprFlags.None)
 		{
 			Separator = separator;
 			Values = values;
 		}
 
-		public override string Compute(BgExprContext context) => String.Join(Separator.Compute(context), Values.Compute(context));
+		public override void Write(BgBytecodeWriter writer)
+		{
+			writer.WriteOpcode(BgOpcode.StrJoin);
+			writer.WriteExpr(Separator);
+			writer.WriteExpr(Values);
+		}
 
 		public override string ToString() => $"Join({Separator}, {Values})";
 	}
 
-	class BgStringIsMatchExpr : BgBool
+	class BgStringSplitExpr : BgList<BgString>
+	{
+		public BgString Source { get; }
+		public BgString Separator { get; }
+
+		public BgStringSplitExpr(BgString source, BgString separator)
+			: base(source.Flags & separator.Flags & BgExprFlags.Eager)
+		{
+			Source = source;
+			Separator = separator;
+		}
+
+		public override void Write(BgBytecodeWriter writer)
+		{
+			writer.WriteOpcode(BgOpcode.StrSplit);
+			writer.WriteExpr(Source);
+			writer.WriteExpr(Separator);
+		}
+
+		public override string ToString() => $"Split({Source}, {Separator})";
+	}
+
+	class BgStringMatchExpr : BgBool
 	{
 		public BgString Input { get; }
 		public string Pattern { get; }
 
-		public BgStringIsMatchExpr(BgString input, string pattern)
+		public BgStringMatchExpr(BgString input, string pattern)
+			: base(input.Flags & BgExprFlags.Eager)
 		{
 			Input = input;
 			Pattern = pattern;
 		}
 
-		public override bool Compute(BgExprContext context) => Regex.IsMatch(Input.Compute(context), Pattern);
+		public override void Write(BgBytecodeWriter writer)
+		{
+			writer.WriteOpcode(BgOpcode.StrIsMatch);
+			writer.WriteExpr(Input);
+			writer.WriteString(Pattern);
+		}
 	}
 
 	class BgStringReplaceExpr : BgString
 	{
 		public BgString Input { get; }
 		public string Pattern { get; }
-		public string Replacement { get; }
+		public BgString Replacement { get; }
 
-		public BgStringReplaceExpr(BgString input, string pattern, string replacement)
+		public BgStringReplaceExpr(BgString input, string pattern, BgString replacement)
+			: base(input.Flags & replacement.Flags & BgExprFlags.Eager)
 		{
 			Input = input;
 			Pattern = pattern;
 			Replacement = replacement;
 		}
 
-		public override string Compute(BgExprContext context) => Regex.Replace(Input.Compute(context), Pattern, Replacement);
+		public override void Write(BgBytecodeWriter writer)
+		{
+			writer.WriteOpcode(BgOpcode.StrReplace);
+			writer.WriteExpr(Input);
+			writer.WriteString(Pattern);
+			writer.WriteExpr(Replacement);
+		}
 	}
 
 	class BgStringFormatExpr : BgString
 	{
-		new string Format { get; }
-		BgString[] Arguments { get; }
+		readonly BgString _format;
+		readonly BgExpr[] _arguments;
 
-		public BgStringFormatExpr(string format, object?[] args)
+		public BgStringFormatExpr(BgString format, BgExpr[] arguments)
+			: base(BgExprFlags.None)
 		{
-			Format = format;
+			_format = format;
+			_arguments = arguments;
+		}
 
-			Arguments = new BgString[args.Length];
-			for (int idx = 0; idx < args.Length; idx++)
+		public override void Write(BgBytecodeWriter writer)
+		{
+			writer.WriteOpcode(BgOpcode.StrFormat);
+			writer.WriteExpr(_format);
+			writer.WriteUnsignedInteger((ulong)_arguments.Length);
+			foreach (BgExpr argument in _arguments)
 			{
-				if (args[idx] is IBgExpr expr)
-				{
-					Arguments[idx] = expr.ToBgString();
-				}
-				else if (args[idx] != null)
-				{
-					Arguments[idx] = args[idx]?.ToString() ?? String.Empty;
-				}
-				else
-				{
-					Arguments[idx] = String.Empty;
-				}
+				writer.WriteExpr(argument);
 			}
 		}
-
-		public override string Compute(BgExprContext context) => String.Format(Format, Arguments.Select(x => x.Compute(context)).ToArray());
-
-		public override string ToString() => $"Format({String.Join(", ", new[] { Format }.Concat(Arguments.Select(x => x.ToString())))})";
 	}
 
-	class BgStringChooseExpr : BgString
+	class BgStringEmptyExpr : BgString
 	{
-		public BgBool Condition { get; }
-		public BgString ValueIfTrue { get; }
-		public BgString ValueIfFalse { get; }
-
-		public BgStringChooseExpr(BgBool condition, BgString valueIfTrue, BgString valueIfFalse)
+		public BgStringEmptyExpr()
+			: base(BgExprFlags.NotInterned | BgExprFlags.Eager)
 		{
-			Condition = condition;
-			ValueIfTrue = valueIfTrue;
-			ValueIfFalse = valueIfFalse;
 		}
 
-		public override string Compute(BgExprContext context) => Condition.Compute(context) ? ValueIfTrue.Compute(context) : ValueIfFalse.Compute(context);
-	}
-
-	class BgStringVariableExpr : BgString, IBgExprVariable<BgString>
-	{
-		public BgString Value { get; set; } = BgString.Empty;
-
-		public BgString Variable => this;
-
-		public override string Compute(BgExprContext context) => Value.Compute(context);
+		public override void Write(BgBytecodeWriter writer)
+		{
+			writer.WriteOpcode(BgOpcode.StrEmpty);
+		}
 	}
 
 	class BgStringConstantExpr : BgString
@@ -244,26 +270,29 @@ namespace EpicGames.BuildGraph.Expressions
 		public string Value { get; }
 
 		public BgStringConstantExpr(string value)
+			: base(BgExprFlags.Eager)
 		{
 			Value = value;
 		}
 
-		public override string Compute(BgExprContext context) => Value;
-
-		public override bool Equals(object? obj)
+		public override void Write(BgBytecodeWriter writer)
 		{
-			return obj is string objStr && String.Equals(Value, objStr, StringComparison.OrdinalIgnoreCase);
+			writer.WriteOpcode(BgOpcode.StrLiteral);
+			writer.WriteString(Value);
 		}
+	}
 
-		public override int GetHashCode()
-		{
-			return String.GetHashCode(Value, StringComparison.OrdinalIgnoreCase);
-		}
+	class BgStringWrappedExpr : BgString
+	{
+		BgExpr Expr { get; }
 
-		public override string ToString()
+		public BgStringWrappedExpr(BgExpr expr)
+			: base(expr.Flags)
 		{
-			return $"\"Value\"";
+			Expr = expr;
 		}
+
+		public override void Write(BgBytecodeWriter writer) => Expr.Write(writer);
 	}
 
 	#endregion
diff --git a/Engine/Source/Programs/Shared/EpicGames.BuildGraph/Expressions/BgType.cs b/Engine/Source/Programs/Shared/EpicGames.BuildGraph/Expressions/BgType.cs
new file mode 100644
--- /dev/null
+++ b/Engine/Source/Programs/Shared/EpicGames.BuildGraph/Expressions/BgType.cs
@@ -0,0 +1,153 @@
+// Copyright Epic Games, Inc. All Rights Reserved.
+
+using System;
+using System.Collections.Concurrent;
+using System.Reflection;
+
+namespace EpicGames.BuildGraph.Expressions
+{
+	/// <summary>
+	/// Interface for type traits
+	/// </summary>
+	public interface IBgType
+	{
+		/// <summary>
+		/// Creates a constant expression wrapping the given value
+		/// </summary>
+		/// <param name="value">Value to wrap</param>
+		/// <returns></returns>
+		public abstract BgExpr Constant(object value);
+	}
+
+	/// <summary>
+	/// Base class for type functions
+	/// </summary>
+	public abstract class BgType<T> : IBgType where T : BgExpr
+	{
+		/// <inheritdoc/>
+		BgExpr IBgType.Constant(object value) => Constant(value);
+
+		/// <inheritdoc cref="IBgType.Constant(Object)"/>
+		public abstract T Constant(object value);
+
+		/// <summary>
+		/// Wraps an untyped expression as a strongly typed value
+		/// </summary>
+		/// <param name="expr">Expression to wrap</param>
+		/// <returns></returns>
+		public abstract T Wrap(BgExpr expr);
+	}
+
+	/// <summary>
+	/// Attribute used to specify the converter class to use for a type
+	/// </summary>
+	class BgTypeAttribute : Attribute
+	{
+		/// <summary>
+		/// The converter type
+		/// </summary>
+		public Type Type { get; }
+
+		/// <summary>
+		/// Constructor
+		/// </summary>
+		/// <param name="type"></param>
+		public BgTypeAttribute(Type type)
+		{
+			Type = type;
+		}
+	}
+
+	/// <summary>
+	/// Utility methods for type traits
+	/// </summary>
+	public abstract class BgType
+	{
+		class BgTypeCache<T> where T : BgExpr
+		{
+			public static BgType<T> Instance { get; } = (BgType<T>)BgType.CreateInstance(typeof(T));
+		}
+
+		static readonly ConcurrentDictionary<Type, IBgType> s_typeCache = new ConcurrentDictionary<Type, IBgType>();
+
+		/// <summary>
+		/// Create a converter instance for the given type
+		/// </summary>
+		/// <param name="type"></param>
+		/// <returns></returns>
+		static object CreateInstance(Type type)
+		{
+			BgTypeAttribute? converterAttr = type.GetCustomAttribute<BgTypeAttribute>(true);
+			if (converterAttr == null)
+			{
+				throw new InvalidOperationException($"Missing [BgType] attribute on {type.Name}");
+			}
+
+			Type converterType = converterAttr.Type;
+			if (converterType.IsGenericType)
+			{
+				converterType = converterType.MakeGenericType(type.GetGenericArguments());
+			}
+
+			return Activator.CreateInstance(converterType)!;
+		}
+
+		/// <summary>
+		/// Gets the type interface for a type
+		/// </summary>
+		/// <param name="type"></param>
+		/// <returns></returns>
+		public static IBgType Get(Type type)
+		{
+			return s_typeCache.GetOrAdd(type, CreateTypeInstance);
+		}
+
+		static IBgType CreateTypeInstance(Type type)
+		{
+			Type concreteType = typeof(BgTypeCache<>).MakeGenericType(type);
+			return (IBgType)concreteType.GetProperty(nameof(BgTypeCache<BgInt>.Instance))!.GetValue(null)!;
+		}
+
+		/// <summary>
+		/// Create a traits instance for the given type
+		/// </summary>
+		/// <returns></returns>
+		public static BgType<T> Get<T>() where T : BgExpr
+		{
+			return BgTypeCache<T>.Instance;
+		}
+
+		/// <summary>
+		/// Wraps an expression as a different type
+		/// </summary>
+		/// <typeparam name="T"></typeparam>
+		/// <param name="expr"></param>
+		/// <returns></returns>
+		public static T Wrap<T>(BgExpr expr) where T : BgExpr
+		{
+			return Get<T>().Wrap(expr);
+		}
+
+		/// <summary>
+		/// Creates a constant of the given type
+		/// </summary>
+		/// <typeparam name="T"></typeparam>
+		/// <param name="value"></param>
+		/// <returns></returns>
+		public static T Constant<T>(object value) where T : BgExpr
+		{
+			return Get<T>().Constant(value);
+		}
+
+		/// <summary>
+		/// Creates a constant of the given type
+		/// </summary>
+		/// <param name="type"></param>
+		/// <param name="value"></param>
+		/// <returns></returns>
+		public static BgExpr Constant(Type type, object value)
+		{
+			return Get(type).Constant(value);
+		}
+	}
+}
diff --git a/Engine/Source/Programs/Shared/EpicGames.BuildGraph/Expressions/IBgExpr.cs b/Engine/Source/Programs/Shared/EpicGames.BuildGraph/Expressions/IBgExpr.cs
deleted file mode 100644
diff --git a/Engine/Source/Programs/Shared/EpicGames.BuildGraph/Expressions/IBgType.cs b/Engine/Source/Programs/Shared/EpicGames.BuildGraph/Expressions/IBgType.cs
deleted file mode 100644
