BuildGraph: Rename classes to reduce boilerplate when writing scripts.

* Expression classes (nodes, aggregates, graphs, etc...) are now called BgNode, BgAggregate, etc...
* Evaluated and instantiated objects used by BuildGraph internals are now called BgNodeDef, BgAggregateDef, etc...

diff --git a/Engine/Source/Programs/AutomationTool/BuildGraph/BgGraphBuilder.cs b/Engine/Source/Programs/AutomationTool/BuildGraph/BgGraphBuilder.cs
--- a/Engine/Source/Programs/AutomationTool/BuildGraph/BgGraphBuilder.cs
+++ b/Engine/Source/Programs/AutomationTool/BuildGraph/BgGraphBuilder.cs
@@ -17,6 +17,6 @@ namespace AutomationTool
 		/// Callback used to instantiate the graph
 		/// </summary>
 		/// <param name="env">The graph context</param>
-		public abstract BgGraphSpec CreateGraph(BgEnvironment env);
+		public abstract BgGraph CreateGraph(BgEnvironment env);
 	}
 }
diff --git a/Engine/Source/Programs/AutomationTool/BuildGraph/BgNodeExecutor.cs b/Engine/Source/Programs/AutomationTool/BuildGraph/BgNodeExecutor.cs
--- a/Engine/Source/Programs/AutomationTool/BuildGraph/BgNodeExecutor.cs
+++ b/Engine/Source/Programs/AutomationTool/BuildGraph/BgNodeExecutor.cs
@@ -150,7 +150,7 @@ namespace AutomationTool
 	}
 
 	/// <summary>
-	/// Implementation of <see cref="BgNode"/> for graphs defined through XML syntax
+	/// Implementation of <see cref="BgNodeDef"/> for graphs defined through XML syntax
 	/// </summary>
 	class BgScriptNodeExecutor : BgNodeExecutor
 	{
diff --git a/Engine/Source/Programs/AutomationTool/BuildGraph/BgScriptReader.cs b/Engine/Source/Programs/AutomationTool/BuildGraph/BgScriptReader.cs
--- a/Engine/Source/Programs/AutomationTool/BuildGraph/BgScriptReader.cs
+++ b/Engine/Source/Programs/AutomationTool/BuildGraph/BgScriptReader.cs
@@ -275,9 +275,9 @@ namespace AutomationTool
 	}
 
 	/// <summary>
-	/// Overridden version of <see cref="BgNode"/> which contains a list of tasks
+	/// Overridden version of <see cref="BgNodeDef"/> which contains a list of tasks
 	/// </summary>
-	class BgScriptNode : BgNode
+	class BgScriptNode : BgNodeDef
 	{
 		/// <summary>
 		/// List of tasks to execute
@@ -287,7 +287,7 @@ namespace AutomationTool
 		/// <summary>
 		/// Constructor
 		/// </summary>
-		public BgScriptNode(string name, IReadOnlyList<BgNodeOutput> inputs, IReadOnlyList<string> outputNames, IReadOnlyList<BgNode> inputDependencies, IReadOnlyList<BgNode> orderDependencies, IReadOnlyList<FileReference> requiredTokens)
+		public BgScriptNode(string name, IReadOnlyList<BgNodeOutput> inputs, IReadOnlyList<string> outputNames, IReadOnlyList<BgNodeDef> inputDependencies, IReadOnlyList<BgNodeDef> orderDependencies, IReadOnlyList<FileReference> requiredTokens)
 			: base(name, inputs, outputNames, inputDependencies, orderDependencies, requiredTokens)
 		{
 		}
@@ -318,7 +318,7 @@ namespace AutomationTool
 		/// <summary>
 		/// The current graph
 		/// </summary>
-		readonly BgGraph _graph = new BgGraph();
+		readonly BgGraphDef _graph = new BgGraphDef();
 
 		/// <summary>
 		/// Arguments for evaluating the graph
@@ -345,7 +345,7 @@ namespace AutomationTool
 		/// </summary>
 		public int NumErrors { get; private set; }
 
-		BgAgent? _enclosingAgent;
+		BgAgentDef? _enclosingAgent;
 		BgScriptNode? _enclosingNode;
 
 		/// <summary>
@@ -382,7 +382,7 @@ namespace AutomationTool
 		/// <param name="logger">Logger for output messages</param>
 		/// <param name="singleNodeName">If a single node will be processed, the name of that node.</param>
 		/// <returns>True if the graph was read, false if there were errors</returns>
-		public static async Task<BgGraph?> ReadAsync(FileReference file, Dictionary<string, string> arguments, Dictionary<string, string> defaultProperties, BgScriptSchema schema, ILogger logger, string? singleNodeName = null)
+		public static async Task<BgGraphDef?> ReadAsync(FileReference file, Dictionary<string, string> arguments, Dictionary<string, string> defaultProperties, BgScriptSchema schema, ILogger logger, string? singleNodeName = null)
 		{
 			// Read the file and build the graph
 			BgScriptReader reader = new BgScriptReader(defaultProperties, arguments, singleNodeName, schema, logger);
@@ -950,7 +950,7 @@ namespace AutomationTool
 				string[] types = ReadListAttribute(element, "Type");
 
 				// Create the agent object, or continue an existing one
-				BgAgent? agent;
+				BgAgentDef? agent;
 				if (_graph.NameToAgent.TryGetValue(name, out agent))
 				{
 					if (types.Length > 0 && agent.PossibleTypes.Length > 0)
@@ -967,7 +967,7 @@ namespace AutomationTool
 					{
 						LogError(element, "Missing type for agent '{0}'", name);
 					}
-					agent = new BgAgent(name, types);
+					agent = new BgAgentDef(name, types);
 					_graph.NameToAgent.Add(name, agent);
 					_graph.Agents.Add(agent);
 				}
@@ -1048,8 +1048,8 @@ namespace AutomationTool
 			{
 				string[] requiredNames = ReadListAttribute(element, "Requires");
 
-				BgAggregate newAggregate = new BgAggregate(name);
-				foreach (BgNode referencedNode in ResolveReferences(element, requiredNames))
+				BgAggregateDef newAggregate = new BgAggregateDef(name);
+				foreach (BgNodeDef referencedNode in ResolveReferences(element, requiredNames))
 				{
 					newAggregate.RequiredNodes.Add(referencedNode);
 				}
@@ -1058,21 +1058,21 @@ namespace AutomationTool
 				string labelCategoryName = ReadAttribute(element, "Label");
 				if (!String.IsNullOrEmpty(labelCategoryName))
 				{
-					BgLabel label;
+					BgLabelDef label;
 
 					// Create the label
 					int slashIdx = labelCategoryName.IndexOf('/');
 					if (slashIdx != -1)
 					{
-						label = new BgLabel(labelCategoryName.Substring(slashIdx + 1), labelCategoryName.Substring(0, slashIdx), null, null, BgLabelChange.Current);
+						label = new BgLabelDef(labelCategoryName.Substring(slashIdx + 1), labelCategoryName.Substring(0, slashIdx), null, null, BgLabelChange.Current);
 					}
 					else
 					{
-						label = new BgLabel(labelCategoryName, "Other", null, null, BgLabelChange.Current);
+						label = new BgLabelDef(labelCategoryName, "Other", null, null, BgLabelChange.Current);
 					}
 
 					// Find all the included nodes
-					foreach (BgNode requiredNode in newAggregate.RequiredNodes)
+					foreach (BgNodeDef requiredNode in newAggregate.RequiredNodes)
 					{
 						label.RequiredNodes.Add(requiredNode);
 						label.IncludedNodes.Add(requiredNode);
@@ -1080,14 +1080,14 @@ namespace AutomationTool
 					}
 
 					string[] includedNames = ReadListAttribute(element, "Include");
-					foreach (BgNode includedNode in ResolveReferences(element, includedNames))
+					foreach (BgNodeDef includedNode in ResolveReferences(element, includedNames))
 					{
 						label.IncludedNodes.Add(includedNode);
 						label.IncludedNodes.UnionWith(includedNode.OrderDependencies);
 					}
 
 					string[] excludedNames = ReadListAttribute(element, "Exclude");
-					foreach (BgNode excludedNode in ResolveReferences(element, excludedNames))
+					foreach (BgNodeDef excludedNode in ResolveReferences(element, excludedNames))
 					{
 						label.IncludedNodes.Remove(excludedNode);
 						label.IncludedNodes.ExceptWith(excludedNode.OrderDependencies);
@@ -1110,7 +1110,7 @@ namespace AutomationTool
 				string[] requiredNames = ReadListAttribute(element, "Requires");
 
 				BgReport newReport = new BgReport(name);
-				foreach (BgNode referencedNode in ResolveReferences(element, requiredNames))
+				foreach (BgNodeDef referencedNode in ResolveReferences(element, requiredNames))
 				{
 					newReport.Nodes.Add(referencedNode);
 					newReport.Nodes.UnionWith(referencedNode.OrderDependencies);
@@ -1133,12 +1133,12 @@ namespace AutomationTool
 				string project = ReadAttribute(element, "Project");
 				int change = ReadIntegerAttribute(element, "Change", 0);
 
-				BgBadge newBadge = new BgBadge(name, project, change);
-				foreach (BgNode referencedNode in ResolveReferences(element, requiredNames))
+				BgBadgeDef newBadge = new BgBadgeDef(name, project, change);
+				foreach (BgNodeDef referencedNode in ResolveReferences(element, requiredNames))
 				{
 					newBadge.Nodes.Add(referencedNode);
 				}
-				foreach (BgNode referencedNode in ResolveReferences(element, targetNames))
+				foreach (BgNodeDef referencedNode in ResolveReferences(element, targetNames))
 				{
 					newBadge.Nodes.Add(referencedNode);
 					newBadge.Nodes.UnionWith(referencedNode.OrderDependencies);
@@ -1172,19 +1172,19 @@ namespace AutomationTool
 
 				BgLabelChange change = ReadEnumAttribute<BgLabelChange>(element, "Change", BgLabelChange.Current);
 
-				BgLabel newLabel = new BgLabel(name, category, ugsBadge, ugsProject, change);
-				foreach (BgNode referencedNode in ResolveReferences(element, requiredNames))
+				BgLabelDef newLabel = new BgLabelDef(name, category, ugsBadge, ugsProject, change);
+				foreach (BgNodeDef referencedNode in ResolveReferences(element, requiredNames))
 				{
 					newLabel.RequiredNodes.Add(referencedNode);
 					newLabel.IncludedNodes.Add(referencedNode);
 					newLabel.IncludedNodes.UnionWith(referencedNode.OrderDependencies);
 				}
-				foreach (BgNode includedNode in ResolveReferences(element, includedNames))
+				foreach (BgNodeDef includedNode in ResolveReferences(element, includedNames))
 				{
 					newLabel.IncludedNodes.Add(includedNode);
 					newLabel.IncludedNodes.UnionWith(includedNode.OrderDependencies);
 				}
-				foreach (BgNode excludedNode in ResolveReferences(element, excludedNames))
+				foreach (BgNodeDef excludedNode in ResolveReferences(element, excludedNames))
 				{
 					newLabel.IncludedNodes.Remove(excludedNode);
 					newLabel.IncludedNodes.ExceptWith(excludedNode.OrderDependencies);
@@ -1214,8 +1214,8 @@ namespace AutomationTool
 				HashSet<BgNodeOutput> inputs = ResolveInputReferences(element, requiresNames);
 
 				// Gather up all the input dependencies, and check they're all upstream of the current node
-				HashSet<BgNode> inputDependencies = new HashSet<BgNode>();
-				foreach (BgNode inputDependency in inputs.Select(x => x.ProducingNode).Distinct())
+				HashSet<BgNodeDef> inputDependencies = new HashSet<BgNodeDef>();
+				foreach (BgNodeDef inputDependency in inputs.Select(x => x.ProducingNode).Distinct())
 				{
 					inputDependencies.Add(inputDependency);
 				}
@@ -1224,7 +1224,7 @@ namespace AutomationTool
 				HashSet<FileReference> requiredTokens = new HashSet<FileReference>(tokenFileNames.Select(x => new FileReference(x)));
 
 				// Recursively include all their dependencies too
-				foreach (BgNode inputDependency in inputDependencies.ToArray())
+				foreach (BgNodeDef inputDependency in inputDependencies.ToArray())
 				{
 					requiredTokens.UnionWith(inputDependency.RequiredTokens);
 					inputDependencies.UnionWith(inputDependency.InputDependencies);
@@ -1250,11 +1250,11 @@ namespace AutomationTool
 				}
 
 				// Gather up all the order dependencies
-				HashSet<BgNode> orderDependencies = new HashSet<BgNode>(inputDependencies);
+				HashSet<BgNodeDef> orderDependencies = new HashSet<BgNodeDef>(inputDependencies);
 				orderDependencies.UnionWith(ResolveReferences(element, afterNames));
 
 				// Recursively include all their order dependencies too
-				foreach (BgNode orderDependency in orderDependencies.ToArray())
+				foreach (BgNodeDef orderDependency in orderDependencies.ToArray())
 				{
 					orderDependencies.UnionWith(orderDependency.OrderDependencies);
 				}
@@ -1263,7 +1263,7 @@ namespace AutomationTool
 				int agentIdx = _graph.Agents.IndexOf(_enclosingAgent!);
 				for (int idx = agentIdx + 1; idx < _graph.Agents.Count; idx++)
 				{
-					foreach (BgNode node in _graph.Agents[idx].Nodes.Where(x => orderDependencies.Contains(x)))
+					foreach (BgNodeDef node in _graph.Agents[idx].Nodes.Where(x => orderDependencies.Contains(x)))
 					{
 						LogError(element, "Node '{0}' has a dependency on '{1}', which was declared after the initial definition of '{2}'.", name, node.Name, _enclosingAgent!.Name);
 					}
@@ -1537,11 +1537,11 @@ namespace AutomationTool
 				bool bAbsolute = element.HasAttribute("Absolute") && ReadBooleanAttribute(element, "Absolute", true);
 
 				// Find the list of targets which are included, and recurse through all their dependencies
-				HashSet<BgNode> nodes = new HashSet<BgNode>();
+				HashSet<BgNodeDef> nodes = new HashSet<BgNodeDef>();
 				if (targetNames != null)
 				{
-					HashSet<BgNode> targetNodes = ResolveReferences(element, targetNames);
-					foreach (BgNode node in targetNodes)
+					HashSet<BgNodeDef> targetNodes = ResolveReferences(element, targetNames);
+					foreach (BgNodeDef node in targetNodes)
 					{
 						nodes.Add(node);
 						nodes.UnionWith(node.InputDependencies);
@@ -1551,19 +1551,19 @@ namespace AutomationTool
 				// Add all the individually referenced nodes
 				if (individualNodeNames != null)
 				{
-					HashSet<BgNode> individualNodes = ResolveReferences(element, individualNodeNames);
+					HashSet<BgNodeDef> individualNodes = ResolveReferences(element, individualNodeNames);
 					nodes.UnionWith(individualNodes);
 				}
 
 				// Exclude all the exceptions
 				if (exceptNames != null)
 				{
-					HashSet<BgNode> exceptNodes = ResolveReferences(element, exceptNames);
+					HashSet<BgNodeDef> exceptNodes = ResolveReferences(element, exceptNames);
 					nodes.ExceptWith(exceptNodes);
 				}
 
 				// Update all the referenced nodes with the settings
-				foreach (BgNode node in nodes)
+				foreach (BgNodeDef node in nodes)
 				{
 					if (users != null)
 					{
@@ -1626,11 +1626,11 @@ namespace AutomationTool
 				Dictionary<string, string> annotations = ReadAnnotationsAttribute(element, "Values");
 
 				// Find the list of targets which are included, and recurse through all their dependencies
-				HashSet<BgNode> nodes = new HashSet<BgNode>();
+				HashSet<BgNodeDef> nodes = new HashSet<BgNodeDef>();
 				if (targetNames != null)
 				{
-					HashSet<BgNode> targetNodes = ResolveReferences(element, targetNames);
-					foreach (BgNode node in targetNodes)
+					HashSet<BgNodeDef> targetNodes = ResolveReferences(element, targetNames);
+					foreach (BgNodeDef node in targetNodes)
 					{
 						nodes.Add(node);
 						nodes.UnionWith(node.InputDependencies);
@@ -1640,19 +1640,19 @@ namespace AutomationTool
 				// Add all the individually referenced nodes
 				if (individualNodeNames != null)
 				{
-					HashSet<BgNode> individualNodes = ResolveReferences(element, individualNodeNames);
+					HashSet<BgNodeDef> individualNodes = ResolveReferences(element, individualNodeNames);
 					nodes.UnionWith(individualNodes);
 				}
 
 				// Exclude all the exceptions
 				if (exceptNames != null)
 				{
-					HashSet<BgNode> exceptNodes = ResolveReferences(element, exceptNames);
+					HashSet<BgNodeDef> exceptNodes = ResolveReferences(element, exceptNames);
 					nodes.ExceptWith(exceptNodes);
 				}
 
 				// Update all the referenced nodes with the settings
-				foreach (BgNode node in nodes)
+				foreach (BgNodeDef node in nodes)
 				{
 					foreach ((string key, string value) in annotations)
 					{
@@ -1673,7 +1673,7 @@ namespace AutomationTool
 			{
 				string message = ReadAttribute(element, "Message");
 
-				BgDiagnostic diagnostic = new BgDiagnostic(element.Location.File.FullName, element.Location.LineNumber, level, message);
+				BgDiagnosticDef diagnostic = new BgDiagnosticDef(element.Location.File.FullName, element.Location.LineNumber, level, message);
 				if (_enclosingNode != null)
 				{
 					_enclosingNode.Diagnostics.Add(diagnostic);
@@ -1741,12 +1741,12 @@ namespace AutomationTool
 		/// <param name="element">Element used to locate any errors</param>
 		/// <param name="referenceNames">Sequence of names to look up</param>
 		/// <returns>Hashset of all the nodes included by the given names</returns>
-		HashSet<BgNode> ResolveReferences(BgScriptElement element, IEnumerable<string> referenceNames)
+		HashSet<BgNodeDef> ResolveReferences(BgScriptElement element, IEnumerable<string> referenceNames)
 		{
-			HashSet<BgNode> nodes = new HashSet<BgNode>();
+			HashSet<BgNodeDef> nodes = new HashSet<BgNodeDef>();
 			foreach (string referenceName in referenceNames)
 			{
-				BgNode[]? otherNodes;
+				BgNodeDef[]? otherNodes;
 				if (_graph.TryResolveReference(referenceName, out otherNodes))
 				{
 					nodes.UnionWith(otherNodes);
diff --git a/Engine/Source/Programs/AutomationTool/BuildGraph/BgScriptWriter.cs b/Engine/Source/Programs/AutomationTool/BuildGraph/BgScriptWriter.cs
--- a/Engine/Source/Programs/AutomationTool/BuildGraph/BgScriptWriter.cs
+++ b/Engine/Source/Programs/AutomationTool/BuildGraph/BgScriptWriter.cs
@@ -27,7 +27,7 @@ namespace AutomationTool
 		/// <param name="graph">Graph to output</param>
 		/// <param name="file">The file to load</param>
 		/// <param name="schemaFile">Schema file for validation</param>
-		public static void Write(this BgGraph graph, FileReference file, FileReference schemaFile)
+		public static void Write(this BgGraphDef graph, FileReference file, FileReference schemaFile)
 		{
 			XmlWriterSettings settings = new XmlWriterSettings();
 			settings.Indent = true;
@@ -42,12 +42,12 @@ namespace AutomationTool
 					writer.WriteAttributeString("schemaLocation", "http://www.w3.org/2001/XMLSchema-instance", "http://www.epicgames.com/BuildGraph " + schemaFile.MakeRelativeTo(file.Directory));
 				}
 
-				foreach (BgAgent agent in graph.Agents)
+				foreach (BgAgentDef agent in graph.Agents)
 				{
 					agent.Write(writer);
 				}
 
-				foreach (BgAggregate aggregate in graph.NameToAggregate.Values)
+				foreach (BgAggregateDef aggregate in graph.NameToAggregate.Values)
 				{
 					// If the aggregate has no required elements, skip it.
 					if (aggregate.RequiredNodes.Count == 0)
@@ -61,7 +61,7 @@ namespace AutomationTool
 					writer.WriteEndElement();
 				}
 
-				foreach (BgLabel label in graph.Labels)
+				foreach (BgLabelDef label in graph.Labels)
 				{
 					writer.WriteStartElement("Label");
 					if (label.DashboardCategory != null)
@@ -71,7 +71,7 @@ namespace AutomationTool
 					writer.WriteAttributeString("Name", label.DashboardName);
 					writer.WriteAttributeString("Requires", String.Join(";", label.RequiredNodes.Select(x => x.Name)));
 
-					HashSet<BgNode> includedNodes = new HashSet<BgNode>(label.IncludedNodes);
+					HashSet<BgNodeDef> includedNodes = new HashSet<BgNodeDef>(label.IncludedNodes);
 					includedNodes.ExceptWith(label.IncludedNodes.SelectMany(x => x.InputDependencies));
 					includedNodes.ExceptWith(label.RequiredNodes);
 					if (includedNodes.Count > 0)
@@ -79,7 +79,7 @@ namespace AutomationTool
 						writer.WriteAttributeString("Include", String.Join(";", includedNodes.Select(x => x.Name)));
 					}
 
-					HashSet<BgNode> excludedNodes = new HashSet<BgNode>(label.IncludedNodes);
+					HashSet<BgNodeDef> excludedNodes = new HashSet<BgNodeDef>(label.IncludedNodes);
 					excludedNodes.UnionWith(label.IncludedNodes.SelectMany(x => x.InputDependencies));
 					excludedNodes.ExceptWith(label.IncludedNodes);
 					excludedNodes.ExceptWith(excludedNodes.ToArray().SelectMany(x => x.InputDependencies));
@@ -98,7 +98,7 @@ namespace AutomationTool
 					writer.WriteEndElement();
 				}
 
-				foreach (BgBadge badge in graph.Badges)
+				foreach (BgBadgeDef badge in graph.Badges)
 				{
 					writer.WriteStartElement("Badge");
 					writer.WriteAttributeString("Name", badge.Name);
@@ -123,12 +123,12 @@ namespace AutomationTool
 		/// </summary>
 		/// <param name="agent">Agent to output</param>
 		/// <param name="writer">The XML writer to output to</param>
-		public static void Write(this BgAgent agent, XmlWriter writer)
+		public static void Write(this BgAgentDef agent, XmlWriter writer)
 		{
 			writer.WriteStartElement("Agent");
 			writer.WriteAttributeString("Name", agent.Name);
 			writer.WriteAttributeString("Type", String.Join(";", agent.PossibleTypes));
-			foreach (BgNode node in agent.Nodes)
+			foreach (BgNodeDef node in agent.Nodes)
 			{
 				node.Write(writer);
 			}
@@ -140,7 +140,7 @@ namespace AutomationTool
 		/// </summary>
 		/// <param name="node">Node to output</param>
 		/// <param name="writer">The writer to output the node to</param>
-		public static void Write(this BgNode node, XmlWriter writer)
+		public static void Write(this BgNodeDef node, XmlWriter writer)
 		{
 			writer.WriteStartElement("Node");
 			writer.WriteAttributeString("Name", node.Name);
diff --git a/Engine/Source/Programs/AutomationTool/BuildGraph/BuildGraph.cs b/Engine/Source/Programs/AutomationTool/BuildGraph/BuildGraph.cs
--- a/Engine/Source/Programs/AutomationTool/BuildGraph/BuildGraph.cs
+++ b/Engine/Source/Programs/AutomationTool/BuildGraph/BuildGraph.cs
@@ -360,7 +360,7 @@ namespace AutomationTool
 			}
 
 			// Create the graph
-			BgGraph Graph;
+			BgGraphDef Graph;
 			if (ClassName != null)
 			{
 				// Find all the graph builders
@@ -381,13 +381,13 @@ namespace AutomationTool
 				}
 
 				BgGraphBuilder Builder = (BgGraphBuilder)Activator.CreateInstance(BuilderType);
-				BgGraphSpec GraphSpec = Builder.CreateGraph(new BgEnvironment(P4Enabled ? P4Env : null));
+				BgGraph GraphSpec = Builder.CreateGraph(new BgEnvironment(P4Enabled ? P4Env : null));
 
 				(byte[] Data, BgMethod[] Methods) = BgCompiler.Compile(GraphSpec);
 
 				BgInterpreter Interpreter = new BgInterpreter(Data, Methods, Arguments);
 				Interpreter.Disassemble(Logger);
-				Graph = (BgGraph)Interpreter.Evaluate();
+				Graph = (BgGraphDef)Interpreter.Evaluate();
 			}
 			else
 			{
@@ -457,7 +457,7 @@ namespace AutomationTool
 			}
 
 			// Convert the supplied target references into nodes 
-			HashSet<BgNode> TargetNodes = new HashSet<BgNode>();
+			HashSet<BgNodeDef> TargetNodes = new HashSet<BgNodeDef>();
 			if(TargetNames.Length == 0)
 			{
 				if (!bListOnly && SingleNodeName == null)
@@ -484,7 +484,7 @@ namespace AutomationTool
 
 				foreach (string TargetName in NodesToResolve)
 				{
-					BgNode[] Nodes;
+					BgNodeDef[] Nodes;
 					if(!Graph.TryResolveReference(TargetName, out Nodes))
 					{
 						LogError("Target '{0}' is not in graph", TargetName);
@@ -504,7 +504,7 @@ namespace AutomationTool
 				if (SingleNodeName == null)
 				{
 					CommandUtils.LogInformation("Required tokens:");
-					foreach(BgNode Node in TargetNodes)
+					foreach(BgNodeDef Node in TargetNodes)
 					{
 						foreach(FileReference RequiredToken in Node.RequiredTokens)
 						{
@@ -537,13 +537,13 @@ namespace AutomationTool
 					if(bSkipTargetsWithoutTokens)
 					{
 						// Find all the nodes we're going to skip
-						HashSet<BgNode> SkipNodes = new HashSet<BgNode>();
+						HashSet<BgNodeDef> SkipNodes = new HashSet<BgNodeDef>();
 						foreach(IGrouping<string, FileReference> MissingTokensForBuild in MissingTokens.GroupBy(x => x.Value, x => x.Key))
 						{
 							LogInformation("Skipping the following nodes due to {0}:", MissingTokensForBuild.Key);
 							foreach(FileReference MissingToken in MissingTokensForBuild)
 							{
-								foreach(BgNode SkipNode in TargetNodes.Where(x => x.RequiredTokens.Contains(MissingToken) && SkipNodes.Add(x)))
+								foreach(BgNodeDef SkipNode in TargetNodes.Where(x => x.RequiredTokens.Contains(MissingToken) && SkipNodes.Add(x)))
 								{
 									LogInformation("    {0}", SkipNode);
 								}
@@ -555,7 +555,7 @@ namespace AutomationTool
 						{
 							TargetNodes.ExceptWith(SkipNodes);
 							LogInformation("Remaining target nodes:");
-							foreach(BgNode TargetNode in TargetNodes)
+							foreach(BgNodeDef TargetNode in TargetNodes)
 							{
 								LogInformation("    {0}", TargetNode);
 							}
@@ -569,7 +569,7 @@ namespace AutomationTool
 					{
 						foreach(KeyValuePair<FileReference, string> Pair in MissingTokens)
 						{
-							List<BgNode> SkipNodes = TargetNodes.Where(x => x.RequiredTokens.Contains(Pair.Key)).ToList();
+							List<BgNodeDef> SkipNodes = TargetNodes.Where(x => x.RequiredTokens.Contains(Pair.Key)).ToList();
 							LogError("Cannot run {0} due to previous build: {1}", String.Join(", ", SkipNodes), Pair.Value);
 						}
 						foreach(FileReference CreatedToken in CreatedTokens)
@@ -602,7 +602,7 @@ namespace AutomationTool
 			}
 
 			// If we're just building a single node, find it 
-			BgNode SingleNode = null;
+			BgNodeDef SingleNode = null;
 			if(SingleNodeName != null && !Graph.NameToNode.TryGetValue(SingleNodeName, out SingleNode))
 			{
 				LogError("Node '{0}' is not in the trimmed graph", SingleNodeName);
@@ -612,15 +612,15 @@ namespace AutomationTool
 			// If we just want to show the contents of the graph, do so and exit.
 			if(bListOnly)
 			{ 
-				HashSet<BgNode> CompletedNodes = FindCompletedNodes(Graph, Storage);
+				HashSet<BgNodeDef> CompletedNodes = FindCompletedNodes(Graph, Storage);
 				Graph.Print(CompletedNodes, PrintOptions, Log.Logger);
 			}
 
 			// Print out all the diagnostic messages which still apply, unless we're running a step as part of a build system or just listing the contents of the file. 
 			if(SingleNode == null && (!bListOnly || bShowDiagnostics))
 			{
-				List<BgDiagnostic> Diagnostics = Graph.GetAllDiagnostics();
-				foreach (BgDiagnostic Diagnostic in Diagnostics)
+				List<BgDiagnosticDef> Diagnostics = Graph.GetAllDiagnostics();
+				foreach (BgDiagnosticDef Diagnostic in Diagnostics)
 				{
 					if(Diagnostic.Level == LogLevel.Information)
 					{
@@ -644,14 +644,14 @@ namespace AutomationTool
 			// Export the graph to a file
 			if(ExportFileName != null)
 			{
-				HashSet<BgNode> CompletedNodes = FindCompletedNodes(Graph, Storage);
+				HashSet<BgNodeDef> CompletedNodes = FindCompletedNodes(Graph, Storage);
 				Graph.Print(CompletedNodes, PrintOptions, Log.Logger);
 				Graph.Export(new FileReference(ExportFileName), CompletedNodes);
 				return ExitCode.Success;
 			}
 
 			// Create tasks for the entire graph
-			Dictionary<BgNode, BgNodeExecutor> NodeToExecutor = new Dictionary<BgNode, BgNodeExecutor>();
+			Dictionary<BgNodeDef, BgNodeExecutor> NodeToExecutor = new Dictionary<BgNodeDef, BgNodeExecutor>();
 			if (bSkipValidation && SingleNode != null)
 			{
 				if (!BindNodes(SingleNode, NameToTask, Graph.TagNameToNodeOutput, NodeToExecutor))
@@ -688,12 +688,12 @@ namespace AutomationTool
 			return ExitCode.Success;
 		}
 
-		bool BindNodes(BgGraph Graph, Dictionary<string, ScriptTaskBinding> NameToTask, Dictionary<BgNode, BgNodeExecutor> NodeToExecutor)
+		bool BindNodes(BgGraphDef Graph, Dictionary<string, ScriptTaskBinding> NameToTask, Dictionary<BgNodeDef, BgNodeExecutor> NodeToExecutor)
 		{
 			bool bResult = true;
-			foreach (BgAgent Agent in Graph.Agents)
+			foreach (BgAgentDef Agent in Graph.Agents)
 			{
-				foreach (BgNode Node in Agent.Nodes)
+				foreach (BgNodeDef Node in Agent.Nodes)
 				{
 					bResult &= BindNodes(Node, NameToTask, Graph.TagNameToNodeOutput, NodeToExecutor);
 				}
@@ -701,7 +701,7 @@ namespace AutomationTool
 			return bResult;
 		}
 
-		bool BindNodes(BgNode Node, Dictionary<string, ScriptTaskBinding> NameToTask, Dictionary<string, BgNodeOutput> TagNameToNodeOutput, Dictionary<BgNode, BgNodeExecutor> NodeToExecutor)
+		bool BindNodes(BgNodeDef Node, Dictionary<string, ScriptTaskBinding> NameToTask, Dictionary<string, BgNodeOutput> TagNameToNodeOutput, Dictionary<BgNodeDef, BgNodeExecutor> NodeToExecutor)
 		{
 			if (Node is BgScriptNode ScriptNode)
 			{
@@ -887,10 +887,10 @@ namespace AutomationTool
 		/// </summary>
 		/// <param name="Graph">The graph instance</param>
 		/// <param name="Storage">The temp storage backend which stores the shared state</param>
-		HashSet<BgNode> FindCompletedNodes(BgGraph Graph, TempStorage Storage)
+		HashSet<BgNodeDef> FindCompletedNodes(BgGraphDef Graph, TempStorage Storage)
 		{
-			HashSet<BgNode> CompletedNodes = new HashSet<BgNode>();
-			foreach(BgNode Node in Graph.Agents.SelectMany(x => x.Nodes))
+			HashSet<BgNodeDef> CompletedNodes = new HashSet<BgNodeDef>();
+			foreach(BgNodeDef Node in Graph.Agents.SelectMany(x => x.Nodes))
 			{
 				if(Storage.IsComplete(Node.Name))
 				{
@@ -908,15 +908,15 @@ namespace AutomationTool
 		/// <param name="NodeToExecutor">Map from node to executor</param>
 		/// <param name="Storage">The temp storage backend which stores the shared state</param>
 		/// <returns>True if everything built successfully</returns>
-		async Task<bool> BuildAllNodesAsync(JobContext Job, BgGraph Graph, Dictionary<BgNode, BgNodeExecutor> NodeToExecutor, TempStorage Storage)
+		async Task<bool> BuildAllNodesAsync(JobContext Job, BgGraphDef Graph, Dictionary<BgNodeDef, BgNodeExecutor> NodeToExecutor, TempStorage Storage)
 		{
 			// Build a flat list of nodes to execute, in order
-			BgNode[] NodesToExecute = Graph.Agents.SelectMany(x => x.Nodes).ToArray();
+			BgNodeDef[] NodesToExecute = Graph.Agents.SelectMany(x => x.Nodes).ToArray();
 
 			// Check the integrity of any local nodes that have been completed. It's common to run formal builds locally between regular development builds, so we may have 
 			// stale local state. Rather than failing later, detect and clean them up now.
-			HashSet<BgNode> CleanedNodes = new HashSet<BgNode>();
-			foreach(BgNode NodeToExecute in NodesToExecute)
+			HashSet<BgNodeDef> CleanedNodes = new HashSet<BgNodeDef>();
+			foreach(BgNodeDef NodeToExecute in NodesToExecute)
 			{
 				if(NodeToExecute.InputDependencies.Any(x => CleanedNodes.Contains(x)) || !Storage.CheckLocalIntegrity(NodeToExecute.Name, NodeToExecute.Outputs.Select(x => x.TagName)))
 				{
@@ -927,7 +927,7 @@ namespace AutomationTool
 
 			// Execute them in order
 			int NodeIdx = 0;
-			foreach(BgNode NodeToExecute in NodesToExecute)
+			foreach(BgNodeDef NodeToExecute in NodesToExecute)
 			{
 				LogInformation("****** [{0}/{1}] {2}", ++NodeIdx, NodesToExecute.Length, NodeToExecute.Name);
 				if(!Storage.IsComplete(NodeToExecute.Name))
@@ -953,7 +953,7 @@ namespace AutomationTool
 		/// <param name="Storage">The temp storage backend which stores the shared state</param>
 		/// <param name="bWithBanner">Whether to write a banner before and after this node's log output</param>
 		/// <returns>True if the node built successfully, false otherwise.</returns>
-		async Task<bool> BuildNodeAsync(JobContext Job, BgGraph Graph, BgNode Node, Dictionary<BgNode, BgNodeExecutor> NodeToExecutor, TempStorage Storage, bool bWithBanner)
+		async Task<bool> BuildNodeAsync(JobContext Job, BgGraphDef Graph, BgNodeDef Node, Dictionary<BgNodeDef, BgNodeExecutor> NodeToExecutor, TempStorage Storage, bool bWithBanner)
 		{
 			DirectoryReference RootDir = new DirectoryReference(CommandUtils.CmdEnv.LocalRoot);
 
@@ -1047,10 +1047,10 @@ namespace AutomationTool
 
 			// Determine all the output files which are required to be copied to temp storage (because they're referenced by nodes in another agent)
 			HashSet<FileReference> ReferencedOutputFiles = new HashSet<FileReference>();
-			foreach(BgAgent Agent in Graph.Agents)
+			foreach(BgAgentDef Agent in Graph.Agents)
 			{
 				bool bSameAgent = Agent.Nodes.Contains(Node);
-				foreach(BgNode OtherNode in Agent.Nodes)
+				foreach(BgNodeDef OtherNode in Agent.Nodes)
 				{
 					if(!bSameAgent)
 					{
diff --git a/Engine/Source/Programs/AutomationTool/BuildGraph/Scripts/InstalledBuild.cs b/Engine/Source/Programs/AutomationTool/BuildGraph/Scripts/InstalledBuild.cs
--- a/Engine/Source/Programs/AutomationTool/BuildGraph/Scripts/InstalledBuild.cs
+++ b/Engine/Source/Programs/AutomationTool/BuildGraph/Scripts/InstalledBuild.cs
@@ -81,7 +81,7 @@ namespace AutomationTool
 			return Projects;
 		}
 
-		public override BgGraphSpec CreateGraph(BgEnvironment Context)
+		public override BgGraph CreateGraph(BgEnvironment Context)
 		{
 			string RootDir = CommandUtils.CmdEnv.LocalRoot;
 
@@ -132,63 +132,63 @@ namespace AutomationTool
 			BgString CrashReporterCompileArgs = "";
 			CrashReporterCompileArgs = CrashReporterCompileArgs.If(CrashReporterAPIURL != "" & CrashReporterAPIKey != "", BgString.Format("-define:CRC_TELEMETRY_URL=\"{0}\" -define:CRC_TELEMETRY_KEY_DEV=\"{1}\" -define:CRC_TELEMETRY_KEY_RELEASE=\"{1}\" -OverrideBuildEnvironment", CrashReporterAPIURL, CrashReporterAPIKey));
 
-			List<BgAggregateSpec> Aggregates = new List<BgAggregateSpec>();
+			List<BgAggregate> Aggregates = new List<BgAggregate>();
 
 			/////// EDITORS ////////////////////////////////////////////////
 
 			//// Windows ////
-			BgAgentSpec EditorWin64 = new BgAgentSpec("Editor Win64", "Win64_Licensee");
+			BgAgent EditorWin64 = new BgAgent("Editor Win64", "Win64_Licensee");
 
-			BgNodeSpec VersionFilesNode = EditorWin64
+			BgNode VersionFilesNode = EditorWin64
 				.AddNode(x => UpdateVersionFilesAsync(x))
 				.Construct();
 
-			BgNodeSpec<BgFileSet> WinUhtNode = EditorWin64
+			BgNode<BgFileSet> WinUhtNode = EditorWin64
 				.AddNode(x => CompileUnrealHeaderToolAsync(x, UnrealTargetPlatform.Win64))
 				.Requires(VersionFilesNode)
 				.Construct();
 
-			BgNodeSpec<BgFileSet> WinEditorNode = EditorWin64
+			BgNode<BgFileSet> WinEditorNode = EditorWin64
 				.AddNode(x => CompileUnrealEditorWin64Async(x, CrashReporterCompileArgs, EmbedSrcSrvInfo, CompileDatasmithPlugins, WithFullDebugInfo, SignExecutables))
 				.Requires(WinUhtNode)
 				.Construct();
 
-			Aggregates.Add(new BgAggregateSpec("Win64 Editor", WinEditorNode, label: "Editors/Win64"));
+			Aggregates.Add(new BgAggregate("Win64 Editor", WinEditorNode, label: "Editors/Win64"));
 
 
 			/////// TARGET PLATFORMS ////////////////////////////////////////////////
 
 			//// Win64 ////
 
-			BgAgentSpec TargetWin64 = new BgAgentSpec("Target Win64", "Win64_Licensee");
+			BgAgent TargetWin64 = new BgAgent("Target Win64", "Win64_Licensee");
 
-			BgNodeSpec<BgFileSet> WinGame = TargetWin64
+			BgNode<BgFileSet> WinGame = TargetWin64
 				.AddNode(x => CompileUnrealGameWin64(x, GameConfigurations, EmbedSrcSrvInfo, WithFullDebugInfo, SignExecutables))
 				.Requires(WinUhtNode)
 				.Construct();
 
-			Aggregates.Add(new BgAggregateSpec("TargetPlatforms_Win64", WinGame));
+			Aggregates.Add(new BgAggregate("TargetPlatforms_Win64", WinGame));
 
 			/////// TOOLS //////////////////////////////////////////////////////////
 
 			//// Build Rules ////
 
-			BgAgentSpec BuildRules = new BgAgentSpec("BuildRules", "Win64_Licensee");
+			BgAgent BuildRules = new BgAgent("BuildRules", "Win64_Licensee");
 
-			BgNodeSpec RulesAssemblies = BuildRules
+			BgNode RulesAssemblies = BuildRules
 				.AddNode(x => CompileRulesAssemblies(x))
 				.Construct();
 
 			//// Win Tools ////
 
-			BgAgentSpec ToolsGroupWin64 = new BgAgentSpec("Tools Group Win64", "Win64_Licensee");
+			BgAgent ToolsGroupWin64 = new BgAgent("Tools Group Win64", "Win64_Licensee");
 
-			BgNodeSpec<BgFileSet> WinTools = ToolsGroupWin64
+			BgNode<BgFileSet> WinTools = ToolsGroupWin64
 				.AddNode(x => BuildToolsWin64Async(x, CrashReporterCompileArgs))
 				.Requires(WinUhtNode)
 				.Construct();
 
-			BgNodeSpec<BgFileSet> CsTools = ToolsGroupWin64
+			BgNode<BgFileSet> CsTools = ToolsGroupWin64
 				.AddNode(x => BuildToolsCSAsync(x, SignExecutables))
 				.Requires(VersionFilesNode)
 				.Construct();
@@ -196,12 +196,12 @@ namespace AutomationTool
 
 			/////// DDC //////////////////////////////////////////////////////////
 
-			BgAgentSpec DDCGroupWin64 = new BgAgentSpec("DDC Group Win64", "Win64_Licensee");
+			BgAgent DDCGroupWin64 = new BgAgent("DDC Group Win64", "Win64_Licensee");
 
 			BgList<BgString> DDCPlatformsWin64 = BgList<BgString>.Create("WindowsEditor");
 			DDCPlatformsWin64 = DDCPlatformsWin64.If(WithWin64, x => x.Add("Windows"));
 
-			BgNodeSpec DdcNode = DDCGroupWin64
+			BgNode DdcNode = DDCGroupWin64
 				.AddNode(x => BuildDDCWin64Async(x, DDCPlatformsWin64, BgList<BgFileSet>.Create(WinUhtNode.Output, WinEditorNode.Output, WinTools.Output)))
 				.Requires(WinEditorNode, WinTools)
 				.Construct();
@@ -210,7 +210,7 @@ namespace AutomationTool
 			/////// STAGING ///////
 
 			// Windows 
-			BgAgentSpec WinStageAgent = new BgAgentSpec("Installed Build Group Win64", "Win64_Licensee");
+			BgAgent WinStageAgent = new BgAgent("Installed Build Group Win64", "Win64_Licensee");
 
 			BgList<BgFileSet> WinInstalledFiles = BgList<BgFileSet>.Empty;
 			WinInstalledFiles = WinInstalledFiles.Add(WinUhtNode.Output);
@@ -227,14 +227,14 @@ namespace AutomationTool
 			BgString WinOutputDir = "LocalBuilds/Engine/Windows";
 			BgString WinFinalizeArgs = BgString.Format("-OutputDir=\"{0}\" -Platforms={1} -ContentOnlyPlatforms={2}", WinOutputDir, BgString.Join(";", WinPlatforms), BgString.Join(";", WinContentOnlyPlatforms));
 
-			BgNodeSpec WinInstalledNode = WinStageAgent
+			BgNode WinInstalledNode = WinStageAgent
 				.AddNode(x => MakeInstalledBuildWin64Async(x, WinInstalledFiles, WinFinalizeArgs, WinOutputDir))
 				.Requires(WinInstalledFiles)
 				.Construct();
 
-			Aggregates.Add(new BgAggregateSpec("HostPlatforms_Win64", WinInstalledNode, label: "Builds/Win64"));
+			Aggregates.Add(new BgAggregate("HostPlatforms_Win64", WinInstalledNode, label: "Builds/Win64"));
 
-			return new BgGraphSpec(BgList<BgNodeSpec>.Empty, Aggregates);
+			return new BgGraph(BgList<BgNode>.Empty, Aggregates);
 		}
 
 		/// <summary>
diff --git a/Engine/Source/Programs/Shared/EpicGames.BuildGraph/BgAgent.cs b/Engine/Source/Programs/Shared/EpicGames.BuildGraph/BgAgent.cs
deleted file mode 100644
diff --git a/Engine/Source/Programs/Shared/EpicGames.BuildGraph/BgAgentDef.cs b/Engine/Source/Programs/Shared/EpicGames.BuildGraph/BgAgentDef.cs
new file mode 100644
--- /dev/null
+++ b/Engine/Source/Programs/Shared/EpicGames.BuildGraph/BgAgentDef.cs
@@ -0,0 +1,48 @@
+// Copyright Epic Games, Inc. All Rights Reserved.
+
+using System;
+using System.Collections.Generic;
+using System.Diagnostics;
+using System.Xml;
+
+namespace EpicGames.BuildGraph
+{
+	/// <summary>
+	/// Stores a list of nodes which can be executed on a single agent
+	/// </summary>
+	[DebuggerDisplay("{Name}")]
+	public class BgAgentDef
+	{
+		/// <summary>
+		/// Name of this agent. Used for display purposes in a build system.
+		/// </summary>
+		public string Name { get; }
+
+		/// <summary>
+		/// Array of valid agent types that these nodes may run on. When running in the build system, this determines the class of machine that should
+		/// be selected to run these nodes. The first defined agent type for this branch will be used.
+		/// </summary>
+		public string[] PossibleTypes { get; }
+
+		/// <summary>
+		/// List of nodes in this agent group.
+		/// </summary>
+		public List<BgNodeDef> Nodes { get; set; } = new List<BgNodeDef>();
+
+		/// <summary>
+		/// Diagnostics to output if executing this agent
+		/// </summary>
+		public List<BgDiagnosticDef> Diagnostics { get; } = new List<BgDiagnosticDef>();
+
+		/// <summary>
+		/// Constructor
+		/// </summary>
+		/// <param name="inName">Name of this agent group</param>
+		/// <param name="inPossibleTypes">Array of valid agent types. See comment for AgentTypes member.</param>
+		public BgAgentDef(string inName, string[] inPossibleTypes)
+		{
+			Name = inName;
+			PossibleTypes = inPossibleTypes;
+		}
+	}
+}
diff --git a/Engine/Source/Programs/Shared/EpicGames.BuildGraph/BgAgentSpec.cs b/Engine/Source/Programs/Shared/EpicGames.BuildGraph/BgAgentSpec.cs
deleted file mode 100644
diff --git a/Engine/Source/Programs/Shared/EpicGames.BuildGraph/BgAggregate.cs b/Engine/Source/Programs/Shared/EpicGames.BuildGraph/BgAggregate.cs
deleted file mode 100644
diff --git a/Engine/Source/Programs/Shared/EpicGames.BuildGraph/BgAggregateDef.cs b/Engine/Source/Programs/Shared/EpicGames.BuildGraph/BgAggregateDef.cs
new file mode 100644
--- /dev/null
+++ b/Engine/Source/Programs/Shared/EpicGames.BuildGraph/BgAggregateDef.cs
@@ -0,0 +1,61 @@
+// Copyright Epic Games, Inc. All Rights Reserved.
+
+using System.Collections.Generic;
+
+namespace EpicGames.BuildGraph
+{
+	/// <summary>
+	/// Defines a agggregate within a graph, which give the combined status of one or more job steps, and allow building several steps at once.
+	/// </summary>
+	public class BgAggregateDef
+	{
+		/// <summary>
+		/// Name of this badge
+		/// </summary>
+		public string Name { get; }
+
+		/// <summary>
+		/// Set of nodes that must be run for this label to be shown.
+		/// </summary>
+		public HashSet<BgNodeDef> RequiredNodes { get; } = new HashSet<BgNodeDef>();
+
+		/// <summary>
+		/// Constructor
+		/// </summary>
+		/// <param name="name">Name of this aggregate</param>
+		public BgAggregateDef(string name)
+		{
+			Name = name;
+		}
+
+		/// <summary>
+		/// Get the name of this label
+		/// </summary>
+		/// <returns>The name of this label</returns>
+		public override string ToString()
+		{
+			return Name;
+		}
+	}
+
+	/// <summary>
+	/// Aggregate that was created from bytecode
+	/// </summary>
+	public class BgBytecodeAggregateDef : BgAggregateDef
+	{
+		/// <summary>
+		/// Labels to add this aggregate to
+		/// </summary>
+		public BgLabelDef? Label { get; }
+
+		/// <summary>
+		/// Constructor
+		/// </summary>
+		public BgBytecodeAggregateDef(string name, IEnumerable<BgNodeDef> nodes, BgLabelDef? label)
+			: base(name)
+		{
+			RequiredNodes.UnionWith(nodes);
+			Label = label;
+		}
+	}
+}
diff --git a/Engine/Source/Programs/Shared/EpicGames.BuildGraph/BgAggregateSpec.cs b/Engine/Source/Programs/Shared/EpicGames.BuildGraph/BgAggregateSpec.cs
deleted file mode 100644
diff --git a/Engine/Source/Programs/Shared/EpicGames.BuildGraph/BgBadge.cs b/Engine/Source/Programs/Shared/EpicGames.BuildGraph/BgBadge.cs
deleted file mode 100644
diff --git a/Engine/Source/Programs/Shared/EpicGames.BuildGraph/BgBadgeDef.cs b/Engine/Source/Programs/Shared/EpicGames.BuildGraph/BgBadgeDef.cs
new file mode 100644
--- /dev/null
+++ b/Engine/Source/Programs/Shared/EpicGames.BuildGraph/BgBadgeDef.cs
@@ -0,0 +1,54 @@
+// Copyright Epic Games, Inc. All Rights Reserved.
+
+using System.Collections.Generic;
+
+namespace EpicGames.BuildGraph
+{
+	/// <summary>
+	/// Defines a badge which gives an at-a-glance summary of part of the build, and can be displayed in UGS
+	/// </summary>
+	public class BgBadgeDef
+	{
+		/// <summary>
+		/// Name of this badge
+		/// </summary>
+		public string Name { get; }
+
+		/// <summary>
+		/// Depot path to the project that this badge applies to. Used for filtering in UGS.
+		/// </summary>
+		public string Project { get; }
+
+		/// <summary>
+		/// The changelist to post the badge for
+		/// </summary>
+		public int Change { get; }
+
+		/// <summary>
+		/// Set of nodes that this badge reports the status of
+		/// </summary>
+		public HashSet<BgNodeDef> Nodes { get; } = new HashSet<BgNodeDef>();
+
+		/// <summary>
+		/// Constructor
+		/// </summary>
+		/// <param name="inName">Name of this report</param>
+		/// <param name="inProject">Depot path to the project that this badge applies to</param>
+		/// <param name="inChange">The changelist to post the badge for</param>
+		public BgBadgeDef(string inName, string inProject, int inChange)
+		{
+			Name = inName;
+			Project = inProject;
+			Change = inChange;
+		}
+
+		/// <summary>
+		/// Get the name of this badge
+		/// </summary>
+		/// <returns>The name of this badge</returns>
+		public override string ToString()
+		{
+			return Name;
+		}
+	}
+}
diff --git a/Engine/Source/Programs/Shared/EpicGames.BuildGraph/BgDiagnostic.cs b/Engine/Source/Programs/Shared/EpicGames.BuildGraph/BgDiagnostic.cs
deleted file mode 100644
diff --git a/Engine/Source/Programs/Shared/EpicGames.BuildGraph/BgDiagnosticDef.cs b/Engine/Source/Programs/Shared/EpicGames.BuildGraph/BgDiagnosticDef.cs
new file mode 100644
--- /dev/null
+++ b/Engine/Source/Programs/Shared/EpicGames.BuildGraph/BgDiagnosticDef.cs
@@ -0,0 +1,45 @@
+// Copyright Epic Games, Inc. All Rights Reserved.
+
+using EpicGames.Core;
+using Microsoft.Extensions.Logging;
+
+namespace EpicGames.BuildGraph
+{
+	/// <summary>
+	/// Diagnostic message from the graph script. These messages are parsed at startup, then culled along with the rest of the graph nodes before output. Doing so
+	/// allows errors and warnings which are only output if a node is part of the graph being executed.
+	/// </summary>
+	public class BgDiagnosticDef
+	{
+		/// <summary>
+		/// File containing the diagnostic
+		/// </summary>
+		public string File { get; }
+
+		/// <summary>
+		/// Line number containing the diagnostic
+		/// </summary>
+		public int Line { get; }
+
+		/// <summary>
+		/// The diagnostic event type
+		/// </summary>
+		public LogLevel Level { get; }
+
+		/// <summary>
+		/// The message to display
+		/// </summary>
+		public string Message { get; }
+
+		/// <summary>
+		/// Constructor
+		/// </summary>
+		public BgDiagnosticDef(string file, int line, LogLevel level, string message)
+		{
+			File = file;
+			Line = line;
+			Level = level;
+			Message = message;
+		}
+	}
+}
diff --git a/Engine/Source/Programs/Shared/EpicGames.BuildGraph/BgGraph.cs b/Engine/Source/Programs/Shared/EpicGames.BuildGraph/BgGraph.cs
deleted file mode 100644
diff --git a/Engine/Source/Programs/Shared/EpicGames.BuildGraph/BgGraphDef.cs b/Engine/Source/Programs/Shared/EpicGames.BuildGraph/BgGraphDef.cs
new file mode 100644
--- /dev/null
+++ b/Engine/Source/Programs/Shared/EpicGames.BuildGraph/BgGraphDef.cs
@@ -0,0 +1,646 @@
+// Copyright Epic Games, Inc. All Rights Reserved.
+
+using System;
+using System.Collections.Generic;
+using System.Diagnostics.CodeAnalysis;
+using System.Linq;
+using System.Text;
+using System.Xml;
+using EpicGames.Core;
+using Microsoft.Extensions.Logging;
+
+namespace EpicGames.BuildGraph
+{
+	/// <summary>
+	/// Options for how the graph should be printed
+	/// </summary>
+	public enum GraphPrintOptions
+	{
+		/// <summary>
+		/// Includes a list of the graph options
+		/// </summary>
+		ShowCommandLineOptions = 0x1,
+
+		/// <summary>
+		/// Includes the list of dependencies for each node
+		/// </summary>
+		ShowDependencies = 0x2,
+
+		/// <summary>
+		/// Includes the list of notifiers for each node
+		/// </summary>
+		ShowNotifications = 0x4,
+	}
+
+	/// <summary>
+	/// Definition of a graph.
+	/// </summary>
+	public class BgGraphDef
+	{
+		/// <summary>
+		/// List of options, in the order they were specified
+		/// </summary>
+		public List<BgOption> Options { get; } = new List<BgOption>();
+
+		/// <summary>
+		/// List of agents containing nodes to execute
+		/// </summary>
+		public List<BgAgentDef> Agents { get; } = new List<BgAgentDef>();
+
+		/// <summary>
+		/// Mapping from name to agent
+		/// </summary>
+		public Dictionary<string, BgAgentDef> NameToAgent { get; set; } = new Dictionary<string, BgAgentDef>(StringComparer.OrdinalIgnoreCase);
+
+		/// <summary>
+		/// Mapping of names to the corresponding node.
+		/// </summary>
+		public Dictionary<string, BgNodeDef> NameToNode { get; set; } = new Dictionary<string, BgNodeDef>(StringComparer.OrdinalIgnoreCase);
+
+		/// <summary>
+		/// Mapping of names to the corresponding report.
+		/// </summary>
+		public Dictionary<string, BgReport> NameToReport { get; private set; } = new Dictionary<string, BgReport>(StringComparer.OrdinalIgnoreCase);
+
+		/// <summary>
+		/// Mapping of names to their corresponding node output.
+		/// </summary>
+		public Dictionary<string, BgNodeOutput> TagNameToNodeOutput { get; private set; } = new Dictionary<string, BgNodeOutput>(StringComparer.OrdinalIgnoreCase);
+
+		/// <summary>
+		/// Mapping of aggregate names to their respective nodes
+		/// </summary>
+		public Dictionary<string, BgAggregateDef> NameToAggregate { get; private set; } = new Dictionary<string, BgAggregateDef>(StringComparer.OrdinalIgnoreCase);
+
+		/// <summary>
+		/// List of badges that can be displayed for this build
+		/// </summary>
+		public List<BgBadgeDef> Badges { get; } = new List<BgBadgeDef>();
+
+		/// <summary>
+		/// List of labels that can be displayed for this build
+		/// </summary>
+		public List<BgLabelDef> Labels { get; } = new List<BgLabelDef>();
+
+		/// <summary>
+		/// Diagnostics at graph scope
+		/// </summary>
+		public List<BgDiagnosticDef> Diagnostics { get; } = new List<BgDiagnosticDef>();
+
+		/// <summary>
+		/// Default constructor
+		/// </summary>
+		public BgGraphDef()
+		{
+		}
+
+		/// <summary>
+		/// Checks whether a given name already exists
+		/// </summary>
+		/// <param name="name">The name to check.</param>
+		/// <returns>True if the name exists, false otherwise.</returns>
+		public bool ContainsName(string name)
+		{
+			return NameToNode.ContainsKey(name) || NameToReport.ContainsKey(name) || NameToAggregate.ContainsKey(name);
+		}
+
+		/// <summary>
+		/// Gets diagnostics from all graph structures
+		/// </summary>
+		/// <returns>List of diagnostics</returns>
+		public List<BgDiagnosticDef> GetAllDiagnostics()
+		{
+			List<BgDiagnosticDef> diagnostics = new List<BgDiagnosticDef>(Diagnostics);
+			foreach (BgAgentDef agent in Agents)
+			{
+				diagnostics.AddRange(agent.Diagnostics);
+				foreach (BgNodeDef node in agent.Nodes)
+				{
+					diagnostics.AddRange(node.Diagnostics);
+				}
+			}
+			return diagnostics;
+		}
+
+		/// <summary>
+		/// Tries to resolve the given name to one or more nodes. Checks for aggregates, and actual nodes.
+		/// </summary>
+		/// <param name="name">The name to search for</param>
+		/// <param name="outNodes">If the name is a match, receives an array of nodes and their output names</param>
+		/// <returns>True if the name was found, false otherwise.</returns>
+		public bool TryResolveReference(string name, [NotNullWhen(true)] out BgNodeDef[]? outNodes)
+		{
+			// Check if it's a tag reference or node reference
+			if (name.StartsWith("#"))
+			{
+				// Check if it's a regular node or output name
+				BgNodeOutput? output;
+				if (TagNameToNodeOutput.TryGetValue(name, out output))
+				{
+					outNodes = new BgNodeDef[] { output.ProducingNode };
+					return true;
+				}
+			}
+			else
+			{
+				// Check if it's a regular node or output name
+				BgNodeDef? node;
+				if (NameToNode.TryGetValue(name, out node))
+				{
+					outNodes = new BgNodeDef[] { node };
+					return true;
+				}
+
+				// Check if it's an aggregate name
+				BgAggregateDef? aggregate;
+				if (NameToAggregate.TryGetValue(name, out aggregate))
+				{
+					outNodes = aggregate.RequiredNodes.ToArray();
+					return true;
+				}
+
+				// Check if it's a group name
+				BgAgentDef? agent;
+				if (NameToAgent.TryGetValue(name, out agent))
+				{
+					outNodes = agent.Nodes.ToArray();
+					return true;
+				}
+			}
+
+			// Otherwise fail
+			outNodes = null;
+			return false;
+		}
+
+		/// <summary>
+		/// Tries to resolve the given name to one or more node outputs. Checks for aggregates, and actual nodes.
+		/// </summary>
+		/// <param name="name">The name to search for</param>
+		/// <param name="outOutputs">If the name is a match, receives an array of nodes and their output names</param>
+		/// <returns>True if the name was found, false otherwise.</returns>
+		public bool TryResolveInputReference(string name, [NotNullWhen(true)] out BgNodeOutput[]? outOutputs)
+		{
+			// Check if it's a tag reference or node reference
+			if (name.StartsWith("#"))
+			{
+				// Check if it's a regular node or output name
+				BgNodeOutput? output;
+				if (TagNameToNodeOutput.TryGetValue(name, out output))
+				{
+					outOutputs = new BgNodeOutput[] { output };
+					return true;
+				}
+			}
+			else
+			{
+				// Check if it's a regular node or output name
+				BgNodeDef? node;
+				if (NameToNode.TryGetValue(name, out node))
+				{
+					outOutputs = node.Outputs.Union(node.Inputs).ToArray();
+					return true;
+				}
+
+				// Check if it's an aggregate name
+				BgAggregateDef? aggregate;
+				if (NameToAggregate.TryGetValue(name, out aggregate))
+				{
+					outOutputs = aggregate.RequiredNodes.SelectMany(x => x.Outputs.Union(x.Inputs)).Distinct().ToArray();
+					return true;
+				}
+			}
+
+			// Otherwise fail
+			outOutputs = null;
+			return false;
+		}
+
+		static void AddDependencies(BgNodeDef node, HashSet<BgNodeDef> retainNodes)
+		{
+			if (retainNodes.Add(node))
+			{
+				foreach (BgNodeDef inputDependency in node.InputDependencies)
+				{
+					AddDependencies(inputDependency, retainNodes);
+				}
+			}
+		}
+
+		/// <summary>
+		/// Cull the graph to only include the given nodes and their dependencies
+		/// </summary>
+		/// <param name="targetNodes">A set of target nodes to build</param>
+		public void Select(IEnumerable<BgNodeDef> targetNodes)
+		{
+			// Find this node and all its dependencies
+			HashSet<BgNodeDef> retainNodes = new HashSet<BgNodeDef>();
+			foreach (BgNodeDef targetNode in targetNodes)
+			{
+				AddDependencies(targetNode, retainNodes);
+			}
+
+			// Remove all the nodes which are not marked to be kept
+			foreach (BgAgentDef agent in Agents)
+			{
+				agent.Nodes = agent.Nodes.Where(x => retainNodes.Contains(x)).ToList();
+			}
+
+			// Remove all the empty agents
+			Agents.RemoveAll(x => x.Nodes.Count == 0);
+
+			// Trim down the list of nodes for each report to the ones that are being built
+			foreach (BgReport report in NameToReport.Values)
+			{
+				report.Nodes.RemoveWhere(x => !retainNodes.Contains(x));
+			}
+
+			// Remove all the empty reports
+			NameToReport = NameToReport.Where(x => x.Value.Nodes.Count > 0).ToDictionary(pair => pair.Key, pair => pair.Value, StringComparer.InvariantCultureIgnoreCase);
+
+			// Remove all the order dependencies which are no longer part of the graph. Since we don't need to build them, we don't need to wait for them
+			foreach (BgNodeDef node in retainNodes)
+			{
+				node.OrderDependencies = node.OrderDependencies.Where(x => retainNodes.Contains(x)).ToArray();
+			}
+
+			// Create a new list of aggregates for everything that's left
+			Dictionary<string, BgAggregateDef> newNameToAggregate = new Dictionary<string, BgAggregateDef>(NameToAggregate.Comparer);
+			foreach (BgAggregateDef aggregate in NameToAggregate.Values)
+			{
+				if (aggregate.RequiredNodes.All(x => retainNodes.Contains(x)))
+				{
+					newNameToAggregate[aggregate.Name] = aggregate;
+				}
+			}
+			NameToAggregate = newNameToAggregate;
+
+			// Remove any labels that are no longer value
+			foreach (BgLabelDef label in Labels)
+			{
+				label.RequiredNodes.RemoveWhere(x => !retainNodes.Contains(x));
+				label.IncludedNodes.RemoveWhere(x => !retainNodes.Contains(x));
+			}
+			Labels.RemoveAll(x => x.RequiredNodes.Count == 0);
+
+			// Remove any badges which do not have all their dependencies
+			Badges.RemoveAll(x => x.Nodes.Any(y => !retainNodes.Contains(y)));
+		}
+
+		/// <summary>
+		/// Export the build graph to a Json file, for parallel execution by the build system
+		/// </summary>
+		/// <param name="file">Output file to write</param>
+		/// <param name="completedNodes">Set of nodes which have been completed</param>
+		public void Export(FileReference file, HashSet<BgNodeDef> completedNodes)
+		{
+			// Find all the nodes which we're actually going to execute. We'll use this to filter the graph.
+			HashSet<BgNodeDef> nodesToExecute = new HashSet<BgNodeDef>();
+			foreach (BgNodeDef node in Agents.SelectMany(x => x.Nodes))
+			{
+				if (!completedNodes.Contains(node))
+				{
+					nodesToExecute.Add(node);
+				}
+			}
+
+			// Open the output file
+			using (JsonWriter jsonWriter = new JsonWriter(file.FullName))
+			{
+				jsonWriter.WriteObjectStart();
+
+				// Write all the agents
+				jsonWriter.WriteArrayStart("Groups");
+				foreach (BgAgentDef agent in Agents)
+				{
+					BgNodeDef[] nodes = agent.Nodes.Where(x => nodesToExecute.Contains(x)).ToArray();
+					if (nodes.Length > 0)
+					{
+						jsonWriter.WriteObjectStart();
+						jsonWriter.WriteValue("Name", agent.Name);
+						jsonWriter.WriteArrayStart("Agent Types");
+						foreach (string agentType in agent.PossibleTypes)
+						{
+							jsonWriter.WriteValue(agentType);
+						}
+						jsonWriter.WriteArrayEnd();
+						jsonWriter.WriteArrayStart("Nodes");
+						foreach (BgNodeDef node in nodes)
+						{
+							jsonWriter.WriteObjectStart();
+							jsonWriter.WriteValue("Name", node.Name);
+							jsonWriter.WriteValue("DependsOn", String.Join(";", node.GetDirectOrderDependencies().Where(x => nodesToExecute.Contains(x))));
+							jsonWriter.WriteValue("RunEarly", node.RunEarly);
+							jsonWriter.WriteObjectStart("Notify");
+							jsonWriter.WriteValue("Default", String.Join(";", node.NotifyUsers));
+							jsonWriter.WriteValue("Submitters", String.Join(";", node.NotifySubmitters));
+							jsonWriter.WriteValue("Warnings", node.NotifyOnWarnings);
+							jsonWriter.WriteObjectEnd();
+							jsonWriter.WriteObjectEnd();
+						}
+						jsonWriter.WriteArrayEnd();
+						jsonWriter.WriteObjectEnd();
+					}
+				}
+				jsonWriter.WriteArrayEnd();
+
+				// Write all the badges
+				jsonWriter.WriteArrayStart("Badges");
+				foreach (BgBadgeDef badge in Badges)
+				{
+					BgNodeDef[] dependencies = badge.Nodes.Where(x => nodesToExecute.Contains(x)).ToArray();
+					if (dependencies.Length > 0)
+					{
+						// Reduce that list to the smallest subset of direct dependencies
+						HashSet<BgNodeDef> directDependencies = new HashSet<BgNodeDef>(dependencies);
+						foreach (BgNodeDef dependency in dependencies)
+						{
+							directDependencies.ExceptWith(dependency.OrderDependencies);
+						}
+
+						jsonWriter.WriteObjectStart();
+						jsonWriter.WriteValue("Name", badge.Name);
+						if (!String.IsNullOrEmpty(badge.Project))
+						{
+							jsonWriter.WriteValue("Project", badge.Project);
+						}
+						if (badge.Change != 0)
+						{
+							jsonWriter.WriteValue("Change", badge.Change);
+						}
+						jsonWriter.WriteValue("AllDependencies", String.Join(";", Agents.SelectMany(x => x.Nodes).Where(x => dependencies.Contains(x)).Select(x => x.Name)));
+						jsonWriter.WriteValue("DirectDependencies", String.Join(";", directDependencies.Select(x => x.Name)));
+						jsonWriter.WriteObjectEnd();
+					}
+				}
+				jsonWriter.WriteArrayEnd();
+
+				// Write all the triggers and reports. 
+				jsonWriter.WriteArrayStart("Reports");
+				foreach (BgReport report in NameToReport.Values)
+				{
+					BgNodeDef[] dependencies = report.Nodes.Where(x => nodesToExecute.Contains(x)).ToArray();
+					if (dependencies.Length > 0)
+					{
+						// Reduce that list to the smallest subset of direct dependencies
+						HashSet<BgNodeDef> directDependencies = new HashSet<BgNodeDef>(dependencies);
+						foreach (BgNodeDef dependency in dependencies)
+						{
+							directDependencies.ExceptWith(dependency.OrderDependencies);
+						}
+
+						jsonWriter.WriteObjectStart();
+						jsonWriter.WriteValue("Name", report.Name);
+						jsonWriter.WriteValue("AllDependencies", String.Join(";", Agents.SelectMany(x => x.Nodes).Where(x => dependencies.Contains(x)).Select(x => x.Name)));
+						jsonWriter.WriteValue("DirectDependencies", String.Join(";", directDependencies.Select(x => x.Name)));
+						jsonWriter.WriteValue("Notify", String.Join(";", report.NotifyUsers));
+						jsonWriter.WriteValue("IsTrigger", false);
+						jsonWriter.WriteObjectEnd();
+					}
+				}
+				jsonWriter.WriteArrayEnd();
+
+				jsonWriter.WriteObjectEnd();
+			}
+		}
+
+		/// <summary>
+		/// Export the build graph to a Json file for parsing by Horde
+		/// </summary>
+		/// <param name="file">Output file to write</param>
+		public void ExportForHorde(FileReference file)
+		{
+			DirectoryReference.CreateDirectory(file.Directory);
+			using (JsonWriter jsonWriter = new JsonWriter(file.FullName))
+			{
+				jsonWriter.WriteObjectStart();
+				jsonWriter.WriteArrayStart("Groups");
+				foreach (BgAgentDef agent in Agents)
+				{
+					jsonWriter.WriteObjectStart();
+					jsonWriter.WriteArrayStart("Types");
+					foreach (string possibleType in agent.PossibleTypes)
+					{
+						jsonWriter.WriteValue(possibleType);
+					}
+					jsonWriter.WriteArrayEnd();
+					jsonWriter.WriteArrayStart("Nodes");
+					foreach (BgNodeDef node in agent.Nodes)
+					{
+						jsonWriter.WriteObjectStart();
+						jsonWriter.WriteValue("Name", node.Name);
+						jsonWriter.WriteValue("RunEarly", node.RunEarly);
+						jsonWriter.WriteValue("Warnings", node.NotifyOnWarnings);
+
+						jsonWriter.WriteArrayStart("InputDependencies");
+						foreach (string inputDependency in node.GetDirectInputDependencies().Select(x => x.Name))
+						{
+							jsonWriter.WriteValue(inputDependency);
+						}
+						jsonWriter.WriteArrayEnd();
+
+						jsonWriter.WriteArrayStart("OrderDependencies");
+						foreach (string orderDependency in node.GetDirectOrderDependencies().Select(x => x.Name))
+						{
+							jsonWriter.WriteValue(orderDependency);
+						}
+						jsonWriter.WriteArrayEnd();
+
+						jsonWriter.WriteObjectStart("Annotations");
+						foreach ((string key, string value) in node.Annotations)
+						{
+							jsonWriter.WriteValue(key, value);
+						}
+						jsonWriter.WriteObjectEnd();
+
+						jsonWriter.WriteObjectEnd();
+					}
+					jsonWriter.WriteArrayEnd();
+					jsonWriter.WriteObjectEnd();
+				}
+				jsonWriter.WriteArrayEnd();
+
+				jsonWriter.WriteArrayStart("Aggregates");
+				foreach (BgAggregateDef aggregate in NameToAggregate.Values)
+				{
+					jsonWriter.WriteObjectStart();
+					jsonWriter.WriteValue("Name", aggregate.Name);
+					jsonWriter.WriteArrayStart("Nodes");
+					foreach (BgNodeDef requiredNode in aggregate.RequiredNodes.OrderBy(x => x.Name))
+					{
+						jsonWriter.WriteValue(requiredNode.Name);
+					}
+					jsonWriter.WriteArrayEnd();
+					jsonWriter.WriteObjectEnd();
+				}
+				jsonWriter.WriteArrayEnd();
+
+				jsonWriter.WriteArrayStart("Labels");
+				foreach (BgLabelDef label in Labels)
+				{
+					jsonWriter.WriteObjectStart();
+					if (!String.IsNullOrEmpty(label.DashboardName))
+					{
+						jsonWriter.WriteValue("Name", label.DashboardName);
+					}
+					if (!String.IsNullOrEmpty(label.DashboardCategory))
+					{
+						jsonWriter.WriteValue("Category", label.DashboardCategory);
+					}
+					if (!String.IsNullOrEmpty(label.UgsBadge))
+					{
+						jsonWriter.WriteValue("UgsBadge", label.UgsBadge);
+					}
+					if (!String.IsNullOrEmpty(label.UgsProject))
+					{
+						jsonWriter.WriteValue("UgsProject", label.UgsProject);
+					}
+					if (label.Change != BgLabelChange.Current)
+					{
+						jsonWriter.WriteValue("Change", label.Change.ToString());
+					}
+
+					jsonWriter.WriteArrayStart("RequiredNodes");
+					foreach (BgNodeDef requiredNode in label.RequiredNodes.OrderBy(x => x.Name))
+					{
+						jsonWriter.WriteValue(requiredNode.Name);
+					}
+					jsonWriter.WriteArrayEnd();
+					jsonWriter.WriteArrayStart("IncludedNodes");
+					foreach (BgNodeDef includedNode in label.IncludedNodes.OrderBy(x => x.Name))
+					{
+						jsonWriter.WriteValue(includedNode.Name);
+					}
+					jsonWriter.WriteArrayEnd();
+					jsonWriter.WriteObjectEnd();
+				}
+				jsonWriter.WriteArrayEnd();
+
+				jsonWriter.WriteArrayStart("Badges");
+				foreach (BgBadgeDef badge in Badges)
+				{
+					HashSet<BgNodeDef> dependencies = badge.Nodes;
+					if (dependencies.Count > 0)
+					{
+						// Reduce that list to the smallest subset of direct dependencies
+						HashSet<BgNodeDef> directDependencies = new HashSet<BgNodeDef>(dependencies);
+						foreach (BgNodeDef dependency in dependencies)
+						{
+							directDependencies.ExceptWith(dependency.OrderDependencies);
+						}
+
+						jsonWriter.WriteObjectStart();
+						jsonWriter.WriteValue("Name", badge.Name);
+						if (!String.IsNullOrEmpty(badge.Project))
+						{
+							jsonWriter.WriteValue("Project", badge.Project);
+						}
+						if (badge.Change != 0)
+						{
+							jsonWriter.WriteValue("Change", badge.Change);
+						}
+						jsonWriter.WriteValue("Dependencies", String.Join(";", directDependencies.Select(x => x.Name)));
+						jsonWriter.WriteObjectEnd();
+					}
+				}
+				jsonWriter.WriteArrayEnd();
+
+				jsonWriter.WriteObjectEnd();
+			}
+		}
+
+		/// <summary>
+		/// Print the contents of the graph
+		/// </summary>
+		/// <param name="completedNodes">Set of nodes which are already complete</param>
+		/// <param name="printOptions">Options for how to print the graph</param>
+		/// <param name="logger"></param>
+		public void Print(HashSet<BgNodeDef> completedNodes, GraphPrintOptions printOptions, ILogger logger)
+		{
+			// Print the options
+			if ((printOptions & GraphPrintOptions.ShowCommandLineOptions) != 0)
+			{
+				// Get the list of messages
+				List<KeyValuePair<string, string>> parameters = new List<KeyValuePair<string, string>>();
+				foreach (BgOption option in Options)
+				{
+					string name = String.Format("-set:{0}=...", option.Name);
+
+					StringBuilder description = new StringBuilder(option.Description);
+					if (!String.IsNullOrEmpty(option.DefaultValue))
+					{
+						description.AppendFormat(" (Default: {0})", option.DefaultValue);
+					}
+
+					parameters.Add(new KeyValuePair<string, string>(name, description.ToString()));
+				}
+
+				// Format them to the log
+				if (parameters.Count > 0)
+				{
+					logger.LogInformation("");
+					logger.LogInformation("Options:");
+					logger.LogInformation("");
+
+					List<string> lines = new List<string>();
+					HelpUtils.FormatTable(parameters, 4, 24, HelpUtils.WindowWidth - 20, lines);
+
+					foreach (string line in lines)
+					{
+						logger.Log(LogLevel.Information, line);
+					}
+				}
+			}
+
+			// Output all the triggers in order
+			logger.LogInformation("");
+			logger.LogInformation("Graph:");
+			foreach (BgAgentDef agent in Agents)
+			{
+				logger.LogInformation("        Agent: {0} ({1})", agent.Name, String.Join(";", agent.PossibleTypes));
+				foreach (BgNodeDef node in agent.Nodes)
+				{
+					logger.LogInformation("            Node: {0}{1}", node.Name, completedNodes.Contains(node) ? " (completed)" : node.RunEarly ? " (early)" : "");
+					if (printOptions.HasFlag(GraphPrintOptions.ShowDependencies))
+					{
+						HashSet<BgNodeDef> inputDependencies = new HashSet<BgNodeDef>(node.GetDirectInputDependencies());
+						foreach (BgNodeDef inputDependency in inputDependencies)
+						{
+							logger.LogInformation("                input> {0}", inputDependency.Name);
+						}
+						HashSet<BgNodeDef> orderDependencies = new HashSet<BgNodeDef>(node.GetDirectOrderDependencies());
+						foreach (BgNodeDef orderDependency in orderDependencies.Except(inputDependencies))
+						{
+							logger.LogInformation("                after> {0}", orderDependency.Name);
+						}
+					}
+					if (printOptions.HasFlag(GraphPrintOptions.ShowNotifications))
+					{
+						string label = node.NotifyOnWarnings ? "warnings" : "errors";
+						foreach (string user in node.NotifyUsers)
+						{
+							logger.LogInformation("                {0}> {1}", label, user);
+						}
+						foreach (string submitter in node.NotifySubmitters)
+						{
+							logger.LogInformation("                {0}> submitters to {1}", label, submitter);
+						}
+					}
+				}
+			}
+			logger.LogInformation("");
+
+			// Print out all the non-empty aggregates
+			BgAggregateDef[] aggregates = NameToAggregate.Values.OrderBy(x => x.Name).ToArray();
+			if (aggregates.Length > 0)
+			{
+				logger.LogInformation("Aggregates:");
+				foreach (string aggregateName in aggregates.Select(x => x.Name))
+				{
+					logger.LogInformation("    {0}", aggregateName);
+				}
+				logger.LogInformation("");
+			}
+		}
+	}
+}
diff --git a/Engine/Source/Programs/Shared/EpicGames.BuildGraph/BgGraphSpec.cs b/Engine/Source/Programs/Shared/EpicGames.BuildGraph/BgGraphSpec.cs
deleted file mode 100644
diff --git a/Engine/Source/Programs/Shared/EpicGames.BuildGraph/BgInterpreter.cs b/Engine/Source/Programs/Shared/EpicGames.BuildGraph/BgInterpreter.cs
--- a/Engine/Source/Programs/Shared/EpicGames.BuildGraph/BgInterpreter.cs
+++ b/Engine/Source/Programs/Shared/EpicGames.BuildGraph/BgInterpreter.cs
@@ -551,7 +551,7 @@ namespace EpicGames.BuildGraph
 				case BgOpcode.Node:
 					{
 						string name = (string)Evaluate(frame);
-						BgAgent agent = (BgAgent)Evaluate(frame);
+						BgAgentDef agent = (BgAgentDef)Evaluate(frame);
 						BgMethod method = _methods[(int)ReadUnsignedInteger(frame)];
 
 						int argumentCount = (int)ReadUnsignedInteger(frame);
@@ -566,18 +566,18 @@ namespace EpicGames.BuildGraph
 						List<BgNodeOutput> fences = ReadNodeOutputList(frame);
 						int taggedOutputCount = (int)ReadUnsignedInteger(frame);
 						bool runEarly = (bool)Evaluate(frame);
-						List<BgLabel> labels = ((IEnumerable<object>)Evaluate(frame)).Select(x => (BgLabel)x).ToList();
+						List<BgLabelDef> labels = ((IEnumerable<object>)Evaluate(frame)).Select(x => (BgLabelDef)x).ToList();
 
-						HashSet<BgNode> inputDependencies = new HashSet<BgNode>();
+						HashSet<BgNodeDef> inputDependencies = new HashSet<BgNodeDef>();
 						foreach (BgNodeOutput input in inputs)
 						{
 							inputDependencies.Add(input.ProducingNode);
 							inputDependencies.UnionWith(input.ProducingNode.InputDependencies);
 						}
 
-						BgNode[] orderDependencies = fences.Select(x => x.ProducingNode).Distinct().ToArray();
+						BgNodeDef[] orderDependencies = fences.Select(x => x.ProducingNode).Distinct().ToArray();
 
-						string[] outputNames = Enumerable.Range(0, taggedOutputCount).Select(x => BgNodeSpec.GetDefaultTagName(name, x)).ToArray();
+						string[] outputNames = Enumerable.Range(0, taggedOutputCount).Select(x => BgNode.GetDefaultTagName(name, x)).ToArray();
 
 						BgBytecodeNode node = new BgBytecodeNode(agent, name, method, arguments, inputs, outputNames, inputDependencies.ToArray(), orderDependencies.ToArray(), Array.Empty<FileReference>(), labels);
 						return node;
@@ -586,15 +586,15 @@ namespace EpicGames.BuildGraph
 					{
 						string name = (string)Evaluate(frame);
 						string[] types = ((IEnumerable<object>)Evaluate(frame)).Select(x => (string)x).ToArray();
-						return new BgAgent(name, types);
+						return new BgAgentDef(name, types);
 					}
 				case BgOpcode.Aggregate:
 					{
 						string name = (string)Evaluate(frame);
 						List<BgNodeOutput> requires = ReadNodeOutputList(frame);
-						BgLabel? label = (BgLabel?)Evaluate(frame);
+						BgLabelDef? label = (BgLabelDef?)Evaluate(frame);
 
-						return new BgBytecodeAggregate(name, requires.Select(x => x.ProducingNode), label);
+						return new BgBytecodeAggregateDef(name, requires.Select(x => x.ProducingNode), label);
 					}
 				case BgOpcode.Label:
 					{
@@ -610,37 +610,37 @@ namespace EpicGames.BuildGraph
 							labelChange = Enum.Parse<BgLabelChange>(change);
 						}
 
-						return new BgLabel(dashboardName, dashboardCategory, ugsBadge, ugsProject, labelChange);
+						return new BgLabelDef(dashboardName, dashboardCategory, ugsBadge, ugsProject, labelChange);
 					}
 				case BgOpcode.Graph:
 					{
 						List<BgBytecodeNode> nodes = EvaluateList<BgBytecodeNode>(frame);
-						List<BgBytecodeAggregate> aggregates = EvaluateList<BgBytecodeAggregate>(frame);
+						List<BgBytecodeAggregateDef> aggregates = EvaluateList<BgBytecodeAggregateDef>(frame);
 
-						BgGraph graph = new BgGraph();
-						foreach (BgBytecodeAggregate aggregate in aggregates)
+						BgGraphDef graph = new BgGraphDef();
+						foreach (BgBytecodeAggregateDef aggregate in aggregates)
 						{
 							graph.NameToAggregate[aggregate.Name] = aggregate;
 							nodes.AddRange(aggregate.RequiredNodes.Select(x => (BgBytecodeNode)x));
 						}
 
-						HashSet<BgNode> uniqueNodes = new HashSet<BgNode>();
-						HashSet<BgAgent> uniqueAgents = new HashSet<BgAgent>();
+						HashSet<BgNodeDef> uniqueNodes = new HashSet<BgNodeDef>();
+						HashSet<BgAgentDef> uniqueAgents = new HashSet<BgAgentDef>();
 						foreach (BgBytecodeNode node in nodes)
 						{
 							RegisterNode(graph, node, uniqueNodes, uniqueAgents);
 						}
 
-						HashSet<BgLabel> labels = new HashSet<BgLabel>();
+						HashSet<BgLabelDef> labels = new HashSet<BgLabelDef>();
 						foreach (BgBytecodeNode node in nodes)
 						{
-							foreach (BgLabel label in node.Labels)
+							foreach (BgLabelDef label in node.Labels)
 							{
 								labels.Add(label);
 								label.RequiredNodes.Add(node);
 							}
 						}
-						foreach (BgBytecodeAggregate aggregate in aggregates)
+						foreach (BgBytecodeAggregateDef aggregate in aggregates)
 						{
 							if (aggregate.Label != null)
 							{
@@ -659,7 +659,7 @@ namespace EpicGames.BuildGraph
 			}
 		}
 
-		void RegisterNode(BgGraph graph, BgBytecodeNode node, HashSet<BgNode> uniqueNodes, HashSet<BgAgent> uniqueAgents)
+		void RegisterNode(BgGraphDef graph, BgBytecodeNode node, HashSet<BgNodeDef> uniqueNodes, HashSet<BgAgentDef> uniqueAgents)
 		{
 			if (uniqueNodes.Add(node))
 			{
@@ -668,7 +668,7 @@ namespace EpicGames.BuildGraph
 					RegisterNode(graph, inputNode, uniqueNodes, uniqueAgents);
 				}
 
-				BgAgent agent = node.Agent;
+				BgAgentDef agent = node.Agent;
 				if (uniqueAgents.Add(agent))
 				{
 					graph.Agents.Add(agent);
diff --git a/Engine/Source/Programs/Shared/EpicGames.BuildGraph/BgLabel.cs b/Engine/Source/Programs/Shared/EpicGames.BuildGraph/BgLabel.cs
deleted file mode 100644
diff --git a/Engine/Source/Programs/Shared/EpicGames.BuildGraph/BgLabelDef.cs b/Engine/Source/Programs/Shared/EpicGames.BuildGraph/BgLabelDef.cs
new file mode 100644
--- /dev/null
+++ b/Engine/Source/Programs/Shared/EpicGames.BuildGraph/BgLabelDef.cs
@@ -0,0 +1,98 @@
+// Copyright Epic Games, Inc. All Rights Reserved.
+
+using System;
+using System.Collections.Generic;
+
+namespace EpicGames.BuildGraph
+{
+	/// <summary>
+	/// Which changelist to show a UGS badge for
+	/// </summary>
+	public enum BgLabelChange
+	{
+		/// <summary>
+		/// The current changelist being built
+		/// </summary>
+		Current,
+
+		/// <summary>
+		/// The last code changelist
+		/// </summary>
+		Code,
+	}
+
+	/// <summary>
+	/// Defines a label within a graph. Labels are similar to badges, and give the combined status of one or more job steps. Unlike badges, they
+	/// separate the requirements for its status and optional nodes to be included in its status, allowing this to be handled externally.
+	/// </summary>
+	public class BgLabelDef
+	{
+		/// <summary>
+		/// Name of this badge
+		/// </summary>
+		public string? DashboardName { get; }
+
+		/// <summary>
+		/// Category for this label
+		/// </summary>
+		public string? DashboardCategory { get; }
+
+		/// <summary>
+		/// Name of the badge in UGS
+		/// </summary>
+		public string? UgsBadge { get; }
+
+		/// <summary>
+		/// Path to the project folder in UGS
+		/// </summary>
+		public string? UgsProject { get; }
+
+		/// <summary>
+		/// Which change to show the badge for
+		/// </summary>
+		public BgLabelChange Change { get; }
+
+		/// <summary>
+		/// Set of nodes that must be run for this label to be shown.
+		/// </summary>
+		public HashSet<BgNodeDef> RequiredNodes { get; } = new HashSet<BgNodeDef>();
+
+		/// <summary>
+		/// Set of nodes that will be included in this label if present.
+		/// </summary>
+		public HashSet<BgNodeDef> IncludedNodes { get; } = new HashSet<BgNodeDef>();
+
+		/// <summary>
+		/// Constructor
+		/// </summary>
+		/// <param name="inDashboardName">Name of this label</param>
+		/// <param name="inDashboardCategory">Type of this label</param>
+		/// <param name="inUgsBadge">The UGS badge name</param>
+		/// <param name="inUgsProject">Project to display this badge for</param>
+		/// <param name="inChange">The change to show this badge on in UGS</param>
+		public BgLabelDef(string? inDashboardName, string? inDashboardCategory, string? inUgsBadge, string? inUgsProject, BgLabelChange inChange)
+		{
+			DashboardName = inDashboardName;
+			DashboardCategory = inDashboardCategory;
+			UgsBadge = inUgsBadge;
+			UgsProject = inUgsProject;
+			Change = inChange;
+		}
+
+		/// <summary>
+		/// Get the name of this label
+		/// </summary>
+		/// <returns>The name of this label</returns>
+		public override string ToString()
+		{
+			if (!String.IsNullOrEmpty(DashboardName))
+			{
+				return String.Format("{0}/{1}", DashboardCategory, DashboardName);
+			}
+			else
+			{
+				return UgsBadge ?? "Unknown";
+			}
+		}
+	}
+}
diff --git a/Engine/Source/Programs/Shared/EpicGames.BuildGraph/BgLabelSpec.cs b/Engine/Source/Programs/Shared/EpicGames.BuildGraph/BgLabelSpec.cs
deleted file mode 100644
diff --git a/Engine/Source/Programs/Shared/EpicGames.BuildGraph/BgNode.cs b/Engine/Source/Programs/Shared/EpicGames.BuildGraph/BgNode.cs
deleted file mode 100644
diff --git a/Engine/Source/Programs/Shared/EpicGames.BuildGraph/BgNodeDef.cs b/Engine/Source/Programs/Shared/EpicGames.BuildGraph/BgNodeDef.cs
new file mode 100644
--- /dev/null
+++ b/Engine/Source/Programs/Shared/EpicGames.BuildGraph/BgNodeDef.cs
@@ -0,0 +1,219 @@
+// Copyright Epic Games, Inc. All Rights Reserved.
+
+using System;
+using System.Collections.Generic;
+using System.Linq;
+using System.Linq.Expressions;
+using System.Reflection;
+using System.Threading.Tasks;
+using System.Xml;
+using EpicGames.Core;
+
+namespace EpicGames.BuildGraph
+{
+	/// <summary>
+	/// Reference to an output tag from a particular node
+	/// </summary>
+	public class BgNodeOutput
+	{
+		/// <summary>
+		/// The node which produces the given output
+		/// </summary>
+		public BgNodeDef ProducingNode { get; }
+
+		/// <summary>
+		/// Name of the tag
+		/// </summary>
+		public string TagName { get; }
+
+		/// <summary>
+		/// Constructor
+		/// </summary>
+		/// <param name="inProducingNode">Node which produces the given output</param>
+		/// <param name="inTagName">Name of the tag</param>
+		public BgNodeOutput(BgNodeDef inProducingNode, string inTagName)
+		{
+			ProducingNode = inProducingNode;
+			TagName = inTagName;
+		}
+
+		/// <summary>
+		/// Returns a string representation of this output for debugging purposes
+		/// </summary>
+		/// <returns>The name of this output</returns>
+		public override string ToString()
+		{
+			return String.Format("{0} [{1}]", TagName, ProducingNode.Name);
+		}
+	}
+
+	/// <summary>
+	/// Defines a node, a container for tasks and the smallest unit of execution that can be run as part of a build graph.
+	/// </summary>
+	public class BgNodeDef
+	{
+		/// <summary>
+		/// The node's name
+		/// </summary>
+		public string Name { get; }
+
+		/// <summary>
+		/// Array of inputs which this node requires to run
+		/// </summary>
+		public IReadOnlyList<BgNodeOutput> Inputs { get; }
+
+		/// <summary>
+		/// Array of outputs produced by this node
+		/// </summary>
+		public IReadOnlyList<BgNodeOutput> Outputs { get; }
+
+		/// <summary>
+		/// Nodes which this node has input dependencies on
+		/// </summary>
+		public IReadOnlyList<BgNodeDef> InputDependencies { get; set; }
+
+		/// <summary>
+		/// Nodes which this node needs to run after
+		/// </summary>
+		public IReadOnlyList<BgNodeDef> OrderDependencies { get; set; }
+
+		/// <summary>
+		/// Tokens which must be acquired for this node to run
+		/// </summary>
+		public IReadOnlyList<FileReference> RequiredTokens { get; }
+
+		/// <summary>
+		/// List of email addresses to notify if this node fails.
+		/// </summary>
+		public HashSet<string> NotifyUsers { get; set; } = new HashSet<string>(StringComparer.InvariantCultureIgnoreCase);
+
+		/// <summary>
+		/// If set, anyone that has submitted to one of the given paths will be notified on failure of this node
+		/// </summary>
+		public HashSet<string> NotifySubmitters { get; set; } = new HashSet<string>(StringComparer.InvariantCultureIgnoreCase);
+
+		/// <summary>
+		/// Whether to start this node as soon as its dependencies are satisfied, rather than waiting for all of its agent's dependencies to be met.
+		/// </summary>
+		public bool RunEarly { get; set; } = false;
+
+		/// <summary>
+		/// Whether to ignore warnings produced by this node
+		/// </summary>
+		public bool NotifyOnWarnings { get; set; } = true;
+
+		/// <summary>
+		/// Custom annotations for this node
+		/// </summary>
+		public Dictionary<string, string> Annotations { get; } = new Dictionary<string, string>(StringComparer.OrdinalIgnoreCase);
+
+		/// <summary>
+		/// Diagnostics to output if executing this node
+		/// </summary>
+		public List<BgDiagnosticDef> Diagnostics { get; } = new List<BgDiagnosticDef>();
+
+		/// <summary>
+		/// Constructor
+		/// </summary>
+		/// <param name="name">The name of this node</param>
+		/// <param name="inputs">Inputs that this node depends on</param>
+		/// <param name="outputNames">Names of the outputs that this node produces</param>
+		/// <param name="inputDependencies">Nodes which this node is dependent on for its inputs</param>
+		/// <param name="orderDependencies">Nodes which this node needs to run after. Should include all input dependencies.</param>
+		/// <param name="requiredTokens">Optional tokens which must be required for this node to run</param>
+		public BgNodeDef(string name, IReadOnlyList<BgNodeOutput> inputs, IReadOnlyList<string> outputNames, IReadOnlyList<BgNodeDef> inputDependencies, IReadOnlyList<BgNodeDef> orderDependencies, IReadOnlyList<FileReference> requiredTokens)
+		{
+			Name = name;
+			Inputs = inputs;
+
+			List<BgNodeOutput> allOutputs = new List<BgNodeOutput>();
+			allOutputs.Add(new BgNodeOutput(this, "#" + Name));
+			allOutputs.AddRange(outputNames.Where(x => String.Compare(x, Name, StringComparison.InvariantCultureIgnoreCase) != 0).Select(x => new BgNodeOutput(this, x)));
+			Outputs = allOutputs.ToArray();
+
+			InputDependencies = inputDependencies;
+			OrderDependencies = orderDependencies;
+			RequiredTokens = requiredTokens;
+		}
+
+		/// <summary>
+		/// Returns the default output for this node, which includes all build products
+		/// </summary>
+		public BgNodeOutput DefaultOutput => Outputs[0];
+
+		/// <summary>
+		/// Determines the minimal set of direct input dependencies for this node to run
+		/// </summary>
+		/// <returns>Sequence of nodes that are direct inputs to this node</returns>
+		public IEnumerable<BgNodeDef> GetDirectInputDependencies()
+		{
+			HashSet<BgNodeDef> directDependencies = new HashSet<BgNodeDef>(InputDependencies);
+			foreach (BgNodeDef inputDependency in InputDependencies)
+			{
+				directDependencies.ExceptWith(inputDependency.InputDependencies);
+			}
+			return directDependencies;
+		}
+
+		/// <summary>
+		/// Determines the minimal set of direct order dependencies for this node to run
+		/// </summary>
+		/// <returns>Sequence of nodes that are direct order dependencies of this node</returns>
+		public IEnumerable<BgNodeDef> GetDirectOrderDependencies()
+		{
+			HashSet<BgNodeDef> directDependencies = new HashSet<BgNodeDef>(OrderDependencies);
+			foreach (BgNodeDef orderDependency in OrderDependencies)
+			{
+				directDependencies.ExceptWith(orderDependency.OrderDependencies);
+			}
+			return directDependencies;
+		}
+
+		/// <summary>
+		/// Returns the name of this node
+		/// </summary>
+		/// <returns>The name of this node</returns>
+		public override string ToString()
+		{
+			return Name;
+		}
+	}
+
+	/// <summary>
+	/// Node constructed from a bytecode expression
+	/// </summary>
+	public class BgBytecodeNode : BgNodeDef
+	{
+		/// <summary>
+		/// Agent declaring this node
+		/// </summary>
+		public BgAgentDef Agent { get; }
+
+		/// <summary>
+		/// Offset in the bytecode of this node. Used to lookup method calling information.
+		/// </summary>
+		public BgMethod Method { get; }
+
+		/// <summary>
+		/// Expression arguments to the method
+		/// </summary>
+		public IReadOnlyList<object> Arguments { get; }
+
+		/// <summary>
+		/// Labels to add this node to
+		/// </summary>
+		public IReadOnlyList<BgLabelDef> Labels { get; }
+
+		/// <summary>
+		/// Constructor
+		/// </summary>
+		public BgBytecodeNode(BgAgentDef agent, string name, BgMethod method, IReadOnlyList<object> arguments, IReadOnlyList<BgNodeOutput> inputs, IReadOnlyList<string> outputNames, IReadOnlyList<BgNodeDef> inputDependencies, IReadOnlyList<BgNodeDef> orderDependencies, IReadOnlyList<FileReference> requiredTokens, IReadOnlyList<BgLabelDef> labels)
+			: base(name, inputs, outputNames, inputDependencies, orderDependencies, requiredTokens)
+		{
+			Agent = agent;
+			Method = method;
+			Arguments = arguments;
+			Labels = labels;
+		}
+	}
+}
diff --git a/Engine/Source/Programs/Shared/EpicGames.BuildGraph/BgNodeSpec.cs b/Engine/Source/Programs/Shared/EpicGames.BuildGraph/BgNodeSpec.cs
deleted file mode 100644
diff --git a/Engine/Source/Programs/Shared/EpicGames.BuildGraph/BgOpcode.cs b/Engine/Source/Programs/Shared/EpicGames.BuildGraph/BgOpcode.cs
--- a/Engine/Source/Programs/Shared/EpicGames.BuildGraph/BgOpcode.cs
+++ b/Engine/Source/Programs/Shared/EpicGames.BuildGraph/BgOpcode.cs
@@ -339,27 +339,27 @@ namespace EpicGames.BuildGraph
 		#region Graph structures
 
 		/// <summary>
-		/// Defines a node (see <see cref="BgNodeSpec"/>)
+		/// Defines a node (see <see cref="BgNode"/>)
 		/// </summary>
 		Node = 0xf0,
 
 		/// <summary>
-		/// Defines an agent to execute work on (see <see cref="BgAgentSpec"/>)
+		/// Defines an agent to execute work on (see <see cref="BgAgent"/>)
 		/// </summary>
 		Agent = 0xf1,
 		
 		/// <summary>
-		/// Defines a label (see <see cref="BgLabelSpec"/>)
+		/// Defines a label (see <see cref="BgLabel"/>)
 		/// </summary>
 		Label = 0xf2,
 
 		/// <summary>
-		/// Defines an aggregate (see <see cref="BgAggregateSpec"/>)
+		/// Defines an aggregate (see <see cref="BgAggregate"/>)
 		/// </summary>
 		Aggregate = 0xf3,
 
 		/// <summary>
-		/// Defines a graph (see <see cref="BgGraphSpec"/>)
+		/// Defines a graph (see <see cref="BgGraph"/>)
 		/// </summary>
 		Graph = 0xf4,
 
diff --git a/Engine/Source/Programs/Shared/EpicGames.BuildGraph/BgReport.cs b/Engine/Source/Programs/Shared/EpicGames.BuildGraph/BgReport.cs
--- a/Engine/Source/Programs/Shared/EpicGames.BuildGraph/BgReport.cs
+++ b/Engine/Source/Programs/Shared/EpicGames.BuildGraph/BgReport.cs
@@ -18,7 +18,7 @@ namespace EpicGames.BuildGraph
 		/// <summary>
 		/// Set of nodes to include in the report
 		/// </summary>
-		public HashSet<BgNode> Nodes { get; } = new HashSet<BgNode>();
+		public HashSet<BgNodeDef> Nodes { get; } = new HashSet<BgNodeDef>();
 
 		/// <summary>
 		/// List of users to notify with this report
diff --git a/Engine/Source/Programs/Shared/EpicGames.BuildGraph/Expressions/BgAgent.cs b/Engine/Source/Programs/Shared/EpicGames.BuildGraph/Expressions/BgAgent.cs
new file mode 100644
--- /dev/null
+++ b/Engine/Source/Programs/Shared/EpicGames.BuildGraph/Expressions/BgAgent.cs
@@ -0,0 +1,57 @@
+// Copyright Epic Games, Inc. All Rights Reserved.
+
+using System;
+using System.Collections.Generic;
+using System.Linq.Expressions;
+using System.Threading.Tasks;
+using EpicGames.BuildGraph.Expressions;
+using EpicGames.Core;
+using Microsoft.Extensions.Logging;
+
+namespace EpicGames.BuildGraph
+{
+	/// <summary>
+	/// Describes an agent that can execute execute build steps
+	/// </summary>
+	public class BgAgent : BgExpr
+	{
+		/// <summary>
+		/// Name of the agent
+		/// </summary>
+		public BgString Name { get; }
+
+		/// <summary>
+		/// List of agent types to select from, in order of preference. The first agent type supported by a stream will be used.
+		/// </summary>
+		public BgList<BgString> Types { get; set; }
+
+		/// <summary>
+		/// Constructor
+		/// </summary>
+		public BgAgent(BgString name, BgString type)
+			: this(name, BgList<BgString>.Create(type))
+		{
+		}
+
+		/// <summary>
+		/// Constructor
+		/// </summary>
+		public BgAgent(BgString name, BgList<BgString> types)
+			: base(BgExprFlags.ForceFragment)
+		{
+			Name = name;
+			Types = types;
+		}
+
+		/// <inheritdoc/>
+		public override void Write(BgBytecodeWriter writer)
+		{
+			writer.WriteOpcode(BgOpcode.Agent);
+			writer.WriteExpr(Name);
+			writer.WriteExpr(Types);
+		}
+
+		/// <inheritdoc/>
+		public override BgString ToBgString() => "{Agent}";
+	}
+}
diff --git a/Engine/Source/Programs/Shared/EpicGames.BuildGraph/Expressions/BgAggregate.cs b/Engine/Source/Programs/Shared/EpicGames.BuildGraph/Expressions/BgAggregate.cs
new file mode 100644
--- /dev/null
+++ b/Engine/Source/Programs/Shared/EpicGames.BuildGraph/Expressions/BgAggregate.cs
@@ -0,0 +1,75 @@
+// Copyright Epic Games, Inc. All Rights Reserved.
+
+using System;
+using EpicGames.BuildGraph.Expressions;
+
+namespace EpicGames.BuildGraph
+{
+	/// <summary>
+	/// Specification for an aggregate target in the graph
+	/// </summary>
+	public class BgAggregate : BgExpr
+	{
+		/// <summary>
+		/// Name of the aggregate
+		/// </summary>
+		public BgString Name { get; }
+
+		/// <summary>
+		/// Outputs required for the aggregate
+		/// </summary>
+		public BgList<BgFileSet> Requires { get; }
+
+		/// <summary>
+		/// Label to apply to this aggregate
+		/// </summary>
+		public BgLabel? Label { get; }
+
+		/// <summary>
+		/// Constructor.
+		/// </summary>
+		public BgAggregate(BgString name, params BgFileSet[] requires)
+			: this(name, BgList<BgFileSet>.Create(requires))
+		{
+		}
+
+		/// <summary>
+		/// Constructor.
+		/// </summary>
+		public BgAggregate(BgString name, params BgList<BgFileSet>[] requires)
+			: this(name, BgList<BgFileSet>.Concat(requires))
+		{
+		}
+
+		/// <summary>
+		/// Constructor.
+		/// </summary>
+		public BgAggregate(BgString name, BgList<BgFileSet> requires, string label)
+			: this(name, requires, new BgLabel(label))
+		{
+		}
+
+		/// <summary>
+		/// Constructor.
+		/// </summary>
+		public BgAggregate(BgString name, BgList<BgFileSet> requires, BgLabel? label = null)
+			: base(BgExprFlags.None)
+		{
+			Name = name;
+			Requires = requires;
+			Label = label;
+		}
+
+		/// <inheritdoc/>
+		public override void Write(BgBytecodeWriter writer)
+		{
+			writer.WriteOpcode(BgOpcode.Aggregate);
+			writer.WriteExpr(Name);
+			writer.WriteExpr(Requires);
+			writer.WriteExpr(Label ?? BgExpr.Null);
+		}
+
+		/// <inheritdoc/>
+		public override BgString ToBgString() => Name;
+	}
+}
diff --git a/Engine/Source/Programs/Shared/EpicGames.BuildGraph/Expressions/BgFileSet.cs b/Engine/Source/Programs/Shared/EpicGames.BuildGraph/Expressions/BgFileSet.cs
--- a/Engine/Source/Programs/Shared/EpicGames.BuildGraph/Expressions/BgFileSet.cs
+++ b/Engine/Source/Programs/Shared/EpicGames.BuildGraph/Expressions/BgFileSet.cs
@@ -98,9 +98,9 @@ namespace EpicGames.BuildGraph.Expressions
 
 	class BgFileSetFromNodeExpr : BgFileSet
 	{
-		public BgNodeSpec Node { get; }
+		public BgNode Node { get; }
 
-		public BgFileSetFromNodeExpr(BgNodeSpec node)
+		public BgFileSetFromNodeExpr(BgNode node)
 			: base(BgExprFlags.NotInterned)
 		{
 			Node = node;
@@ -115,10 +115,10 @@ namespace EpicGames.BuildGraph.Expressions
 
 	class BgFileSetFromNodeOutputExpr : BgFileSet
 	{
-		public BgNodeSpec Node { get; }
+		public BgNode Node { get; }
 		public int OutputIndex { get; }
 
-		public BgFileSetFromNodeOutputExpr(BgNodeSpec node, int outputIndex)
+		public BgFileSetFromNodeOutputExpr(BgNode node, int outputIndex)
 			: base(BgExprFlags.NotInterned)
 		{
 			Node = node;
diff --git a/Engine/Source/Programs/Shared/EpicGames.BuildGraph/Expressions/BgGraph.cs b/Engine/Source/Programs/Shared/EpicGames.BuildGraph/Expressions/BgGraph.cs
new file mode 100644
--- /dev/null
+++ b/Engine/Source/Programs/Shared/EpicGames.BuildGraph/Expressions/BgGraph.cs
@@ -0,0 +1,78 @@
+// Copyright Epic Games, Inc. All Rights Reserved.
+
+using System;
+using System.Collections.Generic;
+using EpicGames.BuildGraph.Expressions;
+using Microsoft.Extensions.Logging;
+
+namespace EpicGames.BuildGraph
+{
+	/// <summary>
+	/// Specification for a graph in fluent syntax
+	/// </summary>
+	public class BgGraph : BgExpr
+	{
+		/// <summary>
+		/// Nodes for the graph
+		/// </summary>
+		public BgList<BgNode> Nodes { get; }
+
+		/// <summary>
+		/// Aggregates for the graph
+		/// </summary>
+		public BgList<BgAggregate> Aggregates { get; }
+
+		/// <summary>
+		/// Constructor
+		/// </summary>
+		public BgGraph(BgList<BgNode> nodes, BgList<BgAggregate> aggregates)
+			: base(BgExprFlags.ForceFragment)
+		{
+			Nodes = nodes;
+			Aggregates = aggregates;
+		}
+
+		/// <summary>
+		/// Implicit conversion from a node spec
+		/// </summary>
+		public static implicit operator BgGraph(BgNode node)
+		{
+			return new BgGraph(node, BgList<BgAggregate>.Empty);
+		}
+
+		/// <summary>
+		/// Implicit conversion from a list of node specs
+		/// </summary>
+		public static implicit operator BgGraph(BgList<BgNode> nodes)
+		{
+			return new BgGraph(nodes, BgList<BgAggregate>.Empty);
+		}
+
+		/// <summary>
+		/// Implicit conversion from an aggregate spec
+		/// </summary>
+		public static implicit operator BgGraph(BgAggregate aggregate)
+		{
+			return new BgGraph(BgList<BgNode>.Empty, aggregate);
+		}
+
+		/// <summary>
+		/// Implicit conversion from a list of node specs
+		/// </summary>
+		public static implicit operator BgGraph(BgList<BgAggregate> aggregates)
+		{
+			return new BgGraph(BgList<BgNode>.Empty, aggregates);
+		}
+
+		/// <inheritdoc/>
+		public override void Write(BgBytecodeWriter writer)
+		{
+			writer.WriteOpcode(BgOpcode.Graph);
+			writer.WriteExpr(Nodes);
+			writer.WriteExpr(Aggregates);
+		}
+
+		/// <inheritdoc/>
+		public override BgString ToBgString() => "{Graph}";
+	}
+}
diff --git a/Engine/Source/Programs/Shared/EpicGames.BuildGraph/Expressions/BgLabel.cs b/Engine/Source/Programs/Shared/EpicGames.BuildGraph/Expressions/BgLabel.cs
new file mode 100644
--- /dev/null
+++ b/Engine/Source/Programs/Shared/EpicGames.BuildGraph/Expressions/BgLabel.cs
@@ -0,0 +1,66 @@
+// Copyright Epic Games, Inc. All Rights Reserved.
+
+using System;
+using System.Linq;
+using EpicGames.BuildGraph.Expressions;
+
+namespace EpicGames.BuildGraph
+{
+	/// <summary>
+	/// Specification for a label
+	/// </summary>
+	public class BgLabel : BgExpr
+	{
+		/// <summary>
+		/// Name of this badge
+		/// </summary>
+		public BgString? DashboardName { get; }
+
+		/// <summary>
+		/// Category for this label
+		/// </summary>
+		public BgString? DashboardCategory { get; }
+
+		/// <summary>
+		/// Name of the badge in UGS
+		/// </summary>
+		public BgString? UgsBadge { get; }
+
+		/// <summary>
+		/// Path to the project folder in UGS
+		/// </summary>
+		public BgString? UgsProject { get; }
+
+		/// <summary>
+		/// Which change to show the badge for
+		/// </summary>
+		public BgString? Change { get; }
+
+		/// <summary>
+		/// Constructor
+		/// </summary>
+		public BgLabel(BgString? name = null, BgString? category = null, BgString? ugsBadge = null, BgString? ugsProject = null, BgString? change = null)
+			: base(BgExprFlags.ForceFragment)
+		{
+			DashboardName = name;
+			DashboardCategory = category;
+			UgsBadge = ugsBadge;
+			UgsProject = ugsProject;
+			Change = change;
+		}
+
+		/// <inheritdoc/>
+		public override void Write(BgBytecodeWriter writer)
+		{
+			writer.WriteOpcode(BgOpcode.Label);
+			writer.WriteExpr(DashboardName ?? BgString.Empty);
+			writer.WriteExpr(DashboardCategory ?? BgString.Empty);
+			writer.WriteExpr(UgsBadge ?? BgString.Empty);
+			writer.WriteExpr(UgsProject ?? BgString.Empty);
+			writer.WriteExpr(Change ?? BgString.Empty);
+		}
+
+		/// <inheritdoc/>
+		public override BgString ToBgString() => "{Label}";
+	}
+}
diff --git a/Engine/Source/Programs/Shared/EpicGames.BuildGraph/Expressions/BgNode.cs b/Engine/Source/Programs/Shared/EpicGames.BuildGraph/Expressions/BgNode.cs
new file mode 100644
--- /dev/null
+++ b/Engine/Source/Programs/Shared/EpicGames.BuildGraph/Expressions/BgNode.cs
@@ -0,0 +1,612 @@
+// Copyright Epic Games, Inc. All Rights Reserved.
+
+using System;
+using System.Collections.Generic;
+using System.Linq;
+using System.Linq.Expressions;
+using System.Reflection;
+using System.Text;
+using System.Threading.Tasks;
+using EpicGames.BuildGraph.Expressions;
+using EpicGames.Core;
+
+namespace EpicGames.BuildGraph
+{
+	/// <summary>
+	/// Exception for constructing nodes
+	/// </summary>
+	public sealed class BgNodeException : Exception
+	{
+		/// <summary>
+		/// Constructor
+		/// </summary>
+		/// <param name="message"></param>
+		public BgNodeException(string message) : base(message)
+		{
+		}
+	}
+
+	/// <summary>
+	/// Speecifies the node name for a method. Parameters from the method may be embedded in the name using the {ParamName} syntax.
+	/// </summary>
+	[AttributeUsage(AttributeTargets.Method)]
+	public class BgNodeNameAttribute : Attribute
+	{
+		/// <summary>
+		/// The format string
+		/// </summary>
+		public string Template { get; }
+
+		/// <summary>
+		/// Constructor
+		/// </summary>
+		/// <param name="template">Format string for the name</param>
+		public BgNodeNameAttribute(string template)
+		{
+			Template = template;
+		}
+	}
+
+	/// <summary>
+	/// Information about the method bound to execute a node
+	/// </summary>
+	public class BgMethod
+	{
+		/// <summary>
+		/// Method to call
+		/// </summary>
+		public MethodInfo Method { get; }
+
+		/// <summary>
+		/// Arguments to the method
+		/// </summary>
+		public IReadOnlyList<object?> Arguments { get; }
+
+		/// <summary>
+		/// Constructor
+		/// </summary>
+		public BgMethod(MethodInfo method, IReadOnlyList<object?> arguments)
+		{
+			Method = method;
+			Arguments = arguments;
+		}
+	}
+
+	/// <summary>
+	/// Specification for a node to execute
+	/// </summary>
+	public class BgNode : BgExpr
+	{
+		/// <summary>
+		/// Name of the node
+		/// </summary>
+		public BgString Name { get; }
+
+		/// <summary>
+		/// Method to execute
+		/// </summary>
+		protected BgMethod Method { get; }
+
+		/// <summary>
+		/// Expression arguments from the bound function.
+		/// </summary>
+		protected BgExpr[] ExpressionArguments { get; }
+
+		/// <summary>
+		/// The default output of this node. Includes all other outputs. 
+		/// </summary>
+		public BgFileSet DefaultOutput { get; }
+
+		/// <summary>
+		/// Explicitly tagged outputs from this node
+		/// </summary>
+		public IReadOnlyList<BgExpr> TaggedOutputs { get; protected set; } = Array.Empty<BgExpr>();
+
+		/// <summary>
+		/// Agent for the node to be run on
+		/// </summary>
+		public BgAgent Agent { get; }
+
+		/// <summary>
+		/// Tokens for inputs of this node
+		/// </summary>
+		public BgList<BgFileSet> Inputs { get; }
+
+		/// <summary>
+		/// Weak dependency on outputs that must be generated for the node to run, without making those dependencies inputs.
+		/// </summary>
+		public BgList<BgFileSet> Fences { get; }
+
+		/// <summary>
+		/// Whether this node should start running as soon as its dependencies are ready, even if other nodes in the same agent are not.
+		/// </summary>
+		public BgBool RunEarly { get; } = BgBool.False;
+
+		/// <summary>
+		/// Labels that this node contributes to
+		/// </summary>
+		public BgList<BgLabel> Labels { get; }
+
+		/// <summary>
+		/// Constructor
+		/// </summary>
+		public BgNode(BgString name, BgMethod method, BgAgent agent, BgList<BgFileSet> inputs, BgList<BgFileSet> fences, BgBool runEarly, BgList<BgLabel> labels)
+			: base(BgExprFlags.ForceFragment)
+		{
+			Name = name;
+
+			ParameterInfo[] parameters = method.Method.GetParameters();
+			ExpressionArguments = new BgExpr[parameters.Length];
+
+			object?[] constantArguments = new object?[method.Arguments.Count];
+			for (int idx = 0; idx < parameters.Length; idx++)
+			{
+				if (typeof(BgExpr).IsAssignableFrom(parameters[idx].ParameterType))
+				{
+					ExpressionArguments[idx] = (BgExpr)method.Arguments[idx]!;
+					constantArguments[idx] = null;
+				}
+				else
+				{
+					ExpressionArguments[idx] = BgBool.False;
+					constantArguments[idx] = method.Arguments[idx];
+				}
+			}
+			Method = new BgMethod(method.Method, constantArguments);
+
+			Agent = agent;
+			DefaultOutput = new BgFileSetFromNodeOutputExpr(this, 0);
+			Inputs = inputs;
+			Fences = fences;
+			RunEarly = runEarly;
+			Labels = labels;
+		}
+
+		int written = 0;
+
+		/// <inheritdoc/>
+		public override void Write(BgBytecodeWriter writer)
+		{
+			writer.WriteOpcode(BgOpcode.Node);
+
+			writer.WriteExpr(Name ?? BgString.Empty);
+			writer.WriteExpr(Agent);
+			writer.WriteMethod(Method);
+
+			writer.WriteUnsignedInteger((uint)ExpressionArguments.Length);
+			foreach (BgExpr argument in ExpressionArguments)
+			{
+				writer.WriteExpr(argument);
+			}
+
+			writer.WriteExpr(Inputs);
+			writer.WriteExpr(Fences);
+			writer.WriteUnsignedInteger(TaggedOutputs.Count);
+			writer.WriteExpr(RunEarly);
+			writer.WriteExpr(Labels);
+			written++;
+		}
+
+		/// <summary>
+		/// Gets the default tag name for the numbered output index
+		/// </summary>
+		/// <param name="name">Name of the node</param>
+		/// <param name="index">Index of the output. Index zero is the default, others are explicit.</param>
+		/// <returns></returns>
+		internal static string GetDefaultTagName(string name, int index)
+		{
+			return $"#{name}${index}";
+		}
+
+		/// <summary>
+		/// Implicit conversion to a fileset
+		/// </summary>
+		/// <param name="node"></param>
+		public static implicit operator BgFileSet(BgNode node)
+		{
+			return new BgFileSetFromNodeExpr(node);
+		}
+
+		/// <summary>
+		/// Implicit conversion to a fileset
+		/// </summary>
+		/// <param name="node"></param>
+		public static implicit operator BgList<BgFileSet>(BgNode node)
+		{
+			return (BgFileSet)node;
+		}
+
+		/// <inheritdoc/>
+		public override BgString ToBgString() => Name ?? BgString.Empty;
+	}
+
+	/// <summary>
+	/// Nodespec with a typed return value
+	/// </summary>
+	/// <typeparam name="T"></typeparam>
+	public class BgNode<T> : BgNode
+	{
+		/// <summary>
+		/// Output from this node
+		/// </summary>
+		public T Output { get; }
+
+		/// <summary>
+		/// Constructor
+		/// </summary>
+		public BgNode(BgString name, BgMethod handler, BgAgent agent, BgList<BgFileSet> inputs, BgList<BgFileSet> fences, BgBool runEarly, BgList<BgLabel> labels)
+			: base(name, handler, agent, inputs, fences, runEarly, labels)
+		{
+			Output = CreateOutput();
+		}
+
+		T CreateOutput()
+		{
+			Type type = typeof(T);
+			if (IsValueTuple(type))
+			{
+				BgExpr[] outputs = CreateOutputExprs(type.GetGenericArguments());
+				TaggedOutputs = outputs;
+				return (T)Activator.CreateInstance(type, outputs)!;
+			}
+			else
+			{
+				BgExpr[] outputs = CreateOutputExprs(new[] { type });
+				TaggedOutputs = outputs;
+				return (T)(object)outputs[0];
+			}
+		}
+
+		BgExpr[] CreateOutputExprs(Type[] types)
+		{
+			BgExpr[] outputs = new BgExpr[types.Length];
+			for (int idx = 0; idx < types.Length; idx++)
+			{
+				outputs[idx] = CreateOutputExpr(types[idx], idx);
+			}
+			return outputs;
+		}
+
+		BgExpr CreateOutputExpr(Type type, int index)
+		{
+			if (type == typeof(BgFileSet))
+			{
+				return new BgFileSetFromNodeOutputExpr(this, index);
+			}
+			else
+			{
+				throw new NotImplementedException();
+			}
+		}
+
+		internal static bool IsValueTuple(Type returnType)
+		{
+			if (returnType.IsGenericType)
+			{
+				Type genericType = returnType.GetGenericTypeDefinition();
+				if (genericType.FullName != null && genericType.FullName.StartsWith("System.ValueTuple`", StringComparison.Ordinal))
+				{
+					return true;
+				}
+			}
+			return false;
+		}
+	}
+
+	/// <summary>
+	/// Class that allows constructing a nodespec
+	/// </summary>
+	public class BgNodeSpecBuilder
+	{
+		/// <inheritdoc cref="BgNode.Name"/>
+		public BgString Name { get; }
+
+		/// <inheritdoc cref="BgNode.Method"/>
+		protected BgMethod Handler { get; }
+
+		/// <inheritdoc cref="BgNode.Agent"/>
+		public BgAgent Agent { get; }
+
+		/// <inheritdoc cref="BgNode.Inputs"/>
+		public BgList<BgFileSet> Inputs { get; set; } = BgList<BgFileSet>.Empty;
+
+		/// <inheritdoc cref="BgNode.Fences"/>
+		public BgList<BgFileSet> Fences { get; set; } = BgList<BgFileSet>.Empty;
+
+		/// <inheritdoc cref="BgNode.RunEarly"/>
+		public BgBool RunEarly { get; set; } = BgBool.False;
+
+		/// <inheritdoc cref="BgNode.Labels"/>
+		public BgList<BgLabel> Labels { get; set; } = BgList<BgLabel>.Empty;
+
+		/// <summary>
+		/// Constructor
+		/// </summary>
+		public BgNodeSpecBuilder(BgString? name, MethodCallExpression call, BgAgent agent)
+		{
+			try
+			{
+				object?[] arguments = CreateArgumentExprs(call);
+
+				Name = name ?? GetDefaultNodeName(call.Method, arguments);
+				Handler = new BgMethod(call.Method, arguments);
+				Agent = agent;
+			}
+			catch (Exception ex)
+			{
+				ExceptionUtils.AddContext(ex, $"while calling method {call.Method}");
+				throw;
+			}
+		}
+
+		static object?[] CreateArgumentExprs(MethodCallExpression call)
+		{
+			object?[] args = new object?[call.Arguments.Count];
+			for (int idx = 0; idx < call.Arguments.Count; idx++)
+			{
+				Expression argumentExpr = call.Arguments[idx];
+				if (argumentExpr is ParameterExpression parameterExpr)
+				{
+					if (parameterExpr.Type != typeof(BgContext))
+					{
+						throw new BgNodeException($"Unable to determine type of parameter '{parameterExpr.Name}'");
+					}
+				}
+				else
+				{
+					Delegate compiled = Expression.Lambda(argumentExpr).Compile();
+					args[idx] = compiled.DynamicInvoke();
+				}
+			}
+			return args;
+		}
+
+		internal static bool IsValueTuple(Type returnType)
+		{
+			if (returnType.IsGenericType)
+			{
+				Type genericType = returnType.GetGenericTypeDefinition();
+				if (genericType.FullName != null && genericType.FullName.StartsWith("System.ValueTuple`", StringComparison.Ordinal))
+				{
+					return true;
+				}
+			}
+			return false;
+		}
+
+		static BgString GetDefaultNodeName(MethodInfo methodInfo, object?[] args)
+		{
+			// Check if it's got an attribute override for the node name
+			BgNodeNameAttribute? nameAttr = methodInfo.GetCustomAttribute<BgNodeNameAttribute>();
+			if (nameAttr != null)
+			{
+				return GetNodeNameFromTemplate(nameAttr.Template, methodInfo.GetParameters(), args);
+			}
+			else
+			{
+				return GetNodeNameFromMethodName(methodInfo.Name);
+			}
+		}
+
+		static BgString GetNodeNameFromTemplate(string template, ParameterInfo[] parameters, object?[] args)
+		{
+			// Create a list of lazily computed string fragments which comprise the evaluated name
+			List<BgString> fragments = new List<BgString>();
+
+			int lastIdx = 0;
+			for (int nextIdx = 0; nextIdx < template.Length; nextIdx++)
+			{
+				if (template[nextIdx] == '{')
+				{
+					if (nextIdx + 1 < template.Length && template[nextIdx + 1] == '{')
+					{
+						fragments.Add(template.Substring(lastIdx, nextIdx - lastIdx));
+						lastIdx = ++nextIdx;
+					}
+					else
+					{
+						fragments.Add(template.Substring(lastIdx, nextIdx - lastIdx));
+						nextIdx++;
+
+						int endIdx = template.IndexOf('}', nextIdx);
+						if (endIdx == -1)
+						{
+							throw new BgNodeException($"Unterminated parameter expression for {nameof(BgNodeNameAttribute)} in {template}");
+						}
+
+						StringView paramName = new StringView(template, nextIdx, endIdx - nextIdx);
+
+						int paramIdx = Array.FindIndex(parameters, x => x.Name != null && paramName.Equals(x.Name, StringComparison.Ordinal));
+						if (paramIdx == -1)
+						{
+							throw new BgNodeException($"Unable to find parameter named {paramName} in {template}");
+						}
+
+						object? arg = args[paramIdx];
+						if (arg != null)
+						{
+							if (arg is BgExpr expr)
+							{
+								fragments.Add(expr.ToBgString());
+							}
+							else
+							{
+								fragments.Add(arg.ToString() ?? String.Empty);
+							}
+						}
+
+						lastIdx = nextIdx = endIdx + 1;
+					}
+				}
+				else if (template[nextIdx] == '}')
+				{
+					if (nextIdx + 1 < template.Length && template[nextIdx + 1] == '{')
+					{
+						fragments.Add(template.Substring(lastIdx, nextIdx - lastIdx));
+						lastIdx = ++nextIdx;
+					}
+				}
+			}
+			fragments.Add(template.Substring(lastIdx, template.Length - lastIdx));
+
+			if (fragments.Count == 1)
+			{
+				return fragments[0];
+			}
+			else
+			{
+				return BgString.Join(BgString.Empty, fragments);
+			}
+		}
+
+		/// <summary>
+		/// Inserts spaces into a PascalCase method name to create a node name
+		/// </summary>
+		public static string GetNodeNameFromMethodName(string methodName)
+		{
+			StringBuilder name = new StringBuilder();
+			name.Append(methodName[0]);
+
+			int length = methodName.Length;
+			if (length > 5 && methodName.EndsWith("Async", StringComparison.Ordinal))
+			{
+				length -= 5;
+			}
+
+			bool bIsAcronym = false;
+			for (int idx = 1; idx < length; idx++)
+			{
+				bool bLastIsUpper = Char.IsUpper(methodName[idx - 1]);
+				bool bNextIsUpper = Char.IsUpper(methodName[idx]);
+				if (bLastIsUpper && bNextIsUpper)
+				{
+					bIsAcronym = true;
+				}
+				else if (bIsAcronym)
+				{
+					name.Insert(name.Length - 2, ' ');
+					bIsAcronym = false;
+				}
+				else if (!bLastIsUpper && bNextIsUpper)
+				{
+					name.Append(' ');
+				}
+				name.Append(methodName[idx]);
+			}
+
+			return name.ToString();
+		}
+
+		/// <summary>
+		/// Construct a nodespec from the current parameters
+		/// </summary>
+		/// <returns>New nodespec</returns>
+		public BgNode Construct()
+		{
+			return new BgNode(Name, Handler, Agent, Inputs, Fences, RunEarly, Labels);
+		}
+	}
+
+	/// <summary>
+	/// Specification for a node to execute that returns one or more tagged outputs
+	/// </summary>
+	/// <typeparam name="T">Return type from the node. Must consist of latent Bg* types (eg. BgToken)</typeparam>
+	public class BgNodeSpecBuilder<T> : BgNodeSpecBuilder
+	{
+		/// <summary>
+		/// Constructor
+		/// </summary>
+		internal BgNodeSpecBuilder(BgString? name, MethodCallExpression call, BgAgent agent)
+			: base(name, call, agent)
+		{
+		}
+
+		/// <summary>
+		/// Construct a nodespec from the current parameters
+		/// </summary>
+		/// <returns>New nodespec</returns>
+		public new BgNode<T> Construct()
+		{
+			return new BgNode<T>(Name, Handler, Agent, Inputs, Fences, RunEarly, Labels);
+		}
+	}
+
+	/// <summary>
+	/// Extension methods for BgNode types
+	/// </summary>
+	public static class BgNodeExtensions
+	{
+		/// <summary>
+		/// Creates a node builder for the given agent
+		/// </summary>
+		/// <param name="agent">Agent to run the node</param>
+		/// <param name="func">Function to execute</param>
+		/// <returns>Node builder</returns>
+		public static BgNodeSpecBuilder AddNode(this BgAgent agent, Expression<Func<BgContext, Task>> func)
+		{
+			MethodCallExpression call = (MethodCallExpression)func.Body;
+			return new BgNodeSpecBuilder(null, call, agent);
+		}
+
+		/// <summary>
+		/// Creates a node builder for the given agent
+		/// </summary>
+		/// <param name="agent">Agent to run the node</param>
+		/// <param name="func">Function to execute</param>
+		/// <returns>Node builder</returns>
+		public static BgNodeSpecBuilder<T> AddNode<T>(this BgAgent agent, Expression<Func<BgContext, Task<T>>> func)
+		{
+			MethodCallExpression call = (MethodCallExpression)func.Body;
+			return new BgNodeSpecBuilder<T>(null, call, agent);
+		}
+
+		/// <summary>
+		/// Add dependencies onto other nodes or outputs. Outputs from the given tokens will be copied to the current machine before execution of the node.
+		/// </summary>
+		/// <param name="builder">The node builder</param>
+		/// <param name="inputs">Files to add as inputs</param>
+		/// <returns>The current node spec, to allow chaining calls</returns>
+		public static T Requires<T>(this T builder, params BgNode[] inputs) where T : BgNodeSpecBuilder
+		{
+			builder.Inputs = builder.Inputs.Add(inputs.Select(x => (BgFileSet)x));
+			return builder;
+		}
+
+		/// <summary>
+		/// Add dependencies onto other nodes or outputs. Outputs from the given tokens will be copied to the current machine before execution of the node.
+		/// </summary>
+		/// <param name="builder">The node builder</param>
+		/// <param name="inputs">Files to add as inputs</param>
+		/// <returns>The current node spec, to allow chaining calls</returns>
+		public static T Requires<T>(this T builder, params BgFileSet[] inputs) where T : BgNodeSpecBuilder
+		{
+			builder.Inputs = builder.Inputs.Add(inputs);
+			return builder;
+		}
+
+		/// <summary>
+		/// Add dependencies onto other nodes or outputs. Outputs from the given tokens will be copied to the current machine before execution of the node.
+		/// </summary>
+		/// <param name="builder">The node builder</param>
+		/// <param name="inputs">Files to add as inputs</param>
+		/// <returns>The current node spec, to allow chaining calls</returns>
+		public static T Requires<T>(this T builder, BgList<BgFileSet> inputs) where T : BgNodeSpecBuilder
+		{
+			builder.Inputs = builder.Inputs.Add(inputs);
+			return builder;
+		}
+
+		/// <summary>
+		/// Add weak dependencies onto other nodes or outputs. The producing nodes must complete successfully if they are part of the graph, but outputs from them will not be 
+		/// transferred to the machine running this node.
+		/// </summary>
+		/// <param name="builder">The node builder</param>
+		/// <param name="inputs">Files to add as inputs</param>
+		/// <returns>The current node spec, to allow chaining calls</returns>
+		public static T After<T>(this T builder, params BgFileSet[] inputs) where T : BgNodeSpecBuilder
+		{
+			builder.Fences = builder.Fences.Add(inputs);
+			return builder;
+		}
+	}
+}
