BuildGraph: Move option classes alongside their underlying types, and fix up a lot of namespaces.

diff --git a/Engine/Source/Programs/AutomationTool/BuildGraph/BgNodeExecutor.cs b/Engine/Source/Programs/AutomationTool/BuildGraph/BgNodeExecutor.cs
--- a/Engine/Source/Programs/AutomationTool/BuildGraph/BgNodeExecutor.cs
+++ b/Engine/Source/Programs/AutomationTool/BuildGraph/BgNodeExecutor.cs
@@ -65,9 +65,9 @@ namespace AutomationTool
 			public override bool IsBuildMachine => CommandUtils.IsBuildMachine;
 		}
 
-		readonly BgBytecodeNode Node;
+		readonly BgExpressionNodeDef Node;
 
-		public BgBytecodeNodeExecutor(BgBytecodeNode node)
+		public BgBytecodeNodeExecutor(BgExpressionNodeDef node)
 		{
 			Node = node;
 		}
diff --git a/Engine/Source/Programs/AutomationTool/BuildGraph/BuildGraph.cs b/Engine/Source/Programs/AutomationTool/BuildGraph/BuildGraph.cs
--- a/Engine/Source/Programs/AutomationTool/BuildGraph/BuildGraph.cs
+++ b/Engine/Source/Programs/AutomationTool/BuildGraph/BuildGraph.cs
@@ -709,7 +709,7 @@ namespace AutomationTool
 				NodeToExecutor[Node] = executor;
 				return executor.Bind(NameToTask, TagNameToNodeOutput, Logger);
 			}
-			else if (Node is BgBytecodeNode BytecodeNode)
+			else if (Node is BgExpressionNodeDef BytecodeNode)
 			{
 				BgBytecodeNodeExecutor executor = new BgBytecodeNodeExecutor(BytecodeNode);
 				NodeToExecutor[Node] = executor;
diff --git a/Engine/Source/Programs/AutomationTool/BuildGraph/Scripts/InstalledBuild.cs b/Engine/Source/Programs/AutomationTool/BuildGraph/Scripts/InstalledBuild.cs
--- a/Engine/Source/Programs/AutomationTool/BuildGraph/Scripts/InstalledBuild.cs
+++ b/Engine/Source/Programs/AutomationTool/BuildGraph/Scripts/InstalledBuild.cs
@@ -115,7 +115,7 @@ namespace AutomationTool
 			BgBoolOption HostPlatformDDCOnly = new BgBoolOption("HostPlatformDDCOnly", "Whether to include DDC for the host platform only", true);
 			BgBoolOption SignExecutables = new BgBoolOption("SignExecutables", "Sign the executables produced where signing is available", false);
 
-			BgStrOption AnalyticsTypeOverride = new BgStrOption("AnalyticsTypeOverride", "Identifier for analytic events to send", "");
+			BgStringOption AnalyticsTypeOverride = new BgStringOption("AnalyticsTypeOverride", "Identifier for analytic events to send", "");
 			BgBool EmbedSrcSrvInfo = new BgBoolOption("EmbedSrcSrvInfo", "Whether to add Source indexing to Windows game apps so they can be added to a symbol server", false);
 
 			BgList<BgString> DefaultGameConfigurations = BgList<BgString>.Create(nameof(UnrealTargetConfiguration.DebugGame), nameof(UnrealTargetConfiguration.Development), nameof(UnrealTargetConfiguration.Shipping));
@@ -124,10 +124,10 @@ namespace AutomationTool
 
 			BgBoolOption WithFullDebugInfo = new BgBoolOption("WithFullDebugInfo", "Generate full debug info for binary editor and packaged application builds", false);
 
-			BgStrOption BuiltDirectory = new BgStrOption("BuiltDirectory", "Directory for outputting the built engine", RootDir + "/LocalBuilds/Engine");
+			BgStringOption BuiltDirectory = new BgStringOption("BuiltDirectory", "Directory for outputting the built engine", RootDir + "/LocalBuilds/Engine");
 
-			BgStrOption CrashReporterAPIURL = new BgStrOption("CrashReporterAPIURL", "The URL to use to talk to the CrashReporterClient API.", "");
-			BgStrOption CrashReporterAPIKey = new BgStrOption("CrashReporterAPIKey", "The API key to use to talk to the CrashReporterClient API.", "");
+			BgStringOption CrashReporterAPIURL = new BgStringOption("CrashReporterAPIURL", "The URL to use to talk to the CrashReporterClient API.", "");
+			BgStringOption CrashReporterAPIKey = new BgStringOption("CrashReporterAPIKey", "The API key to use to talk to the CrashReporterClient API.", "");
 
 			BgString CrashReporterCompileArgs = "";
 			CrashReporterCompileArgs = CrashReporterCompileArgs.If(CrashReporterAPIURL != "" & CrashReporterAPIKey != "", BgString.Format("-define:CRC_TELEMETRY_URL=\"{0}\" -define:CRC_TELEMETRY_KEY_DEV=\"{1}\" -define:CRC_TELEMETRY_KEY_RELEASE=\"{1}\" -OverrideBuildEnvironment", CrashReporterAPIURL, CrashReporterAPIKey));
diff --git a/Engine/Source/Programs/Shared/EpicGames.BuildGraph/BgBytecodeWriter.cs b/Engine/Source/Programs/Shared/EpicGames.BuildGraph/BgBytecodeWriter.cs
--- a/Engine/Source/Programs/Shared/EpicGames.BuildGraph/BgBytecodeWriter.cs
+++ b/Engine/Source/Programs/Shared/EpicGames.BuildGraph/BgBytecodeWriter.cs
@@ -10,6 +10,7 @@ using System.Runtime.CompilerServices;
 using System.Text;
 using EpicGames.Core;
 using EpicGames.Serialization;
+using EpicGames.BuildGraph.Expressions;
 
 namespace EpicGames.BuildGraph
 {
diff --git a/Engine/Source/Programs/Shared/EpicGames.BuildGraph/BgInterpreter.cs b/Engine/Source/Programs/Shared/EpicGames.BuildGraph/BgInterpreter.cs
--- a/Engine/Source/Programs/Shared/EpicGames.BuildGraph/BgInterpreter.cs
+++ b/Engine/Source/Programs/Shared/EpicGames.BuildGraph/BgInterpreter.cs
@@ -459,12 +459,12 @@ namespace EpicGames.BuildGraph
 
 				case BgOpcode.FileSetFromNode:
 					{
-						BgBytecodeNode node = (BgBytecodeNode)Evaluate(frame);
+						BgExpressionNodeDef node = (BgExpressionNodeDef)Evaluate(frame);
 						return node.InputDependencies.SelectMany(x => x.Outputs).Concat(node.Outputs).ToArray();
 					}
 				case BgOpcode.FileSetFromNodeOutput:
 					{
-						BgBytecodeNode node = (BgBytecodeNode)Evaluate(frame);
+						BgExpressionNodeDef node = (BgExpressionNodeDef)Evaluate(frame);
 						int outputIndex = (int)ReadUnsignedInteger(frame);
 						return new[] { node.Outputs[outputIndex] };
 					}
@@ -579,7 +579,7 @@ namespace EpicGames.BuildGraph
 
 						string[] outputNames = Enumerable.Range(0, taggedOutputCount).Select(x => BgNode.GetDefaultTagName(name, x)).ToArray();
 
-						BgBytecodeNode node = new BgBytecodeNode(agent, name, method, arguments, inputs, outputNames, inputDependencies.ToArray(), orderDependencies.ToArray(), Array.Empty<FileReference>(), labels);
+						BgExpressionNodeDef node = new BgExpressionNodeDef(agent, name, method, arguments, inputs, outputNames, inputDependencies.ToArray(), orderDependencies.ToArray(), Array.Empty<FileReference>(), labels);
 						return node;
 					}
 				case BgOpcode.Agent:
@@ -614,25 +614,25 @@ namespace EpicGames.BuildGraph
 					}
 				case BgOpcode.Graph:
 					{
-						List<BgBytecodeNode> nodes = EvaluateList<BgBytecodeNode>(frame);
+						List<BgExpressionNodeDef> nodes = EvaluateList<BgExpressionNodeDef>(frame);
 						List<BgBytecodeAggregateDef> aggregates = EvaluateList<BgBytecodeAggregateDef>(frame);
 
 						BgGraphDef graph = new BgGraphDef();
 						foreach (BgBytecodeAggregateDef aggregate in aggregates)
 						{
 							graph.NameToAggregate[aggregate.Name] = aggregate;
-							nodes.AddRange(aggregate.RequiredNodes.Select(x => (BgBytecodeNode)x));
+							nodes.AddRange(aggregate.RequiredNodes.Select(x => (BgExpressionNodeDef)x));
 						}
 
 						HashSet<BgNodeDef> uniqueNodes = new HashSet<BgNodeDef>();
 						HashSet<BgAgentDef> uniqueAgents = new HashSet<BgAgentDef>();
-						foreach (BgBytecodeNode node in nodes)
+						foreach (BgExpressionNodeDef node in nodes)
 						{
 							RegisterNode(graph, node, uniqueNodes, uniqueAgents);
 						}
 
 						HashSet<BgLabelDef> labels = new HashSet<BgLabelDef>();
-						foreach (BgBytecodeNode node in nodes)
+						foreach (BgExpressionNodeDef node in nodes)
 						{
 							foreach (BgLabelDef label in node.Labels)
 							{
@@ -659,11 +659,11 @@ namespace EpicGames.BuildGraph
 			}
 		}
 
-		void RegisterNode(BgGraphDef graph, BgBytecodeNode node, HashSet<BgNodeDef> uniqueNodes, HashSet<BgAgentDef> uniqueAgents)
+		void RegisterNode(BgGraphDef graph, BgExpressionNodeDef node, HashSet<BgNodeDef> uniqueNodes, HashSet<BgAgentDef> uniqueAgents)
 		{
 			if (uniqueNodes.Add(node))
 			{
-				foreach (BgBytecodeNode inputNode in node.InputDependencies)
+				foreach (BgExpressionNodeDef inputNode in node.InputDependencies)
 				{
 					RegisterNode(graph, inputNode, uniqueNodes, uniqueAgents);
 				}
diff --git a/Engine/Source/Programs/Shared/EpicGames.BuildGraph/BgNodeDef.cs b/Engine/Source/Programs/Shared/EpicGames.BuildGraph/BgNodeDef.cs
--- a/Engine/Source/Programs/Shared/EpicGames.BuildGraph/BgNodeDef.cs
+++ b/Engine/Source/Programs/Shared/EpicGames.BuildGraph/BgNodeDef.cs
@@ -8,6 +8,7 @@ using System.Reflection;
 using System.Threading.Tasks;
 using System.Xml;
 using EpicGames.Core;
+using EpicGames.BuildGraph.Expressions;
 
 namespace EpicGames.BuildGraph
 {
@@ -182,7 +183,7 @@ namespace EpicGames.BuildGraph
 	/// <summary>
 	/// Node constructed from a bytecode expression
 	/// </summary>
-	public class BgBytecodeNode : BgNodeDef
+	public class BgExpressionNodeDef : BgNodeDef
 	{
 		/// <summary>
 		/// Agent declaring this node
@@ -207,7 +208,7 @@ namespace EpicGames.BuildGraph
 		/// <summary>
 		/// Constructor
 		/// </summary>
-		public BgBytecodeNode(BgAgentDef agent, string name, BgMethod method, IReadOnlyList<object> arguments, IReadOnlyList<BgNodeOutput> inputs, IReadOnlyList<string> outputNames, IReadOnlyList<BgNodeDef> inputDependencies, IReadOnlyList<BgNodeDef> orderDependencies, IReadOnlyList<FileReference> requiredTokens, IReadOnlyList<BgLabelDef> labels)
+		public BgExpressionNodeDef(BgAgentDef agent, string name, BgMethod method, IReadOnlyList<object> arguments, IReadOnlyList<BgNodeOutput> inputs, IReadOnlyList<string> outputNames, IReadOnlyList<BgNodeDef> inputDependencies, IReadOnlyList<BgNodeDef> orderDependencies, IReadOnlyList<FileReference> requiredTokens, IReadOnlyList<BgLabelDef> labels)
 			: base(name, inputs, outputNames, inputDependencies, orderDependencies, requiredTokens)
 		{
 			Agent = agent;
diff --git a/Engine/Source/Programs/Shared/EpicGames.BuildGraph/BgOpcode.cs b/Engine/Source/Programs/Shared/EpicGames.BuildGraph/BgOpcode.cs
--- a/Engine/Source/Programs/Shared/EpicGames.BuildGraph/BgOpcode.cs
+++ b/Engine/Source/Programs/Shared/EpicGames.BuildGraph/BgOpcode.cs
@@ -2,6 +2,7 @@
 
 using System;
 using EpicGames.Serialization;
+using EpicGames.BuildGraph.Expressions;
 
 namespace EpicGames.BuildGraph
 {
@@ -166,7 +167,7 @@ namespace EpicGames.BuildGraph
 		StrReplace = 0x28,
 
 		/// <summary>
-		/// A string option (see <see cref="BgStrOption"/>)
+		/// A string option (see <see cref="BgStringOption"/>)
 		/// </summary>
 		StrOption = 0x29,
 
diff --git a/Engine/Source/Programs/Shared/EpicGames.BuildGraph/BgOptionSpec.cs b/Engine/Source/Programs/Shared/EpicGames.BuildGraph/BgOptionSpec.cs
deleted file mode 100644
diff --git a/Engine/Source/Programs/Shared/EpicGames.BuildGraph/Expressions/BgAgent.cs b/Engine/Source/Programs/Shared/EpicGames.BuildGraph/Expressions/BgAgent.cs
--- a/Engine/Source/Programs/Shared/EpicGames.BuildGraph/Expressions/BgAgent.cs
+++ b/Engine/Source/Programs/Shared/EpicGames.BuildGraph/Expressions/BgAgent.cs
@@ -8,7 +8,7 @@ using EpicGames.BuildGraph.Expressions;
 using EpicGames.Core;
 using Microsoft.Extensions.Logging;
 
-namespace EpicGames.BuildGraph
+namespace EpicGames.BuildGraph.Expressions
 {
 	/// <summary>
 	/// Describes an agent that can execute execute build steps
diff --git a/Engine/Source/Programs/Shared/EpicGames.BuildGraph/Expressions/BgAggregate.cs b/Engine/Source/Programs/Shared/EpicGames.BuildGraph/Expressions/BgAggregate.cs
--- a/Engine/Source/Programs/Shared/EpicGames.BuildGraph/Expressions/BgAggregate.cs
+++ b/Engine/Source/Programs/Shared/EpicGames.BuildGraph/Expressions/BgAggregate.cs
@@ -1,9 +1,8 @@
 // Copyright Epic Games, Inc. All Rights Reserved.
 
 using System;
-using EpicGames.BuildGraph.Expressions;
 
-namespace EpicGames.BuildGraph
+namespace EpicGames.BuildGraph.Expressions
 {
 	/// <summary>
 	/// Specification for an aggregate target in the graph
diff --git a/Engine/Source/Programs/Shared/EpicGames.BuildGraph/Expressions/BgBool.cs b/Engine/Source/Programs/Shared/EpicGames.BuildGraph/Expressions/BgBool.cs
--- a/Engine/Source/Programs/Shared/EpicGames.BuildGraph/Expressions/BgBool.cs
+++ b/Engine/Source/Programs/Shared/EpicGames.BuildGraph/Expressions/BgBool.cs
@@ -164,4 +164,60 @@ namespace EpicGames.BuildGraph.Expressions
 	}
 
 	#endregion
+
+	/// <summary>
+	/// A boolean option expression
+	/// </summary>
+	public class BgBoolOption : BgBool
+	{
+		/// <summary>
+		/// Name of the option
+		/// </summary>
+		public BgString Name { get; }
+
+		/// <summary>
+		/// Label to display next to the option
+		/// </summary>
+		public BgString? Label { get; }
+
+		/// <summary>
+		/// Help text to display for the user
+		/// </summary>
+		public BgString? Description { get; }
+
+		/// <summary>
+		/// Default value for the option
+		/// </summary>
+		public BgBool? DefaultValue { get; }
+
+		/// <summary>
+		/// Constructor
+		/// </summary>
+		public BgBoolOption(BgString name, BgString? description = null, BgBool? defaultValue = null)
+			: this(name, null, description, defaultValue)
+		{
+		}
+
+		/// <summary>
+		/// Constructor
+		/// </summary>
+		public BgBoolOption(BgString name, BgString? label, BgString? description, BgBool? defaultValue)
+			: base(BgExprFlags.None)
+		{
+			Name = name;
+			Label = label;
+			Description = description;
+			DefaultValue = defaultValue;
+		}
+
+		/// <inheritdoc/>
+		public override void Write(BgBytecodeWriter writer)
+		{
+			writer.WriteOpcode(BgOpcode.BoolOption);
+			writer.WriteExpr(Name);
+			writer.WriteExpr(Label ?? BgString.Empty);
+			writer.WriteExpr(Description ?? BgString.Empty);
+			writer.WriteExpr(DefaultValue ?? BgBool.False);
+		}
+	}
 }
diff --git a/Engine/Source/Programs/Shared/EpicGames.BuildGraph/Expressions/BgGraph.cs b/Engine/Source/Programs/Shared/EpicGames.BuildGraph/Expressions/BgGraph.cs
--- a/Engine/Source/Programs/Shared/EpicGames.BuildGraph/Expressions/BgGraph.cs
+++ b/Engine/Source/Programs/Shared/EpicGames.BuildGraph/Expressions/BgGraph.cs
@@ -2,10 +2,9 @@
 
 using System;
 using System.Collections.Generic;
-using EpicGames.BuildGraph.Expressions;
 using Microsoft.Extensions.Logging;
 
-namespace EpicGames.BuildGraph
+namespace EpicGames.BuildGraph.Expressions
 {
 	/// <summary>
 	/// Specification for a graph in fluent syntax
diff --git a/Engine/Source/Programs/Shared/EpicGames.BuildGraph/Expressions/BgInt.cs b/Engine/Source/Programs/Shared/EpicGames.BuildGraph/Expressions/BgInt.cs
--- a/Engine/Source/Programs/Shared/EpicGames.BuildGraph/Expressions/BgInt.cs
+++ b/Engine/Source/Programs/Shared/EpicGames.BuildGraph/Expressions/BgInt.cs
@@ -226,4 +226,66 @@ namespace EpicGames.BuildGraph.Expressions
 	}
 
 	#endregion
+
+	/// <summary>
+	/// An integer option expression
+	/// </summary>
+	public class BgIntOption : BgInt
+	{
+		/// <summary>
+		/// Name of the option
+		/// </summary>
+		public BgString Name { get; }
+
+		/// <summary>
+		/// Label to display next to the option
+		/// </summary>
+		public BgString? Label { get; }
+
+		/// <summary>
+		/// Help text to display for the user
+		/// </summary>
+		public BgString? Description { get; }
+
+		/// <summary>
+		/// Default value for the option
+		/// </summary>
+		public BgInt? DefaultValue { get; }
+
+		/// <summary>
+		/// Minimum allowed value
+		/// </summary>
+		public BgInt? MinValue { get; }
+
+		/// <summary>
+		/// Maximum allowed value
+		/// </summary>
+		public BgInt? MaxValue { get; }
+
+		/// <summary>
+		/// Constructor
+		/// </summary>
+		public BgIntOption(string name, BgString? description = null, BgInt? defaultValue = null, BgInt? minValue = null, BgInt? maxValue = null, BgString? label = null)
+			: base(BgExprFlags.None)
+		{
+			Name = name;
+			Label = label;
+			Description = description;
+			DefaultValue = defaultValue;
+			MinValue = minValue;
+			MaxValue = maxValue;
+		}
+
+		/// <inheritdoc/>
+		public override void Write(BgBytecodeWriter writer)
+		{
+			writer.WriteOpcode(BgOpcode.IntOption);
+			writer.WriteExpr(Name);
+			writer.WriteExpr(Label ?? BgString.Empty);
+			writer.WriteExpr(Description ?? BgString.Empty);
+			writer.WriteExpr(DefaultValue ?? (BgInt)0);
+			writer.WriteExpr(MinValue ?? (BgInt)(-1));
+			writer.WriteExpr(MaxValue ?? (BgInt)(-1));
+		}
+	}
 }
diff --git a/Engine/Source/Programs/Shared/EpicGames.BuildGraph/Expressions/BgLabel.cs b/Engine/Source/Programs/Shared/EpicGames.BuildGraph/Expressions/BgLabel.cs
--- a/Engine/Source/Programs/Shared/EpicGames.BuildGraph/Expressions/BgLabel.cs
+++ b/Engine/Source/Programs/Shared/EpicGames.BuildGraph/Expressions/BgLabel.cs
@@ -2,9 +2,8 @@
 
 using System;
 using System.Linq;
-using EpicGames.BuildGraph.Expressions;
 
-namespace EpicGames.BuildGraph
+namespace EpicGames.BuildGraph.Expressions
 {
 	/// <summary>
 	/// Specification for a label
diff --git a/Engine/Source/Programs/Shared/EpicGames.BuildGraph/Expressions/BgList.cs b/Engine/Source/Programs/Shared/EpicGames.BuildGraph/Expressions/BgList.cs
--- a/Engine/Source/Programs/Shared/EpicGames.BuildGraph/Expressions/BgList.cs
+++ b/Engine/Source/Programs/Shared/EpicGames.BuildGraph/Expressions/BgList.cs
@@ -444,4 +444,103 @@ namespace EpicGames.BuildGraph.Expressions
 	}
 
 	#endregion
+
+	/// <summary>
+	/// Style for a list option
+	/// </summary>
+	public enum BgListOptionStyle
+	{
+		/// <summary>
+		/// List of checkboxes
+		/// </summary>
+		CheckList = 0,
+
+		/// <summary>
+		/// Tag picker
+		/// </summary>
+		TagPicker = 1,
+	}
+
+	/// <summary>
+	/// A list option expression
+	/// </summary>
+	public class BgListOptionSpec : BgList<BgString>
+	{
+		/// <summary>
+		/// Name of the option
+		/// </summary>
+		public BgString Name { get; }
+
+		/// <summary>
+		/// Label to display next to the option
+		/// </summary>
+		public BgString? Label { get; }
+
+		/// <summary>
+		/// Help text to display for the user
+		/// </summary>
+		public BgString? Description { get; }
+
+		/// <summary>
+		/// Style for this list box
+		/// </summary>
+		public BgListOptionStyle Style { get; }
+
+		/// <summary>
+		/// Default value for the option
+		/// </summary>
+		public BgString? DefaultValue { get; set; }
+
+		/// <summary>
+		/// Regex for validating values for the option
+		/// </summary>
+		public BgString? Pattern { get; set; }
+
+		/// <summary>
+		/// Message to display if validation fails
+		/// </summary>
+		public BgString? PatternFailed { get; set; }
+
+		/// <summary>
+		/// List of values to choose from
+		/// </summary>
+		public BgList<BgString>? Values { get; set; }
+
+		/// <summary>
+		/// Matching list of descriptions for each value
+		/// </summary>
+		public BgList<BgString>? ValueDescriptions { get; set; }
+
+		/// <summary>
+		/// Constructor
+		/// </summary>
+		public BgListOptionSpec(string name, BgString? label = null, BgString? description = null, BgString? defaultValue = null, BgListOptionStyle style = BgListOptionStyle.CheckList, BgString? pattern = null, BgString? patternFailed = null, BgList<BgString>? values = null, BgList<BgString>? valueDescriptions = null)
+			: base(BgExprFlags.None)
+		{
+			Name = name;
+			Label = label;
+			Description = description;
+			DefaultValue = defaultValue;
+			Style = style;
+			Pattern = pattern;
+			PatternFailed = patternFailed;
+			Values = values;
+			ValueDescriptions = valueDescriptions;
+		}
+
+		/// <inheritdoc/>
+		public override void Write(BgBytecodeWriter writer)
+		{
+			writer.WriteOpcode(BgOpcode.ListOption);
+			writer.WriteExpr(Name);
+			writer.WriteExpr(Label ?? BgString.Empty);
+			writer.WriteExpr(Description ?? BgString.Empty);
+			writer.WriteUnsignedInteger((int)Style);
+			writer.WriteExpr(DefaultValue ?? BgString.Empty);
+			writer.WriteExpr(Pattern ?? BgString.Empty);
+			writer.WriteExpr(PatternFailed ?? BgString.Empty);
+			writer.WriteExpr(Values ?? BgList<BgString>.Empty);
+			writer.WriteExpr(ValueDescriptions ?? BgList<BgString>.Empty);
+		}
+	}
 }
diff --git a/Engine/Source/Programs/Shared/EpicGames.BuildGraph/Expressions/BgNode.cs b/Engine/Source/Programs/Shared/EpicGames.BuildGraph/Expressions/BgNode.cs
--- a/Engine/Source/Programs/Shared/EpicGames.BuildGraph/Expressions/BgNode.cs
+++ b/Engine/Source/Programs/Shared/EpicGames.BuildGraph/Expressions/BgNode.cs
@@ -7,10 +7,9 @@ using System.Linq.Expressions;
 using System.Reflection;
 using System.Text;
 using System.Threading.Tasks;
-using EpicGames.BuildGraph.Expressions;
 using EpicGames.Core;
 
-namespace EpicGames.BuildGraph
+namespace EpicGames.BuildGraph.Expressions
 {
 	/// <summary>
 	/// Exception for constructing nodes
diff --git a/Engine/Source/Programs/Shared/EpicGames.BuildGraph/Expressions/BgString.cs b/Engine/Source/Programs/Shared/EpicGames.BuildGraph/Expressions/BgString.cs
--- a/Engine/Source/Programs/Shared/EpicGames.BuildGraph/Expressions/BgString.cs
+++ b/Engine/Source/Programs/Shared/EpicGames.BuildGraph/Expressions/BgString.cs
@@ -296,4 +296,103 @@ namespace EpicGames.BuildGraph.Expressions
 	}
 
 	#endregion
+
+	/// <summary>
+	/// Style for a string option
+	/// </summary>
+	public enum BgStringOptionStyle
+	{
+		/// <summary>
+		/// Free-form text entry
+		/// </summary>
+		Text,
+
+		/// <summary>
+		/// List of options
+		/// </summary>
+		DropList,
+	}
+
+	/// <summary>
+	/// A string option expression
+	/// </summary>
+	public class BgStringOption : BgString
+	{
+		/// <summary>
+		/// Name of the option
+		/// </summary>
+		public BgString Name { get; }
+
+		/// <summary>
+		/// Label to display next to the option
+		/// </summary>
+		public BgString? Label { get; }
+
+		/// <summary>
+		/// Help text to display for the user
+		/// </summary>
+		public BgString? Description { get; }
+
+		/// <summary>
+		/// Default value for the option
+		/// </summary>
+		public BgString? DefaultValue { get; set; }
+
+		/// <summary>
+		/// Style for this option
+		/// </summary>
+		public BgStringOptionStyle Style { get; }
+
+		/// <summary>
+		/// Regex for validating values for the option
+		/// </summary>
+		public BgString? Pattern { get; set; }
+
+		/// <summary>
+		/// Message to display if validation fails
+		/// </summary>
+		public BgString? PatternFailed { get; set; }
+
+		/// <summary>
+		/// List of values to choose from
+		/// </summary>
+		public BgList<BgString>? Values { get; set; }
+
+		/// <summary>
+		/// Matching list of descriptions for each value
+		/// </summary>
+		public BgList<BgString>? ValueDescriptions { get; set; }
+
+		/// <summary>
+		/// Constructor
+		/// </summary>
+		public BgStringOption(string name, BgString? label = null, BgString? description = null, BgString? defaultValue = null, BgStringOptionStyle style = BgStringOptionStyle.Text, BgString? pattern = null, BgString? patternFailed = null, BgList<BgString>? values = null, BgList<BgString>? valueDescriptions = null)
+			: base(BgExprFlags.None)
+		{
+			Name = name;
+			Label = label;
+			Description = description;
+			Style = style;
+			DefaultValue = defaultValue;
+			Pattern = pattern;
+			PatternFailed = patternFailed;
+			Values = values;
+			ValueDescriptions = valueDescriptions;
+		}
+
+		/// <inheritdoc/>
+		public override void Write(BgBytecodeWriter writer)
+		{
+			writer.WriteOpcode(BgOpcode.StrOption);
+			writer.WriteExpr(Name);
+			writer.WriteExpr(Label ?? BgString.Empty);
+			writer.WriteExpr(Description ?? BgString.Empty);
+			writer.WriteUnsignedInteger((int)Style);
+			writer.WriteExpr(DefaultValue ?? BgString.Empty);
+			writer.WriteExpr(Pattern ?? BgString.Empty);
+			writer.WriteExpr(PatternFailed ?? BgString.Empty);
+			writer.WriteExpr(Values ?? BgList<BgString>.Empty);
+			writer.WriteExpr(ValueDescriptions ?? BgList<BgString>.Empty);
+		}
+	}
 }
