BuildGraph: Fix path matching with UNC paths.

diff --git a/Engine/Source/Programs/Shared/EpicGames.Core/FileFilter.cs b/Engine/Source/Programs/Shared/EpicGames.Core/FileFilter.cs
--- a/Engine/Source/Programs/Shared/EpicGames.Core/FileFilter.cs
+++ b/Engine/Source/Programs/Shared/EpicGames.Core/FileFilter.cs
@@ -4,6 +4,7 @@ using System;
 using System.Collections.Generic;
 using System.IO;
 using System.Linq;
+using System.Text.RegularExpressions;
 
 namespace EpicGames.Core
 {
@@ -260,6 +261,24 @@ namespace EpicGames.Core
 			}
 		}
 
+		/// <summary>
+		/// Normalizes a path to use forward slashes as a path separator
+		/// </summary>
+		/// <param name="path">Path to process</param>
+		/// <returns>Normalized path</returns>
+		static string NormalizePath(string path)
+		{
+			Match match = Regex.Match(path, @"^\\\\[^\\/]+\\[^\\/]+");
+			if (match.Success)
+			{
+				return match.Groups[0].Value + path.Substring(match.Groups[0].Length).Replace('\\', '/');
+			}
+			else
+			{
+				return path.Replace('\\', '/');
+			}
+		}
+
 		/// <summary>
 		/// Adds an include or exclude rule to the filter
 		/// </summary>
@@ -267,7 +286,7 @@ namespace EpicGames.Core
 		/// <param name="type">Whether to include or exclude files matching this rule</param>
 		public void AddRule(string pattern, FileFilterType type)
 		{
-			string normalizedPattern = pattern.Replace('\\', '/');
+			string normalizedPattern = NormalizePath(pattern);
 
 			// Remove the slash from the start of the pattern. Any exclude pattern that doesn't contain a directory separator is assumed to apply to any directory (eg. *.cpp), otherwise it's 
 			// taken relative to the root.
@@ -346,7 +365,7 @@ namespace EpicGames.Core
 		/// <returns>True if the file passes the filter</returns>
 		public bool Matches(string fileName)
 		{
-			string[] tokens = fileName.TrimStart('/', '\\').Split('/', '\\');
+			string[] tokens = NormalizePath(fileName).TrimStart('/').Split('/');
 
 			FileFilterNode matchingNode = FindMatchingNode(_rootNode, tokens, 0, _defaultNode);
 
@@ -360,7 +379,7 @@ namespace EpicGames.Core
 		/// <returns>True if the file passes the filter</returns>
 		public bool PossiblyMatches(string folderName)
 		{
-			string[] tokens = folderName.Trim('/', '\\').Split('/', '\\');
+			string[] tokens = NormalizePath(folderName).Trim('/').Split('/');
 
 			FileFilterNode matchingNode = FindMatchingNode(_rootNode, tokens.Union(new string[]{ "" }).ToArray(), 0, _defaultNode);
 
@@ -416,7 +435,7 @@ namespace EpicGames.Core
 			List<FileReference> matchingFileNames = new List<FileReference>();
 			if (DirectoryReference.Exists(directoryName))
 			{
-				FindMatchesFromDirectory(new DirectoryInfo(directoryName.FullName), prefixPath.Replace('\\', '/').TrimEnd('/') + "/", bIgnoreSymlinks, matchingFileNames);
+				FindMatchesFromDirectory(new DirectoryInfo(directoryName.FullName), NormalizePath(prefixPath).TrimEnd('/') + "/", bIgnoreSymlinks, matchingFileNames);
 			}
 			return matchingFileNames;
 		}
