BuildGraph: Add handler for int options.

diff --git a/Engine/Source/Programs/Shared/EpicGames.BuildGraph/BgInterpreter.cs b/Engine/Source/Programs/Shared/EpicGames.BuildGraph/BgInterpreter.cs
--- a/Engine/Source/Programs/Shared/EpicGames.BuildGraph/BgInterpreter.cs
+++ b/Engine/Source/Programs/Shared/EpicGames.BuildGraph/BgInterpreter.cs
@@ -234,6 +234,23 @@ namespace EpicGames.BuildGraph
 					}
 				case BgOpcode.IntNegate:
 					return -(int)Evaluate(frame);
+				case BgOpcode.IntOption:
+					{
+						string name = (string)Evaluate(frame);
+						string label = (string)Evaluate(frame);
+						string description = (string)Evaluate(frame);
+						int defaultValue = (int)Evaluate(frame);
+						int minValue = (int)Evaluate(frame);
+						int maxValue = (int)Evaluate(frame);
+
+						int value = defaultValue;
+						if (_options.TryGetValue(name, out string? str))
+						{
+							value = Int32.Parse(str);
+						}
+
+						return value;
+					}
 
 				#endregion
 				#region String opcodes
@@ -866,6 +883,14 @@ namespace EpicGames.BuildGraph
 				case BgOpcode.IntNegate:
 					Disassemble(frame, logger);
 					break;
+				case BgOpcode.IntOption:
+					Disassemble(frame, logger);
+					Disassemble(frame, logger);
+					Disassemble(frame, logger);
+					Disassemble(frame, logger);
+					Disassemble(frame, logger);
+					Disassemble(frame, logger);
+					break;
 
 				#endregion
 				#region String opcodes
