EpicGames.OIDC - Made sure the same provider can be read from different applications to allows us to share state.

diff --git a/Engine/Source/Programs/Shared/EpicGames.OIDC/OidcTokenManager.cs b/Engine/Source/Programs/Shared/EpicGames.OIDC/OidcTokenManager.cs
--- a/Engine/Source/Programs/Shared/EpicGames.OIDC/OidcTokenManager.cs
+++ b/Engine/Source/Programs/Shared/EpicGames.OIDC/OidcTokenManager.cs
@@ -287,8 +287,8 @@ namespace EpicGames.OIDC
 
 			// generate a simple http page to show the user
 			HttpListenerResponse response = context.Response;
-			const string httpPage = "<html><head></head><body><h1 style=\"text-align: center;\">Please close this browser to finish your login.<h1></body></html>";
-			byte[] buffer = Encoding.UTF8.GetBytes(httpPage);
+			const string HttpPage = "<html><head></head><body><h1 style=\"text-align: center;\">Please close this browser to finish your login.<h1></body></html>";
+			byte[] buffer = Encoding.UTF8.GetBytes(HttpPage);
 			response.ContentLength64 = buffer.Length;
 			Stream responseOutput = response.OutputStream;
 			await responseOutput.WriteAsync(buffer, 0, buffer.Length);
@@ -445,7 +445,6 @@ namespace EpicGames.OIDC
 
 	public class OidcTokenOptions
 	{
-		[System.Diagnostics.CodeAnalysis.SuppressMessage("Usage", "CA2227:Collection properties should be read only", Justification = "Used by the configuration system")]
 		public Dictionary<string, ProviderInfo> Providers { get; set; } = new Dictionary<string, ProviderInfo>();
 	}
 
diff --git a/Engine/Source/Programs/Shared/EpicGames.OIDC/WindowsTokenStore.cs b/Engine/Source/Programs/Shared/EpicGames.OIDC/WindowsTokenStore.cs
--- a/Engine/Source/Programs/Shared/EpicGames.OIDC/WindowsTokenStore.cs
+++ b/Engine/Source/Programs/Shared/EpicGames.OIDC/WindowsTokenStore.cs
@@ -114,14 +114,8 @@ namespace EpicGames.OIDC
 		private static byte[] GetEntropy(string oidcProvider)
 		{
 			byte[] providerBytes = Encoding.UTF8.GetBytes(oidcProvider);
-			byte[] assemblyBytes = Encoding.UTF8.GetBytes(Assembly.GetEntryAssembly()!.GetName().FullName);
 
-			byte[] mergedBytes = new byte[providerBytes.Length + assemblyBytes.Length];
-			Array.Copy(providerBytes, mergedBytes, providerBytes.Length);
-			int bytesWritten = providerBytes.Length;
-			Array.Copy(assemblyBytes, 0, mergedBytes, bytesWritten, assemblyBytes.Length);
-
-			return mergedBytes;
+			return providerBytes;
 		}
 
 		public void AddRefreshToken(string providerIdentifier, string refreshToken)
