Add a common interface for reading/writing to memory (IMemoryReader/IMemoryWriter), which is implemented by the legacy MemoryReader/MemoryWriter classes as well as the newer ByteArrayBuilder.

Also moves VarInt into EpicGames.Core.

diff --git a/Engine/Source/Programs/Horde/Horde.Build/Logs/Builder/RedisLogBuilder.cs b/Engine/Source/Programs/Horde/Horde.Build/Logs/Builder/RedisLogBuilder.cs
--- a/Engine/Source/Programs/Horde/Horde.Build/Logs/Builder/RedisLogBuilder.cs
+++ b/Engine/Source/Programs/Horde/Horde.Build/Logs/Builder/RedisLogBuilder.cs
@@ -252,7 +252,7 @@ namespace Horde.Build.Logs.Builder
 					int subChunkLineIndex = lineIndex;
 
 					List<LogSubChunkData> subChunks = new List<LogSubChunkData>();
-					while (reader.Offset < reader.Length)
+					while (reader.Memory.Length > 0)
 					{
 						LogSubChunkData subChunkData = reader.ReadLogSubChunkData(subChunkOffset, subChunkLineIndex);
 						subChunkOffset += subChunkData.Length;
diff --git a/Engine/Source/Programs/Horde/Horde.Build/Logs/Data/LogChunkData.cs b/Engine/Source/Programs/Horde/Horde.Build/Logs/Data/LogChunkData.cs
--- a/Engine/Source/Programs/Horde/Horde.Build/Logs/Data/LogChunkData.cs
+++ b/Engine/Source/Programs/Horde/Horde.Build/Logs/Data/LogChunkData.cs
@@ -154,7 +154,7 @@ namespace Horde.Build.Logs.Data
 			{
 				List<LogSubChunkData> subChunks = new List<LogSubChunkData>();
 				subChunks.Add(new LogSubChunkData(LogType.Json, offset, lineIndex, new ReadOnlyLogText(reader.Memory)));
-				reader.Offset = reader.Memory.Length;
+				reader.Advance(reader.Memory.Length);
 				return new LogChunkData(offset, lineIndex, subChunks);
 			}
 
@@ -162,7 +162,7 @@ namespace Horde.Build.Logs.Data
 			if (version == 0)
 			{
 				List<LogSubChunkData> subChunks = ReadSubChunkList(reader, offset, lineIndex);
-				reader.ReadVariableLengthBytes();
+				reader.ReadVariableLengthBytesWithInt32Length();
 				return new LogChunkData(offset, lineIndex, subChunks);
 			}
 			else
@@ -206,7 +206,7 @@ namespace Horde.Build.Logs.Data
 		{
 			MemoryReader reader = new MemoryReader(memory);
 			LogChunkData chunkData = Read(reader, offset, lineIndex);
-			reader.CheckOffset(memory.Length);
+			reader.CheckEmpty();
 			return chunkData;
 		}
 
@@ -235,7 +235,7 @@ namespace Horde.Build.Logs.Data
 			byte[] data = new byte[GetSerializedSize(logger)];
 			MemoryWriter writer = new MemoryWriter(data);
 			Write(writer, logger);
-			writer.CheckOffset(data.Length);
+			writer.CheckEmpty();
 			return data;
 		}
 
diff --git a/Engine/Source/Programs/Horde/Horde.Build/Logs/Data/LogIndexData.cs b/Engine/Source/Programs/Horde/Horde.Build/Logs/Data/LogIndexData.cs
--- a/Engine/Source/Programs/Horde/Horde.Build/Logs/Data/LogIndexData.cs
+++ b/Engine/Source/Programs/Horde/Horde.Build/Logs/Data/LogIndexData.cs
@@ -78,7 +78,7 @@ namespace Horde.Build.Logs.Data
 		{
 			int lineIndex = reader.ReadInt32();
 			int lineCount = reader.ReadInt32();
-			ReadOnlyMemory<byte> compressedPlainText = reader.ReadVariableLengthBytes();
+			ReadOnlyMemory<byte> compressedPlainText = reader.ReadVariableLengthBytesWithInt32Length();
 			return new LogIndexBlock(lineIndex, lineCount, null, compressedPlainText);
 		}
 
@@ -91,7 +91,7 @@ namespace Horde.Build.Logs.Data
 		{
 			writer.WriteInt32(block.LineIndex);
 			writer.WriteInt32(block.LineCount);
-			writer.WriteVariableLengthBytes(block.CompressedPlainText.Span);
+			writer.WriteVariableLengthBytesWithInt32Length(block.CompressedPlainText.Span);
 		}
 
 		/// <summary>
@@ -519,7 +519,7 @@ namespace Horde.Build.Logs.Data
 				{
 					int lineIndex = (idx > 0) ? (blocks[idx - 1].LineIndex + blocks[idx - 1].LineCount) : 0;
 					int lineCount = reader.ReadInt32();
-					ReadOnlyMemory<byte> compressedPlainText = reader.ReadVariableLengthBytes();
+					ReadOnlyMemory<byte> compressedPlainText = reader.ReadVariableLengthBytesWithInt32Length();
 					reader.ReadTrie();
 					blocks[idx] = new LogIndexBlock(lineIndex, lineCount, null, compressedPlainText);
 				}
@@ -536,7 +536,7 @@ namespace Horde.Build.Logs.Data
 				{
 					int lineIndex = reader.ReadInt32();
 					int lineCount = reader.ReadInt32();
-					ReadOnlyMemory<byte> compressedPlainText = reader.ReadVariableLengthBytes();
+					ReadOnlyMemory<byte> compressedPlainText = reader.ReadVariableLengthBytesWithInt32Length();
 					reader.ReadTrie();
 					blocks[idx] = new LogIndexBlock(lineIndex, lineCount, null, compressedPlainText);
 				}
@@ -597,7 +597,7 @@ namespace Horde.Build.Logs.Data
 
 			MemoryWriter writer = new MemoryWriter(buffer);
 			Write(writer);
-			writer.CheckOffset(buffer.Length);
+			writer.CheckEmpty();
 
 			return buffer;
 		}
diff --git a/Engine/Source/Programs/Horde/Horde.Build/Logs/Data/LogSubChunkData.cs b/Engine/Source/Programs/Horde/Horde.Build/Logs/Data/LogSubChunkData.cs
--- a/Engine/Source/Programs/Horde/Horde.Build/Logs/Data/LogSubChunkData.cs
+++ b/Engine/Source/Programs/Horde/Horde.Build/Logs/Data/LogSubChunkData.cs
@@ -149,8 +149,8 @@ namespace Horde.Build.Logs.Data
 				LogType type = (LogType)reader.ReadInt32();
 				int length = reader.ReadInt32();
 				int lineCount = reader.ReadInt32();
-				ReadOnlyMemory<byte> compressedText = reader.ReadVariableLengthBytes();
-				reader.ReadVariableLengthBytes();
+				ReadOnlyMemory<byte> compressedText = reader.ReadVariableLengthBytesWithInt32Length();
+				reader.ReadVariableLengthBytesWithInt32Length();
 				return new LogSubChunkData(type, offset, length, lineIndex, lineCount, compressedText, null);
 			}
 			else if (version == 1)
@@ -158,8 +158,8 @@ namespace Horde.Build.Logs.Data
 				LogType type = (LogType)reader.ReadInt32();
 				int length = reader.ReadInt32();
 				int lineCount = reader.ReadInt32();
-				ReadOnlyMemory<byte> compressedText = reader.ReadVariableLengthBytes();
-				ReadOnlyMemory<byte> compressedPlainText = reader.ReadVariableLengthBytes();
+				ReadOnlyMemory<byte> compressedText = reader.ReadVariableLengthBytesWithInt32Length();
+				ReadOnlyMemory<byte> compressedPlainText = reader.ReadVariableLengthBytesWithInt32Length();
 				ReadOnlyTrie trie = reader.ReadTrie();
 				LogIndexBlock indexBlock = new LogIndexBlock(lineIndex, lineCount, null, compressedPlainText);
 				LogIndexData index = new LogIndexData(trie, 0, new[] { indexBlock });
@@ -170,7 +170,7 @@ namespace Horde.Build.Logs.Data
 				LogType type = (LogType)reader.ReadInt32();
 				int length = reader.ReadInt32();
 				int lineCount = reader.ReadInt32();
-				ReadOnlyMemory<byte> compressedText = reader.ReadVariableLengthBytes();
+				ReadOnlyMemory<byte> compressedText = reader.ReadVariableLengthBytesWithInt32Length();
 				LogIndexData index = reader.ReadLogIndexData();
 				index.SetBaseLineIndex(lineIndex); // Fix for incorrectly saved data
 				return new LogSubChunkData(type, offset, length, lineIndex, lineCount, compressedText, index);
@@ -190,7 +190,7 @@ namespace Horde.Build.Logs.Data
 			writer.WriteInt32(Length);
 			writer.WriteInt32(LineCount);
 
-			writer.WriteVariableLengthBytes(DeflateText().Span);
+			writer.WriteVariableLengthBytesWithInt32Length(DeflateText().Span);
 			writer.WriteLogIndexData(BuildIndex(logger));
 		}
 
@@ -204,7 +204,7 @@ namespace Horde.Build.Logs.Data
 
 			MemoryWriter writer = new MemoryWriter(data);
 			Write(writer, logger);
-			writer.CheckOffset(data.Length);
+			writer.CheckEmpty();
 
 			return data;
 		}
diff --git a/Engine/Source/Programs/Horde/Horde.Build/Logs/Storage/PersistentLogStorage.cs b/Engine/Source/Programs/Horde/Horde.Build/Logs/Storage/PersistentLogStorage.cs
--- a/Engine/Source/Programs/Horde/Horde.Build/Logs/Storage/PersistentLogStorage.cs
+++ b/Engine/Source/Programs/Horde/Horde.Build/Logs/Storage/PersistentLogStorage.cs
@@ -82,9 +82,9 @@ namespace Horde.Build.Logs.Storage
 			MemoryReader reader = new MemoryReader(data.Value);
 			LogChunkData chunkData = reader.ReadLogChunkData(offset, lineIndex);
 
-			if (reader.Offset != data.Value.Length)
+			if (reader.Memory.Length > 0)
 			{
-				throw new Exception($"Serialization of persistent chunk {path} is not at expected offset (expected {data.Value.Length}, actual {reader.Offset})");
+				throw new Exception($"Serialization of persistent chunk {path} is not at expected offset ({reader.Memory.Length} bytes remaining)");
 			}
 
 			return chunkData;
@@ -99,7 +99,7 @@ namespace Horde.Build.Logs.Storage
 			byte[] data = new byte[chunkData.GetSerializedSize(_logger)];
 			MemoryWriter writer = new MemoryWriter(data);
 			writer.WriteLogChunkData(chunkData, _logger);
-			writer.CheckOffset(data.Length);
+			writer.CheckEmpty();
 
 			return _storageProvider.WriteBytesAsync(path, data);
 		}
diff --git a/Engine/Source/Programs/Horde/Horde.Build/Utilities/ReadOnlyTrie.cs b/Engine/Source/Programs/Horde/Horde.Build/Utilities/ReadOnlyTrie.cs
--- a/Engine/Source/Programs/Horde/Horde.Build/Utilities/ReadOnlyTrie.cs
+++ b/Engine/Source/Programs/Horde/Horde.Build/Utilities/ReadOnlyTrie.cs
@@ -244,7 +244,7 @@ namespace Horde.Build.Logs
 		/// <returns>New trie</returns>
 		public static ReadOnlyTrie Read(MemoryReader reader)
 		{
-			ReadOnlyMemory<byte> nodes = reader.ReadVariableLengthBytes();
+			ReadOnlyMemory<byte> nodes = reader.ReadVariableLengthBytesWithInt32Length();
 			ushort[] nodeData = MemoryMarshal.Cast<byte, ushort>(nodes.Span).ToArray();
 			return new ReadOnlyTrie(nodeData);
 		}
@@ -255,7 +255,7 @@ namespace Horde.Build.Logs
 		/// <param name="writer">Writer to output to</param>
 		public void Write(MemoryWriter writer)
 		{
-			writer.WriteVariableLengthBytes(MemoryMarshal.AsBytes<ushort>(NodeData));
+			writer.WriteVariableLengthBytesWithInt32Length(MemoryMarshal.AsBytes<ushort>(NodeData));
 		}
 
 		/// <summary>
diff --git a/Engine/Source/Programs/Shared/EpicGames.BuildGraph/BgOpcode.cs b/Engine/Source/Programs/Shared/EpicGames.BuildGraph/BgOpcode.cs
--- a/Engine/Source/Programs/Shared/EpicGames.BuildGraph/BgOpcode.cs
+++ b/Engine/Source/Programs/Shared/EpicGames.BuildGraph/BgOpcode.cs
@@ -132,7 +132,7 @@ namespace EpicGames.BuildGraph
 		StrLiteral = 0x21,
 
 		/// <summary>
-		/// Compares two strings for equality, using a <see cref="StringComparison"/> value encoded into the byte stream as an unsigned <see cref="VarInt"/>.
+		/// Compares two strings for equality, using a <see cref="StringComparison"/> value encoded into the byte stream as an unsigned <see cref="EpicGames.Core.VarInt"/>.
 		/// </summary>
 		StrCompare = 0x22,
 
diff --git a/Engine/Source/Programs/Shared/EpicGames.Core.Tests/VarIntTests.cs b/Engine/Source/Programs/Shared/EpicGames.Core.Tests/VarIntTests.cs
new file mode 100644
--- /dev/null
+++ b/Engine/Source/Programs/Shared/EpicGames.Core.Tests/VarIntTests.cs
@@ -0,0 +1,26 @@
+// Copyright Epic Games, Inc. All Rights Reserved.
+
+using Microsoft.VisualStudio.TestTools.UnitTesting;
+
+namespace EpicGames.Core.Tests
+{
+	[TestClass]
+	public class VarIntTests
+	{
+		[TestMethod]
+		public void TestVarInt()
+		{
+			byte[] buffer = new byte[20];
+
+			int length = VarInt.WriteUnsigned(buffer, -1);
+			Assert.AreEqual(9, length);
+			Assert.AreEqual(9, VarInt.MeasureUnsigned(-1));
+
+			Assert.AreEqual(9, VarInt.Measure(buffer));
+			int value = (int)(long)VarInt.ReadUnsigned(buffer, out int bytesRead);
+			Assert.AreEqual(9, bytesRead);
+
+			Assert.AreEqual(-1, value);
+		}
+	}
+}
diff --git a/Engine/Source/Programs/Shared/EpicGames.Core/ByteArrayBuilder.cs b/Engine/Source/Programs/Shared/EpicGames.Core/ByteArrayBuilder.cs
--- a/Engine/Source/Programs/Shared/EpicGames.Core/ByteArrayBuilder.cs
+++ b/Engine/Source/Programs/Shared/EpicGames.Core/ByteArrayBuilder.cs
@@ -9,7 +9,7 @@ namespace EpicGames.Core
 	/// <summary>
 	/// Class for building byte sequences, similar to StringBuilder. Allocates memory in chunks to avoid copying data.
 	/// </summary>
-	public class ByteArrayBuilder
+	public class ByteArrayBuilder : IMemoryWriter
 	{
 		class Chunk
 		{
@@ -88,6 +88,24 @@ namespace EpicGames.Core
 			sequence.CopyTo(target);
 		}
 
+		/// <inheritdoc/>
+		public Memory<byte> GetMemory(int minSize)
+		{
+			if (_currentChunk.Length + minSize > _currentChunk.Data.Length)
+			{
+				_currentChunk = new Chunk(_currentChunk.RunningIndex + _currentChunk.Length, Math.Max(minSize, _chunkSize));
+				_chunks.Add(_currentChunk);
+			}
+			return _currentChunk.Data.AsMemory(_currentChunk.Length, minSize);
+		}
+
+		/// <inheritdoc/>
+		public void Advance(int length)
+		{
+			_currentChunk.Length += length;
+			Length += length;
+		}
+
 		/// <summary>
 		/// Gets a span of bytes to write to, and increases the length of the buffer
 		/// </summary>
diff --git a/Engine/Source/Programs/Shared/EpicGames.Core/Digest.cs b/Engine/Source/Programs/Shared/EpicGames.Core/Digest.cs
--- a/Engine/Source/Programs/Shared/EpicGames.Core/Digest.cs
+++ b/Engine/Source/Programs/Shared/EpicGames.Core/Digest.cs
@@ -327,7 +327,7 @@ namespace EpicGames.Core
 		/// <returns></returns>
 		public static Digest ReadDigest(this MemoryReader reader)
 		{
-			return new Digest(reader.ReadVariableLengthBytes());
+			return new Digest(reader.ReadVariableLengthBytesWithInt32Length());
 		}
 
 		/// <summary>
@@ -348,7 +348,7 @@ namespace EpicGames.Core
 		/// <param name="digest"></param>
 		public static void WriteDigest(this MemoryWriter writer, Digest digest)
 		{
-			writer.WriteVariableLengthBytes(digest.Span);
+			writer.WriteVariableLengthBytesWithInt32Length(digest.Span);
 		}
 
 		/// <summary>
diff --git a/Engine/Source/Programs/Shared/EpicGames.Core/MemoryReader.cs b/Engine/Source/Programs/Shared/EpicGames.Core/MemoryReader.cs
--- a/Engine/Source/Programs/Shared/EpicGames.Core/MemoryReader.cs
+++ b/Engine/Source/Programs/Shared/EpicGames.Core/MemoryReader.cs
@@ -6,35 +6,41 @@ using System.Buffers.Binary;
 namespace EpicGames.Core
 {
 	/// <summary>
-	/// Reads data from a memory buffer
+	/// Interface for reading from data in memory
 	/// </summary>
-	public class MemoryReader
+	public interface IMemoryReader
 	{
 		/// <summary>
-		/// The memory to read from
+		/// Gets a block of memory with at least the given size
 		/// </summary>
-		public ReadOnlyMemory<byte> Memory
-		{
-			get;
-		}
+		/// <param name="minSize">Minimum size of the returned data</param>
+		/// <returns>Memory of at least the given size</returns>
+		ReadOnlyMemory<byte> GetMemory(int minSize);
 
 		/// <summary>
-		/// Offset to read from
+		/// Updates the current position within the input buffer
 		/// </summary>
-		public int Offset
-		{
-			get; set;
-		}
+		/// <param name="length">Number of bytes to advance by</param>
+		void Advance(int length);
+	}
 
+	/// <summary>
+	/// Reads data from a memory buffer
+	/// </summary>
+	public class MemoryReader : IMemoryReader
+	{
 		/// <summary>
-		/// Length of the data
+		/// The memory to read from
 		/// </summary>
-		public int Length => Memory.Length;
+		public ReadOnlyMemory<byte> Memory
+		{
+			get; private set;
+		}
 
 		/// <summary>
 		/// Returns the memory at the current offset
 		/// </summary>
-		public ReadOnlySpan<byte> Span => Memory.Span.Slice(Offset);
+		public ReadOnlySpan<byte> Span => Memory.Span;
 
 		/// <summary>
 		/// Constructor
@@ -43,156 +49,199 @@ namespace EpicGames.Core
 		public MemoryReader(ReadOnlyMemory<byte> memory)
 		{
 			Memory = memory;
-			Offset = 0;
 		}
 
 		/// <summary>
 		/// Checks that we've used the exact buffer length
 		/// </summary>
 		/// <param name="expectedOffset">Expected offset within the output buffer</param>
-		public void CheckOffset(int expectedOffset)
+		public void CheckEmpty()
 		{
-			if (Offset != expectedOffset)
+			if (Memory.Length > 0)
 			{
-				throw new Exception($"Serialization is not at expected offset within the input buffer (expected {expectedOffset}, actual {Offset})");
+				throw new Exception($"Serialization is not at expected offset within the input buffer ({Memory.Length} bytes unused)");
 			}
 		}
 
+		/// <inheritdoc/>
+		public ReadOnlyMemory<byte> GetMemory(int minSize) => Memory;
+
+		/// <inheritdoc/>
+		public void Advance(int length) => Memory = Memory.Slice(length);
+	}
+
+	/// <summary>
+	/// Extension methods for <see cref="IMemoryReader"/>
+	/// </summary>
+	public static class MemoryReaderExtensions
+	{
+		/// <summary>
+		/// Gets a span of data with the given length
+		/// </summary>
+		/// <param name="reader">Reader to deserialize from</param>
+		/// <param name="minSize">Minimum size of the returned span</param>
+		/// <returns>Span of at least the given size</returns>
+		public static ReadOnlySpan<byte> GetSpan(this IMemoryReader reader, int minSize) => reader.GetMemory(minSize).Span;
+
 		/// <summary>
 		/// Read a boolean from the buffer
 		/// </summary>
+		/// <param name="reader">Reader to deserialize from</param>
 		/// <returns>The value read from the buffer</returns>
-		public bool ReadBoolean()
+		public static bool ReadBoolean(this IMemoryReader reader)
 		{
-			return ReadUInt8() != 0;
+			return ReadUInt8(reader) != 0;
 		}
 
 		/// <summary>
 		/// Reads a byte from the buffer
 		/// </summary>
+		/// <param name="reader">Reader to deserialize from</param>
 		/// <returns>The value read from the buffer</returns>
-		public sbyte ReadInt8()
+		public static sbyte ReadInt8(this IMemoryReader reader)
 		{
-			sbyte value = (sbyte)Memory.Span[Offset];
-			Offset++;
+			sbyte value = (sbyte)reader.GetSpan(1)[0];
+			reader.Advance(1);
 			return value;
 		}
 
 		/// <summary>
 		/// Reads a byte from the buffer
 		/// </summary>
+		/// <param name="reader">Reader to deserialize from</param>
 		/// <returns>The value read from the buffer</returns>
-		public byte ReadUInt8()
+		public static byte ReadUInt8(this IMemoryReader reader)
 		{
-			byte value = Memory.Span[Offset];
-			Offset++;
+			byte value = reader.GetSpan(1)[0];
+			reader.Advance(1);
 			return value;
 		}
 
 		/// <summary>
 		/// Reads an int16
 		/// </summary>
+		/// <param name="reader">Reader to deserialize from</param>
 		/// <returns>The value read from the buffer</returns>
-		public short ReadInt16()
+		public static short ReadInt16(this IMemoryReader reader)
 		{
-			short value = BinaryPrimitives.ReadInt16LittleEndian(Memory.Span.Slice(Offset));
-			Offset += sizeof(short);
+			short value = BinaryPrimitives.ReadInt16LittleEndian(reader.GetSpan(sizeof(short)));
+			reader.Advance(sizeof(short));
 			return value;
 		}
 
 		/// <summary>
 		/// Reads a uint16
 		/// </summary>
+		/// <param name="reader">Reader to deserialize from</param>
 		/// <returns>The value read from the buffer</returns>
-		public ushort ReadUInt16()
+		public static ushort ReadUInt16(this IMemoryReader reader)
 		{
-			ushort value = BinaryPrimitives.ReadUInt16LittleEndian(Memory.Span.Slice(Offset));
-			Offset += sizeof(ushort);
+			ushort value = BinaryPrimitives.ReadUInt16LittleEndian(reader.GetSpan(sizeof(ushort)));
+			reader.Advance(sizeof(ushort));
 			return value;
 		}
 
 		/// <summary>
 		/// Reads an int32
 		/// </summary>
+		/// <param name="reader">Reader to deserialize from</param>
 		/// <returns>The value that was read from the buffer</returns>
-		public int ReadInt32()
+		public static int ReadInt32(this IMemoryReader reader)
 		{
-			int value = BinaryPrimitives.ReadInt32LittleEndian(Memory.Span.Slice(Offset));
-			Offset += sizeof(int);
+			int value = BinaryPrimitives.ReadInt32LittleEndian(reader.GetSpan(sizeof(int)));
+			reader.Advance(sizeof(int));
 			return value;
 		}
 
 		/// <summary>
 		/// Reads a uint32
 		/// </summary>
+		/// <param name="reader">Reader to deserialize from</param>
 		/// <returns>The value that was read from the buffer</returns>
-		public uint ReadUInt32()
+		public static uint ReadUInt32(this IMemoryReader reader)
 		{
-			uint value = BinaryPrimitives.ReadUInt32LittleEndian(Memory.Span.Slice(Offset));
-			Offset += sizeof(uint);
+			uint value = BinaryPrimitives.ReadUInt32LittleEndian(reader.GetSpan(sizeof(uint)));
+			reader.Advance(sizeof(uint));
 			return value;
 		}
 
 		/// <summary>
 		/// Reads an int64
 		/// </summary>
+		/// <param name="reader">Reader to deserialize from</param>
 		/// <returns>The value that was read from the buffer</returns>
-		public long ReadInt64()
+		public static long ReadInt64(this IMemoryReader reader)
 		{
-			long value = BinaryPrimitives.ReadInt64LittleEndian(Memory.Span.Slice(Offset));
-			Offset += sizeof(long);
+			long value = BinaryPrimitives.ReadInt64LittleEndian(reader.GetSpan(sizeof(long)));
+			reader.Advance(sizeof(long));
 			return value;
 		}
 
 		/// <summary>
 		/// Reads an int64
 		/// </summary>
+		/// <param name="reader">Reader to deserialize from</param>
 		/// <returns>The value that was read from the buffer</returns>
-		public ulong ReadUInt64()
+		public static ulong ReadUInt64(this IMemoryReader reader)
 		{
-			ulong value = BinaryPrimitives.ReadUInt64LittleEndian(Memory.Span.Slice(Offset));
-			Offset += sizeof(ulong);
+			ulong value = BinaryPrimitives.ReadUInt64LittleEndian(reader.GetSpan(sizeof(ulong)));
+			reader.Advance(sizeof(ulong));
 			return value;
 		}
 
 		/// <summary>
 		/// Reads a sequence of bytes from the buffer
 		/// </summary>
+		/// <param name="reader">Reader to deserialize from</param>
 		/// <returns>Sequence of bytes</returns>
-		public ReadOnlyMemory<byte> ReadVariableLengthBytes()
+		public static ReadOnlyMemory<byte> ReadVariableLengthBytes(this IMemoryReader reader)
+		{
+			int length = (int)reader.ReadUnsignedVarInt();
+			return ReadFixedLengthBytes(reader, length);
+		}
+
+		/// <summary>
+		/// Writes a variable length span of bytes
+		/// </summary>
+		/// <param name="reader">Reader to deserialize from</param>
+		public static ReadOnlyMemory<byte> ReadVariableLengthBytesWithInt32Length(this IMemoryReader reader)
 		{
-			int length = ReadInt32();
-			return ReadFixedLengthBytes(length);
+			int length = ReadInt32(reader);
+			return ReadFixedLengthBytes(reader, length);
 		}
 
 		/// <summary>
 		/// Reads a sequence of bytes from the buffer
 		/// </summary>
+		/// <param name="reader">Reader to deserialize from</param>
 		/// <param name="length">Number of bytes to read</param>
 		/// <returns>Sequence of bytes</returns>
-		public ReadOnlyMemory<byte> ReadFixedLengthBytes(int length)
+		public static ReadOnlyMemory<byte> ReadFixedLengthBytes(this IMemoryReader reader, int length)
 		{
-			ReadOnlyMemory<byte> bytes = Memory.Slice(Offset, length);
-			Offset += length;
+			ReadOnlyMemory<byte> bytes = reader.GetMemory(length);
+			reader.Advance(length);
 			return bytes;
 		}
 
 		/// <summary>
 		/// Reads a variable length array
 		/// </summary>
+		/// <param name="reader">Reader to deserialize from</param>
 		/// <param name="readItem">Delegate to write an individual item</param>
-		public T[] ReadVariableLengthArray<T>(Func<T> readItem)
+		public static T[] ReadVariableLengthArray<T>(this IMemoryReader reader, Func<T> readItem)
 		{
-			int length = ReadInt32();
-			return ReadFixedLengthArray(length, readItem);
+			int length = ReadInt32(reader);
+			return ReadFixedLengthArray(reader, length, readItem);
 		}
 
 		/// <summary>
 		/// Writes a fixed length array
 		/// </summary>
+		/// <param name="reader">Reader to deserialize from</param>
 		/// <param name="length">Length of the array to read</param>
 		/// <param name="readItem">Delegate to read an individual item</param>
-		public T[] ReadFixedLengthArray<T>(int length, Func<T> readItem)
+		[System.Diagnostics.CodeAnalysis.SuppressMessage("Style", "IDE0060:Remove unused parameter")]
+		public static T[] ReadFixedLengthArray<T>(this IMemoryReader reader, int length, Func<T> readItem)
 		{
 			T[] array = new T[length];
 			for (int idx = 0; idx < length; idx++)
diff --git a/Engine/Source/Programs/Shared/EpicGames.Core/MemoryWriter.cs b/Engine/Source/Programs/Shared/EpicGames.Core/MemoryWriter.cs
--- a/Engine/Source/Programs/Shared/EpicGames.Core/MemoryWriter.cs
+++ b/Engine/Source/Programs/Shared/EpicGames.Core/MemoryWriter.cs
@@ -1,24 +1,37 @@
 // Copyright Epic Games, Inc. All Rights Reserved.
 
 using System;
+using System.Buffers;
 using System.Buffers.Binary;
+using System.Text;
 
 namespace EpicGames.Core
 {
+	public interface IMemoryWriter
+	{
+		/// <summary>
+		/// Gets a block of memory with at least the given size
+		/// </summary>
+		/// <param name="minSize">Minimum size of the returned data</param>
+		/// <returns>Memory of at least the given size</returns>
+		Memory<byte> GetMemory(int minSize);
+
+		/// <summary>
+		/// Updates the current position within the input buffer
+		/// </summary>
+		/// <param name="length">Number of bytes to advance by</param>
+		void Advance(int length);
+	}
+
 	/// <summary>
 	/// Writes into a fixed size memory block
 	/// </summary>
-	public class MemoryWriter
+	public class MemoryWriter : IMemoryWriter
 	{
 		/// <summary>
 		/// The memory block to write to
 		/// </summary>
-		readonly Memory<byte> _memory;
-
-		/// <summary>
-		/// Current offset within the memory
-		/// </summary>
-		int _offset;
+		Memory<byte> _memory;
 
 		/// <summary>
 		/// Constructor
@@ -27,147 +40,209 @@ namespace EpicGames.Core
 		public MemoryWriter(Memory<byte> memory)
 		{
 			_memory = memory;
-			_offset = 0;
 		}
 
 		/// <summary>
 		/// Checks that we've used the exact buffer length
 		/// </summary>
 		/// <param name="expectedOffset">Expected offset within the output buffer</param>
-		public void CheckOffset(int expectedOffset)
+		public void CheckEmpty()
 		{
-			if (_offset != expectedOffset)
+			if (_memory.Length > 0)
 			{
-				throw new Exception("Serialization is not at expected offset within the output buffer");
+				throw new Exception($"Serialization is not at expected offset within the output buffer ({_memory.Length} bytes unused)");
 			}
 		}
 
+		/// <inheritdoc/>
+		public Memory<byte> GetMemory(int length) => _memory;
+
+		/// <inheritdoc/>
+		public void Advance(int length) => _memory = _memory.Slice(length);
+	}
+
+	/// <summary>
+	/// Extension methods for <see cref="IMemoryWriter"/>
+	/// </summary>
+	public static class MemoryWriterExtensions
+	{
+		/// <summary>
+		/// Writes a byte to memory
+		/// </summary>
+		/// <param name="writer">Writer to serialize to</param>
+		/// <param name="minSize">Minimum size for the returned span</param>
+		public static Span<byte> GetSpan(this IMemoryWriter writer, int minSize) => writer.GetMemory(minSize).Span;
+
 		/// <summary>
 		/// Writes a boolean to memory
 		/// </summary>
+		/// <param name="writer">Writer to serialize to</param>
 		/// <param name="value">Value to write</param>
-		public void WriteBoolean(bool value)
+		public static void WriteBoolean(this IMemoryWriter writer, bool value)
 		{
-			WriteUInt8(value ? (byte)1 : (byte)0);
+			WriteUInt8(writer, value ? (byte)1 : (byte)0);
 		}
 
 		/// <summary>
 		/// Writes a byte to memory
 		/// </summary>
+		/// <param name="writer">Writer to serialize to</param>
 		/// <param name="value">Value to write</param>
-		public void WriteInt8(sbyte value)
+		public static void WriteInt8(this IMemoryWriter writer, sbyte value)
 		{
-			_memory.Span[_offset] = (byte)value;
-			_offset++;
+			WriteUInt8(writer, (byte)value);
 		}
 
 		/// <summary>
 		/// Writes a byte to memory
 		/// </summary>
+		/// <param name="writer">Writer to serialize to</param>
 		/// <param name="value">Value to write</param>
-		public void WriteUInt8(byte value)
+		public static void WriteUInt8(this IMemoryWriter writer, byte value)
 		{
-			_memory.Span[_offset] = value;
-			_offset++;
+			writer.GetSpan(1)[0] = (byte)value;
+			writer.Advance(1);
 		}
 
 		/// <summary>
 		/// Writes an int16 to the memory
 		/// </summary>
+		/// <param name="writer">Writer to serialize to</param>
 		/// <param name="value">Value to write</param>
-		public void WriteInt16(short value)
+		public static void WriteInt16(this IMemoryWriter writer, short value)
 		{
-			BinaryPrimitives.WriteInt16LittleEndian(_memory.Span.Slice(_offset), value);
-			_offset += sizeof(short);
+			Span<byte> span = GetSpan(writer, sizeof(short));
+			BinaryPrimitives.WriteInt16LittleEndian(span, value);
+			writer.Advance(sizeof(short));
 		}
 
 		/// <summary>
 		/// Writes a uint16 to the memory
 		/// </summary>
+		/// <param name="writer">Writer to serialize to</param>
 		/// <param name="value">Value to write</param>
-		public void WriteUInt16(ushort value)
+		public static void WriteUInt16(this IMemoryWriter writer, ushort value)
 		{
-			BinaryPrimitives.WriteUInt16LittleEndian(_memory.Span.Slice(_offset), value);
-			_offset += sizeof(ushort);
+			Span<byte> span = GetSpan(writer, sizeof(ushort));
+			BinaryPrimitives.WriteUInt16LittleEndian(span, value);
+			writer.Advance(sizeof(ushort));
 		}
 
 		/// <summary>
 		/// Writes an int32 to the memory
 		/// </summary>
+		/// <param name="writer">Writer to serialize to</param>
 		/// <param name="value">Value to write</param>
-		public void WriteInt32(int value)
+		public static void WriteInt32(this IMemoryWriter writer, int value)
 		{
-			BinaryPrimitives.WriteInt32LittleEndian(_memory.Span.Slice(_offset), value);
-			_offset += sizeof(int);
+			Span<byte> span = GetSpan(writer, sizeof(int));
+			BinaryPrimitives.WriteInt32LittleEndian(span, value);
+			writer.Advance(sizeof(int));
 		}
 
 		/// <summary>
 		/// Writes a uint32 to the memory
 		/// </summary>
+		/// <param name="writer">Writer to serialize to</param>
 		/// <param name="value">Value to write</param>
-		public void WriteUInt32(uint value)
+		public static void WriteUInt32(this IMemoryWriter writer, uint value)
 		{
-			BinaryPrimitives.WriteUInt32LittleEndian(_memory.Span.Slice(_offset), value);
-			_offset += sizeof(uint);
+			Span<byte> span = GetSpan(writer, sizeof(uint));
+			BinaryPrimitives.WriteUInt32LittleEndian(span, value);
+			writer.Advance(sizeof(uint));
 		}
 
 		/// <summary>
 		/// Writes an int64 to the memory
 		/// </summary>
+		/// <param name="writer">Writer to serialize to</param>
 		/// <param name="value">Value to write</param>
-		public void WriteInt64(long value)
+		public static void WriteInt64(this IMemoryWriter writer, long value)
 		{
-			BinaryPrimitives.WriteInt64LittleEndian(_memory.Span.Slice(_offset), value);
-			_offset += sizeof(long);
+			Span<byte> span = GetSpan(writer, sizeof(long));
+			BinaryPrimitives.WriteInt64LittleEndian(span, value);
+			writer.Advance(sizeof(long));
 		}
 
 		/// <summary>
 		/// Writes a uint64 to the memory
 		/// </summary>
+		/// <param name="writer">Writer to serialize to</param>
 		/// <param name="value">Value to write</param>
-		public void WriteUInt64(ulong value)
+		public static void WriteUInt64(this IMemoryWriter writer, ulong value)
+		{
+			Span<byte> span = GetSpan(writer, sizeof(ulong));
+			BinaryPrimitives.WriteUInt64LittleEndian(span, value);
+			writer.Advance(sizeof(ulong));
+		}
+
+		/// <summary>
+		/// Appends a sequence of bytes to the buffer
+		/// </summary>
+		/// <param name="writer">Writer to serialize to</param>
+		/// <param name="sequence">Sequence to append</param>
+		public static void WriteSequence(this IMemoryWriter writer, ReadOnlySequence<byte> sequence)
+		{
+			Memory<byte> memory = writer.GetMemory((int)sequence.Length);
+			sequence.CopyTo(memory.Span);
+		}
+
+		/// <summary>
+		/// Writes a variable length span of bytes
+		/// </summary>
+		/// <param name="writer">Writer to serialize to</param>
+		/// <param name="bytes">The bytes to write</param>
+		public static void WriteVariableLengthBytes(this IMemoryWriter writer, ReadOnlySpan<byte> bytes)
 		{
-			BinaryPrimitives.WriteUInt64LittleEndian(_memory.Span.Slice(_offset), value);
-			_offset += sizeof(ulong);
+			int lengthBytes = VarInt.MeasureUnsigned(bytes.Length);
+			Span<byte> span = GetSpan(writer, lengthBytes + bytes.Length);
+			VarInt.WriteUnsigned(span, bytes.Length);
+			bytes.CopyTo(span[lengthBytes..]);
 		}
 
 		/// <summary>
 		/// Writes a variable length span of bytes
 		/// </summary>
+		/// <param name="writer">Writer to serialize to</param>
 		/// <param name="bytes">The bytes to write</param>
-		public void WriteVariableLengthBytes(ReadOnlySpan<byte> bytes)
+		public static void WriteVariableLengthBytesWithInt32Length(this IMemoryWriter writer, ReadOnlySpan<byte> bytes)
 		{
-			WriteInt32(bytes.Length);
-			WriteFixedLengthBytes(bytes);
+			writer.WriteInt32(bytes.Length);
+			writer.WriteFixedLengthBytes(bytes);
 		}
 
 		/// <summary>
 		/// Write a fixed-length sequence of bytes to the buffer
 		/// </summary>
+		/// <param name="writer">Writer to serialize to</param>
 		/// <param name="bytes">The bytes to write</param>
-		public void WriteFixedLengthBytes(ReadOnlySpan<byte> bytes)
+		public static void WriteFixedLengthBytes(this IMemoryWriter writer, ReadOnlySpan<byte> bytes)
 		{
-			bytes.CopyTo(_memory.Span.Slice(_offset));
-			_offset += bytes.Length;
+			Span<byte> span = GetSpan(writer, bytes.Length);
+			bytes.CopyTo(span);
+			writer.Advance(bytes.Length);
 		}
 
 		/// <summary>
 		/// Writes a variable length array
 		/// </summary>
+		/// <param name="writer">Writer to serialize to</param>
 		/// <param name="array">The array to write</param>
 		/// <param name="writeItem">Delegate to write an individual item</param>
-		public void WriteVariableLengthArray<T>(T[] array, Action<T> writeItem)
+		public static void WriteVariableLengthArray<T>(this IMemoryWriter writer, T[] array, Action<T> writeItem)
 		{
-			WriteInt32(array.Length);
-			WriteFixedLengthArray(array, writeItem);
+			WriteInt32(writer, array.Length);
+			WriteFixedLengthArray(writer, array, writeItem);
 		}
 
 		/// <summary>
 		/// Writes a fixed length array
 		/// </summary>
+		/// <param name="writer">Writer to serialize to</param>
 		/// <param name="array">The array to write</param>
 		/// <param name="writeItem">Delegate to write an individual item</param>
-		public void WriteFixedLengthArray<T>(T[] array, Action<T> writeItem)
+		[System.Diagnostics.CodeAnalysis.SuppressMessage("Style", "IDE0060:Remove unused parameter")]
+		public static void WriteFixedLengthArray<T>(this IMemoryWriter writer, T[] array, Action<T> writeItem)
 		{
 			for (int idx = 0; idx < array.Length; idx++)
 			{
@@ -176,15 +251,52 @@ namespace EpicGames.Core
 		}
 
 		/// <summary>
-		/// Allocate a writable span from the buffer
+		/// Read a string from memory
+		/// </summary>
+		/// <param name="reader">Reader to deserialize from</param>
+		/// <returns>The string that was read</returns>
+		public static string ReadString(this IMemoryReader reader) => ReadString(reader, Encoding.UTF8);
+
+		/// <summary>
+		/// Read a string from memory
+		/// </summary>
+		/// <param name="reader">Reader to deserialize from</param>
+		/// <param name="encoding">Encoding to use for the string</param>
+		/// <returns>The string that was read</returns>
+		public static string ReadString(this IMemoryReader reader, Encoding encoding)
+		{
+			int length = (int)reader.ReadUnsignedVarInt();
+
+			ReadOnlySpan<byte> span = reader.GetSpan(length).Slice(0, length);
+			string str = encoding.GetString(span);
+			reader.Advance(length);
+
+			return str;
+		}
+
+		/// <summary>
+		/// Write a string to memory
+		/// </summary>
+		/// <param name="writer">Reader to deserialize from</param>
+		/// <param name="str">The string to be written</param>
+		public static void WriteString(this IMemoryWriter writer, string str) => WriteString(writer, str, Encoding.UTF8);
+
+		/// <summary>
+		/// Write a string to memory
 		/// </summary>
-		/// <param name="length">Length of the span to allocate</param>
-		/// <returns>Span that can be written to</returns>
-		public Span<byte> AllocateSpan(int length)
+		/// <param name="writer">Reader to deserialize from</param>
+		/// <param name="str">The string to be written</param>
+		/// <param name="encoding">Encoding to use for the string</param>
+		public static void WriteString(this IMemoryWriter writer, string str, Encoding encoding)
 		{
-			Span<byte> span = _memory.Span.Slice(_offset, length);
-			_offset += length;
-			return span;
+			int stringBytes = encoding.GetByteCount(str);
+			int lengthBytes = VarInt.MeasureUnsigned(stringBytes);
+
+			Span<byte> span = writer.GetSpan(lengthBytes + stringBytes);
+			VarInt.WriteUnsigned(span, stringBytes);
+			encoding.GetBytes(str, span[lengthBytes..]);
+
+			writer.Advance(lengthBytes + stringBytes);
 		}
 	}
 }
diff --git a/Engine/Source/Programs/Shared/EpicGames.Core/Sha1Hash.cs b/Engine/Source/Programs/Shared/EpicGames.Core/Sha1Hash.cs
--- a/Engine/Source/Programs/Shared/EpicGames.Core/Sha1Hash.cs
+++ b/Engine/Source/Programs/Shared/EpicGames.Core/Sha1Hash.cs
@@ -191,9 +191,11 @@ namespace EpicGames.Core
 		/// </summary>
 		/// <param name="reader"></param>
 		/// <returns></returns>
-		public static Sha1Hash ReadSha1Hash(this MemoryReader reader)
+		public static Sha1Hash ReadSha1Hash(this IMemoryReader reader)
 		{
-			return new Sha1Hash(reader.ReadFixedLengthBytes(Sha1Hash.NumBytes).Span);
+			ReadOnlySpan<byte> span = reader.GetSpan(Sha1Hash.NumBytes);
+			reader.Advance(Sha1Hash.NumBytes);
+			return new Sha1Hash(span);
 		}
 
 		/// <summary>
@@ -201,9 +203,11 @@ namespace EpicGames.Core
 		/// </summary>
 		/// <param name="writer"></param>
 		/// <param name="hash"></param>
-		public static void WriteSha1Hash(this MemoryWriter writer, Sha1Hash hash)
+		public static void WriteSha1Hash(this IMemoryWriter writer, Sha1Hash hash)
 		{
-			hash.CopyTo(writer.AllocateSpan(Sha1Hash.NumBytes));
+			Span<byte> span = writer.GetSpan(Sha1Hash.NumBytes);
+			writer.Advance(Sha1Hash.NumBytes);
+			hash.CopyTo(span);
 		}
 	}
 
diff --git a/Engine/Source/Programs/Shared/EpicGames.Core/Utf8String.cs b/Engine/Source/Programs/Shared/EpicGames.Core/Utf8String.cs
--- a/Engine/Source/Programs/Shared/EpicGames.Core/Utf8String.cs
+++ b/Engine/Source/Programs/Shared/EpicGames.Core/Utf8String.cs
@@ -474,30 +474,70 @@ namespace EpicGames.Core
 	/// <summary>
 	/// Extension methods for ReadOnlyUtf8String objects
 	/// </summary>
-	public static class MemoryWriterExtensions
+	public static class Utf8StringExtensions
 	{
 		/// <summary>
-		/// Reads a null-terminated utf8 string from the buffer
+		/// Reads a null-terminated utf8 string from the buffer, without copying it
 		/// </summary>
 		/// <returns>The string data</returns>
-		public static Utf8String ReadString(this MemoryReader reader)
+		public static Utf8String ReadUtf8String(this IMemoryReader reader)
 		{
-			ReadOnlySpan<byte> span = reader.Span;
-			int length = span.IndexOf((byte)0);
-			Utf8String value = new Utf8String(reader.ReadFixedLengthBytes(length));
-			reader.ReadInt8();
-			return value;
+			return ReadUtf8StringWithoutCopy(reader).Clone();
+		}
+
+		/// <summary>
+		/// Reads a null-terminated utf8 string from the buffer, without copying it
+		/// </summary>
+		/// <returns>The string data</returns>
+		public static Utf8String ReadUtf8StringWithoutCopy(this IMemoryReader reader)
+		{
+			return new Utf8String(reader.ReadVariableLengthBytes());
+		}
+
+		/// <summary>
+		/// Reads a null-terminated utf8 string from the buffer, without copying it
+		/// </summary>
+		/// <param name="writer">Writer to serialize to</param>
+		/// <param name="str">String to write</param>
+		public static void WriteUtf8String(this IMemoryWriter writer, Utf8String str)
+		{
+			writer.WriteVariableLengthBytes(str.Span);
+		}
+
+		/// <summary>
+		/// Reads a null-terminated utf8 string from the buffer, without copying it
+		/// </summary>
+		/// <returns>The string data</returns>
+		public static Utf8String ReadNullTerminatedUtf8String(this IMemoryReader reader)
+		{
+			int minSize = 1;
+			for (; ; )
+			{
+				ReadOnlyMemory<byte> memory = reader.GetMemory(minSize);
+
+				int length = memory.Span.IndexOf((byte)0);
+				if (length != -1)
+				{
+					Utf8String str = new Utf8String(memory.Slice(0, length));
+					reader.Advance(length + 1);
+					return str;
+				}
+
+				minSize = memory.Length + 1;
+			}
 		}
 
 		/// <summary>
-		/// Writes a UTF8 string into memory with a null terminator
+		/// Writes a null-terminated utf8 string to the buffer
 		/// </summary>
-		/// <param name="writer">The memory writer to serialize to</param>
+		/// <param name="writer">Writer for the output data</param>
 		/// <param name="str">String to write</param>
-		public static void WriteString(this MemoryWriter writer, Utf8String str)
+		public static void WriteNullTerminatedUtf8String(this IMemoryWriter writer, Utf8String str)
 		{
-			writer.WriteFixedLengthBytes(str.Span);
-			writer.WriteInt8(0);
+			Span<byte> span = writer.GetSpan(str.Length + 1);
+			str.Span.CopyTo(span);
+			span[str.Length] = 0;
+			writer.Advance(str.Length + 1);
 		}
 
 		/// <summary>
@@ -505,7 +545,7 @@ namespace EpicGames.Core
 		/// </summary>
 		/// <param name="str">The string to measure</param>
 		/// <returns>Size of the serialized string</returns>
-		public static int GetSerializedSize(this Utf8String str)
+		public static int GetNullTerminatedSize(this Utf8String str)
 		{
 			return str.Length + 1;
 		}
diff --git a/Engine/Source/Programs/Shared/EpicGames.Core/VarInt.cs b/Engine/Source/Programs/Shared/EpicGames.Core/VarInt.cs
new file mode 100644
--- /dev/null
+++ b/Engine/Source/Programs/Shared/EpicGames.Core/VarInt.cs
@@ -0,0 +1,308 @@
+// Copyright Epic Games, Inc. All Rights Reserved.
+
+using System;
+using System.Numerics;
+using EpicGames.Core;
+
+namespace EpicGames.Core
+{
+	/// <summary>
+	/// Methods for reading VarUInt values
+	/// </summary>
+	public static class VarInt
+	{
+		/// <summary>
+		/// Read a variable-length unsigned integer.
+		/// </summary>
+		/// <param name="buffer">A variable-length encoding of an unsigned integer</param>
+		/// <returns></returns>
+		[Obsolete("Call ReadSigned or ReadUnsigned instead")]
+		public static ulong Read(ReadOnlySpan<byte> buffer) => ReadUnsigned(buffer, out _);
+
+		/// <summary>
+		/// Read a variable-length unsigned integer.
+		/// </summary>
+		/// <param name="buffer">A variable-length encoding of an unsigned integer</param>
+		/// <param name="bytesRead">The number of bytes consumed from the input</param>
+		/// <returns></returns>
+		[Obsolete("Call ReadSigned or ReadUnsigned instead")]
+		public static ulong Read(ReadOnlySpan<byte> buffer, out int bytesRead) => ReadUnsigned(buffer, out bytesRead);
+
+		/// <summary>
+		/// Read a variable-length signed integer.
+		/// </summary>
+		/// <param name="buffer">A variable-length encoding of an unsigned integer</param>
+		/// <returns></returns>
+		public static long ReadSigned(ReadOnlySpan<byte> buffer) => ReadSigned(buffer, out _);
+
+		/// <summary>
+		/// Read a variable-length signed integer.
+		/// </summary>
+		/// <param name="buffer">A variable-length encoding of an unsigned integer</param>
+		/// <param name="bytesRead">The number of bytes consumed from the input</param>
+		/// <returns></returns>
+		public static long ReadSigned(ReadOnlySpan<byte> buffer, out int bytesRead)
+		{
+			ulong unsignedValue = ReadUnsigned(buffer, out bytesRead);
+			return DecodeSigned(unsignedValue);
+		}
+
+		/// <summary>
+		/// Read a variable-length unsigned integer.
+		/// </summary>
+		/// <param name="buffer">A variable-length encoding of an unsigned integer</param>
+		/// <returns></returns>
+		public static ulong ReadUnsigned(ReadOnlySpan<byte> buffer) => ReadUnsigned(buffer, out _);
+
+		/// <summary>
+		/// Read a variable-length unsigned integer.
+		/// </summary>
+		/// <param name="buffer">A variable-length encoding of an unsigned integer</param>
+		/// <param name="bytesRead">The number of bytes consumed from the input</param>
+		/// <returns></returns>
+		public static ulong ReadUnsigned(ReadOnlySpan<byte> buffer, out int bytesRead)
+		{
+			bytesRead = (int)Measure(buffer);
+			return ReadUnsignedKnownSize(buffer, bytesRead);
+		}
+
+		/// <summary>
+		/// Read a variable-length unsigned integer.
+		/// </summary>
+		/// <param name="buffer">A variable-length encoding of an unsigned integer</param>
+		/// <param name="numBytes">The number of bytes consumed from the input</param>
+		/// <returns></returns>
+		internal static ulong ReadUnsignedKnownSize(ReadOnlySpan<byte> buffer, int numBytes)
+		{
+			ulong value = (ulong)(buffer[0] & (0xff >> numBytes));
+			for (int i = 1; i < numBytes; i++)
+			{
+				value <<= 8;
+				value |= buffer[i];
+			}
+			return value;
+		}
+
+		/// <summary>
+		/// Measure the length in bytes (1-9) of an encoded variable-length integer.
+		/// </summary>
+		/// <param name="buffer">A variable-length encoding of an(signed or unsigned) integer.</param>
+		/// <returns>The number of bytes used to encode the integer, in the range 1-9.</returns>
+		public static int Measure(ReadOnlySpan<byte> buffer)
+		{
+			byte b = buffer[0];
+			b = (byte)~b;
+			return BitOperations.LeadingZeroCount(b) - 23;
+		}
+
+		/// <summary>
+		/// Measure the number of bytes required to encode the input.
+		/// </summary>
+		/// <param name="value"></param>
+		/// <returns></returns>
+		[Obsolete("Use MeasureSigned or MeasureUnsigned instead")]
+		public static int Measure(uint value) => MeasureUnsigned(value);
+
+		/// <summary>
+		/// Measure the number of bytes required to encode the input.
+		/// </summary>
+		/// <param name="value"></param>
+		/// <returns></returns>
+		[Obsolete("Use MeasureSigned or MeasureUnsigned instead")]
+		public static int Measure(ulong value) => MeasureUnsigned(value);
+
+		/// <summary>
+		/// Measure the number of bytes (1-9) required to encode the input.
+		/// </summary>
+		/// <param name="value"></param>
+		/// <returns></returns>
+		public static int MeasureSigned(long value)
+		{
+			return MeasureUnsigned(EncodeSigned(value));
+		}
+
+		/// <summary>
+		/// Measure the number of bytes (1-5) required to encode the 32-bit input.
+		/// </summary>
+		/// <param name="value"></param>
+		/// <returns></returns>
+		public static int MeasureUnsigned(int value)
+		{
+			return MeasureUnsigned((ulong)(long)value);
+		}
+
+		/// <summary>
+		/// Measure the number of bytes (1-5) required to encode the 32-bit input.
+		/// </summary>
+		/// <param name="value"></param>
+		/// <returns></returns>
+		public static int MeasureUnsigned(uint value)
+		{
+			return BitOperations.Log2(value) / 7 + 1;
+		}
+
+		/// <summary>
+		/// Measure the number of bytes (1-9) required to encode the 64-bit input.
+		/// </summary>
+		/// <param name="value"></param>
+		/// <returns></returns>
+		public static int MeasureUnsigned(ulong value)
+		{
+			return Math.Min(BitOperations.Log2(value) / 7 + 1, 9);
+		}
+
+		/// <summary>
+		/// Write a variable-length unsigned integer.
+		/// </summary>
+		/// <param name="value">An unsigned integer to encode</param>
+		/// <param name="buffer">A buffer of at least 9 bytes to write the output to.</param>
+		/// <returns>The number of bytes used in the output</returns>
+		[Obsolete("Use WriteUnsigned or WriteSigned instead")]
+		public static int Write(Span<byte> buffer, long value) => WriteUnsigned(buffer, (ulong)value);
+
+		/// <summary>
+		/// Write a variable-length unsigned integer.
+		/// </summary>
+		/// <param name="value">An unsigned integer to encode</param>
+		/// <param name="buffer">A buffer of at least 9 bytes to write the output to.</param>
+		/// <returns>The number of bytes used in the output</returns>
+		[Obsolete("Use WriteUnsigned or WriteSigned instead")]
+		public static int Write(Span<byte> buffer, ulong value) => WriteUnsigned(buffer, value);
+
+		/// <summary>
+		/// Write a variable-length unsigned integer.
+		/// </summary>
+		/// <param name="value">An unsigned integer to encode</param>
+		/// <param name="buffer">A buffer of at least 9 bytes to write the output to.</param>
+		/// <returns>The number of bytes used in the output</returns>
+		public static int WriteSigned(Span<byte> buffer, long value)
+		{
+			ulong unsignedValue = EncodeSigned(value);
+			return WriteUnsigned(buffer, unsignedValue);
+		}
+
+		/// <summary>
+		/// Write a variable-length unsigned integer.
+		/// </summary>
+		/// <param name="value">An unsigned integer to encode</param>
+		/// <param name="buffer">A buffer of at least 9 bytes to write the output to.</param>
+		/// <returns>The number of bytes used in the output</returns>
+		public static int WriteUnsigned(Span<byte> buffer, int value) => WriteUnsigned(buffer, (ulong)value);
+
+		/// <summary>
+		/// Write a variable-length unsigned integer.
+		/// </summary>
+		/// <param name="value">An unsigned integer to encode</param>
+		/// <param name="buffer">A buffer of at least 9 bytes to write the output to.</param>
+		/// <returns>The number of bytes used in the output</returns>
+		public static int WriteUnsigned(Span<byte> buffer, long value) => WriteUnsigned(buffer, (ulong)value);
+
+		/// <summary>
+		/// Write a variable-length unsigned integer.
+		/// </summary>
+		/// <param name="value">An unsigned integer to encode</param>
+		/// <param name="buffer">A buffer of at least 9 bytes to write the output to.</param>
+		/// <returns>The number of bytes used in the output</returns>
+		public static int WriteUnsigned(Span<byte> buffer, ulong value)
+		{
+			int byteCount = MeasureUnsigned(value);
+
+			for (int idx = 1; idx < byteCount; idx++)
+			{
+				buffer[byteCount - idx] = (byte)value;
+				value >>= 8;
+			}
+			buffer[0] = (byte)((0xff << (9 - (int)byteCount)) | (byte)value);
+			return byteCount;
+		}
+
+		/// <summary>
+		/// Decode a signed VarInt value from an unsigned value
+		/// </summary>
+		/// <param name="value">Value to decode</param>
+		/// <returns>Decoded value</returns>
+		public static long DecodeSigned(ulong value)
+		{
+			return -(long)(value & 1) ^ (long)(value >> 1);
+		}
+
+		/// <summary>
+		/// Encode a signed VarInt value into an unsigned value
+		/// </summary>
+		/// <param name="value">Value to encode</param>
+		/// <returns>Encoded value</returns>
+		public static ulong EncodeSigned(long value)
+		{
+			return (ulong)((value >> 63) ^ (value << 1));
+		}
+	}
+
+	/// <summary>
+	/// Extension methods for writing VarInt values
+	/// </summary>
+	public static class VarIntExtensions
+	{
+		/// <summary>
+		/// Read an unsigned VarInt from the given reader
+		/// </summary>
+		/// <param name="reader">Reader to deserialize from</param>
+		/// <returns>The deserialized value</returns>
+		public static long ReadSignedVarInt(this IMemoryReader reader)
+		{
+			ulong value = ReadUnsignedVarInt(reader);
+			return VarInt.DecodeSigned(value);
+		}
+
+		/// <summary>
+		/// Read an unsigned VarInt from the given reader
+		/// </summary>
+		/// <param name="reader">Reader to deserialize from</param>
+		/// <returns>The deserialized value</returns>
+		public static ulong ReadUnsignedVarInt(this IMemoryReader reader)
+		{
+			int length = VarInt.Measure(reader.GetSpan(1));
+			ReadOnlySpan<byte> span = reader.GetSpan(length);
+			reader.Advance(length);
+			return VarInt.ReadUnsignedKnownSize(span, length);
+		}
+
+		/// <summary>
+		/// Writes a signed VarInt to a byte array
+		/// </summary>
+		/// <param name="writer">Writer to serialize to</param>
+		/// <param name="value">Value to write</param>
+		public static void WriteSignedVarInt(this IMemoryWriter writer, int value)
+		{
+			WriteUnsignedVarInt(writer, VarInt.EncodeSigned(value));
+		}
+
+		/// <summary>
+		/// Writes a signed VarInt to a byte array
+		/// </summary>
+		/// <param name="writer">Writer to serialize to</param>
+		/// <param name="value">Value to write</param>
+		public static void WriteSignedVarInt(this IMemoryWriter writer, long value)
+		{
+			WriteUnsignedVarInt(writer, VarInt.EncodeSigned(value));
+		}
+
+		/// <summary>
+		/// Writes a unsigned VarInt to a byte array
+		/// </summary>
+		/// <param name="writer">Writer to serialize to</param>
+		/// <param name="value">Value to write</param>
+		public static void WriteUnsignedVarInt(this IMemoryWriter writer, int value) => WriteUnsignedVarInt(writer, (ulong)value);
+
+		/// <summary>
+		/// Writes a unsigned VarInt to a byte array
+		/// </summary>
+		/// <param name="writer">Writer to serialize to</param>
+		/// <param name="value">Value to write</param>
+		public static void WriteUnsignedVarInt(this IMemoryWriter writer, ulong value)
+		{
+			int length = VarInt.MeasureUnsigned(value);
+			VarInt.WriteUnsigned(writer.GetSpan(length), value);
+			writer.Advance(length);
+		}
+	}
+}
diff --git a/Engine/Source/Programs/Shared/EpicGames.IoHash/IoHash.cs b/Engine/Source/Programs/Shared/EpicGames.IoHash/IoHash.cs
--- a/Engine/Source/Programs/Shared/EpicGames.IoHash/IoHash.cs
+++ b/Engine/Source/Programs/Shared/EpicGames.IoHash/IoHash.cs
@@ -266,7 +266,7 @@ namespace EpicGames.Core
 		/// </summary>
 		/// <param name="reader"></param>
 		/// <returns></returns>
-		public static IoHash ReadIoHash(this MemoryReader reader)
+		public static IoHash ReadIoHash(this IMemoryReader reader)
 		{
 			return new IoHash(reader.ReadFixedLengthBytes(IoHash.NumBytes).Span);
 		}
@@ -276,9 +276,10 @@ namespace EpicGames.Core
 		/// </summary>
 		/// <param name="writer"></param>
 		/// <param name="hash"></param>
-		public static void WriteIoHash(this MemoryWriter writer, IoHash hash)
+		public static void WriteIoHash(this IMemoryWriter writer, IoHash hash)
 		{
-			hash.CopyTo(writer.AllocateSpan(IoHash.NumBytes));
+			hash.CopyTo(writer.GetSpan(IoHash.NumBytes));
+			writer.Advance(IoHash.NumBytes);
 		}
 	}
 
diff --git a/Engine/Source/Programs/Shared/EpicGames.Perforce.Managed/FileContentId.cs b/Engine/Source/Programs/Shared/EpicGames.Perforce.Managed/FileContentId.cs
--- a/Engine/Source/Programs/Shared/EpicGames.Perforce.Managed/FileContentId.cs
+++ b/Engine/Source/Programs/Shared/EpicGames.Perforce.Managed/FileContentId.cs
@@ -41,19 +41,19 @@ namespace EpicGames.Perforce.Managed
 		public static FileContentId ReadFileContentId(this MemoryReader reader)
 		{
 			Md5Hash digest = reader.ReadMd5Hash();
-			Utf8String type = reader.ReadString();
+			Utf8String type = reader.ReadNullTerminatedUtf8String();
 			return new FileContentId(digest, type);
 		}
 
 		public static void WriteFileContentId(this MemoryWriter writer, FileContentId fileContentId)
 		{
 			writer.WriteMd5Hash(fileContentId.Digest);
-			writer.WriteString(fileContentId.Type);
+			writer.WriteNullTerminatedUtf8String(fileContentId.Type);
 		}
 
 		public static int GetSerializedSize(this FileContentId fileContentId)
 		{
-			return Digest<Md5>.Length + fileContentId.Type.GetSerializedSize();
+			return Digest<Md5>.Length + fileContentId.Type.GetNullTerminatedSize();
 		}
 	}
 }
diff --git a/Engine/Source/Programs/Shared/EpicGames.Perforce.Managed/ManagedWorkspace.cs b/Engine/Source/Programs/Shared/EpicGames.Perforce.Managed/ManagedWorkspace.cs
--- a/Engine/Source/Programs/Shared/EpicGames.Perforce.Managed/ManagedWorkspace.cs
+++ b/Engine/Source/Programs/Shared/EpicGames.Perforce.Managed/ManagedWorkspace.cs
@@ -393,7 +393,7 @@ namespace EpicGames.Perforce.Managed
 				writer.WriteCachedFileInfo(trackedFile);
 			}
 			writer.WriteWorkspaceDirectoryInfo(_workspace);
-			writer.CheckOffset(serializedSize);
+			writer.CheckEmpty();
 
 			// Write it to disk
 			FileReference dataFile = FileReference.Combine(_baseDir, DataFileName);
diff --git a/Engine/Source/Programs/Shared/EpicGames.Perforce.Managed/WorkspaceDirectoryInfo.cs b/Engine/Source/Programs/Shared/EpicGames.Perforce.Managed/WorkspaceDirectoryInfo.cs
--- a/Engine/Source/Programs/Shared/EpicGames.Perforce.Managed/WorkspaceDirectoryInfo.cs
+++ b/Engine/Source/Programs/Shared/EpicGames.Perforce.Managed/WorkspaceDirectoryInfo.cs
@@ -356,7 +356,7 @@ namespace EpicGames.Perforce.Managed
 			int numSubDirectories = reader.ReadInt32();
 			for (int idx = 0; idx < numSubDirectories; idx++)
 			{
-				Utf8String name = reader.ReadString();
+				Utf8String name = reader.ReadNullTerminatedUtf8String();
 
 				WorkspaceDirectoryInfo subDirectory = new WorkspaceDirectoryInfo(directoryInfo, name, null);
 				reader.ReadWorkspaceDirectoryInfo(subDirectory, version);
@@ -377,14 +377,14 @@ namespace EpicGames.Perforce.Managed
 			writer.WriteInt32(directoryInfo.NameToSubDirectory.Count);
 			foreach (WorkspaceDirectoryInfo subDirectory in directoryInfo.NameToSubDirectory.Values)
 			{
-				writer.WriteString(subDirectory.Name);
+				writer.WriteNullTerminatedUtf8String(subDirectory.Name);
 				writer.WriteWorkspaceDirectoryInfo(subDirectory);
 			}
 		}
 
 		public static int GetSerializedSize(this WorkspaceDirectoryInfo directoryInfo)
 		{
-			return Digest<Sha1>.Length + sizeof(int) + directoryInfo.NameToFile.Values.Sum(x => x.GetSerializedSize()) + sizeof(int) + directoryInfo.NameToSubDirectory.Values.Sum(x => x.Name.GetSerializedSize() + x.GetSerializedSize());
+			return Digest<Sha1>.Length + sizeof(int) + directoryInfo.NameToFile.Values.Sum(x => x.GetSerializedSize()) + sizeof(int) + directoryInfo.NameToSubDirectory.Values.Sum(x => x.Name.GetNullTerminatedSize() + x.GetSerializedSize());
 		}
 	}
 }
diff --git a/Engine/Source/Programs/Shared/EpicGames.Perforce.Managed/WorkspaceFileInfo.cs b/Engine/Source/Programs/Shared/EpicGames.Perforce.Managed/WorkspaceFileInfo.cs
--- a/Engine/Source/Programs/Shared/EpicGames.Perforce.Managed/WorkspaceFileInfo.cs
+++ b/Engine/Source/Programs/Shared/EpicGames.Perforce.Managed/WorkspaceFileInfo.cs
@@ -95,7 +95,7 @@ namespace EpicGames.Perforce.Managed
 	{
 		public static WorkspaceFileInfo ReadWorkspaceFileInfo(this MemoryReader reader, WorkspaceDirectoryInfo directory)
 		{
-			Utf8String name = reader.ReadString();
+			Utf8String name = reader.ReadNullTerminatedUtf8String();
 			long length = reader.ReadInt64();
 			long lastModifiedTicks = reader.ReadInt64();
 			bool bReadOnly = reader.ReadBoolean();
@@ -105,7 +105,7 @@ namespace EpicGames.Perforce.Managed
 
 		public static void WriteWorkspaceFileInfo(this MemoryWriter writer, WorkspaceFileInfo fileInfo)
 		{
-			writer.WriteString(fileInfo.Name);
+			writer.WriteNullTerminatedUtf8String(fileInfo.Name);
 			writer.WriteInt64(fileInfo._length);
 			writer.WriteInt64(fileInfo._lastModifiedTicks);
 			writer.WriteBoolean(fileInfo._bReadOnly);
@@ -114,7 +114,7 @@ namespace EpicGames.Perforce.Managed
 
 		public static int GetSerializedSize(this WorkspaceFileInfo fileInfo)
 		{
-			return fileInfo.Name.GetSerializedSize() + sizeof(long) + sizeof(long) + sizeof(byte) + fileInfo.ContentId.GetSerializedSize();
+			return fileInfo.Name.GetNullTerminatedSize() + sizeof(long) + sizeof(long) + sizeof(byte) + fileInfo.ContentId.GetSerializedSize();
 		}
 	}
 }
diff --git a/Engine/Source/Programs/Shared/EpicGames.Serialization.Tests/VarIntTests.cs b/Engine/Source/Programs/Shared/EpicGames.Serialization.Tests/VarIntTests.cs
deleted file mode 100644
diff --git a/Engine/Source/Programs/Shared/EpicGames.Serialization/VarInt.cs b/Engine/Source/Programs/Shared/EpicGames.Serialization/VarInt.cs
deleted file mode 100644
