EpicGames.Core: Fix some serialization issues with IMemoryReader/IMemoryWriter, and add a class to build ReadOnlySequence<T> objects.

diff --git a/Engine/Source/Programs/Shared/EpicGames.Core.Tests/MemorySerializerTests.cs b/Engine/Source/Programs/Shared/EpicGames.Core.Tests/MemorySerializerTests.cs
new file mode 100644
--- /dev/null
+++ b/Engine/Source/Programs/Shared/EpicGames.Core.Tests/MemorySerializerTests.cs
@@ -0,0 +1,109 @@
+// Copyright Epic Games, Inc. All Rights Reserved.
+
+using System;
+using System.Collections.Generic;
+using System.Text;
+using Microsoft.VisualStudio.TestTools.UnitTesting;
+
+namespace EpicGames.Core.Tests
+{
+	[TestClass]
+	public class MemorySerializerTests
+	{
+		[TestMethod]
+		public void SignedIntegers()
+		{
+			byte[] data;
+			{
+				ByteArrayBuilder writer = new ByteArrayBuilder();
+				writer.WriteInt32(12345);
+				writer.WriteInt32(-12345);
+				writer.WriteInt8(SByte.MinValue);
+				writer.WriteInt8(SByte.MaxValue);
+				writer.WriteInt16(Int16.MinValue);
+				writer.WriteInt16(Int16.MaxValue);
+				writer.WriteInt32(Int32.MinValue);
+				writer.WriteInt32(Int32.MaxValue);
+				writer.WriteInt64(Int64.MinValue);
+				writer.WriteInt64(Int64.MaxValue);
+				data = writer.ToByteArray();
+			}
+			{
+				MemoryReader reader = new MemoryReader(data);
+				Assert.AreEqual(12345, reader.ReadInt32());
+				Assert.AreEqual(-12345, reader.ReadInt32());
+				Assert.AreEqual(SByte.MinValue, reader.ReadInt8());
+				Assert.AreEqual(SByte.MaxValue, reader.ReadInt8());
+				Assert.AreEqual(Int16.MinValue, reader.ReadInt16());
+				Assert.AreEqual(Int16.MaxValue, reader.ReadInt16());
+				Assert.AreEqual(Int32.MinValue, reader.ReadInt32());
+				Assert.AreEqual(Int32.MaxValue, reader.ReadInt32());
+				Assert.AreEqual(Int64.MinValue, reader.ReadInt64());
+				Assert.AreEqual(Int64.MaxValue, reader.ReadInt64());
+				reader.CheckEmpty();
+			}
+		}
+
+		[TestMethod]
+		public void UnsignedIntegers()
+		{
+			byte[] data;
+			{
+				ByteArrayBuilder writer = new ByteArrayBuilder();
+				writer.WriteUInt32(12345);
+				writer.WriteUInt8(Byte.MaxValue);
+				writer.WriteUInt16(UInt16.MaxValue);
+				writer.WriteUInt32(UInt32.MaxValue);
+				writer.WriteUInt64(UInt64.MaxValue);
+				data = writer.ToByteArray();
+			}
+			{
+				MemoryReader reader = new MemoryReader(data);
+				Assert.AreEqual(12345U, reader.ReadUInt32());
+				Assert.AreEqual(Byte.MaxValue, reader.ReadUInt8());
+				Assert.AreEqual(UInt16.MaxValue, reader.ReadUInt16());
+				Assert.AreEqual(UInt32.MaxValue, reader.ReadUInt32());
+				Assert.AreEqual(UInt64.MaxValue, reader.ReadUInt64());
+				reader.CheckEmpty();
+			}
+		}
+
+		[TestMethod]
+		public void TestMemory()
+		{
+			byte[] data;
+			{
+				ByteArrayBuilder writer = new ByteArrayBuilder();
+				writer.WriteVariableLengthBytes(new byte[] { 1, 2, 3 });
+				writer.WriteVariableLengthBytes(new byte[] { 4, 5, 6 });
+				data = writer.ToByteArray();
+			}
+			{
+				MemoryReader reader = new MemoryReader(data);
+				Assert.IsTrue(reader.ReadVariableLengthBytes().Span.SequenceEqual(new byte[] { 1, 2, 3 }));
+				Assert.IsTrue(reader.ReadVariableLengthBytes().Span.SequenceEqual(new byte[] { 4, 5, 6 }));
+				reader.CheckEmpty();
+			}
+		}
+
+		[TestMethod]
+		public void TestStrings()
+		{
+			byte[] data;
+			{
+				ByteArrayBuilder writer = new ByteArrayBuilder();
+				writer.WriteString("hello", Encoding.UTF8);
+				writer.WriteString("there", Encoding.Unicode);
+				writer.WriteString("world", Encoding.UTF32);
+				data = writer.ToByteArray();
+			}
+			{
+				MemoryReader reader = new MemoryReader(data);
+				Assert.AreEqual("hello", reader.ReadString(Encoding.UTF8));
+				Assert.AreEqual("there", reader.ReadString(Encoding.Unicode));
+				Assert.AreEqual("world", reader.ReadString(Encoding.UTF32));
+				reader.CheckEmpty();
+			}
+		}
+	}
+}
diff --git a/Engine/Source/Programs/Shared/EpicGames.Core/ByteArrayBuilder.cs b/Engine/Source/Programs/Shared/EpicGames.Core/ByteArrayBuilder.cs
--- a/Engine/Source/Programs/Shared/EpicGames.Core/ByteArrayBuilder.cs
+++ b/Engine/Source/Programs/Shared/EpicGames.Core/ByteArrayBuilder.cs
@@ -27,16 +27,6 @@ namespace EpicGames.Core
 			public ReadOnlyMemory<byte> WrittenMemory => Data.AsMemory(0, Length);
 		}
 
-		class Segment : ReadOnlySequenceSegment<byte>
-		{
-			public Segment(Chunk chunk, Segment? next)
-			{
-				Memory = chunk.WrittenMemory;
-				RunningIndex = chunk.RunningIndex;
-				Next = next;
-			}
-		}
-
 		readonly List<Chunk> _chunks = new List<Chunk>();
 		readonly int _chunkSize;
 		Chunk _currentChunk;
@@ -125,11 +115,28 @@ namespace EpicGames.Core
 			return span;
 		}
 
+		/// <summary>
+		/// Appends data in this builder to the given sequence
+		/// </summary>
+		/// <param name="builder">Sequence builder</param>
+		public void AppendTo(ReadOnlySequenceBuilder<byte> builder)
+		{
+			foreach (Chunk chunk in _chunks)
+			{
+				builder.Append(chunk.WrittenMemory);
+			}
+		}
+
 		/// <summary>
 		/// Gets a sequence representing the bytes that have been written so far
 		/// </summary>
 		/// <returns>Sequence of bytes</returns>
-		public ReadOnlySequence<byte> AsSequence() => AsSequence(0);
+		public ReadOnlySequence<byte> AsSequence()
+		{
+			ReadOnlySequenceBuilder<byte> builder = new ReadOnlySequenceBuilder<byte>();
+			AppendTo(builder);
+			return builder.Construct();
+		}
 
 		/// <summary>
 		/// Gets a sequence representing the bytes that have been written so far, starting at the given offset
@@ -138,22 +145,8 @@ namespace EpicGames.Core
 		/// <returns>Sequence of bytes</returns>
 		public ReadOnlySequence<byte> AsSequence(int offset)
 		{
-			if (offset == Length)
-			{
-				return ReadOnlySequence<byte>.Empty;
-			}
-			else
-			{
-				Segment last = new Segment(_chunks[^1], null);
-
-				Segment first = last;
-				for (int idx = _chunks.Count - 1; idx >= 0 && offset < _chunks[idx].RunningIndex + _chunks[idx].Length; idx--)
-				{
-					first = new Segment(_chunks[idx], first);
-				}
-
-				return new ReadOnlySequence<byte>(first, offset - (int)first.RunningIndex, last, last.Memory.Length);
-			}
+			// TODO: could do a binary search for offset and work forwards from there
+			return AsSequence().Slice(offset);
 		}
 
 		/// <summary>
@@ -164,8 +157,8 @@ namespace EpicGames.Core
 		/// <returns>Sequence of bytes</returns>
 		public ReadOnlySequence<byte> AsSequence(int offset, int length)
 		{
-			// TODO: could do a binary search for offset + length and work back from there
-			return AsSequence(offset).Slice(0, length);
+			// TODO: could do a binary search for offset and work fowards from there
+			return AsSequence().Slice(offset, length);
 		}
 
 		/// <summary>
diff --git a/Engine/Source/Programs/Shared/EpicGames.Core/MemoryReader.cs b/Engine/Source/Programs/Shared/EpicGames.Core/MemoryReader.cs
--- a/Engine/Source/Programs/Shared/EpicGames.Core/MemoryReader.cs
+++ b/Engine/Source/Programs/Shared/EpicGames.Core/MemoryReader.cs
@@ -218,7 +218,7 @@ namespace EpicGames.Core
 		/// <returns>Sequence of bytes</returns>
 		public static ReadOnlyMemory<byte> ReadFixedLengthBytes(this IMemoryReader reader, int length)
 		{
-			ReadOnlyMemory<byte> bytes = reader.GetMemory(length);
+			ReadOnlyMemory<byte> bytes = reader.GetMemory(length).Slice(0, length);
 			reader.Advance(length);
 			return bytes;
 		}
@@ -230,7 +230,7 @@ namespace EpicGames.Core
 		/// <param name="readItem">Delegate to write an individual item</param>
 		public static T[] ReadVariableLengthArray<T>(this IMemoryReader reader, Func<T> readItem)
 		{
-			int length = ReadInt32(reader);
+			int length = (int)reader.ReadUnsignedVarInt();
 			return ReadFixedLengthArray(reader, length, readItem);
 		}
 
diff --git a/Engine/Source/Programs/Shared/EpicGames.Core/MemoryWriter.cs b/Engine/Source/Programs/Shared/EpicGames.Core/MemoryWriter.cs
--- a/Engine/Source/Programs/Shared/EpicGames.Core/MemoryWriter.cs
+++ b/Engine/Source/Programs/Shared/EpicGames.Core/MemoryWriter.cs
@@ -3,6 +3,7 @@
 using System;
 using System.Buffers;
 using System.Buffers.Binary;
+using System.Collections.Generic;
 using System.Text;
 
 namespace EpicGames.Core
@@ -198,6 +199,7 @@ namespace EpicGames.Core
 			Span<byte> span = GetSpan(writer, lengthBytes + bytes.Length);
 			VarInt.WriteUnsigned(span, bytes.Length);
 			bytes.CopyTo(span[lengthBytes..]);
+			writer.Advance(lengthBytes + bytes.Length);
 		}
 
 		/// <summary>
@@ -227,26 +229,26 @@ namespace EpicGames.Core
 		/// Writes a variable length array
 		/// </summary>
 		/// <param name="writer">Writer to serialize to</param>
-		/// <param name="array">The array to write</param>
+		/// <param name="list">The array to write</param>
 		/// <param name="writeItem">Delegate to write an individual item</param>
-		public static void WriteVariableLengthArray<T>(this IMemoryWriter writer, T[] array, Action<T> writeItem)
+		public static void WriteVariableLengthArray<T>(this IMemoryWriter writer, IReadOnlyList<T> list, Action<T> writeItem)
 		{
-			WriteInt32(writer, array.Length);
-			WriteFixedLengthArray(writer, array, writeItem);
+			writer.WriteUnsignedVarInt(list.Count);
+			WriteFixedLengthArray(writer, list, writeItem);
 		}
 
 		/// <summary>
 		/// Writes a fixed length array
 		/// </summary>
 		/// <param name="writer">Writer to serialize to</param>
-		/// <param name="array">The array to write</param>
+		/// <param name="list">The array to write</param>
 		/// <param name="writeItem">Delegate to write an individual item</param>
 		[System.Diagnostics.CodeAnalysis.SuppressMessage("Style", "IDE0060:Remove unused parameter")]
-		public static void WriteFixedLengthArray<T>(this IMemoryWriter writer, T[] array, Action<T> writeItem)
+		public static void WriteFixedLengthArray<T>(this IMemoryWriter writer, IReadOnlyList<T> list, Action<T> writeItem)
 		{
-			for (int idx = 0; idx < array.Length; idx++)
+			for (int idx = 0; idx < list.Count; idx++)
 			{
-				writeItem(array[idx]);
+				writeItem(list[idx]);
 			}
 		}
 
diff --git a/Engine/Source/Programs/Shared/EpicGames.Core/ReadOnlySequenceBuilder.cs b/Engine/Source/Programs/Shared/EpicGames.Core/ReadOnlySequenceBuilder.cs
new file mode 100644
--- /dev/null
+++ b/Engine/Source/Programs/Shared/EpicGames.Core/ReadOnlySequenceBuilder.cs
@@ -0,0 +1,80 @@
+// Copyright Epic Games, Inc. All Rights Reserved.
+
+using System;
+using System.Buffers;
+using System.Collections.Generic;
+using System.Text;
+
+namespace EpicGames.Core
+{
+	/// <summary>
+	/// Utility class to combine buffers into a <see cref="ReadOnlySequence{T}"/>
+	/// </summary>
+	/// <typeparam name="T">Element type</typeparam>
+	public class ReadOnlySequenceBuilder<T>
+	{
+		class Segment : ReadOnlySequenceSegment<T>
+		{
+			public Segment(long runningIndex, ReadOnlyMemory<T> memory)
+			{
+				RunningIndex = runningIndex;
+				Memory = memory;
+			}
+
+			public void SetNext(Segment next)
+			{
+				Next = next;
+			}
+		}
+
+		readonly List<ReadOnlyMemory<T>> _segments = new List<ReadOnlyMemory<T>>();
+
+		/// <summary>
+		/// Append a block of memory to the end of the sequence
+		/// </summary>
+		/// <param name="memory">Memory to append</param>
+		public void Append(ReadOnlyMemory<T> memory)
+		{
+			if (memory.Length > 0)
+			{
+				_segments.Add(memory);
+			}
+		}
+
+		/// <summary>
+		/// Append another sequence to the end of this one
+		/// </summary>
+		/// <param name="sequence">Sequence to append</param>
+		public void Append(ReadOnlySequence<T> sequence)
+		{
+			foreach (ReadOnlyMemory<T> segment in sequence)
+			{
+				Append(segment);
+			}
+		}
+
+		/// <summary>
+		/// Construct a sequence from the added blocks
+		/// </summary>
+		/// <returns>Sequence for the added memory blocks</returns>
+		public ReadOnlySequence<T> Construct()
+		{
+			if (_segments.Count == 0)
+			{
+				return ReadOnlySequence<T>.Empty;
+			}
+
+			Segment first = new Segment(0, _segments[0]);
+			Segment last = first;
+
+			for (int idx = 1; idx < _segments.Count; idx++)
+			{
+				Segment next = new Segment(last.RunningIndex + last.Memory.Length, _segments[idx]);
+				last.SetNext(next);
+				last = next;
+			}
+
+			return new ReadOnlySequence<T>(first, 0, last, last.Memory.Length);
+		}
+	}
+}
