Fix static analysis warnings.

diff --git a/Engine/Source/Programs/Shared/EpicGames.Core/MemoryReader.cs b/Engine/Source/Programs/Shared/EpicGames.Core/MemoryReader.cs
--- a/Engine/Source/Programs/Shared/EpicGames.Core/MemoryReader.cs
+++ b/Engine/Source/Programs/Shared/EpicGames.Core/MemoryReader.cs
@@ -54,7 +54,6 @@ namespace EpicGames.Core
 		/// <summary>
 		/// Checks that we've used the exact buffer length
 		/// </summary>
-		/// <param name="expectedOffset">Expected offset within the output buffer</param>
 		public void CheckEmpty()
 		{
 			if (Memory.Length > 0)
diff --git a/Engine/Source/Programs/Shared/EpicGames.Core/MemoryWriter.cs b/Engine/Source/Programs/Shared/EpicGames.Core/MemoryWriter.cs
--- a/Engine/Source/Programs/Shared/EpicGames.Core/MemoryWriter.cs
+++ b/Engine/Source/Programs/Shared/EpicGames.Core/MemoryWriter.cs
@@ -8,6 +8,9 @@ using System.Text;
 
 namespace EpicGames.Core
 {
+	/// <summary>
+	/// Interface for serializing to a memory buffer
+	/// </summary>
 	public interface IMemoryWriter
 	{
 		/// <summary>
@@ -46,7 +49,6 @@ namespace EpicGames.Core
 		/// <summary>
 		/// Checks that we've used the exact buffer length
 		/// </summary>
-		/// <param name="expectedOffset">Expected offset within the output buffer</param>
 		public void CheckEmpty()
 		{
 			if (_memory.Length > 0)
