Low level blob storage interface.

* Blobs are not explicitly content addressed, but each uploaded blob is given an implementation-defined BlobId  for retrieval later. An implementation MAY return a hash value for a BlobId if desired.
* BlobIds are utf8 strings, with a format that is implementation defined. They may be hashes, guids, objectids, or locators within some storage hierarchy. Consumers MAY NOT assume any particular format of a BlobId.
* Blobs may reference other blobs via a table of references stored alongside (or within) it. A caller SHOULD enforce referential integrity. An implementation MAY verify and enforce it.
* Mapping from a known key to a particular blob is done via refs, using a similar interface to the existing storage interface.
* Blobs are expected to be small enough to fit into memory. Larger blobs should be split up. Max size is TBD; probably of the order of 1-10mb.

Abstracting blob storage at this level allows for simpler implementations where use cases allow it, such as writing directly to S3 or a filer, and allows for changes in the way we address blobs to allow them to be retrieved from locations encoded in the BlobId without having to ensure global replication.

diff --git a/Engine/Source/Programs/Shared/EpicGames.Horde/Storage/Backends/InMemoryBlobStore.cs b/Engine/Source/Programs/Shared/EpicGames.Horde/Storage/Backends/InMemoryBlobStore.cs
new file mode 100644
--- /dev/null
+++ b/Engine/Source/Programs/Shared/EpicGames.Horde/Storage/Backends/InMemoryBlobStore.cs
@@ -0,0 +1,87 @@
+// Copyright Epic Games, Inc. All Rights Reserved.
+
+using System;
+using System.Buffers;
+using System.Collections.Concurrent;
+using System.Collections.Generic;
+using System.Threading;
+using System.Threading.Tasks;
+
+namespace EpicGames.Horde.Storage.Impl
+{
+	/// <summary>
+	/// Implementation of <see cref="IBlobStore"/> which stores data in memory. Not intended for production use.
+	/// </summary>
+	public class InMemoryBlobStore : IBlobStore
+	{
+		class Blob : IBlob
+		{
+			readonly ReadOnlyMemory<byte> _data;
+			readonly IReadOnlyList<BlobId> _references;
+
+			public Blob(ReadOnlyMemory<byte> data, IReadOnlyList<BlobId> references)
+			{
+				_data = data;
+				_references = references;
+			}
+
+			/// <inheritdoc/>
+			public ValueTask<ReadOnlyMemory<byte>> GetDataAsync() => new ValueTask<ReadOnlyMemory<byte>>(_data);
+
+			/// <inheritdoc/>
+			public ValueTask<IReadOnlyList<BlobId>> GetReferencesAsync() => new ValueTask<IReadOnlyList<BlobId>>(_references);
+		}
+
+		/// <summary>
+		/// Map of blob id to blob data
+		/// </summary>
+		readonly ConcurrentDictionary<BlobId, IBlob> _blobs = new ConcurrentDictionary<BlobId, IBlob>();
+
+		/// <summary>
+		/// Map of ref name to ref data
+		/// </summary>
+		readonly ConcurrentDictionary<RefId, IBlob> _refs = new ConcurrentDictionary<RefId, IBlob>();
+
+		/// <inheritdoc cref="_blobs"/>
+		public IReadOnlyDictionary<BlobId, IBlob> Blobs => _blobs;
+
+		/// <inheritdoc cref="_refs"/>
+		public IReadOnlyDictionary<RefId, IBlob> Refs => _refs;
+
+		#region Blobs
+
+		/// <inheritdoc/>
+		public Task<IBlob> ReadBlobAsync(BlobId blobId, CancellationToken cancellationToken = default) => Task.FromResult(_blobs[blobId]);
+
+		/// <inheritdoc/>
+		public Task<BlobId> WriteBlobAsync(ReadOnlySequence<byte> data, IReadOnlyList<BlobId> references, CancellationToken cancellationToken = default)
+		{
+			BlobId blobId = new BlobId(Guid.NewGuid().ToString());
+			Blob blob = new Blob(data.ToArray(), references);
+			_blobs[blobId] = blob;
+			return Task.FromResult(blobId);
+		}
+
+		#endregion
+
+		#region Refs
+
+		/// <inheritdoc/>
+		public Task<bool> DeleteRefAsync(RefId id, CancellationToken cancellationToken) => Task.FromResult(_refs.TryRemove(id, out _));
+
+		/// <inheritdoc/>
+		public Task<bool> HasRefAsync(RefId id, CancellationToken cancellationToken) => Task.FromResult(_refs.ContainsKey(id));
+
+		/// <inheritdoc/>
+		public Task<IBlob> ReadRefAsync(RefId id, CancellationToken cancellationToken) => Task.FromResult(_refs[id]);
+
+		/// <inheritdoc/>
+		public Task WriteRefAsync(RefId id, ReadOnlySequence<byte> data, IReadOnlyList<BlobId> references, CancellationToken cancellationToken)
+		{
+			_refs[id] = new Blob(data.ToArray(), references);
+			return Task.CompletedTask;
+		}
+
+		#endregion
+	}
+}
diff --git a/Engine/Source/Programs/Shared/EpicGames.Horde/Storage/BlobId.cs b/Engine/Source/Programs/Shared/EpicGames.Horde/Storage/BlobId.cs
new file mode 100644
--- /dev/null
+++ b/Engine/Source/Programs/Shared/EpicGames.Horde/Storage/BlobId.cs
@@ -0,0 +1,126 @@
+// Copyright Epic Games, Inc. All Rights Reserved.
+
+using EpicGames.Core;
+using EpicGames.Serialization;
+using System;
+using System.ComponentModel;
+using System.Globalization;
+using System.Text;
+using System.Text.Json;
+using System.Text.Json.Serialization;
+
+namespace EpicGames.Horde.Storage
+{
+	/// <summary>
+	/// Unique identifier for a blob, as a utf-8 string. Clients should not assume any internal structure to this identifier; it only
+	/// has meaning to the <see cref="IBlobStore"/> implementation.
+	/// </summary>
+	[JsonConverter(typeof(BlobIdJsonConverter))]
+	[TypeConverter(typeof(BlobIdTypeConverter))]
+	[CbConverter(typeof(BlobIdCbConverter))]
+	public struct BlobId : IEquatable<BlobId>
+	{
+		/// <summary>
+		/// Identifier for the ref
+		/// </summary>
+		public Utf8String Inner { get; }
+
+		/// <summary>
+		/// Constructor
+		/// </summary>
+		/// <param name="inner"></param>
+		public BlobId(Utf8String inner)
+		{
+			Inner = inner;
+		}
+
+		/// <inheritdoc/>
+		public override bool Equals(object? obj) => obj is RefId refId && Equals(refId);
+
+		/// <inheritdoc/>
+		public override int GetHashCode() => Inner.GetHashCode();
+
+		/// <inheritdoc/>
+		public bool Equals(BlobId blobId) => Inner == blobId.Inner;
+
+		/// <inheritdoc/>
+		public override string ToString() => Inner.ToString();
+
+		/// <inheritdoc/>
+		public static bool operator ==(BlobId lhs, BlobId rhs) => lhs.Equals(rhs);
+
+		/// <inheritdoc/>
+		public static bool operator !=(BlobId lhs, BlobId rhs) => !lhs.Equals(rhs);
+	}
+
+	/// <summary>
+	/// Type converter for BlobId to and from JSON
+	/// </summary>
+	sealed class BlobIdJsonConverter : JsonConverter<BlobId>
+	{
+		/// <inheritdoc/>
+		public override BlobId Read(ref Utf8JsonReader reader, Type typeToConvert, JsonSerializerOptions options) => new BlobId(new Utf8String(reader.GetUtf8String().ToArray()));
+
+		/// <inheritdoc/>
+		public override void Write(Utf8JsonWriter writer, BlobId value, JsonSerializerOptions options) => writer.WriteStringValue(value.Inner.Span);
+	}
+
+	/// <summary>
+	/// Type converter from strings to BlobId objects
+	/// </summary>
+	sealed class BlobIdTypeConverter : TypeConverter
+	{
+		/// <inheritdoc/>
+		public override bool CanConvertFrom(ITypeDescriptorContext? context, Type sourceType)
+		{
+			return sourceType == typeof(string);
+		}
+
+		/// <inheritdoc/>
+		public override object ConvertFrom(ITypeDescriptorContext? context, CultureInfo? culture, object? value)
+		{
+			return new BlobId((string)value!);
+		}
+	}
+
+	/// <summary>
+	/// Type converter to compact binary
+	/// </summary>
+	sealed class BlobIdCbConverter : CbConverterBase<BlobId>
+	{
+		/// <inheritdoc/>
+		public override BlobId Read(CbField field) => new BlobId(field.AsUtf8String());
+
+		/// <inheritdoc/>
+		public override void Write(CbWriter writer, BlobId value) => writer.WriteUtf8StringValue(value.Inner);
+
+		/// <inheritdoc/>
+		public override void WriteNamed(CbWriter writer, Utf8String name, BlobId value) => writer.WriteUtf8String(name, value.Inner);
+	}
+
+	/// <summary>
+	/// Extension methods for blob ids
+	/// </summary>
+	public static class BlobIdExtensions
+	{
+		/// <summary>
+		/// Deserialize a blob id
+		/// </summary>
+		/// <param name="reader">Reader to deserialize from</param>
+		/// <returns>The blob id that was read</returns>
+		public static BlobId ReadBlobId(this IMemoryReader reader)
+		{
+			return new BlobId(reader.ReadUtf8String());
+		}
+
+		/// <summary>
+		/// Serialize a blob id
+		/// </summary>
+		/// <param name="writer">Writer to serialize to</param>
+		/// <param name="value">Value to serialize</param>
+		public static void WriteBlobId(this IMemoryWriter writer, BlobId value)
+		{
+			writer.WriteUtf8String(value.Inner);
+		}
+	}
+}
diff --git a/Engine/Source/Programs/Shared/EpicGames.Horde/Storage/IBlobStore.cs b/Engine/Source/Programs/Shared/EpicGames.Horde/Storage/IBlobStore.cs
new file mode 100644
--- /dev/null
+++ b/Engine/Source/Programs/Shared/EpicGames.Horde/Storage/IBlobStore.cs
@@ -0,0 +1,93 @@
+// Copyright Epic Games, Inc. All Rights Reserved.
+
+using System;
+using System.Buffers;
+using System.Collections.Generic;
+using System.Threading;
+using System.Threading.Tasks;
+
+namespace EpicGames.Horde.Storage
+{
+	/// <summary>
+	/// Interface for accessing an object's data
+	/// </summary>
+	public interface IBlob
+	{
+		/// <summary>
+		/// Gets the data for this blob
+		/// </summary>
+		/// <returns></returns>
+		ValueTask<ReadOnlyMemory<byte>> GetDataAsync();
+
+		/// <summary>
+		/// Find the outward references for a node
+		/// </summary>
+		/// <returns></returns>
+		ValueTask<IReadOnlyList<BlobId>> GetReferencesAsync();
+	}
+
+	/// <summary>
+	/// Base interface for a low-level storage backend. Blobs added to this store are not content addressed, but referenced by <see cref="BlobId"/>.
+	/// </summary>
+	public interface IBlobStore
+	{
+		#region Blobs
+
+		/// <summary>
+		/// Reads data for a blob from the store
+		/// </summary>
+		/// <param name="id">The blob identifier</param>
+		/// <param name="cancellationToken">Cancellation token for the operation</param>
+		/// <returns></returns>
+		Task<IBlob> ReadBlobAsync(BlobId id, CancellationToken cancellationToken = default);
+
+		/// <summary>
+		/// Writes a new object to the store
+		/// </summary>
+		/// <param name="data">Payload for the object</param>
+		/// <param name="references">Object references</param>
+		/// <param name="cancellationToken">Cancellation token for the operation</param>
+		/// <returns>Unique identifier for the bundle</returns>
+		Task<BlobId> WriteBlobAsync(ReadOnlySequence<byte> data, IReadOnlyList<BlobId> references, CancellationToken cancellationToken = default);
+
+		#endregion
+
+		#region Refs
+
+		/// <summary>
+		/// Checks if the given ref exists
+		/// </summary>
+		/// <param name="id">Name of the reference to look for</param>
+		/// <param name="cancellationToken">Cancellation token for the operation</param>
+		/// <returns>True if the ref exists, false if it did not exist</returns>
+		Task<bool> HasRefAsync(RefId id, CancellationToken cancellationToken = default);
+
+		/// <summary>
+		/// Reads data for a ref from the store
+		/// </summary>
+		/// <param name="id">The blob identifier</param>
+		/// <param name="cancellationToken">Cancellation token for the operation</param>
+		/// <returns></returns>
+		Task<IBlob> ReadRefAsync(RefId id, CancellationToken cancellationToken = default);
+
+		/// <summary>
+		/// Writes a new blob to the store
+		/// </summary>
+		/// <param name="id">Ref for which the blob is a part</param>
+		/// <param name="data">Payload for the object</param>
+		/// <param name="references">Object references</param>
+		/// <param name="cancellationToken">Cancellation token for the operation</param>
+		/// <returns>Unique identifier for the blob</returns>
+		Task WriteRefAsync(RefId id, ReadOnlySequence<byte> data, IReadOnlyList<BlobId> references, CancellationToken cancellationToken = default);
+
+		/// <summary>
+		/// Reads data for a ref from the store
+		/// </summary>
+		/// <param name="id">The ref identifier</param>
+		/// <param name="cancellationToken">Cancellation token for the operation</param>
+		/// <returns></returns>
+		Task<bool> DeleteRefAsync(RefId id, CancellationToken cancellationToken = default);
+
+		#endregion
+	}
+}
