Horde: API for manipulating tree structures in storage. Designed to abstract the logical tree structure away from storage/retrieval when allow working on larger-than-memory data sets.

* ITreeStore defines an API for reading/writing named tree structures.
* ITreeStore<T> allows reading/writing trees of nodes derived from the TreeNode base class, which can be specialized by clients. TreeNode tracks whether the node has been mutated, and allows recursing through the tree with strongly typed objects. Weak references are kept to deserialized TreeNode objects to ensure synchronization with the serialized tree objects.
* Default (and only current) implementation of ITreeStore is BundleStore, which packs nodes together in an efficient binary format and supports incrementally updating bundles. Nodes are hashed and deduplicated within bundles (and the known working set of the BundleStore instance), but packed Bundles are stored using a unique id returned by IBlobStore.
* DirectoryNode and FileNode nodes allow storing file trees in bundles, and implement content-aware chunking of binary data using a rolling BuzHash over an input stream.

diff --git a/Engine/Source/Programs/Shared/EpicGames.Core/ReadOnlySequenceBuilder.cs b/Engine/Source/Programs/Shared/EpicGames.Core/ReadOnlySequenceBuilder.cs
--- a/Engine/Source/Programs/Shared/EpicGames.Core/ReadOnlySequenceBuilder.cs
+++ b/Engine/Source/Programs/Shared/EpicGames.Core/ReadOnlySequenceBuilder.cs
@@ -29,6 +29,36 @@ namespace EpicGames.Core
 
 		readonly List<ReadOnlyMemory<T>> _segments = new List<ReadOnlyMemory<T>>();
 
+		/// <summary>
+		/// Current length of the sequence
+		/// </summary>
+		public long Length { get; private set; }
+
+		/// <summary>
+		/// Create a sequence from a list of segments
+		/// </summary>
+		/// <param name="segments"></param>
+		/// <returns>Sequence for the list of segments</returns>
+		public static ReadOnlySequence<T> Create(IReadOnlyList<ReadOnlyMemory<T>> segments)
+		{
+			if (segments.Count == 0)
+			{
+				return ReadOnlySequence<T>.Empty;
+			}
+
+			Segment first = new Segment(0, segments[0]);
+			Segment last = first;
+
+			for (int idx = 1; idx < segments.Count; idx++)
+			{
+				Segment next = new Segment(last.RunningIndex + last.Memory.Length, segments[idx]);
+				last.SetNext(next);
+				last = next;
+			}
+
+			return new ReadOnlySequence<T>(first, 0, last, last.Memory.Length);
+		}
+
 		/// <summary>
 		/// Append a block of memory to the end of the sequence
 		/// </summary>
@@ -38,6 +68,7 @@ namespace EpicGames.Core
 			if (memory.Length > 0)
 			{
 				_segments.Add(memory);
+				Length += memory.Length;
 			}
 		}
 
@@ -57,24 +88,29 @@ namespace EpicGames.Core
 		/// Construct a sequence from the added blocks
 		/// </summary>
 		/// <returns>Sequence for the added memory blocks</returns>
-		public ReadOnlySequence<T> Construct()
+		public ReadOnlySequence<T> Construct() => Create(_segments);
+	}
+
+	/// <summary>
+	/// Extension methods for <see cref="ReadOnlySequence{T}"/>
+	/// </summary>
+	public static class ReadOnlySequenceExtensions
+	{
+		/// <summary>
+		/// Gets the data from a sequence as a contiguous block of memory
+		/// </summary>
+		/// <param name="sequence">Sequence to return</param>
+		/// <returns>Data for the blob</returns>
+		public static ReadOnlyMemory<T> AsSingleSegment<T>(this ReadOnlySequence<T> sequence)
 		{
-			if (_segments.Count == 0)
+			if (sequence.IsSingleSegment)
 			{
-				return ReadOnlySequence<T>.Empty;
+				return sequence.First;
 			}
-
-			Segment first = new Segment(0, _segments[0]);
-			Segment last = first;
-
-			for (int idx = 1; idx < _segments.Count; idx++)
+			else
 			{
-				Segment next = new Segment(last.RunningIndex + last.Memory.Length, _segments[idx]);
-				last.SetNext(next);
-				last = next;
+				return sequence.ToArray();
 			}
-
-			return new ReadOnlySequence<T>(first, 0, last, last.Memory.Length);
 		}
 	}
 }
diff --git a/Engine/Source/Programs/Shared/EpicGames.Horde.Tests/BundleStoreTests.cs b/Engine/Source/Programs/Shared/EpicGames.Horde.Tests/BundleStoreTests.cs
new file mode 100644
--- /dev/null
+++ b/Engine/Source/Programs/Shared/EpicGames.Horde.Tests/BundleStoreTests.cs
@@ -0,0 +1,306 @@
+// Copyright Epic Games, Inc. All Rights Reserved.
+
+using EpicGames.Horde.Storage;
+using EpicGames.Horde.Storage.Impl;
+using EpicGames.Horde.Storage.Bundles;
+using EpicGames.Horde.Storage.Nodes;
+using Microsoft.Extensions.Caching.Memory;
+using Microsoft.VisualStudio.TestTools.UnitTesting;
+using System;
+using System.Buffers;
+using System.Collections.Generic;
+using System.Linq;
+using System.Threading.Tasks;
+using System.Threading;
+using System.Text;
+using System.IO;
+using System.Security.Cryptography;
+using EpicGames.Core;
+
+namespace EpicGames.Horde.Tests
+{
+	[TestClass]
+	public class BundleStoreTests
+	{
+		[TestMethod]
+		public async Task TestTreeAsync()
+		{
+			InMemoryBlobStore blobStore = new InMemoryBlobStore();
+			await TestTreeAsync(blobStore, new BundleOptions { MaxBlobSize = 1024 * 1024 });
+
+			Assert.AreEqual(0, blobStore.Blobs.Count);
+			Assert.AreEqual(1, blobStore.Refs.Count);
+		}
+
+		[TestMethod]
+		public async Task TestTreeSeparateBlobsAsync()
+		{
+			InMemoryBlobStore blobStore = new InMemoryBlobStore();
+			await TestTreeAsync(blobStore, new BundleOptions { MaxBlobSize = 1 });
+
+			Assert.AreEqual(4, blobStore.Blobs.Count);
+			Assert.AreEqual(1, blobStore.Refs.Count);
+		}
+
+		static async Task TestTreeAsync(IBlobStore blobStore, BundleOptions options)
+		{
+			using IMemoryCache cache = new MemoryCache(new MemoryCacheOptions());
+
+			// Generate a tree
+			using (BundleStore collection = new BundleStore(blobStore, options, cache))
+			{
+				ITreeBlob node1 = await collection.WriteBlobAsync(new ReadOnlySequence<byte>(new byte[] { 1 }), Array.Empty<ITreeBlob>(), CancellationToken.None);
+				ITreeBlob node2 = await collection.WriteBlobAsync(new ReadOnlySequence<byte>(new byte[] { 2 }), new[] { node1 }, CancellationToken.None);
+				ITreeBlob node3 = await collection.WriteBlobAsync(new ReadOnlySequence<byte>(new byte[] { 3 }), new[] { node2 }, CancellationToken.None);
+				ITreeBlob node4 = await collection.WriteBlobAsync(new ReadOnlySequence<byte>(new byte[] { 4 }), Array.Empty<ITreeBlob>(), CancellationToken.None);
+				ITreeBlob node5 = await collection.WriteBlobAsync(new ReadOnlySequence<byte>(new byte[] { 5 }), new[] { node4, node3 }, CancellationToken.None);
+
+				await collection.WriteTreeAsync(new RefId("test"), node5);
+
+				ITreeBlob root = node5;
+				await CheckTree(root);
+			}
+
+			// Check we can read it back in
+			using (BundleStore collection = new BundleStore(blobStore, options, cache))
+			{
+				ITreeBlob root = await collection.ReadTreeAsync(new RefId("test"));
+				await CheckTree(root);
+			}
+		}
+
+		static async Task CheckTree(ITreeBlob root)
+		{
+			ITreeBlob node5 = root;
+			byte[] data5 = (await node5.GetDataAsync()).ToArray();
+			Assert.IsTrue(data5.SequenceEqual(new byte[] { 5 }));
+			IReadOnlyList<ITreeBlob> refs5 = await node5.GetReferencesAsync();
+			Assert.AreEqual(2, refs5.Count);
+
+			ITreeBlob node4 = refs5[0];
+			byte[] data4 = (await node4.GetDataAsync()).ToArray();
+			Assert.IsTrue(data4.SequenceEqual(new byte[] { 4 }));
+			IReadOnlyList<ITreeBlob> refs4 = await node4.GetReferencesAsync();
+			Assert.AreEqual(0, refs4.Count);
+
+			ITreeBlob node3 = refs5[1];
+			byte[] data3 = (await node3.GetDataAsync()).ToArray();
+			Assert.IsTrue(data3.SequenceEqual(new byte[] { 3 }));
+			IReadOnlyList<ITreeBlob> refs3 = await node3.GetReferencesAsync();
+			Assert.AreEqual(1, refs3.Count);
+
+			ITreeBlob node2 = refs3[0];
+			byte[] data2 = (await node2.GetDataAsync()).ToArray();
+			Assert.IsTrue(data2.SequenceEqual(new byte[] { 2 }));
+			IReadOnlyList<ITreeBlob> refs2 = await node2.GetReferencesAsync();
+			Assert.AreEqual(1, refs2.Count);
+
+			ITreeBlob node1 = refs2[0];
+			byte[] data1 = (await node1.GetDataAsync()).ToArray();
+			Assert.IsTrue(data1.SequenceEqual(new byte[] { 1 }));
+			IReadOnlyList<ITreeBlob> refs1 = await node1.GetReferencesAsync();
+			Assert.AreEqual(0, refs1.Count);
+		}
+
+		[TestMethod]
+		public async Task DirectoryNodesAsync()
+		{
+			using IMemoryCache cache = new MemoryCache(new MemoryCacheOptions());
+			InMemoryBlobStore blobStore = new InMemoryBlobStore();
+
+			// Generate a tree
+			using (BundleStore<DirectoryNode> store = new BundleStore<DirectoryNode>(blobStore, new BundleOptions(), cache))
+			{
+				DirectoryNode root = new DirectoryNode();
+				DirectoryNode hello = root.AddDirectory("hello");
+				DirectoryNode world = hello.AddDirectory("world");
+				await store.WriteTreeAsync(new RefId("test"), root, true, CancellationToken.None);
+
+				await CheckDirectoryTreeAsync(root);
+			}
+
+			// Check we can read it back in
+			using (BundleStore<DirectoryNode> store = new BundleStore<DirectoryNode>(blobStore, new BundleOptions(), cache))
+			{
+				DirectoryNode root = await store.ReadTreeAsync(new RefId("test"));
+				await CheckDirectoryTreeAsync(root);
+			}
+		}
+
+		static async Task CheckDirectoryTreeAsync(DirectoryNode root)
+		{
+			Assert.AreEqual(1, root.Directories.Count);
+			Assert.AreEqual("hello", root.Directories.First().Name);
+
+			DirectoryNode hello = await root.Directories.First().ExpandAsync(CancellationToken.None);
+			Assert.AreEqual(1, hello.Directories.Count);
+			Assert.AreEqual("world", hello.Directories.First().Name);
+
+			DirectoryNode world = await hello.Directories.First().ExpandAsync(CancellationToken.None);
+			Assert.AreEqual(0, world.Directories.Count);
+		}
+
+		[TestMethod]
+		public async Task FileNodesAsync()
+		{
+			using IMemoryCache cache = new MemoryCache(new MemoryCacheOptions());
+			InMemoryBlobStore blobStore = new InMemoryBlobStore();
+
+			// Generate a tree
+			using (BundleStore<DirectoryNode> store = new BundleStore<DirectoryNode>(blobStore, new BundleOptions(), cache))
+			{
+				DirectoryNode root = new DirectoryNode();
+				DirectoryNode hello = root.AddDirectory("hello");
+
+				FileNode world = await hello.AddFileAsync("world", FileEntryFlags.None, Encoding.UTF8.GetBytes("world"), new ChunkingOptions(), CancellationToken.None);
+				
+				await store.WriteTreeAsync(new RefId("test"), root, true, CancellationToken.None);
+
+				await CheckFileTreeAsync(root);
+			}
+
+			// Check we can read it back in
+			using (BundleStore<DirectoryNode> store = new BundleStore<DirectoryNode>(blobStore, new BundleOptions(), cache))
+			{
+				DirectoryNode root = await store.ReadTreeAsync(new RefId("test"));
+				await CheckFileTreeAsync(root);
+			}
+		}
+
+		static async Task CheckFileTreeAsync(DirectoryNode root)
+		{
+			Assert.AreEqual(1, root.Directories.Count);
+			Assert.AreEqual("hello", root.Directories.First().Name);
+
+			DirectoryNode hello = await root.Directories.First().ExpandAsync(CancellationToken.None);
+			Assert.AreEqual(0, hello.Directories.Count);
+			Assert.AreEqual(1, hello.Files.Count);
+			Assert.AreEqual("world", hello.Files.First().Name);
+
+			FileNode world = await hello.Files.First().ExpandAsync(CancellationToken.None);
+
+			byte[] worldData = await GetFileDataAsync(world);
+			Assert.IsTrue(worldData.SequenceEqual(Encoding.UTF8.GetBytes("world")));
+		}
+
+		[TestMethod]
+		public void ChunkingTests()
+		{
+			byte[] chunk = RandomNumberGenerator.GetBytes(256);
+
+			byte[] data = new byte[chunk.Length * 1024];
+			for (int idx = 0; idx * chunk.Length < data.Length; idx++)
+			{
+				chunk.CopyTo(data.AsSpan(idx * chunk.Length));
+			}
+
+			ChunkingOptions options = new ChunkingOptions();
+			options.LeafOptions.MinSize = 1;
+			options.LeafOptions.TargetSize = 64;
+			options.LeafOptions.MaxSize = 1024;
+
+			ReadOnlyMemory<byte> remaining = data;
+
+			Dictionary<IoHash, long> uniqueChunks = new Dictionary<IoHash, long>();
+
+			List<LeafFileNode> nodes = new List<LeafFileNode>();
+			while (remaining.Length > 0)
+			{
+				LeafFileNode node = new LeafFileNode();
+				remaining = node.AppendData(remaining, options);
+				nodes.Add(node);
+
+				IoHash hash = IoHash.Compute(node.Data);
+				uniqueChunks[hash] = node.Length;
+			}
+
+			long uniqueSize = uniqueChunks.Sum(x => x.Value);
+			Assert.IsTrue(uniqueSize < data.Length / 3);
+		}
+
+		[TestMethod]
+		public async Task LargeFileTestAsync()
+		{
+			using IMemoryCache cache = new MemoryCache(new MemoryCacheOptions());
+			InMemoryBlobStore blobStore = new InMemoryBlobStore();
+
+			const int length = 1024;
+			const int copies = 4096;
+
+			byte[] chunk = new byte[length];
+			new Random(0).NextBytes(chunk);
+
+			byte[] data = new byte[chunk.Length * copies];
+			for (int idx = 0; idx < copies; idx++)
+			{
+				chunk.CopyTo(data.AsSpan(idx * chunk.Length));
+			}
+
+			// Generate a tree
+			using (BundleStore<DirectoryNode> store = new BundleStore<DirectoryNode>(blobStore, new BundleOptions { MaxInlineBlobSize = 1024, MaxBlobSize = 1024 }, cache))
+			{
+				DirectoryNode root = new DirectoryNode();
+
+				ChunkingOptions options = new ChunkingOptions();
+				options.LeafOptions.MinSize = 1;
+				options.LeafOptions.TargetSize = 128;
+				options.LeafOptions.MaxSize = 64 * 1024;
+
+				FileNode file = await root.AddFileAsync("test", FileEntryFlags.None, data, options, CancellationToken.None);
+
+				await store.WriteTreeAsync(new RefId("test"), root, true, CancellationToken.None);
+
+				await CheckLargeFileTreeAsync(root, data);
+			}
+
+			// Check we can read it back in
+			using (BundleStore<DirectoryNode> store = new BundleStore<DirectoryNode>(blobStore, new BundleOptions(), cache))
+			{
+				DirectoryNode root = await store.ReadTreeAsync(new RefId("test"));
+				await CheckLargeFileTreeAsync(root, data);
+
+				TreeNodeRef<FileNode> file = root.GetFileEntry("test");
+				ITreeBlob blob = await file.CollapseAsync(store, CancellationToken.None);
+
+				Dictionary<IoHash, long> uniqueChunks = new Dictionary<IoHash, long>();
+				await FindUniqueNodes(blob, uniqueChunks);
+
+				long uniqueSize = uniqueChunks.Sum(x => x.Value);
+				Assert.IsTrue(uniqueSize < data.Length / 3); // random fraction meaning "lots of dedupe happened"
+			}
+		}
+
+		static async Task CheckLargeFileTreeAsync(DirectoryNode root, byte[] data)
+		{
+			Assert.AreEqual(0, root.Directories.Count);
+			Assert.AreEqual(1, root.Files.Count);
+
+			FileNode world = await root.Files.First().ExpandAsync(CancellationToken.None);
+
+			byte[] worldData = await GetFileDataAsync(world);
+			Assert.AreEqual(data.Length, worldData.Length);
+			Assert.IsTrue(worldData.SequenceEqual(data));
+		}
+
+		static async Task FindUniqueNodes(ITreeBlob blob, Dictionary<IoHash, long> uniqueHashes)
+		{
+			ReadOnlySequence<byte> data = await blob.GetDataAsync();
+			uniqueHashes[blob.Hash] = data.Length;
+
+			IReadOnlyList<ITreeBlob> references = await blob.GetReferencesAsync();
+			foreach (ITreeBlob reference in references)
+			{
+				await FindUniqueNodes(reference, uniqueHashes);
+			}
+		}
+
+		static async Task<byte[]> GetFileDataAsync(FileNode fileNode)
+		{
+			using (MemoryStream stream = new MemoryStream())
+			{
+				await fileNode.CopyToStreamAsync(stream, CancellationToken.None);
+				return stream.ToArray();
+			}
+		}
+	}
+}
diff --git a/Engine/Source/Programs/Shared/EpicGames.Horde.Tests/ChunkingTests.cs b/Engine/Source/Programs/Shared/EpicGames.Horde.Tests/ChunkingTests.cs
new file mode 100644
--- /dev/null
+++ b/Engine/Source/Programs/Shared/EpicGames.Horde.Tests/ChunkingTests.cs
@@ -0,0 +1,134 @@
+// Copyright Epic Games, Inc. All Rights Reserved.
+
+using EpicGames.Core;
+using EpicGames.Horde.Storage.Nodes;
+using Microsoft.VisualStudio.TestTools.UnitTesting;
+using System;
+using System.Linq;
+using System.Threading.Tasks;
+using System.IO;
+using System.Threading;
+
+namespace EpicGames.Horde.Tests
+{
+	[TestClass]
+	public class ChunkingTests
+	{
+		[TestMethod]
+		public void BuzHashTests()
+		{
+			byte[] data = new byte[4096];
+			new Random(0).NextBytes(data);
+
+			const int WindowSize = 128;
+
+			uint rollingHash = 0;
+			for (int maxIdx = 0; maxIdx < data.Length + WindowSize; maxIdx++)
+			{
+				int minIdx = maxIdx - WindowSize;
+
+				if (maxIdx < data.Length)
+				{
+					rollingHash = BuzHash.Add(rollingHash, data[maxIdx]);
+				}
+
+				int length = Math.Min(maxIdx + 1, data.Length) - Math.Max(minIdx, 0);
+				uint cleanHash = BuzHash.Add(0, data.AsSpan(Math.Max(minIdx, 0), length));
+				Assert.AreEqual(rollingHash, cleanHash);
+
+				if (minIdx >= 0)
+				{
+					rollingHash = BuzHash.Sub(rollingHash, data[minIdx], length);
+				}
+			}
+		}
+
+		[TestMethod]
+		public async Task BasicChunkingTests()
+		{
+			ChunkingOptions options = new ChunkingOptions();
+			options.LeafOptions = new ChunkingOptionsForNodeType(8, 8, 8);
+
+			FileNode node = new LeafFileNode();
+			node.Append(new byte[7], options);
+			Assert.AreEqual(7, node.Length);
+			Assert.AreEqual(0, node.GetReferences().Count);
+			Assert.AreEqual(7, (await node.ToByteArrayAsync(CancellationToken.None)).Length);
+
+			node = new LeafFileNode();
+			node.Append(new byte[8], options);
+			Assert.AreEqual(8, node.Length);
+			Assert.AreEqual(0, node.GetReferences().Count);
+			Assert.AreEqual(8, (await node.ToByteArrayAsync(CancellationToken.None)).Length);
+
+			node = new LeafFileNode();
+			node = node.Append(new byte[9], options);
+			Assert.AreEqual(9, node.Length);
+			Assert.AreEqual(2, node.GetReferences().Count);
+
+			FileNode? childNode1 = (FileNode?)node.GetReferences()[0].Node;
+			Assert.IsNotNull(childNode1);
+			Assert.AreEqual(0, childNode1!.GetReferences().Count);
+			Assert.AreEqual(8, (await childNode1!.ToByteArrayAsync(CancellationToken.None)).Length);
+
+			FileNode? childNode2 = (FileNode?)node.GetReferences()[1].Node;
+			Assert.IsNotNull(childNode2);
+			Assert.AreEqual(0, childNode2!.GetReferences().Count);
+			Assert.AreEqual(1, (await childNode2!.ToByteArrayAsync(CancellationToken.None)).Length);
+		}
+
+		[TestMethod]
+		public async Task FixedSizeChunkingTests()
+		{
+			ChunkingOptions options = new ChunkingOptions();
+			options.LeafOptions = new ChunkingOptionsForNodeType(64, 64, 64);
+			options.InteriorOptions = new ChunkingOptionsForNodeType(IoHash.NumBytes * 4, IoHash.NumBytes * 4, IoHash.NumBytes * 4);
+
+			await TestChunkingAsync(options);
+		}
+
+		[TestMethod]
+		public async Task VariableSizeChunkingTests()
+		{
+			ChunkingOptions options = new ChunkingOptions();
+			options.LeafOptions = new ChunkingOptionsForNodeType(32, 64, 96);
+			options.InteriorOptions = new ChunkingOptionsForNodeType(IoHash.NumBytes * 1, IoHash.NumBytes * 4, IoHash.NumBytes * 12);
+
+			await TestChunkingAsync(options);
+		}
+
+		static async Task TestChunkingAsync(ChunkingOptions options)
+		{
+			byte[] data = new byte[4096];
+			new Random(0).NextBytes(data);
+
+			for (int idx = 0; idx < data.Length; idx++)
+			{
+				data[idx] = (byte)idx;
+			}
+
+			FileNode root = new LeafFileNode();
+
+			const int NumIterations = 100;
+			for (int idx = 0; idx < NumIterations; idx++)
+			{
+				root = root.Append(data, options);
+			}
+
+			byte[] result;
+			using (MemoryStream stream = new MemoryStream())
+			{
+				await root.CopyToStreamAsync(stream, CancellationToken.None);
+				result = stream.ToArray();
+			}
+
+			Assert.AreEqual(NumIterations * data.Length, result.Length);
+
+			for (int idx = 0; idx < NumIterations; idx++)
+			{
+				ReadOnlyMemory<byte> spanData = result.AsMemory(idx * data.Length, data.Length);
+				Assert.IsTrue(spanData.Span.SequenceEqual(data));
+			}
+		}
+	}
+}
diff --git a/Engine/Source/Programs/Shared/EpicGames.Horde/Storage/Bundles/Bundle.cs b/Engine/Source/Programs/Shared/EpicGames.Horde/Storage/Bundles/Bundle.cs
new file mode 100644
--- /dev/null
+++ b/Engine/Source/Programs/Shared/EpicGames.Horde/Storage/Bundles/Bundle.cs
@@ -0,0 +1,404 @@
+// Copyright Epic Games, Inc. All Rights Reserved.
+
+using EpicGames.Core;
+using EpicGames.Serialization;
+using System;
+using System.Buffers;
+using System.Collections.Generic;
+using System.IO;
+using System.Linq;
+using System.Threading;
+using System.Threading.Tasks;
+
+namespace EpicGames.Horde.Storage.Bundles
+{
+	/// <summary>
+	/// Bundle version number
+	/// </summary>
+	[System.Diagnostics.CodeAnalysis.SuppressMessage("Design", "CA1027:Mark enums with FlagsAttribute")]
+	public enum BundleVersion
+	{
+		/// <summary>
+		/// Initial version number
+		/// </summary>
+		Initial = 0,
+
+		/// <summary>
+		/// The current version number
+		/// </summary>
+		Current = Initial,
+	}
+
+	/// <summary>
+	/// Header for the contents of a bundle. May contain an inlined payload object containing the object data itself.
+	/// </summary>
+	public class Bundle
+	{
+		/// <summary>
+		/// Header for the bundle
+		/// </summary>
+		public BundleHeader Header { get; }
+
+		/// <summary>
+		/// Packet data as described in the header
+		/// </summary>
+		public ReadOnlyMemory<byte> Payload { get; }
+
+		/// <summary>
+		/// Constructor
+		/// </summary>
+		public Bundle(BundleHeader header, ReadOnlyMemory<byte> payload)
+		{
+			Header = header;
+			Payload = payload;
+		}
+
+		/// <summary>
+		/// Reads a bundle from a block of memory
+		/// </summary>
+		public Bundle(IMemoryReader reader)
+		{
+			Header = new BundleHeader(reader);
+			int length = (int)reader.ReadUnsignedVarInt();
+			Payload = reader.ReadFixedLengthBytes(length);
+		}
+
+		/// <summary>
+		/// Serializes the bundle to a sequence of bytes
+		/// </summary>
+		/// <returns>Sequence for the bundle</returns>
+		public ReadOnlySequence<byte> AsSequence()
+		{
+			ByteArrayBuilder builder = new ByteArrayBuilder();
+			Header.Write(builder);
+			builder.WriteUnsignedVarInt(Payload.Length);
+
+			ReadOnlySequenceBuilder<byte> sequence = new ReadOnlySequenceBuilder<byte>();
+			sequence.Append(builder.AsSequence());
+			sequence.Append(Payload);
+
+			return sequence.Construct();
+		}
+	}
+
+	/// <summary>
+	/// Header for the contents of a bundle.
+	/// </summary>
+	public class BundleHeader
+	{
+		/// <summary>
+		/// References to nodes in other bundles
+		/// </summary>
+		public IReadOnlyList<BundleImport> Imports { get; } = new List<BundleImport>();
+
+		/// <summary>
+		/// Nodes exported from this bundle
+		/// </summary>
+		public IReadOnlyList<BundleExport> Exports { get; } = new List<BundleExport>();
+
+		/// <summary>
+		/// List of data packets within this bundle
+		/// </summary>
+		public IReadOnlyList<BundlePacket> Packets { get; } = new List<BundlePacket>();
+
+		/// <summary>
+		/// Constructs a new bundle header
+		/// </summary>
+		/// <param name="imports">Imports from other bundles</param>
+		/// <param name="exports">Exports for nodes</param>
+		/// <param name="packets">Compression packets within the bundle</param>
+		public BundleHeader(IReadOnlyList<BundleImport> imports, IReadOnlyList<BundleExport> exports, IReadOnlyList<BundlePacket> packets)
+		{
+			Imports = imports;
+			Exports = exports;
+			Packets = packets;
+		}
+
+		/// <summary>
+		/// Reads a bundle header from memory
+		/// </summary>
+		/// <param name="reader">Reader to deserialize from</param>
+		/// <returns>New header object</returns>
+		public BundleHeader(IMemoryReader reader)
+		{
+			BundleVersion version = (BundleVersion)reader.ReadUnsignedVarInt();
+			if (version != BundleVersion.Current)
+			{
+				throw new InvalidDataException();
+			}
+
+			Imports = reader.ReadVariableLengthArray(() => new BundleImport(reader));
+			Exports = reader.ReadVariableLengthArray(() => new BundleExport(reader));
+			Packets = reader.ReadVariableLengthArray(() => new BundlePacket(reader));
+		}
+
+		/// <summary>
+		/// Serializes a header to memory
+		/// </summary>
+		/// <param name="writer">Writer to serialize to</param>
+		public void Write(IMemoryWriter writer)
+		{
+			writer.WriteUnsignedVarInt((ulong)BundleVersion.Current);
+
+			writer.WriteVariableLengthArray(Imports, x => x.Write(writer));
+			writer.WriteVariableLengthArray(Exports, x => x.Write(writer));
+			writer.WriteVariableLengthArray(Packets, x => x.Write(writer));
+		}
+	}
+
+	/// <summary>
+	/// Reference to another tree pack object
+	/// </summary>
+	public class BundleImport
+	{
+		/// <summary>
+		/// Blob containing the bundle data.
+		/// </summary>
+		public BlobId BlobId { get; }
+
+		/// <summary>
+		/// Number of exports from this blob.
+		/// </summary>
+		public int ExportCount { get; }
+
+		/// <summary>
+		/// Indexes of referenced nodes exported from this bundle
+		/// </summary>
+		public IReadOnlyList<(int, IoHash)> Exports { get; }
+
+		/// <summary>
+		/// Constructor
+		/// </summary>
+		public BundleImport(BlobId blobId, int exportCount, IReadOnlyList<(int, IoHash)> exports)
+		{
+			BlobId = blobId;
+			ExportCount = exportCount;
+			Exports = exports;
+		}
+
+		/// <summary>
+		/// Deserialize a bundle import
+		/// </summary>
+		/// <param name="reader">Reader to deserialize from</param>
+		internal BundleImport(IMemoryReader reader)
+		{
+			BlobId = reader.ReadBlobId();
+
+			ExportCount = (int)reader.ReadUnsignedVarInt();
+
+			int count = (int)reader.ReadUnsignedVarInt();
+
+			(int, IoHash)[] exports = new (int, IoHash)[count];
+			for (int idx = 0; idx < count; idx++)
+			{
+				int index = (int)reader.ReadUnsignedVarInt();
+				IoHash hash = reader.ReadIoHash();
+				exports[idx] = (index, hash);
+			}
+			Exports = exports;
+		}
+
+		/// <summary>
+		/// Serializes a bundle import
+		/// </summary>
+		/// <param name="writer">Writer to serialize to</param>
+		public void Write(IMemoryWriter writer)
+		{
+			writer.WriteBlobId(BlobId);
+			writer.WriteUnsignedVarInt(ExportCount);
+
+			writer.WriteUnsignedVarInt(Exports.Count);
+			for(int idx = 0; idx < Exports.Count; idx++)
+			{
+				(int index, IoHash hash) = Exports[idx];
+				writer.WriteUnsignedVarInt(index);
+				writer.WriteIoHash(hash);
+			}
+		}
+	}
+
+	/// <summary>
+	/// Descriptor for a compression packet
+	/// </summary>
+	public class BundlePacket
+	{
+		/// <summary>
+		/// Encoded length of the packet
+		/// </summary>
+		public int EncodedLength { get; set; }
+
+		/// <summary>
+		/// Decoded length of the packet
+		/// </summary>
+		public int DecodedLength { get; set; }
+
+		/// <summary>
+		/// Constructor
+		/// </summary>
+		/// <param name="encodedLength">Size of the encoded data</param>
+		/// <param name="decodedLength">Size of the decoded data</param>
+		public BundlePacket(int encodedLength, int decodedLength)
+		{
+			EncodedLength = encodedLength;
+			DecodedLength = decodedLength;
+		}
+
+		/// <summary>
+		/// Deserialize a packet header from a memory reader
+		/// </summary>
+		/// <param name="reader">Reader to deserialize from</param>
+		public BundlePacket(IMemoryReader reader)
+		{
+			EncodedLength = (int)reader.ReadUnsignedVarInt();
+			DecodedLength = (int)reader.ReadUnsignedVarInt();
+		}
+
+		/// <summary>
+		/// Serializes this packet header 
+		/// </summary>
+		/// <param name="writer"></param>
+		public void Write(IMemoryWriter writer)
+		{
+			writer.WriteUnsignedVarInt(EncodedLength);
+			writer.WriteUnsignedVarInt(DecodedLength);
+		}
+	}
+
+	/// <summary>
+	/// Entry for a node exported from an object
+	/// </summary>
+	public class BundleExport
+	{
+		/// <summary>
+		/// Hash of the node
+		/// </summary>
+		public IoHash Hash { get; }
+
+		/// <summary>
+		/// Uncompressed length of this node
+		/// </summary>
+		public int Length { get; }
+		
+		/// <summary>
+		/// Nodes referenced by this export. Indices in this array correspond to a lookup table consisting
+        /// of the imported nodes in the order they are declared in the header, followed by nodes listed in the
+        /// export table itself.
+        /// </summary>
+		public IReadOnlyList<int> References { get; }
+
+		/// <summary>
+		/// Constructor
+		/// </summary>
+		public BundleExport(IoHash hash, int length, IReadOnlyList<int> references)
+		{
+			Hash = hash;
+			Length = length;
+			References = references;
+		}
+
+		/// <summary>
+		/// Deserialize an export
+		/// </summary>
+		/// <param name="reader">Reader to deserialize from</param>
+		public BundleExport(IMemoryReader reader)
+		{
+			Hash = reader.ReadIoHash();
+			Length = (int)reader.ReadUnsignedVarInt();
+
+			int numReferences = (int)reader.ReadUnsignedVarInt();
+			int[] references = new int[numReferences];
+
+			for (int idx = 0; idx < numReferences; idx++)
+			{
+				references[idx] = (int)reader.ReadUnsignedVarInt();
+			}
+
+			References = references;
+		}
+
+		/// <summary>
+		/// Serializes an export
+		/// </summary>
+		/// <param name="writer">Writer to serialize to</param>
+		public void Write(IMemoryWriter writer)
+		{
+			writer.WriteIoHash(Hash);
+			writer.WriteUnsignedVarInt(Length);
+
+			writer.WriteUnsignedVarInt(References.Count);
+			for(int idx = 0; idx < References.Count; idx++)
+			{
+				writer.WriteUnsignedVarInt(References[idx]);
+			}
+		}
+	}
+
+	/// <summary>
+	/// Extension methods for serializing bundles
+	/// </summary>
+	public static class BundleExtensions
+	{
+		/// <summary>
+		/// Read a bundle from a blob store
+		/// </summary>
+		/// <param name="store">Blob store to read from</param>
+		/// <param name="id">Identifier for the blob</param>
+		/// <param name="cancellationToken">Cancellation token for the operation</param>
+		/// <returns>The bundle that was read</returns>
+		public static async Task<Bundle> ReadBundleAsync(this IBlobStore store, BlobId id, CancellationToken cancellationToken)
+		{
+			IBlob blob = await store.ReadBlobAsync(id, cancellationToken);
+			return await ReadBundleAsync(blob);
+		}
+
+		/// <summary>
+		/// Read a bundle from a blob store
+		/// </summary>
+		/// <param name="store">Blob store to read from</param>
+		/// <param name="id">Identifier for the blob</param>
+		/// <param name="cancellationToken">Cancellation token for the operation</param>
+		/// <returns>The bundle that was read</returns>
+		public static async Task<Bundle> ReadBundleAsync(this IBlobStore store, RefId id, CancellationToken cancellationToken)
+		{
+			IBlob blob = await store.ReadRefAsync(id, cancellationToken);
+			return await ReadBundleAsync(blob);
+		}
+
+		/// <summary>
+		/// Writes a bundle to a blob store
+		/// </summary>
+		/// <param name="store">Blob store to write to</param>
+		/// <param name="id">Identifier for the blob</param>
+		/// <param name="bundle">Bundle to write</param>
+		/// <param name="cancellationToken">Cancellation token for the operation</param>
+		public static Task WriteBundleAsync(this IBlobStore store, RefId id, Bundle bundle, CancellationToken cancellationToken)
+		{
+			ReadOnlySequence<byte> data = bundle.AsSequence();
+			List<BlobId> imports = bundle.Header.Imports.Select(x => x.BlobId).ToList();
+			return store.WriteRefAsync(id, data, imports, cancellationToken);
+		}
+
+		/// <summary>
+		/// Writes a bundle to a blob store
+		/// </summary>
+		/// <param name="store">Blob store to write to</param>
+		/// <param name="bundle">Bundle to write</param>
+		/// <param name="cancellationToken">Cancellation token for the operation</param>
+		/// <returns>Identifier for the blob</returns>
+		public static Task<BlobId> WriteBundleAsync(this IBlobStore store, Bundle bundle, CancellationToken cancellationToken)
+		{
+			ReadOnlySequence<byte> data = bundle.AsSequence();
+			List<BlobId> imports = bundle.Header.Imports.Select(x => x.BlobId).ToList();
+			return store.WriteBlobAsync(data, imports, cancellationToken);
+		}
+
+		static async Task<Bundle> ReadBundleAsync(IBlob blob)
+		{
+			ReadOnlyMemory<byte> data = await blob.GetDataAsync();
+			MemoryReader reader = new MemoryReader(data);
+			Bundle bundle = new Bundle(reader);
+			reader.CheckEmpty();
+			return bundle;
+		}
+	}
+}
diff --git a/Engine/Source/Programs/Shared/EpicGames.Horde/Storage/Bundles/BundleStore.cs b/Engine/Source/Programs/Shared/EpicGames.Horde/Storage/Bundles/BundleStore.cs
new file mode 100644
--- /dev/null
+++ b/Engine/Source/Programs/Shared/EpicGames.Horde/Storage/Bundles/BundleStore.cs
@@ -0,0 +1,878 @@
+// Copyright Epic Games, Inc. All Rights Reserved.
+
+using System;
+using System.Buffers;
+using System.Collections.Concurrent;
+using System.Collections.Generic;
+using System.Diagnostics;
+using System.IO;
+using System.Linq;
+using System.Threading;
+using System.Threading.Tasks;
+using EpicGames.Core;
+using K4os.Compression.LZ4;
+using Microsoft.Extensions.Caching.Memory;
+
+namespace EpicGames.Horde.Storage.Bundles
+{
+	/// <summary>
+	/// Options for configuring a bundle serializer
+	/// </summary>
+	public class BundleOptions
+	{
+		/// <summary>
+		/// Maximum payload size fo a blob
+		/// </summary>
+		public int MaxBlobSize { get; set; } = 10 * 1024 * 1024;
+
+		/// <summary>
+		/// Maximum payload size for the root blob
+		/// </summary>
+		public int MaxInlineBlobSize { get; set; } = 1024 * 64;
+
+		/// <summary>
+		/// Minimum size of a block to be compressed
+		/// </summary>
+		public int MinCompressionPacketSize { get; set; } = 16 * 1024;
+
+		/// <summary>
+		/// Number of nodes to retain in the working set after performing a partial flush
+		/// </summary>
+		public int TrimIgnoreCount { get; set; } = 150;
+
+		/// <summary>
+		/// Number of reads from storage to allow concurrently. This has an impact on memory usage as well as network throughput.
+		/// </summary>
+		public int MaxConcurrentReads { get; set; } = 5;
+	}
+
+	/// <summary>
+	/// Statistics for the state of a bundle
+	/// </summary>
+	public class BundleStats
+	{
+		/// <summary>
+		/// Number of blobs that have been written
+		/// </summary>
+		public int NewBlobCount { get; set; }
+
+		/// <summary>
+		/// Total size of the blobs that have been written
+		/// </summary>
+		public long NewBlobBytes { get; set; }
+
+		/// <summary>
+		/// Number of refs that have been written
+		/// </summary>
+		public int NewRefCount { get; set; }
+
+		/// <summary>
+		/// Total size of the refs that have been written
+		/// </summary>
+		public long NewRefBytes { get; set; }
+
+		/// <summary>
+		/// Number of nodes in the tree that have been written
+		/// </summary>
+		public int NewNodeCount { get; set; }
+
+		/// <summary>
+		/// Total size of the nodes that have been serialized
+		/// </summary>
+		public long NewNodeBytes { get; set; }
+	}
+
+	/// <summary>
+	/// Manages a collection of bundles.
+	/// </summary>
+	public class BundleStore : ITreeStore, ITreeBlobWriter, IDisposable
+	{
+		/// <summary>
+		/// Information about a node within a bundle.
+		/// </summary>
+		[DebuggerDisplay("{Hash}")]
+		sealed class NodeInfo : ITreeBlob
+		{
+			public readonly BundleStore Owner;
+			public IoHash Hash { get; }
+			public NodeState State { get; set; }
+
+			public NodeInfo(BundleStore owner, IoHash hash, NodeState state)
+			{
+				Owner = owner;
+				Hash = hash;
+				State = state;
+			}
+
+			/// <inheritdoc/>
+			public ValueTask<ReadOnlySequence<byte>> GetDataAsync(CancellationToken cancellationToken) => Owner.GetDataAsync(this, cancellationToken);
+
+			/// <inheritdoc/>
+			public async ValueTask<IReadOnlyList<ITreeBlob>> GetReferencesAsync(CancellationToken cancellationToken) => await Owner.GetReferencesAsync(this, cancellationToken);
+		}
+
+		sealed class PacketInfo
+		{
+			public readonly int Offset;
+			public readonly int DecodedLength;
+			public readonly int EncodedLength;
+
+			public PacketInfo(int offset, BundlePacket packet)
+			{
+				Offset = offset;
+				DecodedLength = packet.DecodedLength;
+				EncodedLength = packet.EncodedLength;
+			}
+		}
+
+		// Type of node state, ordered from low to high in terms of preference when resolving duplicates.
+		enum NodeStatePriority
+		{
+			Imported = 0,
+			InMemory = 1,
+			Exported = 2,
+		}
+
+		abstract class NodeState
+		{
+		}
+
+		class ImportedNodeState : NodeState
+		{
+			public readonly BlobInfo BundleInfo;
+			public readonly int ExportIdx;
+
+			public ImportedNodeState(BlobInfo bundleInfo, int exportIdx)
+			{
+				BundleInfo = bundleInfo;
+				ExportIdx = exportIdx;
+			}
+		}
+
+		class ExportedNodeState : ImportedNodeState
+		{
+			public readonly PacketInfo PacketInfo;
+			public readonly int Offset;
+			public readonly int Length;
+			public readonly IReadOnlyList<NodeInfo> References;
+
+			public ExportedNodeState(BlobInfo bundleInfo, int exportIdx, PacketInfo packetInfo, int offset, int length, IReadOnlyList<NodeInfo> references)
+				: base(bundleInfo, exportIdx)
+			{
+				PacketInfo = packetInfo;
+				Offset = offset;
+				Length = length;
+				References = references;
+			}
+		}
+
+		class InMemoryNodeState : NodeState
+		{
+			public readonly ReadOnlySequence<byte> Data;
+			public readonly IReadOnlyList<NodeInfo> References;
+
+			public InMemoryNodeState(ReadOnlySequence<byte> data, IReadOnlyList<NodeInfo> references)
+			{
+				Data = data;
+				References = references;
+			}
+		}
+
+		/// <summary>
+		/// Information about a blob
+		/// </summary>
+		[DebuggerDisplay("{BlobId}")]
+		class BlobInfo
+		{
+			public readonly BlobId BlobId;
+			public readonly NodeInfo?[] Exports;
+			public Task<NodeInfo>? _mountTask;
+
+			public BlobInfo(BlobId blobId, int exportCount)
+				: this(blobId, new NodeInfo?[exportCount])
+			{
+			}
+
+			public BlobInfo(BlobId blobId, NodeInfo?[] exports)
+			{
+				BlobId = blobId;
+				Exports = exports;
+			}
+		}
+
+		readonly IBlobStore _blobStore;
+		readonly ConcurrentDictionary<IoHash, NodeInfo> _hashToNode = new ConcurrentDictionary<IoHash, NodeInfo>();
+		readonly ConcurrentDictionary<BlobId, BlobInfo> _blobIdToBundleInfo = new ConcurrentDictionary<BlobId, BlobInfo>();
+		readonly IMemoryCache _cache;
+
+		readonly object _lockObject = new object();
+
+		readonly SemaphoreSlim _readSema;
+
+		// Active read tasks at any moment. If a BundleObject is not available in the cache, we start a read and add an entry to this dictionary
+		// so that other threads can also await it.
+		readonly Dictionary<string, Task<Bundle>> _readTasks = new Dictionary<string, Task<Bundle>>();
+
+		readonly BundleOptions _options;
+
+		/// <summary>
+		/// Tracks statistics for the size of written data
+		/// </summary>
+		public BundleStats Stats { get; private set; } = new BundleStats();
+
+		BundleStats _nextStats = new BundleStats();
+		byte[] _blockBuffer = Array.Empty<byte>();
+		byte[] _encodedBuffer = Array.Empty<byte>();
+
+		/// <summary>
+		/// Constructor
+		/// </summary>
+		public BundleStore(IBlobStore blobStore, BundleOptions options, IMemoryCache cache)
+		{
+			_blobStore = blobStore;
+			_options = options;
+			_cache = cache;
+			_readSema = new SemaphoreSlim(options.MaxConcurrentReads);
+		}
+
+		/// <inheritdoc/>
+		public void Dispose()
+		{
+			Dispose(true);
+			GC.SuppressFinalize(this);
+		}
+
+		/// <summary>
+		/// Dispose of this objects resources
+		/// </summary>
+		/// <param name="disposing"></param>
+		protected virtual void Dispose(bool disposing)
+		{
+			_readSema.Dispose();
+		}
+
+		/// <inheritdoc/>
+		public Task<bool> HasTreeAsync(RefId id, CancellationToken cancellationToken) => _blobStore.HasRefAsync(id, cancellationToken);
+
+		/// <inheritdoc/>
+		public async Task<ITreeBlob> ReadTreeAsync(RefId id, CancellationToken cancellationToken = default)
+		{
+			Bundle bundle = await _blobStore.ReadBundleAsync(id, cancellationToken);
+			return MountInMemoryBundle(bundle);
+		}
+
+		/// <inheritdoc/>
+		public async Task WriteTreeAsync(RefId id, ITreeBlob root, bool flush = true, CancellationToken cancellationToken = default)
+		{
+			// Find all the referenced in-memory nodes from the root
+			List<NodeInfo> nodes = new List<NodeInfo>();
+			FindNodesToWrite((NodeInfo)root, nodes, new HashSet<NodeInfo>(), flush);
+
+			// Group the nodes into bundles
+			long bundleSize = 0;
+			List<NodeInfo> bundleNodes = new List<NodeInfo>();
+
+			foreach (NodeInfo node in nodes)
+			{
+				// Check if we need to flush the current bundle
+				long nodeSize = ((InMemoryNodeState)node.State).Data.Length;
+				if (bundleNodes.Count > 0 && bundleSize + nodeSize > _options.MaxBlobSize)
+				{
+					Bundle bundle = CreateBundle(bundleNodes);
+					BlobId blobId = await _blobStore.WriteBundleAsync(bundle, cancellationToken);
+
+					BlobInfo bundleInfo = new BlobInfo(blobId, bundleNodes.ToArray());
+					MountExportedBundle(bundleInfo, bundle);
+
+					bundleSize = 0;
+					bundleNodes.Clear();
+				}
+
+				// Add this node to the list
+				bundleSize += nodeSize;
+				bundleNodes.Add(node);
+			}
+
+			// If we're doing a full flush, also write whatever is left
+			if (flush)
+			{
+				// Write the main ref
+				Bundle rootBundle = CreateBundle(bundleNodes);
+				await _blobStore.WriteBundleAsync(id, rootBundle, cancellationToken);
+
+				// Copy the stats over
+				Stats = _nextStats;
+				_nextStats = new BundleStats();
+			}
+		}
+
+		/// <inheritdoc/>
+		public Task<bool> DeleteTreeAsync(RefId id, CancellationToken cancellationToken) => _blobStore.DeleteRefAsync(id, cancellationToken);
+
+		/// <summary>
+		/// Gets the data for a given node
+		/// </summary>
+		/// <param name="node">The node to get the data for</param>
+		/// <param name="cancellationToken">Cancellation token for the operation</param>
+		/// <returns>The node data</returns>
+		async ValueTask<ReadOnlySequence<byte>> GetDataAsync(NodeInfo node, CancellationToken cancellationToken)
+		{
+			if (node.State is InMemoryNodeState inMemoryState)
+			{
+				return inMemoryState.Data;
+			}
+
+			ExportedNodeState exportedState = await GetExportedStateAsync(node, cancellationToken);
+
+			Bundle bundle = await ReadBundleAsync(exportedState.BundleInfo, cancellationToken);
+			ReadOnlyMemory<byte> packetData = DecodePacket(exportedState.BundleInfo, exportedState.PacketInfo, bundle.Payload);
+			ReadOnlyMemory<byte> nodeData = packetData.Slice(exportedState.Offset, exportedState.Length);
+
+			return new ReadOnlySequence<byte>(nodeData);
+		}
+
+		/// <summary>
+		/// Gets the references from a given node
+		/// </summary>
+		/// <param name="node">The node to get the data for</param>
+		/// <param name="cancellationToken">Cancellation token for the operation</param>
+		/// <returns>The node data</returns>
+		async ValueTask<IReadOnlyList<NodeInfo>> GetReferencesAsync(NodeInfo node, CancellationToken cancellationToken)
+		{
+			if (node.State is InMemoryNodeState inMemoryState)
+			{
+				return inMemoryState.References;
+			}
+
+			ExportedNodeState exportedState = await GetExportedStateAsync(node, cancellationToken);
+			return exportedState.References;
+		}
+
+		async ValueTask<ExportedNodeState> GetExportedStateAsync(NodeInfo node, CancellationToken cancellationToken)
+		{
+			ExportedNodeState? exportedState = node.State as ExportedNodeState;
+			if (exportedState == null)
+			{
+				ImportedNodeState importedState = (ImportedNodeState)node.State;
+				await MountBundleAsync(importedState.BundleInfo, cancellationToken);
+				exportedState = (ExportedNodeState)node.State;
+			}
+			return exportedState;
+		}
+
+		/// <inheritdoc/>
+		public Task<ITreeBlob> WriteBlobAsync(ReadOnlySequence<byte> data, IReadOnlyList<ITreeBlob> references, CancellationToken cancellationToken)
+		{
+			List<NodeInfo> typedReferences = references.ConvertAll(x => (NodeInfo)x);
+			if (typedReferences.Any(x => x.Owner != this))
+			{
+				throw new InvalidDataException("Referenced node does not belong to the current tree.");
+			}
+
+			IoHash hash = ComputeHash(data, typedReferences);
+
+			InMemoryNodeState state = new InMemoryNodeState(data, typedReferences);
+			NodeInfo node = FindOrAddNode(hash, state);
+
+			_nextStats.NewNodeCount++;
+			_nextStats.NewNodeBytes += data.Length;
+
+			return Task.FromResult<ITreeBlob>(node);
+		}
+
+		NodeInfo FindOrAddNode(IoHash hash, NodeState state)
+		{
+			NodeInfo? node;
+			if (!_hashToNode.TryGetValue(hash, out node))
+			{
+				node = _hashToNode.GetOrAdd(hash, new NodeInfo(this, hash, state));
+			}
+			return node;
+		}
+
+		static IoHash ComputeHash(ReadOnlySequence<byte> data, IReadOnlyList<NodeInfo> references)
+		{
+			byte[] buffer = new byte[IoHash.NumBytes * (references.Count + 1)];
+
+			Span<byte> span = buffer;
+			for (int idx = 0; idx < references.Count; idx++)
+			{
+				references[idx].Hash.CopyTo(span);
+				span = span[IoHash.NumBytes..];
+			}
+			IoHash.Compute(data).CopyTo(span);
+
+			return IoHash.Compute(buffer);
+		}
+
+		void FindNodesToWrite(NodeInfo root, List<NodeInfo> nodes, HashSet<NodeInfo> uniqueNodes, bool flush)
+		{
+			if (root.State is InMemoryNodeState inMemoryState)
+			{
+				if (uniqueNodes.Add(root))
+				{
+					int count = inMemoryState.References.Count;
+					for (int idx = 0; idx < count; idx++)
+					{
+						FindNodesToWrite(inMemoryState.References[idx], nodes, uniqueNodes, flush || idx < count - 1);
+					}
+					if (flush)
+					{
+						nodes.Add(root);
+					}
+				}
+			}
+		}
+
+		async Task<NodeInfo> MountBundleAsync(BlobInfo bundle, CancellationToken cancellationToken)
+		{
+			Task<NodeInfo>? mountTask = bundle._mountTask;
+			if (mountTask == null)
+			{
+				lock (bundle)
+				{
+					bundle._mountTask ??= Task.Run(() => MountBundleInternalAsync(bundle, cancellationToken));
+					mountTask = bundle._mountTask;
+				}
+			}
+			return await mountTask;
+		}
+
+		async Task<NodeInfo> MountBundleInternalAsync(BlobInfo bundleInfo, CancellationToken cancellationToken)
+		{
+			Bundle bundle = await ReadBundleAsync(bundleInfo, cancellationToken);
+			return MountExportedBundle(bundleInfo, bundle);
+		}
+
+		List<NodeInfo> RegisterImports(BundleHeader header)
+		{
+			// Create all the nodes as imports first, so we can fix up references later.
+			List<NodeInfo> nodes = new List<NodeInfo>(header.Imports.Sum(x => x.Exports.Count) + header.Exports.Count);
+			foreach (BundleImport import in header.Imports)
+			{
+				BlobInfo importBundle = _blobIdToBundleInfo.GetOrAdd(import.BlobId, id => new BlobInfo(id, import.ExportCount));
+				foreach ((int exportIdx, IoHash exportHash) in import.Exports)
+				{
+					NodeInfo? export = importBundle.Exports[exportIdx];
+					if (export == null)
+					{
+						export = FindOrAddNode(exportHash, new ImportedNodeState(importBundle, exportIdx));
+						importBundle.Exports[exportIdx] = export;
+					}
+					nodes.Add(export);
+				}
+			}
+			return nodes;
+		}
+
+		NodeInfo MountInMemoryBundle(Bundle bundle)
+		{
+			BundleHeader header = bundle.Header;
+			List<NodeInfo> nodes = RegisterImports(header);
+
+			int exportIdx = 0;
+			int packetOffset = 0;
+			foreach (BundlePacket packet in header.Packets)
+			{
+				ReadOnlyMemory<byte> encodedPacket = bundle.Payload.Slice(packetOffset, packet.EncodedLength);
+
+				byte[] decodedPacket = new byte[packet.DecodedLength];
+				LZ4Codec.Decode(encodedPacket.Span, decodedPacket);
+
+				int nodeOffset = 0;
+				for (; exportIdx < header.Exports.Count && nodeOffset + header.Exports[exportIdx].Length <= packet.DecodedLength; exportIdx++)
+				{
+					BundleExport export = header.Exports[exportIdx];
+
+					ReadOnlyMemory<byte> nodeData = decodedPacket.AsMemory(nodeOffset, export.Length);
+					List<NodeInfo> nodeRefs = export.References.ConvertAll(x => nodes[x]);
+
+					InMemoryNodeState nodeState = new InMemoryNodeState(new ReadOnlySequence<byte>(nodeData), nodeRefs);
+					NodeInfo node = FindOrAddNode(export.Hash, nodeState);
+					nodes.Add(node);
+
+					nodeOffset += header.Exports[exportIdx].Length;
+				}
+			}
+
+			return nodes[^1];
+		}
+
+		NodeInfo MountExportedBundle(BlobInfo bundleInfo, Bundle bundle)
+		{
+			BundleHeader header = bundle.Header;
+			List<NodeInfo> nodes = RegisterImports(header);
+
+			int exportIdx = 0;
+			int packetOffset = 0;
+			foreach (BundlePacket packet in header.Packets)
+			{
+				PacketInfo packetInfo = new PacketInfo(packetOffset, packet);
+
+				int nodeOffset = 0;
+				for (; exportIdx < header.Exports.Count && nodeOffset + header.Exports[exportIdx].Length <= packet.DecodedLength; exportIdx++)
+				{
+					BundleExport export = header.Exports[exportIdx];
+
+					List<NodeInfo> nodeRefs = export.References.ConvertAll(x => nodes[x]);
+					ExportedNodeState nodeState = new ExportedNodeState(bundleInfo, exportIdx, packetInfo, nodeOffset, export.Length, nodeRefs);
+
+					NodeInfo node = FindOrAddNode(export.Hash, nodeState);
+					node.State = nodeState;
+					bundleInfo.Exports[exportIdx] = node;
+					nodes.Add(node);
+
+					nodeOffset += header.Exports[exportIdx].Length;
+				}
+
+				packetOffset += packet.EncodedLength;
+			}
+
+			return nodes[^1];
+		}
+
+		async Task<Bundle> ReadBundleAsync(BlobInfo bundleInfo, CancellationToken cancellationToken)
+		{
+			string cacheKey = $"bundle:{bundleInfo.BlobId}";
+			if (!_cache.TryGetValue<Bundle>(cacheKey, out Bundle? bundle))
+			{
+				Task<Bundle>? readTask;
+				lock (_lockObject)
+				{
+					if (!_readTasks.TryGetValue(cacheKey, out readTask))
+					{
+						readTask = Task.Run(() => ReadBundleInternalAsync(cacheKey, bundleInfo.BlobId, cancellationToken));
+						_readTasks.Add(cacheKey, readTask);
+					}
+				}
+				bundle = await readTask;
+			}
+			return bundle;
+		}
+
+		async Task<Bundle> ReadBundleInternalAsync(string cacheKey, BlobId blobId, CancellationToken cancellationToken)
+		{
+			// Perform another (sequenced) check whether an object has been added to the cache, to counteract the race between a read task being added and a task completing.
+			lock (_lockObject)
+			{
+				if (_cache.TryGetValue<Bundle>(cacheKey, out Bundle? cachedObject))
+				{
+					return cachedObject;
+				}
+			}
+
+			// Wait for the read semaphore to avoid triggering too many operations at once.
+			await _readSema.WaitAsync(cancellationToken);
+
+			// Read the data from storage
+			Bundle bundle;
+			try
+			{
+				bundle = await _blobStore.ReadBundleAsync(blobId, cancellationToken);
+			}
+			finally
+			{
+				_readSema.Release();
+			}
+
+			// Add the object to the cache
+			using (ICacheEntry entry = _cache.CreateEntry(cacheKey))
+			{
+				entry.SetSize(bundle.Payload.Length);
+				entry.SetValue(bundle);
+			}
+
+			// Remove this object from the list of read tasks
+			lock (_lockObject)
+			{
+				_readTasks.Remove(cacheKey);
+			}
+			return bundle;
+		}
+
+		/// <summary>
+		/// Gets a decoded block from the store
+		/// </summary>
+		/// <param name="bundleInfo">Information about the bundle</param>
+		/// <param name="packetInfo">The decoded block location and size</param>
+		/// <param name="payload">The bundle payload</param>
+		/// <returns>The decoded data</returns>
+		ReadOnlyMemory<byte> DecodePacket(BlobInfo bundleInfo, PacketInfo packetInfo, ReadOnlyMemory<byte> payload)
+		{
+			string cacheKey = $"bundle-packet:{bundleInfo.BlobId}@{packetInfo.Offset}";
+			return _cache.GetOrCreate<ReadOnlyMemory<byte>>(cacheKey, entry =>
+			{
+				ReadOnlyMemory<byte> decodedPacket = DecodePacketUncached(packetInfo, payload);
+				entry.SetSize(packetInfo.DecodedLength);
+				return decodedPacket;
+			});
+		}
+
+		static ReadOnlyMemory<byte> DecodePacketUncached(PacketInfo packetInfo, ReadOnlyMemory<byte> payload)
+		{
+			ReadOnlyMemory<byte> encodedPacket = payload.Slice(packetInfo.Offset, packetInfo.EncodedLength);
+
+			byte[] decodedPacket = new byte[packetInfo.DecodedLength];
+			LZ4Codec.Decode(encodedPacket.Span, decodedPacket);
+
+			return decodedPacket;
+		}
+
+		async Task FindLiveNodes(NodeInfo node, Dictionary<BlobInfo, HashSet<NodeInfo>> bundleToLiveNodes, HashSet<NodeInfo> nodes, CancellationToken cancellationToken)
+		{
+			if (nodes.Add(node))
+			{
+				if (node.State is ImportedNodeState importedState)
+				{
+					HashSet<NodeInfo>? liveNodes;
+					if (!bundleToLiveNodes.TryGetValue(importedState.BundleInfo, out liveNodes))
+					{
+						liveNodes = new HashSet<NodeInfo>();
+						bundleToLiveNodes.Add(importedState.BundleInfo, liveNodes);
+					}
+					liveNodes.Add(node);
+				}
+
+				foreach (NodeInfo reference in await GetReferencesAsync(node, cancellationToken))
+				{
+					await FindLiveNodes(reference, bundleToLiveNodes, nodes, cancellationToken);
+				}
+			}
+		}
+
+		async Task<bool> MarkDependencies(NodeInfo node, CancellationToken cancellationToken)
+		{
+			bool rewrite = false;
+			foreach (NodeInfo reference in await GetReferencesAsync(node, cancellationToken))
+			{
+				rewrite |= await MarkDependencies(reference, cancellationToken);
+			}
+			if (rewrite)
+			{
+				await MakeInMemory(node, cancellationToken);
+			}
+			return node.State is InMemoryNodeState;
+		}
+
+		async Task MakeInMemory(NodeInfo node, CancellationToken cancellationToken)
+		{
+			if (!(node.State is InMemoryNodeState))
+			{
+				ReadOnlySequence<byte> data = await GetDataAsync(node, cancellationToken);
+				IReadOnlyList<NodeInfo> references = await GetReferencesAsync(node, cancellationToken);
+				node.State = new InMemoryNodeState(data, references);
+			}
+		}
+
+#if false
+		/// <summary>
+		/// Splits apart any bundles that have a low utilization
+		/// </summary>
+		/// <returns></returns>
+		public async Task OptimizeAsync(double repackRatio)
+		{
+			// Find all the live nodes in the tree
+			Dictionary<BundleInfo, HashSet<NodeInfo>> bundleToLiveNodes = new Dictionary<BundleInfo, HashSet<NodeInfo>>();
+			await FindLiveNodes(_root, bundleToLiveNodes, new HashSet<NodeInfo>());
+
+			// Figure out which blobs need to be repacked
+			foreach ((BundleInfo bundleInfo, HashSet<NodeInfo> liveNodes) in bundleToLiveNodes)
+			{
+				Bundle bundle = await ReadBundleAsync(bundleInfo);
+				BundleHeader header = bundle.Header;
+
+				// Get the total size of this bundle
+				long totalSize = header.Exports.Sum(x => x.Length + x.References.Count * IoHash.NumBytes);
+
+				// Figure out how much of it we're actually using
+				long usedSize = 0;
+				foreach(NodeInfo liveNode in liveNodes)
+				{
+					if(liveNode.State is ExportedNodeState exportedState)
+					{
+						usedSize += exportedState.Length + exportedState.References.Count * IoHash.NumBytes;
+					}
+				}
+
+				int minUsedSize = (int)(totalSize * repackRatio);
+				if (usedSize < minUsedSize)
+				{
+					foreach (NodeInfo liveNode in liveNodes)
+					{
+						await MakeInMemory(liveNode);
+					}
+				}
+			}
+
+			// Walk backwards through the list of live nodes and make anything that references a repacked node also standalone
+			await MarkDependencies(_root);
+		}
+#endif
+		/// <summary>
+		/// Creates a Bundle containing a set of nodes. 
+		///
+		/// WARNING: The <see cref="Bundle.Payload"/> member will be set to the value of <see cref="_encodedBuffer"/> on return, which will be
+		/// reused on a subsequent call. If the object needs to be persisted across the lifetime of other objects, this field must be duplicated.
+		/// </summary>
+		Bundle CreateBundle(IReadOnlyList<NodeInfo> nodes)
+		{
+			// Create a set from the nodes to be written. We use this to determine references that are imported.
+			HashSet<NodeInfo> nodeSet = new HashSet<NodeInfo>(nodes);
+
+			// Find all the imported nodes by bundle
+			Dictionary<BlobInfo, List<NodeInfo>> bundleToImportedNodes = new Dictionary<BlobInfo, List<NodeInfo>>();
+			foreach (NodeInfo node in nodes)
+			{
+				InMemoryNodeState state = (InMemoryNodeState)node.State;
+				foreach (NodeInfo reference in state.References)
+				{
+					if (nodeSet.Add(reference))
+					{
+						// Get the persisted node info
+						ImportedNodeState importedState = (ImportedNodeState)reference.State;
+						BlobInfo bundleInfo = importedState.BundleInfo;
+
+						// Get the list of nodes within it
+						List<NodeInfo>? importedNodes;
+						if (!bundleToImportedNodes.TryGetValue(bundleInfo, out importedNodes))
+						{
+							importedNodes = new List<NodeInfo>();
+							bundleToImportedNodes.Add(bundleInfo, importedNodes);
+						}
+						importedNodes.Add(reference);
+					}
+				}
+			}
+
+			// Create the import list
+			List<BundleImport> imports = new List<BundleImport>();
+
+			// Add all the imports and assign them identifiers
+			Dictionary<NodeInfo, int> nodeToIndex = new Dictionary<NodeInfo, int>();
+			foreach ((BlobInfo bundleInfo, List<NodeInfo> importedNodes) in bundleToImportedNodes)
+			{
+				(int, IoHash)[] exportInfos = new (int, IoHash)[importedNodes.Count];
+				for (int idx = 0; idx < importedNodes.Count; idx++)
+				{
+					NodeInfo importedNode = importedNodes[idx];
+					nodeToIndex.Add(importedNode, nodeToIndex.Count);
+
+					ImportedNodeState importedNodeState = (ImportedNodeState)importedNode.State;
+					exportInfos[idx] = (importedNodeState.ExportIdx, importedNode.Hash);
+				}
+				imports.Add(new BundleImport(bundleInfo.BlobId, bundleInfo.Exports.Length, exportInfos));
+			}
+
+			// Preallocate data in the encoded buffer to reduce fragmentation if we have to resize
+			CreateFreeSpace(ref _encodedBuffer, 0, _options.MaxBlobSize);
+
+			// Create the export list
+			List<BundleExport> exports = new List<BundleExport>();
+			List<BundlePacket> packets = new List<BundlePacket>();
+
+			// Size of data currently stored in the block buffer
+			int blockSize = 0;
+
+			// Compress all the nodes into the encoded buffer
+			int packetOffset = 0;
+			foreach (NodeInfo node in nodes)
+			{
+				InMemoryNodeState nodeState = (InMemoryNodeState)node.State;
+				ReadOnlySequence<byte> nodeData = nodeState.Data;
+
+				// If we can't fit this data into the current block, flush the contents of it first
+				if (blockSize > 0 && blockSize + nodeData.Length > _options.MinCompressionPacketSize)
+				{
+					FlushPacket(_blockBuffer.AsMemory(0, blockSize), ref packetOffset, packets);
+					blockSize = 0;
+				}
+
+				// Create the export for this node
+				int[] references = nodeState.References.Select(x => nodeToIndex[x]).ToArray();
+				BundleExport export = new BundleExport(node.Hash, (int)nodeData.Length, references);
+				exports.Add(export);
+				nodeToIndex[node] = nodeToIndex.Count;
+
+				// Write out the new block
+				if (nodeData.Length < _options.MinCompressionPacketSize || !nodeData.IsSingleSegment)
+				{
+					int requiredSize = Math.Max(blockSize + (int)nodeData.Length, (int)(_options.MaxBlobSize * 1.2));
+					CreateFreeSpace(ref _blockBuffer, blockSize, requiredSize);
+
+					foreach (ReadOnlyMemory<byte> nodeSegment in nodeData)
+					{
+						nodeSegment.CopyTo(_blockBuffer.AsMemory(blockSize));
+						blockSize += nodeSegment.Length;
+					}
+				}
+				else
+				{
+					FlushPacket(nodeData.First, ref packetOffset, packets);
+				}
+			}
+			FlushPacket(_blockBuffer.AsMemory(0, blockSize), ref packetOffset, packets);
+
+			// Flush the data
+			BundleHeader header = new BundleHeader(imports, exports, packets);
+			ReadOnlyMemory<byte> payload = _encodedBuffer.AsMemory(0, packets.Sum(x => x.EncodedLength));
+			return new Bundle(header, payload);
+		}
+
+		void FlushPacket(ReadOnlyMemory<byte> inputData, ref int packetOffset, List<BundlePacket> packets)
+		{
+			if (inputData.Length > 0)
+			{
+				int minFreeSpace = LZ4Codec.MaximumOutputSize(inputData.Length);
+				CreateFreeSpace(ref _encodedBuffer, packetOffset, packetOffset + minFreeSpace);
+
+				ReadOnlySpan<byte> inputSpan = inputData.Span;
+				Span<byte> outputSpan = _encodedBuffer.AsSpan(packetOffset);
+
+				int encodedLength = LZ4Codec.Encode(inputSpan, outputSpan);
+				packets.Add(new BundlePacket(encodedLength, inputData.Length));
+
+				packetOffset += encodedLength;
+
+				Debug.Assert(encodedLength >= 0);
+			}
+		}
+
+		static void CreateFreeSpace(ref byte[] buffer, int usedSize, int requiredSize)
+		{
+			if (requiredSize > buffer.Length)
+			{
+				byte[] newBuffer = new byte[requiredSize];
+				buffer.AsSpan(0, usedSize).CopyTo(newBuffer);
+				buffer = newBuffer;
+			}
+		}
+	}
+
+	/// <summary>
+	/// Manages a collection of bundles.
+	/// </summary>
+	public class BundleStore<T> : BundleStore, ITreeStore<T> where T : TreeNode
+	{
+		/// <summary>
+		/// Constructor
+		/// </summary>
+		public BundleStore(IBlobStore blobStore, BundleOptions options, IMemoryCache cache)
+			: base(blobStore, options, cache)
+		{
+		}
+
+		/// <inheritdoc/>
+		public new async Task<T> ReadTreeAsync(RefId refName, CancellationToken cancellationToken = default)
+		{
+			ITreeBlob node = await base.ReadTreeAsync(refName, cancellationToken);
+			return await TreeNode.DeserializeAsync<T>(node, cancellationToken);
+		}
+
+		/// <inheritdoc/>
+		public async Task WriteTreeAsync(RefId refName, T root, bool flush, CancellationToken cancellationToken = default)
+		{
+			ITreeBlob rootBlob = await TreeNode.SerializeAsync(root, this, cancellationToken);
+			await base.WriteTreeAsync(refName, rootBlob, flush, cancellationToken);
+		}
+	}
+}
diff --git a/Engine/Source/Programs/Shared/EpicGames.Horde/Storage/ITreeStore.cs b/Engine/Source/Programs/Shared/EpicGames.Horde/Storage/ITreeStore.cs
new file mode 100644
--- /dev/null
+++ b/Engine/Source/Programs/Shared/EpicGames.Horde/Storage/ITreeStore.cs
@@ -0,0 +1,110 @@
+// Copyright Epic Games, Inc. All Rights Reserved.
+
+using System;
+using System.Buffers;
+using System.Collections.Generic;
+using System.Threading;
+using System.Threading.Tasks;
+using EpicGames.Core;
+
+namespace EpicGames.Horde.Storage
+{
+	/// <summary>
+	/// Interface for a blob of data for a tree node
+	/// </summary>
+	public interface ITreeBlob
+	{
+		/// <summary>
+		/// Gets the data for a node with the given hash
+		/// </summary>
+		/// <param name="cancellationToken">Cancellation token for the operation</param>
+		/// <returns>Data for the blob</returns>
+		ValueTask<ReadOnlySequence<byte>> GetDataAsync(CancellationToken cancellationToken = default);
+
+		/// <summary>
+		/// Find the outward references for a node
+		/// </summary>
+		/// <param name="cancellationToken">Cancellation token for the operation</param>
+		/// <returns>List of outward references from the blob</returns>
+		ValueTask<IReadOnlyList<ITreeBlob>> GetReferencesAsync(CancellationToken cancellationToken = default);
+	}
+
+	/// <summary>
+	/// Writer for trees
+	/// </summary>
+	public interface ITreeBlobWriter
+	{
+		/// <summary>
+		/// Write a blob to storage
+		/// </summary>
+		/// <param name="data"></param>
+		/// <param name="references"></param>
+		/// <param name="cancellationToken">Cancellation token for the operation</param>
+		/// <returns></returns>
+		public Task<ITreeBlob> WriteBlobAsync(ReadOnlySequence<byte> data, IReadOnlyList<ITreeBlob> references, CancellationToken cancellationToken);
+	}
+
+	/// <summary>
+	/// Base interface for a tree store
+	/// </summary>
+	public interface ITreeStore
+	{
+		/// <summary>
+		/// Tests whether a tree exists with the given name
+		/// </summary>
+		/// <param name="id">Name of the ref used to store the tree</param>
+		/// <param name="cancellationToken">Cancellation token for the operation</param>
+		/// <returns>Tree reader instance</returns>
+		Task<bool> HasTreeAsync(RefId id, CancellationToken cancellationToken = default);
+
+		/// <summary>
+		/// Deletes a tree with the given name
+		/// </summary>
+		/// <param name="id">Name of the ref used to store the tree</param>
+		/// <param name="cancellationToken">Cancellation token for the operation</param>
+		/// <returns>Tree writer instance</returns>
+		Task<bool> DeleteTreeAsync(RefId id, CancellationToken cancellationToken = default);
+
+		/// <summary>
+		/// Reads the root blob from a tree in storage
+		/// </summary>
+		/// <param name="id">Identifier for the tree</param>
+		/// <param name="cancellationToken">Cancellation token for the operation</param>
+		/// <returns></returns>
+		Task<ITreeBlob> ReadTreeAsync(RefId id, CancellationToken cancellationToken = default);
+
+		/// <summary>
+		/// Creates a writer for a tree with the given name
+		/// </summary>
+		/// <param name="id">Name of the ref used to store the tree</param>
+		/// <param name="root">Node to flush</param>
+		/// <param name="flush">Whether to flush the complete tree state. If false, an implementation may choose to buffer some writes to allow packing multiple nodes together.</param>
+		/// <param name="cancellationToken">Cancellation token for the operation</param>
+		/// <returns>Tree writer instance</returns>
+		Task WriteTreeAsync(RefId id, ITreeBlob root, bool flush = true, CancellationToken cancellationToken = default);
+	}
+
+	/// <summary>
+	/// Interface for a typed tree store
+	/// </summary>
+	public interface ITreeStore<T> : ITreeStore where T : TreeNode
+	{
+		/// <summary>
+		/// Reads a root node from the store
+		/// </summary>
+		/// <param name="id">Name of the tree to fetch</param>
+		/// <param name="cancellationToken">Cancellation token for the operation</param>
+		/// <returns>Deserialized node stored with this name</returns>
+		new Task<T> ReadTreeAsync(RefId id, CancellationToken cancellationToken = default);
+
+		/// <summary>
+		/// Creates a writer for a tree with the given name
+		/// </summary>
+		/// <param name="id">Name of the ref used to store the tree</param>
+		/// <param name="root">Node to flush</param>
+		/// <param name="flush">Whether to flush the complete tree state. If false, an implementation may choose to buffer some writes to allow packing multiple nodes together.</param>
+		/// <param name="cancellationToken">Cancellation token for the operation</param>
+		/// <returns>Tree writer instance</returns>
+		Task WriteTreeAsync(RefId id, T root, bool flush = true, CancellationToken cancellationToken = default);
+	}
+}
diff --git a/Engine/Source/Programs/Shared/EpicGames.Horde/Storage/Nodes/DirectoryNode.cs b/Engine/Source/Programs/Shared/EpicGames.Horde/Storage/Nodes/DirectoryNode.cs
new file mode 100644
--- /dev/null
+++ b/Engine/Source/Programs/Shared/EpicGames.Horde/Storage/Nodes/DirectoryNode.cs
@@ -0,0 +1,650 @@
+// Copyright Epic Games, Inc. All Rights Reserved.
+
+using EpicGames.Core;
+using EpicGames.Serialization;
+using Microsoft.Extensions.Logging;
+using System;
+using System.Buffers;
+using System.Collections.Generic;
+using System.Diagnostics;
+using System.Diagnostics.CodeAnalysis;
+using System.IO;
+using System.Linq;
+using System.Threading;
+using System.Threading.Tasks;
+
+namespace EpicGames.Horde.Storage.Nodes
+{
+	/// <summary>
+	/// Flags for a directory entry
+	/// </summary>
+	[Flags]
+	public enum FileEntryFlags
+	{
+		/// <summary>
+		/// No other flags set
+		/// </summary>
+		None = 0,
+
+		/// <summary>
+		/// Indicates that the referenced file is executable
+		/// </summary>
+		Executable = 4,
+
+		/// <summary>
+		/// File should be stored as read-only
+		/// </summary>
+		ReadOnly = 8,
+
+		/// <summary>
+		/// File contents are utf-8 encoded text. Client may want to replace line-endings with OS-specific format.
+		/// </summary>
+		Text = 16,
+
+		/// <summary>
+		/// The attached entry includes a Git SHA1 of the corresponding blob/tree contents.
+		/// </summary>
+		HasGitSha1 = 32,
+
+		/// <summary>
+		/// The data for this entry is a Perforce depot path and revision rather than the actual file contents.
+		/// </summary>
+		PerforceDepotPathAndRevision = 64,
+	}
+
+	/// <summary>
+	/// Entry for a file within a directory node
+	/// </summary>
+	public class FileEntry : TreeNodeRef<FileNode>
+	{
+		/// <summary>
+		/// Name of this file
+		/// </summary>
+		public Utf8String Name { get; }
+
+		/// <summary>
+		/// Flags for this file
+		/// </summary>
+		public FileEntryFlags Flags { get; }
+
+		/// <summary>
+		/// Length of this entry
+		/// </summary>
+		public long Length => (Node == null) ? _cachedLength : Node.Length;
+
+		/// <summary>
+		/// Cached length of this node
+		/// </summary>
+		long _cachedLength;
+
+		/// <summary>
+		/// Constructor
+		/// </summary>
+		public FileEntry(DirectoryNode owner, Utf8String name, FileEntryFlags flags, FileNode node)
+			: base(owner, node)
+		{
+			Name = name;
+			Flags = flags;
+		}
+
+		/// <summary>
+		/// Constructor
+		/// </summary>
+		public FileEntry(DirectoryNode owner, Utf8String name, FileEntryFlags flags, long length, ITreeBlob target)
+			: base(owner, target)
+		{
+			Name = name;
+			Flags = flags;
+			
+			_cachedLength = length;
+		}
+
+		/// <inheritdoc/>
+		protected override void OnCollapse()
+		{
+			_cachedLength = Node!.Length;
+		}
+
+		/// <inheritdoc/>
+		public override string ToString() => IsDirty() ? $"{Name} (*)" : Name.ToString();
+	}
+
+	/// <summary>
+	/// Entry for a directory within a directory node
+	/// </summary>
+	public class DirectoryEntry : TreeNodeRef<DirectoryNode>
+	{
+		/// <summary>
+		/// Name of this directory
+		/// </summary>
+		public Utf8String Name { get; }
+
+		/// <summary>
+		/// Length of this directory tree
+		/// </summary>
+		public long Length => (Node == null) ? _cachedLength : Node.Length;
+
+		/// <summary>
+		/// Cached value for the length of this tree
+		/// </summary>
+		long _cachedLength;
+
+		/// <summary>
+		/// Constructor
+		/// </summary>
+		public DirectoryEntry(DirectoryNode owner, Utf8String name, DirectoryNode node)
+			: base(owner, node)
+		{
+			Name = name;
+		}
+
+		/// <summary>
+		/// Constructor
+		/// </summary>
+		public DirectoryEntry(DirectoryNode owner, Utf8String name, long length, ITreeBlob target)
+			: base(owner, target)
+		{
+			Name = name;
+
+			_cachedLength = length;
+		}
+
+		/// <inheritdoc/>
+		protected override void OnCollapse()
+		{
+			_cachedLength = Node!.Length;
+		}
+
+		/// <inheritdoc/>
+		public override string ToString() => IsDirty() ? $"{Name} (*)" : Name.ToString();
+	}
+
+	/// <summary>
+	/// A directory node
+	/// </summary>
+	[TreeSerializer(typeof(DirectoryNodeDeserializer))]
+	public class DirectoryNode : TreeNode
+	{
+		internal const byte TypeId = (byte)'d';
+
+		readonly Dictionary<Utf8String, FileEntry> _nameToFileEntry = new Dictionary<Utf8String, FileEntry>();
+		readonly Dictionary<Utf8String, DirectoryEntry> _nameToDirectoryEntry = new Dictionary<Utf8String, DirectoryEntry>();
+
+		/// <summary>
+		/// Total size of this directory
+		/// </summary>
+		public long Length => Files.Sum(x => x.Length) + Directories.Sum(x => x.Length);
+
+		/// <summary>
+		/// All the files within this directory
+		/// </summary>
+		public IReadOnlyCollection<FileEntry> Files => _nameToFileEntry.Values;
+
+		/// <summary>
+		/// All the subdirectories within this directory
+		/// </summary>
+		public IReadOnlyCollection<DirectoryEntry> Directories => _nameToDirectoryEntry.Values;
+
+		/// <summary>
+		/// Clear the contents of this directory
+		/// </summary>
+		public void Clear()
+		{
+			_nameToFileEntry.Clear();
+			_nameToDirectoryEntry.Clear();
+			MarkAsDirty();
+		}
+
+		/// <summary>
+		/// Check whether an entry with the given name exists in this directory
+		/// </summary>
+		/// <param name="name">Name of the entry to search for</param>
+		/// <returns>True if the entry exists</returns>
+		public bool Contains(Utf8String name) => TryGetFileEntry(name, out _) || TryGetDirectoryEntry(name, out _);
+
+		/// <inheritdoc/>
+		public override IReadOnlyList<TreeNodeRef> GetReferences()
+		{
+			List<TreeNodeRef> refs = new List<TreeNodeRef>();
+			refs.AddRange(Files);
+			refs.AddRange(Directories);
+			return refs;
+		}
+
+		#region File operations
+
+		/// <summary>
+		/// Adds a new directory with the given name
+		/// </summary>
+		/// <param name="name">Name of the new directory</param>
+		/// <param name="flags">Flags for the new file</param>
+		/// <param name="data">Data for the file</param>
+		/// <param name="options">Options for chunking the data</param>
+		/// <param name="cancellationToken">Cancellation token for the operation</param>
+		/// <returns>The new directory object</returns>
+		public async Task<FileNode> AddFileAsync(Utf8String name, FileEntryFlags flags, ReadOnlyMemory<byte> data, ChunkingOptions options, CancellationToken cancellationToken)
+		{
+			using ReadOnlyMemoryStream stream = new ReadOnlyMemoryStream(data);
+			return await AddFileAsync(name, flags, stream, options, cancellationToken);
+		}
+
+		/// <summary>
+		/// Adds a new directory with the given name
+		/// </summary>
+		/// <param name="name">Name of the new directory</param>
+		/// <param name="flags">Flags for the new file</param>
+		/// <param name="stream">The stream to read from</param>
+		/// <param name="options">Options for chunking the data</param>
+		/// <param name="cancellationToken">Cancellation token for the operation</param>
+		/// <returns>The new directory object</returns>
+		public async Task<FileNode> AddFileAsync(Utf8String name, FileEntryFlags flags, Stream stream, ChunkingOptions options, CancellationToken cancellationToken)
+		{
+			if(TryGetDirectoryEntry(name, out _))
+			{
+				throw new ArgumentException($"A directory with the name {name} already exists");
+			}
+
+			FileNode newNode = await FileNode.CreateAsync(stream, options, cancellationToken);
+
+			FileEntry entry = new FileEntry(this, name, flags, newNode);
+			_nameToFileEntry[name] = entry;
+			MarkAsDirty();
+
+			return newNode;
+		}
+
+		/// <summary>
+		/// Finds or adds a file with the given path
+		/// </summary>
+		/// <param name="path">Path to the file</param>
+		/// <param name="flags">Flags for the new file</param>
+		/// <param name="stream">Stream to add to the file</param>
+		/// <param name="options">Options for chunking the data</param>
+		/// <param name="cancellationToken">Cancellation token for the operation</param>
+		/// <returns>The new directory object</returns>
+		public async ValueTask<FileNode> AddFileByPathAsync(Utf8String path, FileEntryFlags flags, Stream stream, ChunkingOptions options, CancellationToken cancellationToken)
+		{
+			DirectoryNode directory = this;
+
+			Utf8String remainingPath = path;
+			if (remainingPath[0] == '/' || remainingPath[0] == '\\')
+			{
+				remainingPath = remainingPath.Substring(1);
+			}
+
+			for (; ; )
+			{
+				int length = 0;
+				for (; ; length++)
+				{
+					if (length == remainingPath.Length)
+					{
+						return await directory.AddFileAsync(remainingPath, flags, stream, options, cancellationToken);
+					}
+
+					byte character = remainingPath[length];
+					if (character == '\\' || character == '/')
+					{
+						break;
+					}
+				}
+
+				if (length > 0)
+				{
+					directory = await directory.FindOrAddDirectoryAsync(remainingPath.Slice(0, length), cancellationToken);
+				}
+				remainingPath = remainingPath.Slice(length + 1);
+			}
+		}
+
+		/// <summary>
+		/// Attempts to get a file entry with the given name
+		/// </summary>
+		/// <param name="name">Name of the file</param>
+		/// <returns>Entry for the given name</returns>
+		public FileEntry GetFileEntry(Utf8String name) => _nameToFileEntry[name];
+
+		/// <summary>
+		/// Attempts to get a file entry with the given name
+		/// </summary>
+		/// <param name="name">Name of the file</param>
+		/// <param name="entry">Entry for the file</param>
+		/// <returns>True if the file was found</returns>
+		public bool TryGetFileEntry(Utf8String name, [NotNullWhen(true)] out FileEntry? entry) => _nameToFileEntry.TryGetValue(name, out entry);
+
+		/// <summary>
+		/// Deletes the file entry with the given name
+		/// </summary>
+		/// <param name="name">Name of the entry to delete</param>
+		/// <returns>True if the entry was found, false otherwise</returns>
+		public bool DeleteFile(Utf8String name)
+		{
+			if (_nameToFileEntry.Remove(name))
+			{
+				MarkAsDirty();
+				return true;
+			}
+			return false;
+		}
+
+		/// <summary>
+		/// Deletes a file with the given path
+		/// </summary>
+		/// <param name="path"></param>
+		/// <param name="cancellationToken"></param>
+		/// <returns></returns>
+		public async ValueTask<bool> DeleteFileByPathAsync(Utf8String path, CancellationToken cancellationToken)
+		{
+			Utf8String remainingPath = path;
+			for (DirectoryNode? directory = this; directory != null; )
+			{
+				int length = remainingPath.IndexOf('/');
+				if (length == -1)
+				{
+					return directory.DeleteFile(remainingPath);
+				}
+				if (length > 0)
+				{
+					directory = await directory.FindDirectoryAsync(remainingPath.Slice(0, length), cancellationToken);
+				}
+				remainingPath = remainingPath.Slice(length + 1);
+			}
+			return false;
+		}
+
+		#endregion
+
+		#region Directory operations
+
+		/// <summary>
+		/// Adds a new directory with the given name
+		/// </summary>
+		/// <param name="name">Name of the new directory</param>
+		/// <returns>The new directory object</returns>
+		public DirectoryNode AddDirectory(Utf8String name)
+		{
+			if (TryGetFileEntry(name, out _))
+			{
+				throw new ArgumentException($"A file with the name '{name}' already exists in this directory", nameof(name));
+			}
+
+			DirectoryNode node = new DirectoryNode();
+
+			DirectoryEntry entry = new DirectoryEntry(this, name, node);
+			_nameToDirectoryEntry.Add(name, entry);
+			MarkAsDirty();
+
+			return node;
+		}
+
+		/// <summary>
+		/// Get a directory entry with the given name
+		/// </summary>
+		/// <param name="name">Name of the directory</param>
+		/// <returns>The entry with the given name</returns>
+		public DirectoryEntry GetDirectoryEntry(Utf8String name) => _nameToDirectoryEntry[name];
+
+		/// <summary>
+		/// Attempts to get a directory entry with the given name
+		/// </summary>
+		/// <param name="name">Name of the directory</param>
+		/// <param name="entry">Entry for the directory</param>
+		/// <returns>True if the directory was found</returns>
+		public bool TryGetDirectoryEntry(Utf8String name, [NotNullWhen(true)] out DirectoryEntry? entry) => _nameToDirectoryEntry.TryGetValue(name, out entry);
+
+		/// <summary>
+		/// Tries to get a directory with the given name
+		/// </summary>
+		/// <param name="name">Name of the new directory</param>
+		/// <param name="cancellationToken">Cancellation token for the operation</param>
+		/// <returns>The new directory object</returns>
+		public async ValueTask<DirectoryNode?> FindDirectoryAsync(Utf8String name, CancellationToken cancellationToken)
+		{
+			if (TryGetDirectoryEntry(name, out DirectoryEntry? entry))
+			{
+				return await entry.ExpandAsync(cancellationToken);
+			}
+			else
+			{
+				return null;
+			}
+		}
+
+		/// <summary>
+		/// Tries to get a directory with the given name
+		/// </summary>
+		/// <param name="name">Name of the new directory</param>
+		/// <param name="cancellationToken">Cancellation token for the operation</param>
+		/// <returns>The new directory object</returns>
+		public async ValueTask<DirectoryNode> FindOrAddDirectoryAsync(Utf8String name, CancellationToken cancellationToken)
+		{
+			DirectoryNode? directory = await FindDirectoryAsync(name, cancellationToken);
+			if (directory == null)
+			{
+				directory = AddDirectory(name);
+			}
+			return directory;
+		}
+
+		/// <summary>
+		/// Deletes the file entry with the given name
+		/// </summary>
+		/// <param name="name">Name of the entry to delete</param>
+		/// <returns>True if the entry was found, false otherwise</returns>
+		public bool DeleteDirectory(Utf8String name)
+		{
+			if(_nameToDirectoryEntry.Remove(name))
+			{
+				MarkAsDirty();
+				return true;
+			}
+			return false;
+		}
+
+		#endregion
+
+		/// <inheritdoc/>
+		public async ValueTask<ITreeBlob> SerializeAsync(ITreeBlobWriter writer, CancellationToken cancellationToken)
+		{
+			List<ITreeBlob> childNodes = new List<ITreeBlob>();
+
+			List<FileEntry> fileEntries = _nameToFileEntry.Values.OrderBy(x => x.Name).ToList();
+			foreach (FileEntry fileEntry in fileEntries)
+			{
+				childNodes.Add(await fileEntry.CollapseAsync(writer, cancellationToken));
+			}
+
+			List<DirectoryEntry> directoryEntries = _nameToDirectoryEntry.Values.OrderBy(x => x.Name).ToList();
+			foreach (DirectoryEntry directoryEntry in directoryEntries)
+			{
+				childNodes.Add(await directoryEntry.CollapseAsync(writer, cancellationToken));
+			}
+
+			ReadOnlySequence<byte> data = SerializeData(fileEntries, directoryEntries);
+			return await writer.WriteBlobAsync(data, childNodes, cancellationToken);
+		}
+
+		static ReadOnlySequence<byte> SerializeData(List<FileEntry> fileEntries, List<DirectoryEntry> directoryEntries)
+		{
+			// Measure the required size of the write buffer
+			int size = 1;
+			size += VarInt.MeasureUnsigned(fileEntries.Count);
+			foreach (FileEntry fileEntry in fileEntries)
+			{
+				size += (fileEntry.Name.Length + 1) + 1 + VarInt.MeasureUnsigned((ulong)fileEntry.Length);
+			}
+			size += VarInt.MeasureUnsigned(directoryEntries.Count);
+			foreach (DirectoryEntry directoryEntry in directoryEntries)
+			{
+				size += (directoryEntry.Name.Length + 1) + VarInt.MeasureUnsigned((ulong)directoryEntry.Length);
+			}
+
+			// Allocate the buffer and copy the node to it
+			byte[] data = new byte[size];
+			Span<byte> span = data.AsSpan();
+
+			span[0] = TypeId;
+			span = span.Slice(1);
+
+			int fileCountBytes = VarInt.WriteUnsigned(span, fileEntries.Count);
+			span = span.Slice(fileCountBytes);
+
+			foreach (FileEntry fileEntry in fileEntries)
+			{
+				fileEntry.Name.Span.CopyTo(span);
+				span = span.Slice(fileEntry.Name.Length + 1);
+
+				span[0] = (byte)fileEntry.Flags;
+				span = span.Slice(1);
+
+				int lengthBytes = VarInt.WriteUnsigned(span, fileEntry.Length);
+				span = span.Slice(lengthBytes);
+			}
+
+			int directoryCountBytes = VarInt.WriteUnsigned(span, directoryEntries.Count);
+			span = span.Slice(directoryCountBytes);
+
+			foreach (DirectoryEntry directoryEntry in directoryEntries)
+			{
+				directoryEntry.Name.Span.CopyTo(span);
+				span = span.Slice(directoryEntry.Name.Length + 1);
+
+				int lengthBytes = VarInt.WriteUnsigned(span, directoryEntry.Length);
+				span = span.Slice(lengthBytes);
+			}
+
+			Debug.Assert(span.Length == 0);
+			return new ReadOnlySequence<byte>(data);
+		}
+
+		/// <summary>
+		/// Deserialize a directory node from data
+		/// </summary>
+		/// <param name="payload">The data to read from</param>
+		/// <param name="children">Child nodes</param>
+		/// <returns>The deserialized directory node</returns>
+		public static DirectoryNode Deserialize(ReadOnlySequence<byte> payload, IReadOnlyList<ITreeBlob> children)
+		{
+			ReadOnlySpan<byte> span = payload.AsSingleSegment().Span;
+			if(span[0] != TypeId)
+			{
+				throw new InvalidOperationException("Invalid signature byte for directory");
+			}
+
+			span = span.Slice(1);
+
+			DirectoryNode node = new DirectoryNode();
+
+			int fileCount = (int)VarInt.ReadUnsigned(span, out int fileCountBytes);
+			span = span.Slice(fileCountBytes);
+
+			int childIdx = 0;
+
+			node._nameToFileEntry.EnsureCapacity(fileCount);
+			for (int idx = 0; idx < fileCount; idx++)
+			{
+				int nameLen = span.IndexOf((byte)0);
+				Utf8String name = new Utf8String(span.Slice(0, nameLen).ToArray());
+				span = span.Slice(nameLen + 1);
+
+				FileEntryFlags flags = (FileEntryFlags)span[0];
+				span = span.Slice(1);
+
+				long length = (long)VarInt.ReadUnsigned(span, out int lengthBytes);
+				span = span.Slice(lengthBytes);
+
+				node._nameToFileEntry[name] = new FileEntry(node, name, flags, length, children[childIdx++]);
+			}
+
+			int directoryCount = (int)VarInt.ReadUnsigned(span, out int directoryCountBytes);
+			span = span.Slice(directoryCountBytes);
+
+			node._nameToDirectoryEntry.EnsureCapacity(directoryCount);
+			for (int idx = 0; idx < directoryCount; idx++)
+			{
+				int nameLen = span.IndexOf((byte)0);
+				Utf8String name = new Utf8String(span.Slice(0, nameLen).ToArray());
+				span = span.Slice(nameLen + 1);
+
+				long length = (long)VarInt.ReadUnsigned(span, out int lengthBytes);
+				span = span.Slice(lengthBytes);
+
+				node._nameToDirectoryEntry[name] = new DirectoryEntry(node, name, length, children[childIdx++]);
+			}
+
+			Debug.Assert(span.Length == 0);
+			return node;
+		}
+
+		/// <summary>
+		/// Adds files from a directory on disk
+		/// </summary>
+		/// <param name="directoryInfo"></param>
+		/// <param name="options">Options for chunking file content</param>
+		/// <param name="logger"></param>
+		/// <param name="cancellationToken">Cancellation token for the operation</param>
+		/// <returns></returns>
+		public async Task CopyFromDirectoryAsync(DirectoryInfo directoryInfo, ChunkingOptions options, ILogger logger, CancellationToken cancellationToken)
+		{
+			foreach (DirectoryInfo subDirectoryInfo in directoryInfo.EnumerateDirectories())
+			{
+				logger.LogInformation("Adding {Directory}", subDirectoryInfo.FullName);
+				DirectoryNode subDirectoryNode = AddDirectory(subDirectoryInfo.Name);
+				await subDirectoryNode.CopyFromDirectoryAsync(subDirectoryInfo, options, logger, cancellationToken);
+			}
+			foreach (FileInfo fileInfo in directoryInfo.EnumerateFiles())
+			{
+				logger.LogInformation("Adding {File}", fileInfo.FullName);
+				using Stream stream = fileInfo.OpenRead();
+				FileNode fileNode = await AddFileAsync(fileInfo.Name, 0, stream, options, cancellationToken);
+			}
+		}
+
+		/// <summary>
+		/// Utility function to allow extracting a packed directory to disk
+		/// </summary>
+		/// <param name="directoryInfo"></param>
+		/// <param name="logger"></param>
+		/// <param name="cancellationToken"></param>
+		/// <returns></returns>
+		public async Task CopyToDirectoryAsync(DirectoryInfo directoryInfo, ILogger logger, CancellationToken cancellationToken)
+		{
+			directoryInfo.Create();
+
+			List<Task> tasks = new List<Task>();
+			foreach (FileEntry fileEntry in Files)
+			{
+				FileInfo fileInfo = new FileInfo(Path.Combine(directoryInfo.FullName, fileEntry.Name.ToString()));
+				FileNode fileNode = await fileEntry.ExpandAsync(cancellationToken);
+				tasks.Add(Task.Run(() => fileNode.CopyToFileAsync(fileInfo, cancellationToken), cancellationToken));
+			}
+			foreach (DirectoryEntry directoryEntry in Directories)
+			{
+				DirectoryInfo subDirectoryInfo = directoryInfo.CreateSubdirectory(directoryEntry.Name.ToString());
+				DirectoryNode subDirectoryNode = await directoryEntry.ExpandAsync(cancellationToken);
+				tasks.Add(Task.Run(() => subDirectoryNode.CopyToDirectoryAsync(subDirectoryInfo, logger, cancellationToken), cancellationToken));
+			}
+
+			await Task.WhenAll(tasks);
+		}
+	}
+
+	/// <summary>
+	/// Factory for creating and serializing <see cref="DirectoryNode"/> objects
+	/// </summary>
+	public class DirectoryNodeDeserializer : TreeNodeSerializer<DirectoryNode>
+	{
+		/// <inheritdoc/>
+		public override async ValueTask<DirectoryNode> DeserializeAsync(ITreeBlob blob, CancellationToken cancellationToken) => DirectoryNode.Deserialize(await blob.GetDataAsync(cancellationToken), await blob.GetReferencesAsync(cancellationToken));
+
+		/// <inheritdoc/>
+		public override ValueTask<ITreeBlob> SerializeAsync(ITreeBlobWriter writer, DirectoryNode node, CancellationToken cancellationToken) => node.SerializeAsync(writer, cancellationToken);
+	}
+
+	/// <summary>
+	/// Extension methods for <see cref="DirectoryNode"/>
+	/// </summary>
+	public static class DirectoryNodeExtensions
+	{
+	}
+}
diff --git a/Engine/Source/Programs/Shared/EpicGames.Horde/Storage/Nodes/FileNode.cs b/Engine/Source/Programs/Shared/EpicGames.Horde/Storage/Nodes/FileNode.cs
new file mode 100644
--- /dev/null
+++ b/Engine/Source/Programs/Shared/EpicGames.Horde/Storage/Nodes/FileNode.cs
@@ -0,0 +1,597 @@
+// Copyright Epic Games, Inc. All Rights Reserved.
+
+using EpicGames.Core;
+using EpicGames.Serialization;
+using System;
+using System.Buffers;
+using System.Buffers.Binary;
+using System.Collections.Generic;
+using System.Diagnostics;
+using System.IO;
+using System.Threading;
+using System.Threading.Tasks;
+
+namespace EpicGames.Horde.Storage.Nodes
+{
+	/// <summary>
+	/// Options for creating file nodes
+	/// </summary>
+	public class ChunkingOptions
+	{
+		/// <summary>
+		/// Options for creating leaf nodes
+		/// </summary>
+		public ChunkingOptionsForNodeType LeafOptions { get; set; }
+
+		/// <summary>
+		/// Options for creating interior nodes
+		/// </summary>
+		public ChunkingOptionsForNodeType InteriorOptions { get; set; }
+
+		/// <summary>
+		/// Constructor
+		/// </summary>
+		public ChunkingOptions()
+		{
+			LeafOptions = new ChunkingOptionsForNodeType(32 * 1024, 256 * 1024, 64 * 1024);
+			InteriorOptions = new ChunkingOptionsForNodeType(32 * 1024, 256 * 1024, 64 * 1024);
+		}
+	}
+
+	/// <summary>
+	/// Options for creating a specific type of file nodes
+	/// </summary>
+	public class ChunkingOptionsForNodeType
+	{
+		/// <summary>
+		/// Minimum chunk size
+		/// </summary>
+		public int MinSize { get; set; }
+
+		/// <summary>
+		/// Maximum chunk size. Chunks will be split on this boundary if another match is not found.
+		/// </summary>
+		public int MaxSize { get; set; }
+
+		/// <summary>
+		/// Target chunk size for content-slicing
+		/// </summary>
+		public int TargetSize { get; set; }
+
+		/// <summary>
+		/// Default constructor
+		/// </summary>
+		public ChunkingOptionsForNodeType(int minSize, int maxSize, int targetSize)
+		{
+			MinSize = minSize;
+			MaxSize = maxSize;
+			TargetSize = targetSize;
+		}
+	}
+
+	/// <summary>
+	/// Representation of a data stream, split into chunks along content-aware boundaries using a rolling hash (<see cref="BuzHash"/>).
+	/// Chunks are pushed into a tree hierarchy as data is appended to the root, with nodes of the tree also split along content-aware boundaries with <see cref="IoHash.NumBytes"/> granularity.
+	/// Once a chunk has been written to storage, it is treated as immutable.
+	/// </summary>
+	[TreeSerializer(typeof(FileNodeSerializer))]
+	public abstract class FileNode : TreeNode
+	{
+		/// <summary>
+		/// Length of the node
+		/// </summary>
+		public abstract long Length { get; }
+
+		/// <summary>
+		/// Rolling hash for the current node
+		/// </summary>
+		public abstract uint RollingHash { get; }
+
+		/// <summary>
+		/// Serailize this node
+		/// </summary>
+		/// <param name="writer"></param>
+		/// <param name="cancellationToken"></param>
+		/// <returns></returns>
+		public abstract ValueTask<ITreeBlob> SerializeAsync(ITreeBlobWriter writer, CancellationToken cancellationToken);
+
+		/// <summary>
+		/// Creates a file node from a stream
+		/// </summary>
+		/// <param name="stream">The stream to read from</param>
+		/// <param name="options">Options for chunking the data</param>
+		/// <param name="cancellationToken">Cancellation token for the operation</param>
+		/// <returns>New file node</returns>
+		public static async Task<FileNode> CreateAsync(Stream stream, ChunkingOptions options, CancellationToken cancellationToken)
+		{
+			FileNode node = new LeafFileNode();
+
+			byte[] buffer = new byte[4 * 1024];
+			for (; ; )
+			{
+				int numBytes = await stream.ReadAsync(buffer, cancellationToken);
+				if (numBytes == 0)
+				{
+					break;
+				}
+				node = node.Append(buffer.AsMemory(0, numBytes), options);
+			}
+
+			return node;
+		}
+
+		/// <summary>
+		/// Append data to this chunk. Must only be called on the root node in a chunk tree.
+		/// </summary>
+		/// <param name="input">The data to write</param>
+		/// <param name="options">Settings for chunking the data</param>
+		public FileNode Append(ReadOnlyMemory<byte> input, ChunkingOptions options)
+		{
+			return Append(this, input, options);
+		}
+
+		static FileNode Append(FileNode root, ReadOnlyMemory<byte> input, ChunkingOptions options)
+		{
+			for (; ; )
+			{
+				// Append as much data as possible to the existing tree
+				input = root.AppendToNode(input, options);
+				if (input.IsEmpty)
+				{
+					break;
+				}
+
+				// Increase the height of the tree by pushing the contents of this node into a new child node
+				root = new InteriorFileNode(root.Length, root);
+
+/*
+				FileNode newNode = new InteriorFileNode(root.Length, root);
+				if (_data is InternalNodeData internalData)
+				{
+					foreach (TreeNodeRef<FileNode> childNodeRef in internalData._children)
+					{
+						childNodeRef.Reparent(newNode);
+					}
+				}
+				_data = new InternalNodeData(newNode.Length, new TreeNodeRef<FileNode>(this, newNode));*/
+			}
+			return root;
+		}
+
+		private ReadOnlyMemory<byte> AppendToNode(ReadOnlyMemory<byte> appendData, ChunkingOptions options)
+		{
+			if (appendData.Length == 0 || IsReadOnly())
+			{
+				return appendData;
+			}
+			else
+			{
+				return AppendData(appendData, options);
+			}
+		}
+
+		/// <summary>
+		/// Attempt to append data to the current node.
+		/// </summary>
+		/// <param name="newData">The data to append</param>
+		/// <param name="options">Options for chunking the data</param>
+		/// <returns>Remaining data in the buffer</returns>
+		public abstract ReadOnlyMemory<byte> AppendData(ReadOnlyMemory<byte> newData, ChunkingOptions options);
+
+		/// <summary>
+		/// Copies the contents of this node and its children to the given output stream
+		/// </summary>
+		/// <param name="outputStream">The output stream to receive the data</param>
+		/// <param name="cancellationToken">Cancellation token for the operation</param>
+		public abstract Task CopyToStreamAsync(Stream outputStream, CancellationToken cancellationToken);
+
+		/// <summary>
+		/// Extracts the contents of this node to a file
+		/// </summary>
+		/// <param name="file">File to write with the contents of this node</param>
+		/// <param name="cancellationToken">Cancellation token for the operation</param>
+		/// <returns></returns>
+		public async Task CopyToFileAsync(FileInfo file, CancellationToken cancellationToken)
+		{
+			using (FileStream stream = file.OpenWrite())
+			{
+				await CopyToStreamAsync(stream, cancellationToken);
+			}
+		}
+
+		/// <summary>
+		/// Serialize this node and its children into a byte array
+		/// </summary>
+		/// <param name="cancellationToken"></param>
+		/// <returns>Array of data stored by the tree</returns>
+		public async Task<byte[]> ToByteArrayAsync(CancellationToken cancellationToken)
+		{
+			using (MemoryStream stream = new MemoryStream())
+			{
+				await CopyToStreamAsync(stream, cancellationToken);
+				return stream.ToArray();
+			}
+		}
+	}
+
+	/// <summary>
+	/// File node that contains a chunk of data
+	/// </summary>
+	public sealed class LeafFileNode : FileNode
+	{
+		class DataSegment : ReadOnlySequenceSegment<byte>
+		{
+			public DataSegment(long runningIndex, ReadOnlyMemory<byte> data)
+			{
+				RunningIndex = runningIndex;
+				Memory = data;
+			}
+
+			public void SetNext(DataSegment next)
+			{
+				Next = next;
+			}
+		}
+
+		/// <summary>
+		/// First byte in the serialized data for this class, indicating its type.
+		/// </summary>
+		public const byte TypeId = (byte)'l';
+
+		const int HeaderLength = 1 + sizeof(uint);
+
+		bool _isReadOnly;
+		uint _rollingHash;
+		readonly byte[] _header;
+		readonly DataSegment _firstSegment;
+		DataSegment _lastSegment;
+		ReadOnlySequence<byte> _writtenSequence; // Payload described by _firstSegment -> _lastSegment
+
+		/// <summary>
+		/// Create an empty leaf node
+		/// </summary>
+		public LeafFileNode()
+		{
+			_header = new byte[HeaderLength];
+			_header[0] = TypeId;
+
+			_firstSegment = new DataSegment(0, _header);
+			_lastSegment = _firstSegment;
+
+			_writtenSequence = CreateSequence(_firstSegment, _lastSegment);
+		}
+
+		/// <summary>
+		/// Create a leaf node from the given serialized data
+		/// </summary>
+		/// <param name="data">Data to construct from</param>
+		public LeafFileNode(ReadOnlySequence<byte> data)
+		{
+			_header = Array.Empty<byte>();
+
+			ReadOnlySpan<byte> span = data.Slice(0, HeaderLength).AsSingleSegment().Span;
+			if (span[0] != TypeId)
+			{
+				throw new InvalidDataException($"Invalid type id for {nameof(LeafFileNode)}");
+			}
+
+			_isReadOnly = true;
+			_rollingHash = BinaryPrimitives.ReadUInt32LittleEndian(span.Slice(1, sizeof(uint)));
+
+			_firstSegment = null!;
+			_lastSegment = null!;
+
+			_writtenSequence = data;
+		}
+
+		static ReadOnlySequence<byte> CreateSequence(DataSegment firstSegment, DataSegment lastSegment)
+		{
+			return new ReadOnlySequence<byte>(firstSegment, 0, lastSegment, lastSegment.Memory.Length);
+		}
+
+		/// <inheritdoc/>
+		public override bool IsReadOnly() => _isReadOnly;
+
+		/// <inheritdoc/>
+		public override uint RollingHash => _rollingHash;
+
+		/// <summary>
+		/// Gets the data for this node
+		/// </summary>
+		public ReadOnlySequence<byte> Data => _writtenSequence.Slice(HeaderLength);
+
+		/// <inheritdoc/>
+		public override long Length => _writtenSequence.Length - HeaderLength;
+
+		/// <inheritdoc/>
+		public override IReadOnlyList<TreeNodeRef> GetReferences() => Array.Empty<TreeNodeRef>();
+
+		/// <inheritdoc/>
+		public override async ValueTask<ITreeBlob> SerializeAsync(ITreeBlobWriter writer, CancellationToken cancellationToken)
+		{
+			return await writer.WriteBlobAsync(_writtenSequence, Array.Empty<ITreeBlob>(), cancellationToken);
+		}
+
+		/// <inheritdoc/>
+		public override ReadOnlyMemory<byte> AppendData(ReadOnlyMemory<byte> newData, ChunkingOptions options)
+		{
+			// Fast path for appending data to the buffer up to the chunk window size
+			int windowSize = options.LeafOptions.MinSize;
+			if (Length < windowSize)
+			{
+				int appendLength = Math.Min(windowSize - (int)Length, newData.Length);
+				AppendLeafData(newData.Span.Slice(0, appendLength));
+				newData = newData.Slice(appendLength);
+			}
+
+			// Cap the maximum amount of data to append to this node
+			int maxLength = Math.Min(newData.Length, options.LeafOptions.MaxSize - (int)Length);
+			if (maxLength > 0)
+			{
+				ReadOnlySpan<byte> inputSpan = newData.Span.Slice(0, maxLength);
+				int length = AppendLeafDataToChunkBoundary(inputSpan, options);
+				newData = newData.Slice(length);
+			}
+
+			// Mark this node as complete if we've reached the max size
+			if (Length == options.LeafOptions.MaxSize)
+			{
+				_isReadOnly = true;
+			}
+			return newData;
+		}
+
+		private int AppendLeafDataToChunkBoundary(ReadOnlySpan<byte> headSpan, ChunkingOptions options)
+		{
+			int windowSize = options.LeafOptions.MinSize;
+			Debug.Assert(Length >= windowSize);
+
+			// Get the threshold for the rolling hash
+			uint newRollingHash = _rollingHash;
+			uint rollingHashThreshold = (uint)((1L << 32) / options.LeafOptions.TargetSize);
+
+			// Length of the data taken from the input span, updated as we step through it.
+			int length = 0;
+
+			// Step the window through the tail end of the existing payload window. In this state, update the hash to remove data from the current payload, and add data from the new payload.
+			int splitLength = Math.Min(headSpan.Length, windowSize);
+			ReadOnlySequence<byte> tailSequence = _writtenSequence.Slice(_writtenSequence.Length - windowSize, splitLength);
+
+			foreach (ReadOnlyMemory<byte> tailSegment in tailSequence)
+			{
+				int baseLength = length;
+				int spanLength = length + tailSegment.Length;
+
+				ReadOnlySpan<byte> tailSpan = tailSegment.Span;
+				for (; length < spanLength; length++)
+				{
+					newRollingHash = BuzHash.Add(newRollingHash, headSpan[length]);
+					if (newRollingHash < rollingHashThreshold)
+					{
+						AppendLeafData(headSpan.Slice(0, length), newRollingHash);
+						_isReadOnly = true;
+						return length;
+					}
+					newRollingHash = BuzHash.Sub(newRollingHash, tailSpan[length - baseLength], windowSize);
+				}
+			}
+
+			// Step through the new window until we get to a chunk boundary.
+			for (; length < headSpan.Length; length++)
+			{
+				newRollingHash = BuzHash.Add(newRollingHash, headSpan[length]);
+				if (newRollingHash < rollingHashThreshold)
+				{
+					AppendLeafData(headSpan.Slice(0, length), newRollingHash);
+					_isReadOnly = true;
+					return length;
+				}
+				newRollingHash = BuzHash.Sub(newRollingHash, headSpan[length - windowSize], windowSize);
+			}
+
+			// Otherwise just append all the data.
+			AppendLeafData(headSpan, newRollingHash);
+			return headSpan.Length;
+		}
+
+		private void AppendLeafData(ReadOnlySpan<byte> leafData)
+		{
+			uint newRollingHash = BuzHash.Add(_rollingHash, leafData);
+			AppendLeafData(leafData, newRollingHash);
+		}
+
+		private void AppendLeafData(ReadOnlySpan<byte> leafData, uint newRollingHash)
+		{
+			_rollingHash = newRollingHash;
+			BinaryPrimitives.WriteUInt32LittleEndian(_header.AsSpan(1, sizeof(uint)), newRollingHash);
+
+			DataSegment segment = new DataSegment(_lastSegment.RunningIndex + _lastSegment.Memory.Length, leafData.ToArray());
+			_lastSegment.SetNext(segment);
+			_lastSegment = segment;
+
+			_writtenSequence = CreateSequence(_firstSegment, _lastSegment);
+		}
+
+		/// <inheritdoc/>
+		public override async Task CopyToStreamAsync(Stream outputStream, CancellationToken cancellationToken)
+		{
+			foreach (ReadOnlyMemory<byte> segment in _writtenSequence.Slice(HeaderLength))
+			{
+				await outputStream.WriteAsync(segment, cancellationToken);
+			}
+		}
+	}
+
+	/// <summary>
+	/// An interior file node
+	/// </summary>
+	public class InteriorFileNode : FileNode
+	{
+		/// <summary>
+		/// Type identifier for interior nodes. First byte in the serialized stream.
+		/// </summary>
+		public const byte TypeId = (byte)'i';
+
+		bool _isReadOnly;
+		uint _rollingHash;
+		long _length;
+		readonly List<TreeNodeRef<FileNode>> _children = new List<TreeNodeRef<FileNode>>();
+
+		/// <inheritdoc/>
+		public override long Length => _length;
+
+		/// <inheritdoc/>
+		public override uint RollingHash => _rollingHash;
+
+		/// <summary>
+		/// Constructor
+		/// </summary>
+		/// <param name="length"></param>
+		/// <param name="child"></param>
+		public InteriorFileNode(long length, FileNode child)
+		{
+			_length = length;
+			_children.Add(new TreeNodeRef<FileNode>(this, child));
+		}
+
+		/// <summary>
+		/// Constructor
+		/// </summary>
+		/// <param name="data"></param>
+		/// <param name="children"></param>
+		public InteriorFileNode(ReadOnlySequence<byte> data, IReadOnlyList<ITreeBlob> children)
+		{
+			ReadOnlySpan<byte> span = data.AsSingleSegment().Span;
+			if (span[0] != TypeId)
+			{
+				throw new InvalidDataException($"Invalid type id for {nameof(InteriorFileNode)}");
+			}
+
+			int pos = 1;
+
+			_rollingHash = BinaryPrimitives.ReadUInt32LittleEndian(span[pos..]);
+			pos += sizeof(uint);
+
+			_length = (int)VarInt.ReadUnsigned(span[pos..], out int lengthBytes);
+			pos += lengthBytes;
+
+			_isReadOnly = span[pos] != 0;
+			_children = children.ConvertAll(x => new TreeNodeRef<FileNode>(this, x));
+		}
+
+		/// <inheritdoc/>
+		public override bool IsReadOnly() => _isReadOnly;
+
+		/// <inheritdoc/>
+		public override IReadOnlyList<TreeNodeRef> GetReferences() => _children;
+
+		/// <inheritdoc/>
+		public override async ValueTask<ITreeBlob> SerializeAsync(ITreeBlobWriter writer, CancellationToken cancellationToken)
+		{
+			List<ITreeBlob> children = new List<ITreeBlob>();
+			foreach (TreeNodeRef<FileNode> child in _children)
+			{
+				children.Add(await child.CollapseAsync(writer, cancellationToken));
+			}
+
+			byte[] data = new byte[1 + sizeof(uint) + VarInt.MeasureUnsigned((ulong)_length) + 1];
+			data[0] = TypeId;
+
+			int pos = 1;
+
+			BinaryPrimitives.WriteUInt32LittleEndian(data.AsSpan(pos), _rollingHash);
+			pos += sizeof(uint);
+
+			int lengthBytes = VarInt.WriteUnsigned(data.AsSpan(pos), _length);
+			pos += lengthBytes;
+
+			data[pos] = _isReadOnly ? (byte)1 : (byte)0;
+			pos++;
+
+			Debug.Assert(pos == data.Length);
+			return await writer.WriteBlobAsync(new ReadOnlySequence<byte>(data), children, cancellationToken);
+		}
+
+		/// <inheritdoc/>
+		public override ReadOnlyMemory<byte> AppendData(ReadOnlyMemory<byte> newData, ChunkingOptions options)
+		{
+			Span<byte> hashData = stackalloc byte[4];
+			for (; ; )
+			{
+				Debug.Assert(_children != null);
+
+				// Try to write to the last node
+				if (_children.Count > 0)
+				{
+					FileNode? lastNode = _children[^1].Node;
+					if (lastNode != null)
+					{
+						// Update the length to match the new node
+						_length -= lastNode.Length;
+						newData = lastNode.AppendData(newData, options);
+						_length += lastNode.Length;
+
+						// If the last node is complete, write it to the buffer
+						if (lastNode.IsReadOnly())
+						{
+							// Update the hash
+							BinaryPrimitives.WriteUInt32LittleEndian(hashData, lastNode.RollingHash);
+							_rollingHash = BuzHash.Add(_rollingHash, hashData);
+
+							// Check if it's time to finish this chunk
+							uint hashThreshold = (uint)(((1L << 32) * IoHash.NumBytes) / options.LeafOptions.TargetSize);
+							if ((_children.Count >= options.InteriorOptions.MinSize && _rollingHash < hashThreshold) || (_children.Count >= options.InteriorOptions.MaxSize))
+							{
+								_isReadOnly = true;
+								return newData;
+							}
+						}
+
+						// Bail out if there's nothing left to write
+						if (newData.Length == 0)
+						{
+							return newData;
+						}
+					}
+				}
+
+				// Add a new child node
+				_children.Add(new TreeNodeRef<FileNode>(this, new LeafFileNode()));
+			}
+		}
+
+		/// <inheritdoc/>
+		public override async Task CopyToStreamAsync(Stream outputStream, CancellationToken cancellationToken)
+		{
+			foreach (TreeNodeRef<FileNode> childNodeRef in _children)
+			{
+				FileNode childNode = await childNodeRef.ExpandAsync(cancellationToken);
+				await childNode.CopyToStreamAsync(outputStream, cancellationToken);
+			}
+		}
+	}
+
+	/// <summary>
+	/// Factory class for file nodes
+	/// </summary>
+	public class FileNodeSerializer : TreeNodeSerializer<FileNode>
+	{
+		/// <inheritdoc/>
+		public override async ValueTask<FileNode> DeserializeAsync(ITreeBlob blob, CancellationToken cancellationToken)
+		{
+			ReadOnlySequence<byte> data = await blob.GetDataAsync(cancellationToken);
+			switch (data.FirstSpan[0])
+			{
+				case LeafFileNode.TypeId:
+					return new LeafFileNode(data);
+				case InteriorFileNode.TypeId:
+					return new InteriorFileNode(data, await blob.GetReferencesAsync(cancellationToken));
+				default:
+					throw new InvalidDataException("Unknown type id while deserializing file node");
+			};
+		}
+
+		/// <inheritdoc/>
+		public override ValueTask<ITreeBlob> SerializeAsync(ITreeBlobWriter writer, FileNode node, CancellationToken cancellationToken) => node.SerializeAsync(writer, cancellationToken);
+	}
+}
diff --git a/Engine/Source/Programs/Shared/EpicGames.Horde/Storage/TreeNode.cs b/Engine/Source/Programs/Shared/EpicGames.Horde/Storage/TreeNode.cs
new file mode 100644
--- /dev/null
+++ b/Engine/Source/Programs/Shared/EpicGames.Horde/Storage/TreeNode.cs
@@ -0,0 +1,114 @@
+// Copyright Epic Games, Inc. All Rights Reserved.
+
+using System;
+using System.Buffers;
+using System.Collections.Generic;
+using System.Reflection;
+using System.Threading;
+using System.Threading.Tasks;
+
+namespace EpicGames.Horde.Storage
+{
+	/// <summary>
+	/// Base class for user-defined types that are stored in a tree
+	/// </summary>
+	public abstract class TreeNode
+	{
+		/// <summary>
+		/// Cached incoming reference to the owner of this node.
+		/// </summary>
+		internal TreeNodeRef? IncomingRef { get; set; }
+
+		/// <summary>
+		/// Queries if the node in its current state is read-only. Once we know that nodes are no longer going to be modified, they are favored for spilling to persistent storage.
+		/// </summary>
+		/// <returns>True if the node is read-only.</returns>
+		public virtual bool IsReadOnly() => false;
+
+		/// <summary>
+		/// Mark this node as dirty
+		/// </summary>
+		protected void MarkAsDirty() => IncomingRef?.MarkAsDirty();
+
+		/// <summary>
+		/// Enumerates all the child references from this node
+		/// </summary>
+		/// <returns>Children of this node</returns>
+		public abstract IReadOnlyList<TreeNodeRef> GetReferences();
+
+		/// <summary>
+		/// Static instance of the serializer for a particular <see cref="TreeNode"/> type.
+		/// </summary>
+		static class SerializerInstance<T> where T : TreeNode
+		{
+			static readonly TreeSerializerAttribute _attribute = typeof(T).GetCustomAttribute<TreeSerializerAttribute>()!;
+			public static TreeNodeSerializer<T> Serializer { get; } = (TreeNodeSerializer<T>)Activator.CreateInstance(_attribute.Type)!;
+		}
+
+		/// <summary>
+		/// Serialize a node to data
+		/// </summary>
+		/// <typeparam name="T"></typeparam>
+		/// <param name="node"></param>
+		/// <param name="writer"></param>
+		/// <param name="cancellationToken"></param>
+		/// <returns></returns>
+		public static ValueTask<ITreeBlob> SerializeAsync<T>(T node, ITreeBlobWriter writer, CancellationToken cancellationToken) where T : TreeNode
+		{
+			return SerializerInstance<T>.Serializer.SerializeAsync(writer, node, cancellationToken);
+		}
+
+		/// <summary>
+		/// Deserialize a node from data
+		/// </summary>
+		/// <typeparam name="T"></typeparam>
+		/// <param name="data"></param>
+		/// <param name="cancellationToken"></param>
+		/// <returns></returns>
+		public static ValueTask<T> DeserializeAsync<T>(ITreeBlob data, CancellationToken cancellationToken) where T : TreeNode
+		{
+			return SerializerInstance<T>.Serializer.DeserializeAsync(data, cancellationToken);
+		}
+	}
+
+	/// <summary>
+	/// Factory class for deserializing node types
+	/// </summary>
+	/// <typeparam name="T">The type of node returned</typeparam>
+	public abstract class TreeNodeSerializer<T> where T : TreeNode
+	{
+		/// <summary>
+		/// Serialize a given node
+		/// </summary>
+		/// <param name="writer">Writer for node data</param>
+		/// <param name="node">The user type to be serialized</param>
+		/// <param name="cancellationToken">Cancellation token for the operation</param>
+		/// <returns>Serialized node instance</returns>
+		public abstract ValueTask<ITreeBlob> SerializeAsync(ITreeBlobWriter writer, T node, CancellationToken cancellationToken);
+
+		/// <summary>
+		/// Deserializes data from the given data
+		/// </summary>
+		/// <param name="node">Node to deserialize from</param>
+		/// <param name="cancellationToken">Cancellation token for the operation</param>
+		/// <returns>New node parsed from the data</returns>
+		public abstract ValueTask<T> DeserializeAsync(ITreeBlob node, CancellationToken cancellationToken);
+	}
+
+	/// <summary>
+	/// Attribute used to define a factory for a particular node type
+	/// </summary>
+	[AttributeUsage(AttributeTargets.Class)]
+	public sealed class TreeSerializerAttribute : Attribute
+	{
+		/// <summary>
+		/// The factory type. Should be derived from <see cref="TreeNodeSerializer{T}"/>
+		/// </summary>
+		public Type Type { get; }
+
+		/// <summary>
+		/// Constructor
+		/// </summary>
+		public TreeSerializerAttribute(Type type) => Type = type;
+	}
+}
diff --git a/Engine/Source/Programs/Shared/EpicGames.Horde/Storage/TreeNodeRef.cs b/Engine/Source/Programs/Shared/EpicGames.Horde/Storage/TreeNodeRef.cs
new file mode 100644
--- /dev/null
+++ b/Engine/Source/Programs/Shared/EpicGames.Horde/Storage/TreeNodeRef.cs
@@ -0,0 +1,259 @@
+// Copyright Epic Games, Inc. All Rights Reserved.
+
+using System;
+using System.Diagnostics;
+using System.Reflection;
+using System.Threading;
+using System.Threading.Tasks;
+
+namespace EpicGames.Horde.Storage
+{
+	/// <summary>
+	/// Stores a reference from a parent to child node, which can be resurrected after the child node is flushed to storage if subsequently modified.
+	/// </summary>
+	public class TreeNodeRef
+	{
+		/// <summary>
+		/// Cached reference to the parent node
+		/// </summary>
+		internal TreeNode? _owner;
+
+		/// <summary>
+		/// Strong reference to the current node
+		/// </summary>
+		internal TreeNode? _strongRef;
+
+		/// <summary>
+		/// Weak reference to the child node. Maintained after the object has been flushed.
+		/// </summary>
+		internal WeakReference<TreeNode>? _weakRef;
+
+		/// <summary>
+		/// Last time that the node was modified
+		/// </summary>
+		internal long _lastModifiedTime;
+
+		/// <summary>
+		/// Handle to the node in storage. May be set to null when the node is in memory.
+		/// </summary>
+		internal ITreeBlob? _target;
+
+		/// <summary>
+		/// Reference to the node.
+		/// </summary>
+		public TreeNode? Node
+		{
+			get => _strongRef;
+			set
+			{
+				_strongRef = value;
+				_weakRef = null;
+			}
+		}
+
+		/// <summary>
+		/// Creates a reference to a node with the given hash
+		/// </summary>
+		/// <param name="persistedNode">Hash of the referenced node</param>
+		public TreeNodeRef(ITreeBlob persistedNode)
+		{
+			_target = persistedNode;
+		}
+
+		/// <summary>
+		/// Creates a reference to a node with the given hash
+		/// </summary>
+		/// <param name="owner">The node which owns the reference</param>
+		/// <param name="persistedNode">Hash of the referenced node</param>
+		public TreeNodeRef(TreeNode owner, ITreeBlob persistedNode)
+		{
+			_owner = owner;
+			_target = persistedNode;
+		}
+
+		/// <summary>
+		/// Creates a reference to the given node
+		/// </summary>
+		/// <param name="owner">The node which owns the reference</param>
+		/// <param name="node">The referenced node</param>
+		public TreeNodeRef(TreeNode owner, TreeNode node)
+		{
+			_owner = owner;
+			_strongRef = node;
+			_lastModifiedTime = Stopwatch.GetTimestamp();
+
+			node.IncomingRef = this;
+		}
+
+		/// <summary>
+		/// Reparent this reference to a new owner.
+		/// </summary>
+		public void Reparent(TreeNode newOwner)
+		{
+			MarkAsDirty();
+			_owner = newOwner;
+			MarkAsDirty();
+		}
+
+		/// <summary>
+		/// Figure out whether this ref is dirty
+		/// </summary>
+		/// <returns></returns>
+		public bool IsDirty() => _lastModifiedTime != 0;
+
+		/// <summary>
+		/// Converts the node in this reference from a weak to strong reference.
+		/// </summary>
+		/// <returns>True if the ref contains a valid node on return</returns>
+		protected TreeNode? ConvertToStrongRef()
+		{
+			// Check if it's already a strong ref
+			if (_strongRef != null)
+			{
+				return _strongRef;
+			}
+
+			// Otherwise check if it's a weak ref that hasn't been GC'd yet
+			if (_weakRef != null && _weakRef.TryGetTarget(out _strongRef))
+			{
+				_weakRef = null;
+				return _strongRef;
+			}
+
+			return null;
+		}
+
+		/// <summary>
+		/// Marks this reference as dirty
+		/// </summary>
+		internal void MarkAsDirty()
+		{
+			if (_lastModifiedTime == 0)
+			{
+				_target = null;
+
+				_lastModifiedTime = Stopwatch.GetTimestamp();
+				if (_strongRef == null)
+				{
+					if (_weakRef == null || !_weakRef.TryGetTarget(out _strongRef))
+					{
+						throw new InvalidOperationException("Unable to resolve weak reference to node");
+					}
+					_weakRef = null;
+				}
+				_owner?.IncomingRef?.MarkAsDirty();
+			}
+		}
+
+		/// <summary>
+		/// Marks this reference as clean
+		/// </summary>
+		internal void MarkAsClean(ITreeBlob persistedNode)
+		{
+			_target = persistedNode;
+
+			if (_strongRef != null)
+			{
+				OnCollapse();
+
+				_weakRef = new WeakReference<TreeNode>(_strongRef);
+				_strongRef = null;
+			}
+
+			_lastModifiedTime = 0;
+		}
+
+		/// <summary>
+		/// Callback for when a reference is collapsed to a hash value. At the time of calling, both the hash and node data will be valid.
+		/// </summary>
+		protected virtual void OnCollapse()
+		{
+		}
+	}
+
+	/// <summary>
+	/// Strongly typed reference to a <see cref="TreeNode"/>
+	/// </summary>
+	/// <typeparam name="T">Type of the node</typeparam>
+	public class TreeNodeRef<T> : TreeNodeRef where T : TreeNode
+	{
+		static readonly TreeNodeSerializer<T> Serializer = CreateSerializer();
+
+		/// <inheritdoc cref="TreeNodeRef.Node"/>
+		public new T? Node
+		{
+			get => (T?)base.Node;
+			set => base.Node = value;
+		}
+
+		/// <summary>
+		/// Constructor
+		/// </summary>
+		/// <param name="owner">The node which owns the reference</param>
+		/// <param name="target">Hash of the referenced node</param>
+		public TreeNodeRef(TreeNode owner, ITreeBlob target) : base(owner, target)
+		{
+		}
+
+		/// <summary>
+		/// Constructor
+		/// </summary>
+		/// <param name="owner">The node which owns the reference</param>
+		/// <param name="node">The referenced node</param>
+		public TreeNodeRef(TreeNode owner, T node) : base(owner, node)
+		{
+		}
+
+		/// <summary>
+		/// Create an instance of the serializer for this node type
+		/// </summary>
+		/// <returns>New serializer instance</returns>
+		static TreeNodeSerializer<T> CreateSerializer()
+		{
+			TreeSerializerAttribute? attribute = typeof(T).GetCustomAttribute<TreeSerializerAttribute>();
+			if (attribute == null)
+			{
+				throw new InvalidOperationException($"No serializer is defined for {typeof(T).Name}");
+			}
+			return (TreeNodeSerializer<T>)Activator.CreateInstance(attribute.Type)!;
+		}
+
+		/// <summary>
+		/// Collapse this reference and write its node to storage
+		/// </summary>
+		/// <param name="writer"></param>
+		/// <param name="cancellationToken">Cancellation token for the operation</param>
+		/// <returns></returns>
+		public async ValueTask<ITreeBlob> CollapseAsync(ITreeBlobWriter writer, CancellationToken cancellationToken)
+		{
+			ITreeBlob? node = _target;
+			if (node == null)
+			{
+				node = await Serializer.SerializeAsync(writer, Node!, cancellationToken);
+				MarkAsClean(node);
+			}
+			return node;
+		}
+
+		/// <summary>
+		/// Resolve this reference to a concrete node
+		/// </summary>
+		/// <param name="cancellationToken">Cancellation token for the operation</param>
+		/// <returns></returns>
+		public async ValueTask<T> ExpandAsync(CancellationToken cancellationToken)
+		{
+			T? result = (T?)ConvertToStrongRef();
+			if (result == null)
+			{
+				Debug.Assert(_target != null);
+
+				_strongRef = await Serializer.DeserializeAsync(_target, cancellationToken);
+				_strongRef.IncomingRef = this;
+
+				result = (T)_strongRef;
+			}
+			return result;
+		}
+	}
+}
+
