Horde: Fix tests.

diff --git a/Engine/Source/Programs/Shared/EpicGames.Horde.Tests/BundleStoreTests.cs b/Engine/Source/Programs/Shared/EpicGames.Horde.Tests/BundleStoreTests.cs
--- a/Engine/Source/Programs/Shared/EpicGames.Horde.Tests/BundleStoreTests.cs
+++ b/Engine/Source/Programs/Shared/EpicGames.Horde.Tests/BundleStoreTests.cs
@@ -262,10 +262,10 @@ namespace EpicGames.Horde.Tests
 				TreeNodeRef<FileNode> file = root.GetFileEntry("test");
 				ITreeBlob blob = await file.CollapseAsync(store, CancellationToken.None);
 
-				Dictionary<IoHash, long> uniqueChunks = new Dictionary<IoHash, long>();
-				await FindUniqueNodes(blob, uniqueChunks);
+				Dictionary<ITreeBlob, long> uniqueBlobs = new Dictionary<ITreeBlob, long>();
+				await FindUniqueBlobs(blob, uniqueBlobs);
 
-				long uniqueSize = uniqueChunks.Sum(x => x.Value);
+				long uniqueSize = uniqueBlobs.Sum(x => x.Value);
 				Assert.IsTrue(uniqueSize < data.Length / 3); // random fraction meaning "lots of dedupe happened"
 			}
 		}
@@ -282,15 +282,15 @@ namespace EpicGames.Horde.Tests
 			Assert.IsTrue(worldData.SequenceEqual(data));
 		}
 
-		static async Task FindUniqueNodes(ITreeBlob blob, Dictionary<IoHash, long> uniqueHashes)
+		static async Task FindUniqueBlobs(ITreeBlob blob, Dictionary<ITreeBlob, long> uniqueBlobs)
 		{
 			ReadOnlySequence<byte> data = await blob.GetDataAsync();
-			uniqueHashes[blob.Hash] = data.Length;
+			uniqueBlobs[blob] = data.Length;
 
 			IReadOnlyList<ITreeBlob> references = await blob.GetReferencesAsync();
 			foreach (ITreeBlob reference in references)
 			{
-				await FindUniqueNodes(reference, uniqueHashes);
+				await FindUniqueBlobs(reference, uniqueBlobs);
 			}
 		}
 
