BuildGraph: Various VM improvements.

* Added support for deserializing BgObject types directly into native classes.
* Removed opcodes for creating graph structures. These are now created in user code from BgObject types.
* Removed BgNodeSpecBuilder. BgNode objects can now be modified directly (returning a modified copy).
* Added concrete types for option parameters. The VM now keeps track of any parameters for evaluated options, allowing them to be added into the graph definition.
* Order dependencies now take nodes rather than outputs.
* Added explicit support for native thunks, whose bindings are saved to a sideband channel during compilation and referenced in bytecode as an index. This generalizes code that was previously specific to node definitions.
* Added a name table to bytecode, to optimize situations where we reference the same string mulitple times.

diff --git a/Engine/Source/Programs/AutomationTool/BuildGraph/BgNodeExecutor.cs b/Engine/Source/Programs/AutomationTool/BuildGraph/BgNodeExecutor.cs
--- a/Engine/Source/Programs/AutomationTool/BuildGraph/BgNodeExecutor.cs
+++ b/Engine/Source/Programs/AutomationTool/BuildGraph/BgNodeExecutor.cs
@@ -65,9 +65,9 @@ namespace AutomationTool
 			public override bool IsBuildMachine => CommandUtils.IsBuildMachine;
 		}
 
-		readonly BgExpressionNodeDef Node;
+		readonly BgNodeDef Node;
 
-		public BgBytecodeNodeExecutor(BgExpressionNodeDef node)
+		public BgBytecodeNodeExecutor(BgNodeDef node)
 		{
 			Node = node;
 		}
@@ -85,7 +85,8 @@ namespace AutomationTool
 		/// <returns></returns>
 		public override async Task<bool> ExecuteAsync(JobContext Job, Dictionary<string, HashSet<FileReference>> TagNameToFileSet)
 		{
-			MethodInfo method = Node.Method.Method;
+			BgThunkDef thunk = Node.Thunk!;
+			MethodInfo method = thunk.Method;
 
 			HashSet<FileReference> BuildProducts = TagNameToFileSet[Node.DefaultOutput.TagName];
 
@@ -101,13 +102,9 @@ namespace AutomationTool
 				{
 					arguments[Idx] = Context;
 				}
-				else if (typeof(BgExpr).IsAssignableFrom(ParameterType))
-				{
-					arguments[Idx] = BgType.Constant(ParameterType, Node.Arguments[Idx]);
-				}
 				else
 				{
-					arguments[Idx] = Node.Method.Arguments[Idx];
+					arguments[Idx] = thunk.Arguments[Idx];
 				}
 			}
 
diff --git a/Engine/Source/Programs/AutomationTool/BuildGraph/BgScriptReader.cs b/Engine/Source/Programs/AutomationTool/BuildGraph/BgScriptReader.cs
--- a/Engine/Source/Programs/AutomationTool/BuildGraph/BgScriptReader.cs
+++ b/Engine/Source/Programs/AutomationTool/BuildGraph/BgScriptReader.cs
@@ -724,9 +724,9 @@ namespace AutomationTool
 					else
 					{
 						// Create a new option object to store the settings
-						string description = ReadAttribute(element, "Description");
-						string defaultValue = ReadAttribute(element, "DefaultValue");
-						BgOption option = new BgOption(name, description, defaultValue);
+						BgStringOptionDef option = new BgStringOptionDef(name);
+						option.Description = ReadAttribute(element, "Description");
+						option.DefaultValue = ReadAttribute(element, "DefaultValue");
 						_graph.Options.Add(option);
 
 						// Get the value of this property
@@ -953,9 +953,9 @@ namespace AutomationTool
 				BgAgentDef? agent;
 				if (_graph.NameToAgent.TryGetValue(name, out agent))
 				{
-					if (types.Length > 0 && agent.PossibleTypes.Length > 0)
+					if (types.Length > 0 && agent.PossibleTypes.Count > 0)
 					{
-						if (types.Length != agent.PossibleTypes.Length || !types.SequenceEqual(agent.PossibleTypes, StringComparer.InvariantCultureIgnoreCase))
+						if (types.Length != agent.PossibleTypes.Count || !types.SequenceEqual(agent.PossibleTypes, StringComparer.InvariantCultureIgnoreCase))
 						{
 							LogError(element, "Agent types ({0}) were different than previous agent definition with types ({1}). Must either be empty or match exactly.", String.Join(",", types), String.Join(",", agent.PossibleTypes));
 						}
@@ -967,7 +967,8 @@ namespace AutomationTool
 					{
 						LogError(element, "Missing type for agent '{0}'", name);
 					}
-					agent = new BgAgentDef(name, types);
+					agent = new BgAgentDef(name);
+					agent.PossibleTypes.AddRange(types);
 					_graph.NameToAgent.Add(name, agent);
 					_graph.Agents.Add(agent);
 				}
diff --git a/Engine/Source/Programs/AutomationTool/BuildGraph/BuildGraph.cs b/Engine/Source/Programs/AutomationTool/BuildGraph/BuildGraph.cs
--- a/Engine/Source/Programs/AutomationTool/BuildGraph/BuildGraph.cs
+++ b/Engine/Source/Programs/AutomationTool/BuildGraph/BuildGraph.cs
@@ -383,11 +383,11 @@ namespace AutomationTool
 				BgGraphBuilder Builder = (BgGraphBuilder)Activator.CreateInstance(BuilderType);
 				BgGraph GraphSpec = Builder.CreateGraph(new BgEnvironment(P4Enabled ? P4Env : null));
 
-				(byte[] Data, BgMethod[] Methods) = BgCompiler.Compile(GraphSpec);
+				(byte[] Data, BgThunkDef[] Methods) = BgCompiler.Compile(GraphSpec);
 
 				BgInterpreter Interpreter = new BgInterpreter(Data, Methods, Arguments);
 				Interpreter.Disassemble(Logger);
-				Graph = (BgGraphDef)Interpreter.Evaluate();
+				Graph = ((BgObjectDef)Interpreter.Evaluate()).Deserialize<BgGraphExpressionDef>().ToGraphDef();
 			}
 			else
 			{
@@ -709,7 +709,7 @@ namespace AutomationTool
 				NodeToExecutor[Node] = executor;
 				return executor.Bind(NameToTask, TagNameToNodeOutput, Logger);
 			}
-			else if (Node is BgExpressionNodeDef BytecodeNode)
+			else if (Node is BgNodeExpressionDef BytecodeNode)
 			{
 				BgBytecodeNodeExecutor executor = new BgBytecodeNodeExecutor(BytecodeNode);
 				NodeToExecutor[Node] = executor;
diff --git a/Engine/Source/Programs/AutomationTool/BuildGraph/Scripts/InstalledBuild.cs b/Engine/Source/Programs/AutomationTool/BuildGraph/Scripts/InstalledBuild.cs
--- a/Engine/Source/Programs/AutomationTool/BuildGraph/Scripts/InstalledBuild.cs
+++ b/Engine/Source/Programs/AutomationTool/BuildGraph/Scripts/InstalledBuild.cs
@@ -119,7 +119,7 @@ namespace AutomationTool
 			BgBool EmbedSrcSrvInfo = new BgBoolOption("EmbedSrcSrvInfo", "Whether to add Source indexing to Windows game apps so they can be added to a symbol server", false);
 
 			BgList<BgString> DefaultGameConfigurations = BgList<BgString>.Create(nameof(UnrealTargetConfiguration.DebugGame), nameof(UnrealTargetConfiguration.Development), nameof(UnrealTargetConfiguration.Shipping));
-			BgList<BgString> GameConfigurationStrings = new BgListOptionSpec("GameConfigurations", description: "Which game configurations to include for packaged applications", style: BgListOptionStyle.CheckList, values: DefaultGameConfigurations);
+			BgList<BgString> GameConfigurationStrings = new BgListOption("GameConfigurations", description: "Which game configurations to include for packaged applications", style: BgListOptionStyle.CheckList, values: DefaultGameConfigurations);
 			BgList<BgEnum<UnrealTargetConfiguration>> GameConfigurations = GameConfigurationStrings.Select(x => BgEnum<UnrealTargetConfiguration>.Parse(x));
 
 			BgBoolOption WithFullDebugInfo = new BgBoolOption("WithFullDebugInfo", "Generate full debug info for binary editor and packaged application builds", false);
@@ -140,18 +140,15 @@ namespace AutomationTool
 			BgAgent EditorWin64 = new BgAgent("Editor Win64", "Win64_Licensee");
 
 			BgNode VersionFilesNode = EditorWin64
-				.AddNode(x => UpdateVersionFilesAsync(x))
-				.Construct();
+				.AddNode(x => UpdateVersionFilesAsync(x));
 
 			BgNode<BgFileSet> WinUhtNode = EditorWin64
 				.AddNode(x => CompileUnrealHeaderToolAsync(x, UnrealTargetPlatform.Win64))
-				.Requires(VersionFilesNode)
-				.Construct();
+				.Requires(VersionFilesNode);
 
 			BgNode<BgFileSet> WinEditorNode = EditorWin64
 				.AddNode(x => CompileUnrealEditorWin64Async(x, CrashReporterCompileArgs, EmbedSrcSrvInfo, CompileDatasmithPlugins, WithFullDebugInfo, SignExecutables))
-				.Requires(WinUhtNode)
-				.Construct();
+				.Requires(WinUhtNode);
 
 			Aggregates.Add(new BgAggregate("Win64 Editor", WinEditorNode, label: "Editors/Win64"));
 
@@ -164,8 +161,7 @@ namespace AutomationTool
 
 			BgNode<BgFileSet> WinGame = TargetWin64
 				.AddNode(x => CompileUnrealGameWin64(x, GameConfigurations, EmbedSrcSrvInfo, WithFullDebugInfo, SignExecutables))
-				.Requires(WinUhtNode)
-				.Construct();
+				.Requires(WinUhtNode);
 
 			Aggregates.Add(new BgAggregate("TargetPlatforms_Win64", WinGame));
 
@@ -176,8 +172,7 @@ namespace AutomationTool
 			BgAgent BuildRules = new BgAgent("BuildRules", "Win64_Licensee");
 
 			BgNode RulesAssemblies = BuildRules
-				.AddNode(x => CompileRulesAssemblies(x))
-				.Construct();
+				.AddNode(x => CompileRulesAssemblies(x));
 
 			//// Win Tools ////
 
@@ -185,13 +180,11 @@ namespace AutomationTool
 
 			BgNode<BgFileSet> WinTools = ToolsGroupWin64
 				.AddNode(x => BuildToolsWin64Async(x, CrashReporterCompileArgs))
-				.Requires(WinUhtNode)
-				.Construct();
+				.Requires(WinUhtNode);
 
 			BgNode<BgFileSet> CsTools = ToolsGroupWin64
 				.AddNode(x => BuildToolsCSAsync(x, SignExecutables))
-				.Requires(VersionFilesNode)
-				.Construct();
+				.Requires(VersionFilesNode);
 
 
 			/////// DDC //////////////////////////////////////////////////////////
@@ -203,8 +196,7 @@ namespace AutomationTool
 
 			BgNode DdcNode = DDCGroupWin64
 				.AddNode(x => BuildDDCWin64Async(x, DDCPlatformsWin64, BgList<BgFileSet>.Create(WinUhtNode.Output, WinEditorNode.Output, WinTools.Output)))
-				.Requires(WinEditorNode, WinTools)
-				.Construct();
+				.Requires(WinEditorNode, WinTools);
 
 
 			/////// STAGING ///////
@@ -229,8 +221,7 @@ namespace AutomationTool
 
 			BgNode WinInstalledNode = WinStageAgent
 				.AddNode(x => MakeInstalledBuildWin64Async(x, WinInstalledFiles, WinFinalizeArgs, WinOutputDir))
-				.Requires(WinInstalledFiles)
-				.Construct();
+				.Requires(WinInstalledFiles);
 
 			Aggregates.Add(new BgAggregate("HostPlatforms_Win64", WinInstalledNode, label: "Builds/Win64"));
 
diff --git a/Engine/Source/Programs/Shared/EpicGames.BuildGraph.Tests/GraphTests.cs b/Engine/Source/Programs/Shared/EpicGames.BuildGraph.Tests/GraphTests.cs
--- a/Engine/Source/Programs/Shared/EpicGames.BuildGraph.Tests/GraphTests.cs
+++ b/Engine/Source/Programs/Shared/EpicGames.BuildGraph.Tests/GraphTests.cs
@@ -10,118 +10,115 @@ using System.Text.Json;
 using OpenTracing;
 using EpicGames.BuildGraph.Expressions;
 using Microsoft.Extensions.Logging;
+using System.Threading.Tasks;
 
 namespace EpicGames.BuildGraph.Tests
 {
-	/*
 	[TestClass]
 	public class GraphTests
 	{
+		static Task UpdateVersionFiles() => Task.CompletedTask;
+		static Task CompileShooterGameWin64() => Task.CompletedTask;
+		static Task CookShooterGameWin64() => Task.CompletedTask;
+
 		object Evaluate(BgExpr expr)
 		{
-			byte[] data = BgBytecode.Compile(expr);
-			BgBytecodeInterpreter interpreter = new BgBytecodeInterpreter(data);
+			(byte[] data, BgThunkDef[] methods) = BgCompiler.Compile(expr);
+			BgInterpreter interpreter = new BgInterpreter(data, methods, new Dictionary<string, string>());
 			return interpreter.Evaluate();
 		}
 
 		[TestMethod]
 		public void NodeTest()
 		{
-			BgNodeSpec nodeSpec1 = new BgNodeSpec("Update Version Files");
+			BgAgent agent = new BgAgent("name", "type");
+			BgNode nodeSpec1 = agent.AddNode(x => UpdateVersionFiles()).Construct();
 
-			BgNodeSpec nodeSpec2 = new BgNodeSpec("Compile ShooterGameEditor Win64", requires: BgList<BgFileSet>.Create(nodeSpec1));
+			BgNode nodeSpec2 = agent.AddNode(x => CompileShooterGameWin64()).Requires(BgList<BgFileSet>.Create(nodeSpec1)).Construct();
 
-			BgNodeSpec nodeSpec3 = new BgNodeSpec("Cook ShooterGame Win64", requires: BgList<BgFileSet>.Create(nodeSpec2.DefaultOutput));
+			BgNode nodeSpec3 = agent.AddNode(x => CookShooterGameWin64()).Requires(BgList<BgFileSet>.Create(nodeSpec2.DefaultOutput)).Construct();
 
-			BgNode node3 = (BgNode)Evaluate(nodeSpec3);
-			Assert.AreEqual(node3.Name, "Cook ShooterGame Win64");
+			BgNodeDef node3 = ((BgNodeDef)Evaluate(nodeSpec3));
+			Assert.AreEqual(node3.Name, "Cook Shooter Game Win64");
 			Assert.AreEqual(2, node3.InputDependencies.Count);
 			Assert.AreEqual(1, node3.Inputs.Count);
 
-			BgNode node2 = node3.Inputs[0].ProducingNode;
-			Assert.AreEqual(node2.Name, "Compile ShooterGameEditor Win64");
+			BgNodeDef node2 = node3.Inputs[0].ProducingNode;
+			Assert.AreEqual(node2.Name, "Compile Shooter Game Win64");
 			Assert.AreEqual(1, node2.InputDependencies.Count);
 			Assert.AreEqual(1, node2.Inputs.Count);
 
-			BgNode node1 = node2.Inputs[0].ProducingNode;
+			BgNodeDef node1 = node2.Inputs[0].ProducingNode;
 			Assert.AreEqual(node1.Name, "Update Version Files");
 			Assert.AreEqual(0, node1.InputDependencies.Count);
 			Assert.AreEqual(0, node1.Inputs.Count);
 		}
-
+		
 		[TestMethod]
 		public void AgentTest()
 		{
-			BgNodeSpec nodeSpec = new BgNodeSpec("Update Version Files");
-
-			BgAgentSpec agentSpec = new BgAgentSpec(BgList<BgString>.Create("win64", "incremental"));
+			BgAgent expr = new BgAgent("TestAgent", BgList<BgString>.Create("win64", "incremental"));
 
-			BgAgent agent = (BgAgent)Evaluate(agentSpec);
+			BgAgentDef agent = ((BgObjectDef)Evaluate(expr)).Deserialize<BgAgentDef>();
 			Assert.AreEqual("TestAgent", agent.Name);
 
-			Assert.AreEqual(2, agent.PossibleTypes.Length);
+			Assert.AreEqual(2, agent.PossibleTypes.Count);
 			Assert.AreEqual("win64", agent.PossibleTypes[0]);
 			Assert.AreEqual("incremental", agent.PossibleTypes[1]);
-
+/*
 			Assert.AreEqual(1, agent.Nodes.Count);
 
 			BgNode node = agent.Nodes[0];
 			Assert.AreEqual("Update Version Files", node.Name);
-		}
+*/		}
 
 		[TestMethod]
 		public void AggregateTest()
 		{
-			BgNodeSpec nodeSpec1 = new BgNodeSpec("Update Version Files");
-			BgNodeSpec nodeSpec2 = new BgNodeSpec("Compile UnrealEditor Win64");
+			BgAgent agent = new BgAgent("test", "test");
+
+			BgNode nodeSpec1 = agent.AddNode(x => UpdateVersionFiles()).Construct();
+			BgNode nodeSpec2 = agent.AddNode(x => CompileShooterGameWin64()).Construct();
 
-			BgAggregateSpec aggregateSpec = new BgAggregateSpec("All nodes", BgList<BgFileSet>.Create(nodeSpec1, nodeSpec2));
+			BgAggregate aggregateSpec = new BgAggregate("All nodes", BgList<BgNode>.Create(nodeSpec1, nodeSpec2));
 
-			BgAggregate aggregate = (BgAggregate)Evaluate(aggregateSpec);
+			BgAggregateDef aggregate = ((BgObjectDef)Evaluate(aggregateSpec)).Deserialize<BgAggregateDef>();
 			Assert.AreEqual("All nodes", aggregate.Name);
 
-			BgNode node1 = aggregate.RequiredNodes.OrderBy(x => x.Name).ElementAt(0);
-			Assert.AreEqual("Compile UnrealEditor Win64", node1.Name);
+			BgNodeDef node1 = aggregate.RequiredNodes.OrderBy(x => x.Name).ElementAt(0);
+			Assert.AreEqual("Compile Shooter Game Win64", node1.Name);
 
-			BgNode node2 = aggregate.RequiredNodes.OrderBy(x => x.Name).ElementAt(1);
+			BgNodeDef node2 = aggregate.RequiredNodes.OrderBy(x => x.Name).ElementAt(1);
 			Assert.AreEqual("Update Version Files", node2.Name);
 		}
 
 		[TestMethod]
 		public void LabelTest()
 		{
-			BgNodeSpec nodeSpec1 = new BgNodeSpec("Update Version Files");
-			BgNodeSpec nodeSpec2 = new BgNodeSpec("Compile UnrealEditor Win64");
+			BgLabel labelSpec = new BgLabel("name", "category", "ugsBadge", "ugsProject", BgLabelChange.Code);
 
-			BgLabelSpec labelSpec = new BgLabelSpec("name", "category", "ugsBadge", "ugsProject", "Code", BgList<BgNodeSpec>.Create(nodeSpec2), BgList<BgNodeSpec>.Create(nodeSpec1));
-
-			BgLabel label = (BgLabel)Evaluate(labelSpec);
+			BgLabelDef label = ((BgObjectDef)Evaluate(labelSpec)).Deserialize<BgLabelDef>();
 			Assert.AreEqual("name", label.DashboardName);
 			Assert.AreEqual("category", label.DashboardCategory);
 			Assert.AreEqual("ugsBadge", label.UgsBadge);
 			Assert.AreEqual("ugsProject", label.UgsProject);
 			Assert.AreEqual(BgLabelChange.Code, label.Change);
-			Assert.AreEqual(1, label.RequiredNodes.Count);
-			Assert.AreEqual(1, label.IncludedNodes.Count);
-
-			BgNode node1 = label.RequiredNodes.First();
-			Assert.AreEqual("Compile UnrealEditor Win64", node1.Name);
-
-			BgNode node2 = label.IncludedNodes.First();
-			Assert.AreEqual("Update Version Files", node2.Name);
 		}
 
 		[TestMethod]
-		public void DiagnosticTest()
+		public void GraphTest()
 		{
-			BgDiagnosticSpec diagnosticSpec = new BgDiagnosticSpec(LogLevel.Error, "hello world", "file.cs", 123);
+			BgAgent agent = new BgAgent("test", "test");
+
+			BgNode nodeSpec1 = agent.AddNode(x => UpdateVersionFiles()).Construct();
+			BgNode nodeSpec2 = agent.AddNode(x => CompileShooterGameWin64()).Construct();
+
+			BgAggregate aggregateSpec = new BgAggregate("All nodes", BgList<BgNode>.Create(nodeSpec1, nodeSpec2));
+			BgGraph graphSpec = new BgGraph(BgList.Create(nodeSpec1, nodeSpec2), BgList.Create(aggregateSpec));
 
-			BgDiagnostic diagnostic = (BgDiagnostic)Evaluate(diagnosticSpec);
-			Assert.AreEqual(LogLevel.Error, diagnostic.Level);
-			Assert.AreEqual("hello world", diagnostic.Message);
-			Assert.AreEqual("file.cs", diagnostic.Location.File);
-			Assert.AreEqual(123, diagnostic.Location.LineNumber);
+			BgGraphExpressionDef graph = ((BgObjectDef)Evaluate(graphSpec)).Deserialize<BgGraphExpressionDef>();
+			Assert.AreEqual(2, graph.Nodes.Count);
+			Assert.AreEqual(1, graph.Aggregates.Count);
 		}
 	}
-	*/
 }
diff --git a/Engine/Source/Programs/Shared/EpicGames.BuildGraph.Tests/OpcodeTests.cs b/Engine/Source/Programs/Shared/EpicGames.BuildGraph.Tests/OpcodeTests.cs
new file mode 100644
--- /dev/null
+++ b/Engine/Source/Programs/Shared/EpicGames.BuildGraph.Tests/OpcodeTests.cs
@@ -0,0 +1,607 @@
+// Copyright Epic Games, Inc. All Rights Reserved.
+
+using Microsoft.VisualStudio.TestTools.UnitTesting;
+using System;
+using System.Collections.Generic;
+using System.IO;
+using System.Linq;
+using System.Text;
+using System.Text.Json;
+using OpenTracing;
+using EpicGames.BuildGraph.Expressions;
+using System.Collections.Immutable;
+using System.Threading.Tasks;
+
+namespace EpicGames.BuildGraph.Tests
+{
+	[TestClass]
+	public class OpcodeTests
+	{
+		object Evaluate(BgExpr expr)
+		{
+			(byte[] data, BgThunkDef[] methods) = BgCompiler.Compile(expr);
+			BgInterpreter interpreter = new BgInterpreter(data, methods, new Dictionary<string, string>());
+			return interpreter.Evaluate();
+		}
+
+		BgOptionDef EvaluateOption(BgExpr expr)
+		{
+			(byte[] data, BgThunkDef[] methods) = BgCompiler.Compile(expr);
+			BgInterpreter interpreter = new BgInterpreter(data, methods, new Dictionary<string, string>());
+			interpreter.Evaluate();
+			return interpreter.OptionDefs[0];
+		}
+
+		#region Boolean opcodes
+
+		[TestMethod]
+		public void BoolFalse()
+		{
+			object result = Evaluate((BgBool)false);
+			Assert.AreEqual(false, result);
+		}
+
+		[TestMethod]
+		public void BoolTrue()
+		{
+			object result = Evaluate((BgBool)true);
+			Assert.AreEqual(true, result);
+		}
+
+		[TestMethod]
+		public void BoolNot()
+		{
+			object result1 = Evaluate(!(BgBool)true);
+			Assert.AreEqual(false, result1);
+
+			object result2 = Evaluate(!(BgBool)false);
+			Assert.AreEqual(true, result2);
+		}
+
+		[TestMethod]
+		public void BoolAnd()
+		{
+			object result1 = Evaluate((BgBool)false & (BgBool)false);
+			Assert.AreEqual(false, result1);
+
+			object result2 = Evaluate((BgBool)false & (BgBool)true);
+			Assert.AreEqual(false, result2);
+
+			object result3 = Evaluate((BgBool)true & (BgBool)false);
+			Assert.AreEqual(false, result3);
+
+			object result4 = Evaluate((BgBool)true & (BgBool)true);
+			Assert.AreEqual(true, result4);
+		}
+
+		[TestMethod]
+		public void BoolOr()
+		{
+			object result1 = Evaluate((BgBool)false | (BgBool)false);
+			Assert.AreEqual(false, result1);
+
+			object result2 = Evaluate((BgBool)false | (BgBool)true);
+			Assert.AreEqual(true, result2);
+
+			object result3 = Evaluate((BgBool)true | (BgBool)false);
+			Assert.AreEqual(true, result3);
+
+			object result4 = Evaluate((BgBool)true | (BgBool)true);
+			Assert.AreEqual(true, result4);
+		}
+
+		[TestMethod]
+		public void BoolXor()
+		{
+			object result1 = Evaluate((BgBool)false ^ (BgBool)false);
+			Assert.AreEqual(false, result1);
+
+			object result2 = Evaluate((BgBool)false ^ (BgBool)true);
+			Assert.AreEqual(true, result2);
+
+			object result3 = Evaluate((BgBool)true ^ (BgBool)false);
+			Assert.AreEqual(true, result3);
+
+			object result4 = Evaluate((BgBool)true ^ (BgBool)true);
+			Assert.AreEqual(false, result4);
+		}
+
+		[TestMethod]
+		public void BoolEq()
+		{
+			object result1 = Evaluate((BgBool)false == (BgBool)false);
+			Assert.AreEqual(true, result1);
+
+			object result2 = Evaluate((BgBool)false == (BgBool)true);
+			Assert.AreEqual(false, result2);
+
+			object result3 = Evaluate((BgBool)true == (BgBool)false);
+			Assert.AreEqual(false, result3);
+
+			object result4 = Evaluate((BgBool)true == (BgBool)true);
+			Assert.AreEqual(true, result4);
+		}
+
+		[TestMethod]
+		public void BoolOption()
+		{
+			BgBoolOptionDef result = (BgBoolOptionDef)EvaluateOption(new BgBoolOption("boolOption", "description", defaultValue: true));
+			Assert.AreEqual("boolOption", result.Name);
+			Assert.AreEqual("description", result.Description);
+			Assert.AreEqual(true, result.DefaultValue);
+		}
+
+		#endregion
+
+		#region Integer opcodes
+
+		[TestMethod]
+		public void IntLiteral()
+		{
+			object result = Evaluate((BgInt)123);
+			Assert.AreEqual(123, result);
+		}
+
+		[TestMethod]
+		public void IntEq()
+		{
+			object result1 = Evaluate((BgInt)123 == (BgInt)124);
+			Assert.AreEqual(false, result1);
+
+			object result2 = Evaluate((BgInt)123 == (BgInt)123);
+			Assert.AreEqual(true, result2);
+
+			object result3 = Evaluate((BgInt)123 == (BgInt)122);
+			Assert.AreEqual(false, result3);
+		}
+
+		[TestMethod]
+		public void IntLt()
+		{
+			object result1 = Evaluate((BgInt)123 < (BgInt)124);
+			Assert.AreEqual(true, result1);
+
+			object result2 = Evaluate((BgInt)123 < (BgInt)123);
+			Assert.AreEqual(false, result2);
+
+			object result3 = Evaluate((BgInt)123 < (BgInt)122);
+			Assert.AreEqual(false, result3);
+		}
+
+		[TestMethod]
+		public void IntGt()
+		{
+			object result1 = Evaluate((BgInt)123 > (BgInt)124);
+			Assert.AreEqual(false, result1);
+
+			object result2 = Evaluate((BgInt)123 > (BgInt)123);
+			Assert.AreEqual(false, result2);
+
+			object result3 = Evaluate((BgInt)123 > (BgInt)122);
+			Assert.AreEqual(true, result3);
+		}
+
+		[TestMethod]
+		public void IntAdd()
+		{
+			object result = Evaluate((BgInt)1 + (BgInt)2);
+			Assert.AreEqual(3, result);
+		}
+
+		[TestMethod]
+		public void IntMultiply()
+		{
+			object result = Evaluate((BgInt)3 * (BgInt)7);
+			Assert.AreEqual(21, result);
+		}
+
+		[TestMethod]
+		public void IntDivide()
+		{
+			object result = Evaluate((BgInt)21 / (BgInt)3);
+			Assert.AreEqual(7, result);
+		}
+
+		[TestMethod]
+		public void IntModulo()
+		{
+			object result = Evaluate((BgInt)23 % (BgInt)7);
+			Assert.AreEqual(2, result);
+		}
+
+		[TestMethod]
+		public void IntNegate()
+		{
+			object result = Evaluate(-(BgInt)123);
+			Assert.AreEqual(-123, result);
+		}
+
+		[TestMethod]
+		public void IntOption()
+		{
+			BgIntOptionDef result = (BgIntOptionDef)EvaluateOption(new BgIntOption("intOption", "description", defaultValue: 1, minValue: 4, maxValue: 5));
+			Assert.AreEqual("intOption", result.Name);
+			Assert.AreEqual("description", result.Description);
+			Assert.AreEqual(1, result.DefaultValue);
+			Assert.AreEqual(4, result.MinValue);
+			Assert.AreEqual(5, result.MaxValue);
+		}
+
+		#endregion
+		#region String opcodes
+
+		[TestMethod]
+		public void StrEmpty()
+		{
+			object result = Evaluate(BgString.Empty);
+			Assert.AreEqual(String.Empty, result);
+		}
+
+		[TestMethod]
+		public void StrLiteral()
+		{
+			object result = Evaluate((BgString)"hello world");
+			Assert.AreEqual("hello world", result);
+		}
+
+		[TestMethod]
+		public void StrCompare()
+		{
+			{
+				object result1 = Evaluate((BgString)"hello world" == (BgString)"hello world");
+				Assert.AreEqual(true, result1);
+
+				object result2 = Evaluate((BgString)"hello world" == (BgString)"HELLO WORLD");
+				Assert.AreEqual(false, result2);
+
+				object result3 = Evaluate((BgString)"HELLO WORLD" == (BgString)"hello world");
+				Assert.AreEqual(false, result3);
+
+				object result4 = Evaluate((BgString)"HELLO WORLD" == (BgString)"HELLO WORLD");
+				Assert.AreEqual(true, result4);
+			}
+
+			{
+				int result1 = (int)Evaluate(BgString.Compare("hello world", "hello world", StringComparison.Ordinal));
+				Assert.IsTrue(result1 == 0);
+
+				int result2 = (int)Evaluate(BgString.Compare("hello world", "HELLO WORLD", StringComparison.Ordinal));
+				Assert.IsTrue(result2 > 0);
+
+				int result3 = (int)Evaluate(BgString.Compare("HELLO WORLD", "hello world", StringComparison.Ordinal));
+				Assert.IsTrue(result3 < 0);
+
+				int result4 = (int)Evaluate(BgString.Compare("HELLO WORLD", "HELLO WORLD", StringComparison.Ordinal));
+				Assert.IsTrue(result4 == 0);
+			}
+
+			{
+				int result1 = (int)Evaluate(BgString.Compare("hello world", "hello world", StringComparison.OrdinalIgnoreCase));
+				Assert.IsTrue(result1 == 0);
+
+				int result2 = (int)Evaluate(BgString.Compare("hello world", "HELLO WORLD", StringComparison.OrdinalIgnoreCase));
+				Assert.IsTrue(result2 == 0);
+
+				int result3 = (int)Evaluate(BgString.Compare("HELLO WORLD", "hello world", StringComparison.OrdinalIgnoreCase));
+				Assert.IsTrue(result3 == 0);
+
+				int result4 = (int)Evaluate(BgString.Compare("HELLO WORLD", "HELLO WORLD", StringComparison.OrdinalIgnoreCase));
+				Assert.IsTrue(result4 == 0);
+
+				int result5 = (int)Evaluate(BgString.Compare("HELLO WORLD", "HELLO WORLD2", StringComparison.OrdinalIgnoreCase));
+				Assert.IsTrue(result5 < 0);
+			}
+		}
+
+		[TestMethod]
+		public void StrConcat()
+		{
+			object result = Evaluate((BgString)"hello" + (BgString)" " + (BgString)"world");
+			Assert.AreEqual("hello world", result);
+		}
+
+		[TestMethod]
+		public void StrFormat()
+		{
+			object result = Evaluate(BgString.Format("result is {0}", (BgInt)1 + (BgInt)2));
+			Assert.AreEqual("result is 3", result);
+		}
+
+		[TestMethod]
+		public void StrSplit()
+		{
+			object result = Evaluate(((BgString)"1+2+3").Split("+"));
+			Assert.IsTrue(Enumerable.SequenceEqual((IEnumerable<object>)result, new[] { "1", "2", "3" }));
+		}
+
+		[TestMethod]
+		public void StrJoin()
+		{
+			object result = Evaluate(BgString.Join("+", BgList<BgString>.Create("1", "2", "3")));
+			Assert.AreEqual("1+2+3", result);
+		}
+
+		[TestMethod]
+		public void StrMatch()
+		{
+			object result1 = Evaluate(((BgString)"hello world").Match("wor"));
+			Assert.AreEqual(true, result1);
+
+			object result2 = Evaluate(((BgString)"hello world").Match("war"));
+			Assert.AreEqual(false, result2);
+
+			object result3 = Evaluate(((BgString)"hello world").Match(@"^hello"));
+			Assert.AreEqual(true, result3);
+
+			object result4 = Evaluate(((BgString)"hello world").Match(@"^world"));
+			Assert.AreEqual(false, result4);
+		}
+
+		[TestMethod]
+		public void StrReplace()
+		{
+			object result1 = Evaluate(((BgString)"hello world").Replace("(w[orl]+d)", "-> $1"));
+			Assert.AreEqual("hello -> world", result1);
+		}
+
+		[TestMethod]
+		public void StrOption()
+		{
+			BgStringOptionDef result = (BgStringOptionDef)EvaluateOption(new BgStringOption("stringOption", "description", defaultValue: "value", pattern: "abc", patternFailed: "def", values: BgList<BgString>.Create("value1", "value2"), valueDescriptions: BgList<BgString>.Create("desc1", "desc2")));
+			Assert.AreEqual("stringOption", result.Name);
+			Assert.AreEqual("description", result.Description);
+			Assert.AreEqual("value", result.DefaultValue);
+			Assert.AreEqual("abc", result.Pattern);
+			Assert.AreEqual("def", result.PatternFailed);
+			Assert.IsTrue(result.Values.SequenceEqual(new[] { "value1", "value2" }));
+			Assert.IsTrue(result.ValueDescriptions.SequenceEqual(new[] { "desc1", "desc2" }));
+		}
+
+		#endregion
+
+		#region Enum opcodes
+
+		enum TestEnum
+		{
+			ValueIs123 = 123,
+			ValueIs456 = 456
+		}
+
+		[TestMethod]
+		public void EnumConstant()
+		{
+			object result = Evaluate((BgEnum<TestEnum>)TestEnum.ValueIs123);
+			Assert.AreEqual(result, 123);
+		}
+
+		[TestMethod]
+		public void EnumParse()
+		{
+			object result1 = Evaluate(BgEnum<TestEnum>.Parse(nameof(TestEnum.ValueIs123)));
+			Assert.AreEqual((int)TestEnum.ValueIs123, result1);
+
+			object result2 = Evaluate(BgEnum<TestEnum>.Parse(nameof(TestEnum.ValueIs456)));
+			Assert.AreEqual((int)TestEnum.ValueIs456, result2);
+
+			Assert.ThrowsException<InvalidDataException>(() => Evaluate(BgEnum<TestEnum>.Parse("Other")));
+		}
+
+		[TestMethod]
+		public void EnumToString()
+		{
+			object result1 = Evaluate(((BgEnum<TestEnum>)TestEnum.ValueIs123).ToBgString());
+			Assert.AreEqual(result1, nameof(TestEnum.ValueIs123));
+
+			object result2 = Evaluate(((BgEnum<TestEnum>)TestEnum.ValueIs456).ToBgString());
+			Assert.AreEqual(result2, nameof(TestEnum.ValueIs456));
+
+			object result3 = Evaluate(((BgEnum<TestEnum>)(TestEnum)789).ToBgString());
+			Assert.AreEqual(result3, "789");
+		}
+
+		#endregion
+
+		#region List opcodes
+
+		[TestMethod]
+		public void ListEmpty()
+		{
+			object result = Evaluate(BgList<BgInt>.Empty);
+			Assert.IsTrue(Enumerable.SequenceEqual(Enumerable.Empty<object>(), (IEnumerable<object>)result));
+		}
+
+		[TestMethod]
+		public void ListCount()
+		{
+			object result = Evaluate(BgList<BgInt>.Empty.Add(1, 2, 3).Count);
+			Assert.AreEqual(3, result);
+		}
+
+		[TestMethod]
+		public void ListElement()
+		{
+			object result = Evaluate(BgList<BgInt>.Create(1, 2, 3)[2]);
+			Assert.AreEqual(3, result);
+		}
+
+		[TestMethod]
+		public void ListPush()
+		{
+			object result1 = Evaluate(BgList<BgInt>.Empty.Add(1));
+			Assert.IsTrue(Enumerable.SequenceEqual(new object[] { 1 }, (IEnumerable<object>)result1));
+
+			object result2 = Evaluate(BgList<BgInt>.Empty.Add(1, 2).Add(3));
+			Assert.IsTrue(Enumerable.SequenceEqual(new object[] { 1, 2, 3 }, (IEnumerable<object>)result2));
+		}
+
+		[TestMethod]
+		public void ListConcat()
+		{
+			object result = Evaluate(BgList<BgInt>.Concat(BgList<BgInt>.Create(1, 2), BgList<BgInt>.Create(3)));
+			Assert.IsTrue(Enumerable.SequenceEqual(new object[] { 1, 2, 3 }, (IEnumerable<object>)result));
+		}
+
+		[TestMethod]
+		public void ListUnion()
+		{
+			object result1 = Evaluate(BgList<BgInt>.Create(1, 2).Union(BgList<BgInt>.Create(3)));
+			Assert.IsTrue(Enumerable.SequenceEqual(new object[] { 1, 2, 3 }, (IEnumerable<object>)result1));
+
+			object result2 = Evaluate(BgList<BgInt>.Create(1, 2, 3).Union(BgList<BgInt>.Create(3)));
+			Assert.IsTrue(Enumerable.SequenceEqual(new object[] { 1, 2, 3 }, (IEnumerable<object>)result2));
+		}
+
+		[TestMethod]
+		public void ListExcept()
+		{
+			object result1 = Evaluate(BgList<BgInt>.Create(1, 2).Except(BgList<BgInt>.Create(3)));
+			Assert.IsTrue(Enumerable.SequenceEqual(new object[] { 1, 2 }, (IEnumerable<object>)result1));
+
+			object result2 = Evaluate(BgList<BgInt>.Create(1, 2, 3).Except(BgList<BgInt>.Create(3)));
+			Assert.IsTrue(Enumerable.SequenceEqual(new object[] { 1, 2 }, (IEnumerable<object>)result2));
+		}
+
+		[TestMethod]
+		public void ListSelect()
+		{
+			object result = Evaluate(BgList<BgInt>.Create(1, 1, 2, 2, 3).Select(x => x * 9));
+			Assert.IsTrue(Enumerable.SequenceEqual(new object[] { 9, 9, 18, 18, 27 }, (IEnumerable<object>)result));
+		}
+
+		[TestMethod]
+		public void ListWhere()
+		{
+			object result = Evaluate(BgList<BgInt>.Create(1, 1, 2, 2, 3).Where(x => x >= 2));
+			Assert.IsTrue(Enumerable.SequenceEqual(new object[] { 2, 2, 3 }, (IEnumerable<object>)result));
+		}
+
+		[TestMethod]
+		public void ListDistinct()
+		{
+			object result = Evaluate(BgList<BgInt>.Concat(BgList<BgInt>.Empty.Add(1, 1, 2, 2, 3), BgList<BgInt>.Empty.Add(3)).Distinct());
+			Assert.IsTrue(Enumerable.SequenceEqual(new object[] { 1, 2, 3 }, (IEnumerable<object>)result));
+		}
+
+		[TestMethod]
+		public void ListOption()
+		{
+			BgListOptionDef result = (BgListOptionDef)EvaluateOption(new BgListOption("stringOption", "description", defaultValue: "value", values: BgList<BgString>.Create("value1", "value2"), valueDescriptions: BgList<BgString>.Create("desc1", "desc2")));
+			Assert.AreEqual("stringOption", result.Name);
+			Assert.AreEqual("description", result.Description);
+			Assert.AreEqual("value", result.DefaultValue);
+			Assert.IsTrue(result.Values.SequenceEqual(new[] { "value1", "value2" }));
+			Assert.IsTrue(result.ValueDescriptions.SequenceEqual(new[] { "desc1", "desc2" }));
+		}
+
+		#endregion
+
+		#region Objects
+
+		class TestObject
+		{
+			public int IntegerValue { get; set; } = 1234;
+			public string StringValue { get; set; } = "hello world";
+		}
+
+		[TestMethod]
+		public void ObjEmpty()
+		{
+			object result = Evaluate(BgObject<TestObject>.Empty);
+			Assert.AreEqual(0, ((BgObjectDef)result).Properties.Count);
+		}
+
+		[TestMethod]
+		public void ObjSet()
+		{
+			object result1 = Evaluate(BgObject<TestObject>.Empty.Set<BgInt, int>(x => x.IntegerValue, 999));
+			Assert.AreEqual(999, ((BgObjectDef)result1).Get(nameof(TestObject.IntegerValue), null));
+
+			object result2 = Evaluate(BgObject<TestObject>.Empty.Set<BgString, string>(x => x.StringValue, "foo"));
+			Assert.AreEqual("foo", ((BgObjectDef)result2).Get(nameof(TestObject.StringValue), null));
+		}
+
+		[TestMethod]
+		public void ObjGet()
+		{
+			object result1 = Evaluate(BgObject<TestObject>.Empty.Get<BgInt, int>(x => x.IntegerValue, 1234));
+			Assert.AreEqual(1234, result1);
+
+			object result2 = Evaluate(BgObject<TestObject>.Empty.Get<BgString, string>(x => x.StringValue, "hello world"));
+			Assert.AreEqual("hello world", result2);
+		}
+
+		#endregion
+
+		#region Generic
+
+		[TestMethod]
+		public void Throw()
+		{
+			Assert.ThrowsException<BgBytecodeException>(() => Evaluate(BgExpr.Throw<BgInt>("this is an exception")));
+		}
+
+		static Task NativeMethod(BgInt a, BgString b, int c, string d) => throw new NotImplementedException();
+		static Task<BgFileSet> NativeMethodWithReturnValue(BgInt a, BgString b, int c, string d) => throw new NotImplementedException();
+
+		[TestMethod]
+		public void Thunk()
+		{
+			BgInt a = (BgInt)100 + (BgInt)20 + (BgInt)3;
+			BgString b = (BgString)"4" + (BgString)"5" + (BgString)"6";
+
+			{
+				BgThunk thunk = BgThunk.Create(() => NativeMethod(a, b, 3, "4"));
+				Assert.IsTrue(ReferenceEquals(thunk.Arguments[0], a));
+				Assert.IsTrue(ReferenceEquals(thunk.Arguments[1], b));
+				Assert.AreEqual(3, thunk.Arguments[2]);
+				Assert.AreEqual("4", thunk.Arguments[3]);
+
+				BgThunkDef thunkDef = (BgThunkDef)Evaluate(thunk);
+				Assert.IsTrue(thunkDef.Arguments[0] is BgIntConstantExpr constInt && constInt.Value == 123);
+				Assert.IsTrue(thunkDef.Arguments[1] is BgStringConstantExpr constStr && constStr.Value == "456");
+				Assert.AreEqual(3, thunkDef.Arguments[2]);
+				Assert.AreEqual("4", thunkDef.Arguments[3]);
+			}
+			{
+				BgThunk<BgFileSet> thunk = BgThunk.Create(() => NativeMethodWithReturnValue(a, b, 3, "4"));
+				Assert.IsTrue(ReferenceEquals(thunk.Arguments[0], a));
+				Assert.IsTrue(ReferenceEquals(thunk.Arguments[1], b));
+				Assert.AreEqual(3, thunk.Arguments[2]);
+				Assert.AreEqual("4", thunk.Arguments[3]);
+
+				BgThunkDef thunkDef = (BgThunkDef)Evaluate(thunk);
+				Assert.IsTrue(thunkDef.Arguments[0] is BgIntConstantExpr constInt && constInt.Value == 123);
+				Assert.IsTrue(thunkDef.Arguments[1] is BgStringConstantExpr constStr && constStr.Value == "456");
+				Assert.AreEqual(3, thunkDef.Arguments[2]);
+				Assert.AreEqual("4", thunkDef.Arguments[3]);
+			}
+		}
+
+		#endregion
+
+		#region Functions
+
+		[TestMethod]
+		public void Call()
+		{
+			BgFunc<BgBool> func = new BgFunc<BgBool>(() => true);
+			object result1 = Evaluate(func.Call());
+			Assert.AreEqual(true, result1);
+
+			BgFunc<BgInt, BgInt> func2 = new BgFunc<BgInt, BgInt>(x => (x * 3) + 2);
+			object result2 = Evaluate(func2.Call(7));
+			Assert.AreEqual(23, result2);
+
+			BgFunc<BgInt, BgInt> func3a = new BgFunc<BgInt, BgInt>(x => x * 3);
+			BgFunc<BgInt, BgInt> func3b = new BgFunc<BgInt, BgInt>(x => func3a.Call(x) + 2);
+			object result3 = Evaluate(func3b.Call(7));
+			Assert.AreEqual(23, result3);
+
+			BgFunc<BgInt, BgInt, BgInt> func4 = new BgFunc<BgInt, BgInt, BgInt>((x, y) => (x + y));
+			object result4 = Evaluate(func4.Call(7, 9));
+			Assert.AreEqual(16, result4);
+		}
+
+		#endregion
+	}
+}
diff --git a/Engine/Source/Programs/Shared/EpicGames.BuildGraph.Tests/SerializationTests.cs b/Engine/Source/Programs/Shared/EpicGames.BuildGraph.Tests/SerializationTests.cs
new file mode 100644
--- /dev/null
+++ b/Engine/Source/Programs/Shared/EpicGames.BuildGraph.Tests/SerializationTests.cs
@@ -0,0 +1,59 @@
+// Copyright Epic Games, Inc. All Rights Reserved.
+
+using Microsoft.VisualStudio.TestTools.UnitTesting;
+using System;
+using System.Collections.Generic;
+using System.IO;
+using System.Linq;
+using System.Text;
+using System.Text.Json;
+using OpenTracing;
+using EpicGames.BuildGraph.Expressions;
+using System.Collections.Immutable;
+
+namespace EpicGames.BuildGraph.Tests
+{
+	[TestClass]
+	public class SerializationTests
+	{
+		class NestedObject
+		{
+			public int Value { get; set; }
+		}
+
+		class RootObject
+		{
+			public NestedObject? Nested { get; set; }
+			public List<string>? StringList { get; set; }
+			public List<NestedObject> NestedList { get; } = new List<NestedObject>();
+		}
+
+		[TestMethod]
+		public void BasicTests()
+		{
+			BgObjectDef def = new BgObjectDef().Set(nameof(RootObject.Nested), new BgObjectDef().Set(nameof(NestedObject.Value), 123));
+			RootObject obj = def.Deserialize<RootObject>();
+			Assert.AreEqual(123, obj.Nested!.Value);
+		}
+
+		[TestMethod]
+		public void StringListTests()
+		{
+			BgObjectDef def = new BgObjectDef().Set(nameof(RootObject.StringList), new[] { "hello", "world" });
+			RootObject obj = def.Deserialize<RootObject>();
+			Assert.IsNotNull(obj.StringList);
+			Assert.AreEqual(2, obj.StringList!.Count);
+			Assert.AreEqual("hello", obj.StringList[0]);
+			Assert.AreEqual("world", obj.StringList[1]);
+		}
+
+		[TestMethod]
+		public void ObjectListTests()
+		{
+			BgObjectDef def = new BgObjectDef().Set(nameof(RootObject.NestedList), new[] { new BgObjectDef().Set(nameof(NestedObject.Value), 123) });
+			RootObject obj = def.Deserialize<RootObject>();
+			Assert.AreEqual(1, obj.NestedList.Count);
+			Assert.AreEqual(123, obj.NestedList[0].Value);
+		}
+	}
+}
diff --git a/Engine/Source/Programs/Shared/EpicGames.BuildGraph/BgAgentDef.cs b/Engine/Source/Programs/Shared/EpicGames.BuildGraph/BgAgentDef.cs
--- a/Engine/Source/Programs/Shared/EpicGames.BuildGraph/BgAgentDef.cs
+++ b/Engine/Source/Programs/Shared/EpicGames.BuildGraph/BgAgentDef.cs
@@ -11,6 +11,7 @@ namespace EpicGames.BuildGraph
 	/// Stores a list of nodes which can be executed on a single agent
 	/// </summary>
 	[DebuggerDisplay("{Name}")]
+	[BgObject(typeof(BgAgentDefSerializer))]
 	public class BgAgentDef
 	{
 		/// <summary>
@@ -22,7 +23,7 @@ namespace EpicGames.BuildGraph
 		/// Array of valid agent types that these nodes may run on. When running in the build system, this determines the class of machine that should
 		/// be selected to run these nodes. The first defined agent type for this branch will be used.
 		/// </summary>
-		public string[] PossibleTypes { get; }
+		public List<string> PossibleTypes { get; } = new List<string>();
 
 		/// <summary>
 		/// List of nodes in this agent group.
@@ -37,12 +38,21 @@ namespace EpicGames.BuildGraph
 		/// <summary>
 		/// Constructor
 		/// </summary>
-		/// <param name="inName">Name of this agent group</param>
-		/// <param name="inPossibleTypes">Array of valid agent types. See comment for AgentTypes member.</param>
-		public BgAgentDef(string inName, string[] inPossibleTypes)
+		/// <param name="name">Name of this agent group</param>
+		public BgAgentDef(string name)
 		{
-			Name = inName;
-			PossibleTypes = inPossibleTypes;
+			Name = name;
+		}
+	}
+
+	class BgAgentDefSerializer : BgObjectSerializer<BgAgentDef>
+	{
+		/// <inheritdoc/>
+		public override BgAgentDef Deserialize(BgObjectDef<BgAgentDef> obj)
+		{
+			BgAgentDef agent = new BgAgentDef(obj.Get(x => x.Name, ""));
+			obj.CopyTo(agent);
+			return agent;
 		}
 	}
 }
diff --git a/Engine/Source/Programs/Shared/EpicGames.BuildGraph/BgAggregateDef.cs b/Engine/Source/Programs/Shared/EpicGames.BuildGraph/BgAggregateDef.cs
--- a/Engine/Source/Programs/Shared/EpicGames.BuildGraph/BgAggregateDef.cs
+++ b/Engine/Source/Programs/Shared/EpicGames.BuildGraph/BgAggregateDef.cs
@@ -41,21 +41,48 @@ namespace EpicGames.BuildGraph
 	/// <summary>
 	/// Aggregate that was created from bytecode
 	/// </summary>
-	public class BgBytecodeAggregateDef : BgAggregateDef
+	[BgObject(typeof(BgAggregateExpressionDefSerializer))]
+	public class BgAggregateExpressionDef
 	{
+		/// <inheritdoc cref="BgAggregateDef.Name"/>
+		public string Name { get; set; }
+
+		/// <inheritdoc cref="BgAggregateDef.RequiredNodes"/>
+		public List<BgNodeExpressionDef> RequiredNodes { get; } = new List<BgNodeExpressionDef>();
+
 		/// <summary>
 		/// Labels to add this aggregate to
 		/// </summary>
-		public BgLabelDef? Label { get; }
+		public BgLabelDef? Label { get; set; }
 
 		/// <summary>
 		/// Constructor
 		/// </summary>
-		public BgBytecodeAggregateDef(string name, IEnumerable<BgNodeDef> nodes, BgLabelDef? label)
-			: base(name)
+		public BgAggregateExpressionDef(string name)
+		{
+			Name = name;
+		}
+
+		/// <summary>
+		/// Construct a BgAggregateDef
+		/// </summary>
+		/// <returns></returns>
+		public BgAggregateDef ToAggregateDef()
+		{
+			BgAggregateDef aggregate = new BgAggregateDef(Name);
+			aggregate.RequiredNodes.UnionWith(RequiredNodes);
+			return aggregate;
+		}
+	}
+
+	class BgAggregateExpressionDefSerializer : BgObjectSerializer<BgAggregateExpressionDef>
+	{
+		/// <inheritdoc/>
+		public override BgAggregateExpressionDef Deserialize(BgObjectDef<BgAggregateExpressionDef> obj)
 		{
-			RequiredNodes.UnionWith(nodes);
-			Label = label;
+			BgAggregateExpressionDef aggregate = new BgAggregateExpressionDef(obj.Get(x => x.Name, ""));
+			obj.CopyTo(aggregate);
+			return aggregate;
 		}
 	}
 }
diff --git a/Engine/Source/Programs/Shared/EpicGames.BuildGraph/BgBytecodeWriter.cs b/Engine/Source/Programs/Shared/EpicGames.BuildGraph/BgBytecodeWriter.cs
--- a/Engine/Source/Programs/Shared/EpicGames.BuildGraph/BgBytecodeWriter.cs
+++ b/Engine/Source/Programs/Shared/EpicGames.BuildGraph/BgBytecodeWriter.cs
@@ -31,6 +31,12 @@ namespace EpicGames.BuildGraph
 		/// <param name="str"></param>
 		public abstract void WriteString(string str);
 
+		/// <summary>
+		/// Writes a reference to an interned string to the output
+		/// </summary>
+		/// <param name="name">Name to write</param>
+		public abstract void WriteName(string name);
+
 		/// <summary>
 		/// Writes a signed integer value to the output
 		/// </summary>
@@ -62,9 +68,9 @@ namespace EpicGames.BuildGraph
 		public abstract void WriteExprAsFragment(BgExpr expr);
 
 		/// <summary>
-		/// Writes the method to be called for a particular node
+		/// Writes a thunk to native code.
 		/// </summary>
-		/// <param name="method">Method to be called</param>
-		public abstract void WriteMethod(BgMethod method);
+		/// <param name="thunk">Method to be called</param>
+		public abstract void WriteThunk(BgThunkDef thunk);
 	}
 }
diff --git a/Engine/Source/Programs/Shared/EpicGames.BuildGraph/BgCompiler.cs b/Engine/Source/Programs/Shared/EpicGames.BuildGraph/BgCompiler.cs
--- a/Engine/Source/Programs/Shared/EpicGames.BuildGraph/BgCompiler.cs
+++ b/Engine/Source/Programs/Shared/EpicGames.BuildGraph/BgCompiler.cs
@@ -94,6 +94,9 @@ namespace EpicGames.BuildGraph
 			/// <inheritdoc/>
 			public override void WriteOpcode(BgOpcode opcode) { }
 
+			/// <inheritdoc/>
+			public override void WriteName(string str) { }
+
 			/// <inheritdoc/>
 			public override void WriteString(string str) { }
 
@@ -104,20 +107,28 @@ namespace EpicGames.BuildGraph
 			public override void WriteUnsignedInteger(ulong value) { }
 
 			/// <inheritdoc/>
-			public override void WriteMethod(BgMethod handler) { }
+			public override void WriteThunk(BgThunkDef handler) { }
 		}
 
 		class ForwardWriter : BgBytecodeWriter
 		{
 			readonly ByteArrayBuilder _builder;
 			readonly Dictionary<BgExpr, int> _exprToFragmentIdx;
-			readonly List<BgMethod> _methods;
+			readonly List<string> _names;
+			readonly List<BgThunkDef> _thunks;
+			readonly Dictionary<string, int> _nameToIndex = new Dictionary<string, int>();
 
-			public ForwardWriter(ByteArrayBuilder builder, Dictionary<BgExpr, int> exprToFragmentIdx, List<BgMethod> handlers)
+			public ForwardWriter(ByteArrayBuilder builder, Dictionary<BgExpr, int> exprToFragmentIdx, List<string> names, List<BgThunkDef> thunks)
 			{
 				_builder = builder;
 				_exprToFragmentIdx = exprToFragmentIdx;
-				_methods = handlers;
+				_names = names;
+				_thunks = thunks;
+
+				for (int idx = 0; idx < _names.Count; idx++)
+				{
+					_nameToIndex[_names[idx]] = idx;
+				}
 			}
 
 			/// <inheritdoc/>
@@ -147,6 +158,19 @@ namespace EpicGames.BuildGraph
 				_builder.WriteByte((byte)opcode);
 			}
 
+			/// <inheritdoc/>
+			public override void WriteName(string name)
+			{
+				int index;
+				if (!_nameToIndex.TryGetValue(name, out index))
+				{
+					index = _names.Count;
+					_names.Add(name);
+					_nameToIndex.Add(name, index);
+				}
+				WriteUnsignedInteger((ulong)index);
+			}
+
 			/// <inheritdoc/>
 			public override void WriteString(string str)
 			{
@@ -165,10 +189,10 @@ namespace EpicGames.BuildGraph
 			public override void WriteUnsignedInteger(ulong value) => _builder.WriteUnsignedVarInt(value);
 
 			/// <inheritdoc/>
-			public override void WriteMethod(BgMethod method)
+			public override void WriteThunk(BgThunkDef thunk)
 			{
-				_builder.WriteUnsignedVarInt(_methods.Count);
-				_methods.Add(method);
+				_builder.WriteUnsignedVarInt(_thunks.Count);
+				_thunks.Add(thunk);
 			}
 		}
 
@@ -177,7 +201,7 @@ namespace EpicGames.BuildGraph
 		/// </summary>
 		/// <param name="expr">Expression to compile</param>
 		/// <returns>Compiled bytecode for the expression, suitable for passing to <see cref="BgInterpreter"/></returns>
-		public static (byte[], BgMethod[]) Compile(BgExpr expr)
+		public static (byte[], BgThunkDef[]) Compile(BgExpr expr)
 		{
 			List<BgExpr> fragments = new List<BgExpr>();
 			Dictionary<BgExpr, int> exprToFragmentIndex = new Dictionary<BgExpr, int>(ReferenceEqualityComparer.Instance);
@@ -190,9 +214,10 @@ namespace EpicGames.BuildGraph
 			ByteArrayBuilder code = new ByteArrayBuilder();
 			List<int> fragmentLengths = new List<int>(fragments.Count);
 
-			List<BgMethod> methods = new List<BgMethod>();
+			List<BgThunkDef> thunks = new List<BgThunkDef>();
+			List<string> names = new List<string>();
 
-			ForwardWriter writer = new ForwardWriter(code, exprToFragmentIndex, methods);
+			ForwardWriter writer = new ForwardWriter(code, exprToFragmentIndex, names, thunks);
 			for (int idx = 0; idx < fragments.Count; idx++)
 			{
 				int fragmentOffset = code.Length;
@@ -203,18 +228,15 @@ namespace EpicGames.BuildGraph
 			// Create the header
 			ByteArrayBuilder header = new ByteArrayBuilder();
 			header.WriteUnsignedVarInt((int)BgBytecodeVersion.Current);
-			header.WriteUnsignedVarInt(fragments.Count);
-			for (int idx = 0; idx < fragments.Count; idx++)
-			{
-				header.WriteUnsignedVarInt(fragmentLengths[idx]);
-			}
+			header.WriteVariableLengthArray(names, x => header.WriteString(x));
+			header.WriteVariableLengthArray(fragmentLengths, x => header.WriteUnsignedVarInt(x));
 
 			// Append them together
 			byte[] output = new byte[header.Length + code.Length];
 			header.CopyTo(output);
 			code.CopyTo(output.AsSpan(header.Length));
 
-			return (output, methods.ToArray());
+			return (output, thunks.ToArray());
 		}
 	}
 }
diff --git a/Engine/Source/Programs/Shared/EpicGames.BuildGraph/BgContext.cs b/Engine/Source/Programs/Shared/EpicGames.BuildGraph/BgContext.cs
--- a/Engine/Source/Programs/Shared/EpicGames.BuildGraph/BgContext.cs
+++ b/Engine/Source/Programs/Shared/EpicGames.BuildGraph/BgContext.cs
@@ -87,14 +87,14 @@ namespace EpicGames.BuildGraph
 		/// <typeparam name="TEnum">The enum type</typeparam>
 		/// <param name="expr">Enum expression</param>
 		/// <returns>The enum value</returns>
-		public List<TEnum> Get<TEnum>(BgList<BgEnum<TEnum>> expr) where TEnum : struct => ((BgListConstantExpr<BgEnum<TEnum>>)expr).Value.Select(x => (TEnum)x).ToList();
+		public List<TEnum> Get<TEnum>(BgList<BgEnum<TEnum>> expr) where TEnum : struct => ((BgListConstantExpr<BgEnum<TEnum>>)expr).Value.Select(x => ((BgEnumConstantExpr<TEnum>)x).Value).ToList();
 
 		/// <summary>
 		/// Resolve a list of strings
 		/// </summary>
 		/// <param name="expr">List expression</param>
 		/// <returns></returns>
-		public List<string> Get(BgList<BgString> expr) => (((BgListConstantExpr<BgString>)expr).Value).ConvertAll(x => (string)x);
+		public List<string> Get(BgList<BgString> expr) => (((BgListConstantExpr<BgString>)expr).Value).ConvertAll(x => ((BgStringConstantExpr)x).Value);
 
 		/// <summary>
 		/// Resolve a file set
diff --git a/Engine/Source/Programs/Shared/EpicGames.BuildGraph/BgGraphDef.cs b/Engine/Source/Programs/Shared/EpicGames.BuildGraph/BgGraphDef.cs
--- a/Engine/Source/Programs/Shared/EpicGames.BuildGraph/BgGraphDef.cs
+++ b/Engine/Source/Programs/Shared/EpicGames.BuildGraph/BgGraphDef.cs
@@ -4,11 +4,14 @@ using System;
 using System.Collections.Generic;
 using System.Diagnostics.CodeAnalysis;
 using System.Linq;
+using System.Runtime.CompilerServices;
 using System.Text;
 using System.Xml;
 using EpicGames.Core;
 using Microsoft.Extensions.Logging;
 
+[assembly: InternalsVisibleTo("EpicGames.BuildGraph.Tests")]
+
 namespace EpicGames.BuildGraph
 {
 	/// <summary>
@@ -40,7 +43,7 @@ namespace EpicGames.BuildGraph
 		/// <summary>
 		/// List of options, in the order they were specified
 		/// </summary>
-		public List<BgOption> Options { get; } = new List<BgOption>();
+		public List<BgOptionDef> Options { get; } = new List<BgOptionDef>();
 
 		/// <summary>
 		/// List of agents containing nodes to execute
@@ -261,7 +264,7 @@ namespace EpicGames.BuildGraph
 			// Remove all the order dependencies which are no longer part of the graph. Since we don't need to build them, we don't need to wait for them
 			foreach (BgNodeDef node in retainNodes)
 			{
-				node.OrderDependencies = node.OrderDependencies.Where(x => retainNodes.Contains(x)).ToArray();
+				node.OrderDependencies.RemoveAll(x => !retainNodes.Contains(x));
 			}
 
 			// Create a new list of aggregates for everything that's left
@@ -562,14 +565,29 @@ namespace EpicGames.BuildGraph
 			{
 				// Get the list of messages
 				List<KeyValuePair<string, string>> parameters = new List<KeyValuePair<string, string>>();
-				foreach (BgOption option in Options)
+				foreach (BgOptionDef option in Options)
 				{
 					string name = String.Format("-set:{0}=...", option.Name);
 
 					StringBuilder description = new StringBuilder(option.Description);
-					if (!String.IsNullOrEmpty(option.DefaultValue))
+
+					string? defaultValue = null;
+					if (option is BgBoolOptionDef boolOption)
+					{
+						defaultValue = boolOption.DefaultValue.ToString();
+					}
+					else if (option is BgIntOptionDef intOption)
 					{
-						description.AppendFormat(" (Default: {0})", option.DefaultValue);
+						defaultValue = intOption.DefaultValue.ToString();
+					}
+					else if (option is BgStringOptionDef stringOption)
+					{
+						defaultValue = stringOption.DefaultValue;
+					}
+
+					if (!String.IsNullOrEmpty(defaultValue))
+					{
+						description.AppendFormat(" (Default: {0})", defaultValue);
 					}
 
 					parameters.Add(new KeyValuePair<string, string>(name, description.ToString()));
@@ -643,4 +661,85 @@ namespace EpicGames.BuildGraph
 			}
 		}
 	}
+
+
+	/// <summary>
+	/// Definition of a graph from bytecode. Can be converted to regular graph definition.
+	/// </summary>
+	public class BgGraphExpressionDef
+	{
+		/// <summary>
+		/// Nodes for the graph
+		/// </summary>
+		public List<BgNodeExpressionDef> Nodes { get; } = new List<BgNodeExpressionDef>();
+
+		/// <summary>
+		/// Aggregates for the graph
+		/// </summary>
+		public List<BgAggregateExpressionDef> Aggregates { get; } = new List<BgAggregateExpressionDef>();
+
+		/// <summary>
+		/// Creates a graph definition from this template
+		/// </summary>
+		/// <returns></returns>
+		public BgGraphDef ToGraphDef()
+		{
+			List<BgNodeExpressionDef> nodes = new List<BgNodeExpressionDef>(Nodes);
+
+			BgGraphDef graph = new BgGraphDef();
+			foreach (BgAggregateExpressionDef aggregate in Aggregates)
+			{
+				graph.NameToAggregate[aggregate.Name] = aggregate.ToAggregateDef();
+				nodes.AddRange(aggregate.RequiredNodes.Select(x => (BgNodeExpressionDef)x));
+			}
+
+			HashSet<BgNodeDef> uniqueNodes = new HashSet<BgNodeDef>();
+			HashSet<BgAgentDef> uniqueAgents = new HashSet<BgAgentDef>();
+			foreach (BgNodeExpressionDef node in nodes)
+			{
+				RegisterNode(graph, node, uniqueNodes, uniqueAgents);
+			}
+
+			HashSet<BgLabelDef> labels = new HashSet<BgLabelDef>();
+			foreach (BgNodeExpressionDef node in nodes)
+			{
+				foreach (BgLabelDef label in node.Labels)
+				{
+					labels.Add(label);
+					label.RequiredNodes.Add(node);
+				}
+			}
+			foreach (BgAggregateExpressionDef aggregate in Aggregates)
+			{
+				if (aggregate.Label != null)
+				{
+					aggregate.Label.RequiredNodes.UnionWith(aggregate.RequiredNodes);
+					labels.Add(aggregate.Label);
+				}
+			}
+			graph.Labels.AddRange(labels);
+			return graph;
+		}
+
+		static void RegisterNode(BgGraphDef graph, BgNodeExpressionDef node, HashSet<BgNodeDef> uniqueNodes, HashSet<BgAgentDef> uniqueAgents)
+		{
+			if (uniqueNodes.Add(node))
+			{
+				foreach (BgNodeExpressionDef inputNode in node.InputDependencies)
+				{
+					RegisterNode(graph, inputNode, uniqueNodes, uniqueAgents);
+				}
+
+				BgAgentDef agent = node.Agent;
+				if (uniqueAgents.Add(agent))
+				{
+					graph.Agents.Add(agent);
+					graph.NameToAgent.Add(agent.Name, agent);
+				}
+
+				agent.Nodes.Add(node);
+				graph.NameToNode.Add(node.Name, node);
+			}
+		}
+	}
 }
diff --git a/Engine/Source/Programs/Shared/EpicGames.BuildGraph/BgInterpreter.cs b/Engine/Source/Programs/Shared/EpicGames.BuildGraph/BgInterpreter.cs
--- a/Engine/Source/Programs/Shared/EpicGames.BuildGraph/BgInterpreter.cs
+++ b/Engine/Source/Programs/Shared/EpicGames.BuildGraph/BgInterpreter.cs
@@ -7,6 +7,7 @@ using System.Collections.Immutable;
 using System.Diagnostics;
 using System.IO;
 using System.Linq;
+using System.Reflection;
 using System.Text;
 using System.Text.RegularExpressions;
 using System.Xml;
@@ -81,43 +82,43 @@ namespace EpicGames.BuildGraph
 		}
 
 		readonly byte[] _data;
-		readonly BgMethod[] _methods;
+		readonly BgThunkDef[] _thunks;
+		readonly string[] _names;
 		readonly IReadOnlyDictionary<string, string> _options;
 		readonly BgBytecodeVersion _version;
 		readonly int[] _fragments;
+		readonly List<BgOptionDef> _optionDefs = new List<BgOptionDef>();
+
+		/// <summary>
+		/// The option definitions that were parsed during execution
+		/// </summary>
+		public IReadOnlyList<BgOptionDef> OptionDefs => _optionDefs;
 
 		/// <summary>
 		/// Constructor
 		/// </summary>
 		/// <param name="data"></param>
-		/// <param name="lambdas">Handlers for each node type</param>
+		/// <param name="thunks">Thunks to native code</param>
 		/// <param name="options">Options for evaluating the graph</param>
-		public BgInterpreter(byte[] data, BgMethod[] lambdas, IReadOnlyDictionary<string, string> options)
+		public BgInterpreter(byte[] data, BgThunkDef[] thunks, IReadOnlyDictionary<string, string> options)
 		{
 			_data = data;
-			_methods = lambdas;
+			_thunks = thunks;
 			_options = options;
 
-			int offset = 0;
+			MemoryReader reader = new MemoryReader(data);
 
-			_version = (BgBytecodeVersion)VarInt.ReadUnsigned(data[offset..], out int bytesRead);
-			offset += bytesRead;
+			_version = (BgBytecodeVersion)reader.ReadUnsignedVarInt();
+			_names = reader.ReadVariableLengthArray(() => reader.ReadString());
+			int[] fragmentLengths = reader.ReadVariableLengthArray(() => (int)reader.ReadUnsignedVarInt());
 
-			int fragmentCount = (int)VarInt.ReadUnsigned(data[offset..], out bytesRead);
-			offset += bytesRead;
+			int offset = data.Length - reader.Memory.Length;
 
-			int[] lengths = new int[fragmentCount];
-			for (int idx = 0; idx < fragmentCount; idx++)
-			{
-				lengths[idx] = (int)VarInt.ReadUnsigned(data[offset..], out bytesRead);
-				offset += bytesRead;
-			}
-
-			_fragments = new int[fragmentCount];
-			for (int idx = 0; idx < fragmentCount; idx++)
+			_fragments = new int[fragmentLengths.Length];
+			for (int idx = 0; idx < _fragments.Length; idx++)
 			{
 				_fragments[idx] = offset;
-				offset += lengths[idx];
+				offset += fragmentLengths[idx];
 			}
 		}
 
@@ -171,13 +172,13 @@ namespace EpicGames.BuildGraph
 					}
 				case BgOpcode.BoolOption:
 					{
-						string name = (string)Evaluate(frame);
-						string label = (string)Evaluate(frame);
-						string description = (string)Evaluate(frame);
-						bool defaultValue = (bool)Evaluate(frame);
+						BgObjectDef obj = (BgObjectDef)Evaluate(frame);
+
+						BgBoolOptionDef option = obj.Deserialize<BgBoolOptionDef>();
+						_optionDefs.Add(option);
 
-						bool value = defaultValue;
-						if (_options.TryGetValue(name, out string? str))
+						bool value = option.DefaultValue;
+						if (_options.TryGetValue(option.Name, out string? str))
 						{
 							value = Boolean.Parse(str);
 						}
@@ -236,15 +237,13 @@ namespace EpicGames.BuildGraph
 					return -(int)Evaluate(frame);
 				case BgOpcode.IntOption:
 					{
-						string name = (string)Evaluate(frame);
-						string label = (string)Evaluate(frame);
-						string description = (string)Evaluate(frame);
-						int defaultValue = (int)Evaluate(frame);
-						int minValue = (int)Evaluate(frame);
-						int maxValue = (int)Evaluate(frame);
-
-						int value = defaultValue;
-						if (_options.TryGetValue(name, out string? str))
+						BgObjectDef obj = (BgObjectDef)Evaluate(frame);
+
+						BgIntOptionDef option = obj.Deserialize<BgIntOptionDef>();
+						_optionDefs.Add(option);
+
+						int value = option.DefaultValue;
+						if (_options.TryGetValue(option.Name, out string? str))
 						{
 							value = Int32.Parse(str);
 						}
@@ -314,20 +313,15 @@ namespace EpicGames.BuildGraph
 					}
 				case BgOpcode.StrOption:
 					{
-						string name = (string)Evaluate(frame);
-						string label = (string)Evaluate(frame);
-						string description = (string)Evaluate(frame);
-						int style = (int)ReadUnsignedInteger(frame);
-						string defaultValue = (string)Evaluate(frame);
-						string pattern = (string)Evaluate(frame);
-						string patternFailed = (string)Evaluate(frame);
-						List<string> values = EvaluateList<string>(frame);
-						List<string> valueDescriptions = EvaluateList<string>(frame);
-
-						string value = defaultValue;
-						if (_options.TryGetValue(name, out string? option))
+						BgObjectDef obj = (BgObjectDef)Evaluate(frame);
+
+						BgStringOptionDef option = obj.Deserialize<BgStringOptionDef>();
+						_optionDefs.Add(option);
+
+						string value = option.DefaultValue;
+						if (_options.TryGetValue(option.Name, out string? result))
 						{
-							value = option;
+							value = result;
 						}
 
 						return value;
@@ -452,65 +446,45 @@ namespace EpicGames.BuildGraph
 					}
 				case BgOpcode.ListOption:
 					{
-						string name = (string)Evaluate(frame);
-						string label = (string)Evaluate(frame);
-						string description = (string)Evaluate(frame);
-						BgListOptionStyle style = (BgListOptionStyle)ReadUnsignedInteger(frame);
-						string defaultValue = (string)Evaluate(frame);
-						string pattern = (string)Evaluate(frame);
-						string patternFailed = (string)Evaluate(frame);
-						List<string> values = EvaluateList<string>(frame);
-						List<string> valueDescriptions = EvaluateList<string>(frame);
-
-						string value = defaultValue;
-						if (_options.TryGetValue(name, out string? option))
-						{
-							value = option;
-						}
+						BgObjectDef obj = (BgObjectDef)Evaluate(frame);
 
-						return value.Split('+', ';').Select(x => (object)x).ToList();
-					}
+						BgListOptionDef option = obj.Deserialize<BgListOptionDef>();
+						_optionDefs.Add(option);
 
-				#endregion
-				#region Fileset opcodes
+						string? value = option.DefaultValue ?? String.Empty;
+						if (_options.TryGetValue(option.Name, out string? result))
+						{
+							value = result;
+						}
 
-				case BgOpcode.FileSetFromNode:
-					{
-						BgExpressionNodeDef node = (BgExpressionNodeDef)Evaluate(frame);
-						return node.InputDependencies.SelectMany(x => x.Outputs).Concat(node.Outputs).ToArray();
-					}
-				case BgOpcode.FileSetFromNodeOutput:
-					{
-						BgExpressionNodeDef node = (BgExpressionNodeDef)Evaluate(frame);
-						int outputIndex = (int)ReadUnsignedInteger(frame);
-						return new[] { node.Outputs[outputIndex] };
+						return value!.Split(new[] { '+', ';' }, StringSplitOptions.RemoveEmptyEntries).Select(x => (object)x).ToList();
 					}
 
 				#endregion
 				#region Object opcodes
 
 				case BgOpcode.ObjEmpty:
-					return ImmutableDictionary<string, object>.Empty.WithComparers(StringComparer.Ordinal);
+					return new BgObjectDef();
 				case BgOpcode.ObjGet:
 					{
-						ImmutableDictionary<string, object> obj = (ImmutableDictionary<string, object>)Evaluate(frame);
-						string name = ReadString(frame);
+						BgObjectDef obj = (BgObjectDef)Evaluate(frame);
+						string name = ReadName(frame);
 						object defaultValue = Evaluate(frame);
 
-						object value;
-						if (!obj.TryGetValue(name, out value))
+						object? value;
+						if (!obj.Properties.TryGetValue(name, out value))
 						{
 							value = defaultValue;
 						}
 
-						return value;
+						return value!;
 					}
 				case BgOpcode.ObjSet:
 					{
-						ImmutableDictionary<string, object> obj = (ImmutableDictionary<string, object>)Evaluate(frame);
-						string name = ReadString(frame);
+						BgObjectDef obj = (BgObjectDef)Evaluate(frame);
+						string name = ReadName(frame);
 						object value = Evaluate(frame);
-						return obj.SetItem(name, value);
+						return obj.Set(name, value);
 					}
 
 				#endregion
@@ -561,114 +535,30 @@ namespace EpicGames.BuildGraph
 					}
 				case BgOpcode.Null:
 					return null!;
-
-				#endregion
-				#region Graph opcodes
-
-				case BgOpcode.Node:
+				case BgOpcode.Thunk:
 					{
-						string name = (string)Evaluate(frame);
-						BgAgentDef agent = (BgAgentDef)Evaluate(frame);
-						BgMethod method = _methods[(int)ReadUnsignedInteger(frame)];
-
-						int argumentCount = (int)ReadUnsignedInteger(frame);
-						object[] arguments = new object[argumentCount];
-
-						for (int idx = 0; idx < argumentCount; idx++)
-						{
-							arguments[idx] = Evaluate(frame);
-						}
-
-						List<BgNodeOutput> inputs = ReadNodeOutputList(frame);
-						List<BgNodeOutput> fences = ReadNodeOutputList(frame);
-						int taggedOutputCount = (int)ReadUnsignedInteger(frame);
-						bool runEarly = (bool)Evaluate(frame);
-						List<BgLabelDef> labels = ((IEnumerable<object>)Evaluate(frame)).Select(x => (BgLabelDef)x).ToList();
+						BgThunkDef method = _thunks[(int)ReadUnsignedInteger(frame)];
 
-						HashSet<BgNodeDef> inputDependencies = new HashSet<BgNodeDef>();
-						foreach (BgNodeOutput input in inputs)
-						{
-							inputDependencies.Add(input.ProducingNode);
-							inputDependencies.UnionWith(input.ProducingNode.InputDependencies);
-						}
-
-						BgNodeDef[] orderDependencies = fences.Select(x => x.ProducingNode).Distinct().ToArray();
-
-						string[] outputNames = Enumerable.Range(0, taggedOutputCount).Select(x => BgNode.GetDefaultTagName(name, x)).ToArray();
-
-						BgExpressionNodeDef node = new BgExpressionNodeDef(agent, name, method, arguments, inputs, outputNames, inputDependencies.ToArray(), orderDependencies.ToArray(), Array.Empty<FileReference>(), labels);
-						return node;
-					}
-				case BgOpcode.Agent:
-					{
-						string name = (string)Evaluate(frame);
-						string[] types = ((IEnumerable<object>)Evaluate(frame)).Select(x => (string)x).ToArray();
-						return new BgAgentDef(name, types);
-					}
-				case BgOpcode.Aggregate:
-					{
-						string name = (string)Evaluate(frame);
-						List<BgNodeOutput> requires = ReadNodeOutputList(frame);
-						BgLabelDef? label = (BgLabelDef?)Evaluate(frame);
-
-						return new BgBytecodeAggregateDef(name, requires.Select(x => x.ProducingNode), label);
-					}
-				case BgOpcode.Label:
-					{
-						string? dashboardName = NullForEmptyString((string)Evaluate(frame));
-						string? dashboardCategory = NullForEmptyString((string)Evaluate(frame));
-						string? ugsBadge = NullForEmptyString((string)Evaluate(frame));
-						string? ugsProject = NullForEmptyString((string)Evaluate(frame));
-						string change = (string)Evaluate(frame);
-
-						BgLabelChange labelChange = BgLabelChange.Current;
-						if (change.Length > 0)
-						{
-							labelChange = Enum.Parse<BgLabelChange>(change);
-						}
-
-						return new BgLabelDef(dashboardName, dashboardCategory, ugsBadge, ugsProject, labelChange);
-					}
-				case BgOpcode.Graph:
-					{
-						List<BgExpressionNodeDef> nodes = EvaluateList<BgExpressionNodeDef>(frame);
-						List<BgBytecodeAggregateDef> aggregates = EvaluateList<BgBytecodeAggregateDef>(frame);
-
-						BgGraphDef graph = new BgGraphDef();
-						foreach (BgBytecodeAggregateDef aggregate in aggregates)
-						{
-							graph.NameToAggregate[aggregate.Name] = aggregate;
-							nodes.AddRange(aggregate.RequiredNodes.Select(x => (BgExpressionNodeDef)x));
-						}
-
-						HashSet<BgNodeDef> uniqueNodes = new HashSet<BgNodeDef>();
-						HashSet<BgAgentDef> uniqueAgents = new HashSet<BgAgentDef>();
-						foreach (BgExpressionNodeDef node in nodes)
-						{
-							RegisterNode(graph, node, uniqueNodes, uniqueAgents);
-						}
+						int count = (int)ReadUnsignedInteger(frame);
+						object?[] arguments = new object?[count];
 
-						HashSet<BgLabelDef> labels = new HashSet<BgLabelDef>();
-						foreach (BgExpressionNodeDef node in nodes)
+						ParameterInfo[] parameters = method.Method.GetParameters();
+						for (int idx = 0; idx < count; idx++)
 						{
-							foreach (BgLabelDef label in node.Labels)
+							object value = Evaluate(frame);
+							if (typeof(BgExpr).IsAssignableFrom(parameters[idx].ParameterType))
 							{
-								labels.Add(label);
-								label.RequiredNodes.Add(node);
+								arguments[idx] = BgType.Constant(parameters[idx].ParameterType, value);
 							}
-						}
-						foreach (BgBytecodeAggregateDef aggregate in aggregates)
-						{
-							if (aggregate.Label != null)
+							else
 							{
-								aggregate.Label.RequiredNodes.UnionWith(aggregate.RequiredNodes);
-								labels.Add(aggregate.Label);
+								arguments[idx] = method.Arguments[idx];
 							}
 						}
-						graph.Labels.AddRange(labels);
 
-						return graph;
+						return new BgThunkDef(method.Method, arguments);
 					}
+
 				#endregion
 
 				default:
@@ -676,46 +566,11 @@ namespace EpicGames.BuildGraph
 			}
 		}
 
-		void RegisterNode(BgGraphDef graph, BgExpressionNodeDef node, HashSet<BgNodeDef> uniqueNodes, HashSet<BgAgentDef> uniqueAgents)
-		{
-			if (uniqueNodes.Add(node))
-			{
-				foreach (BgExpressionNodeDef inputNode in node.InputDependencies)
-				{
-					RegisterNode(graph, inputNode, uniqueNodes, uniqueAgents);
-				}
-
-				BgAgentDef agent = node.Agent;
-				if (uniqueAgents.Add(agent))
-				{
-					graph.Agents.Add(agent);
-					graph.NameToAgent.Add(agent.Name, agent);
-				}
-
-				agent.Nodes.Add(node);
-				graph.NameToNode.Add(node.Name, node);
-			}
-		}
-
-		static string? NullForEmptyString(string str) => (str.Length == 0) ? null : str;
-
 		List<T> EvaluateList<T>(Frame frame)
 		{
 			return ((IEnumerable<object>)Evaluate(frame)).Select(x => (T)x).ToList();
 		}
 
-		List<BgNodeOutput> ReadNodeOutputList(Frame frame)
-		{
-			IEnumerable<object> list = (IEnumerable<object>)Evaluate(frame);
-
-			List<BgNodeOutput> outputs = new List<BgNodeOutput>();
-			foreach (object item in list)
-			{
-				outputs.AddRange((BgNodeOutput[])item);
-			}
-			return outputs;
-		}
-
 		IEnumerable<object> LazyEvaluateItem(Frame frame, int fragment)
 		{
 			yield return Jump(frame, fragment);
@@ -768,6 +623,15 @@ namespace EpicGames.BuildGraph
 			return opcode;
 		}
 
+		/// <summary>
+		/// Reads a name from the input stream
+		/// </summary>
+		string ReadName(Frame frame)
+		{
+			int index = (int)ReadUnsignedInteger(frame);
+			return _names[index];
+		}
+
 		/// <summary>
 		/// Reads a string from the input stream
 		/// </summary>
@@ -858,9 +722,6 @@ namespace EpicGames.BuildGraph
 					Disassemble(frame, logger);
 					break;
 				case BgOpcode.BoolOption:
-					Disassemble(frame, logger);
-					Disassemble(frame, logger);
-					Disassemble(frame, logger);
 					Disassemble(frame, logger);
 					break;
 
@@ -884,11 +745,6 @@ namespace EpicGames.BuildGraph
 					Disassemble(frame, logger);
 					break;
 				case BgOpcode.IntOption:
-					Disassemble(frame, logger);
-					Disassemble(frame, logger);
-					Disassemble(frame, logger);
-					Disassemble(frame, logger);
-					Disassemble(frame, logger);
 					Disassemble(frame, logger);
 					break;
 
@@ -933,15 +789,6 @@ namespace EpicGames.BuildGraph
 					Disassemble(frame, logger);
 					break;
 				case BgOpcode.StrOption:
-					Disassemble(frame, logger);
-					Disassemble(frame, logger);
-					Disassemble(frame, logger);
-					Disassemble(frame, logger);
-					Trace(frame, "style", ReadUnsignedInteger, logger);
-					Disassemble(frame, logger);
-					Disassemble(frame, logger);
-					Disassemble(frame, logger);
-					Disassemble(frame, logger);
 					Disassemble(frame, logger);
 					break;
 
@@ -989,29 +836,11 @@ namespace EpicGames.BuildGraph
 					TraceFragment(frame, logger);
 					break;
 				case BgOpcode.ListOption:
-					Disassemble(frame, logger);
-					Disassemble(frame, logger);
-					Disassemble(frame, logger);
-					Trace(frame, "style", ReadUnsignedInteger, logger);
-					Disassemble(frame, logger);
-					Disassemble(frame, logger);
-					Disassemble(frame, logger);
-					Disassemble(frame, logger);
 					Disassemble(frame, logger);
 					break;
 
 				#endregion
-				#region Fileset opcodes
 
-				case BgOpcode.FileSetFromNode:
-					Disassemble(frame, logger);
-					break;
-				case BgOpcode.FileSetFromNodeOutput:
-					Disassemble(frame, logger);
-					TraceFragment(frame, logger);
-					break;
-
-				#endregion
 				#region Object opcodes
 
 				case BgOpcode.ObjEmpty:
@@ -1019,11 +848,12 @@ namespace EpicGames.BuildGraph
 				case BgOpcode.ObjGet:
 				case BgOpcode.ObjSet:
 					Disassemble(frame, logger);
-					Trace(frame, "name", ReadString, logger);
+					Trace(frame, ReadName, x => $"name: {x}", logger);
 					Disassemble(frame, logger);
 					break;
 
 				#endregion
+
 				#region Function opcodes
 
 				case BgOpcode.Call:
@@ -1059,48 +889,19 @@ namespace EpicGames.BuildGraph
 					break;
 				case BgOpcode.Null:
 					break;
-
-				#endregion
-				#region Graph opcodes
-
-				case BgOpcode.Node:
+				case BgOpcode.Thunk:
 					{
-						Disassemble(frame, logger);
-						Disassemble(frame, logger);
-						Trace(frame, "handler", ReadUnsignedInteger, logger);
+						Trace(frame, "method", ReadUnsignedInteger, logger);
 
-						int argCount = (int)Trace(frame, "args", ReadUnsignedInteger, logger);
-						for (int idx = 0; idx < argCount; idx++)
+						int count = (int)Trace(frame, "args", ReadUnsignedInteger, logger);
+						for (int idx = 0; idx < count; idx++)
 						{
 							Disassemble(frame, logger);
 						}
 
-						Disassemble(frame, logger);
-						Disassemble(frame, logger);
-						Disassemble(frame, logger);
-						Disassemble(frame, logger);
 						break;
 					}
-				case BgOpcode.Agent:
-					Disassemble(frame, logger);
-					Disassemble(frame, logger);
-					break;
-				case BgOpcode.Aggregate:
-					Disassemble(frame, logger);
-					Disassemble(frame, logger);
-					break;
-				case BgOpcode.Label:
-					Disassemble(frame, logger);
-					Disassemble(frame, logger);
-					Disassemble(frame, logger);
-					Disassemble(frame, logger);
-					Disassemble(frame, logger);
-					Disassemble(frame, logger);
-					break;
-				case BgOpcode.Graph:
-					Disassemble(frame, logger);
-					Disassemble(frame, logger);
-					break;
+
 				#endregion
 
 				default:
diff --git a/Engine/Source/Programs/Shared/EpicGames.BuildGraph/BgLabelDef.cs b/Engine/Source/Programs/Shared/EpicGames.BuildGraph/BgLabelDef.cs
--- a/Engine/Source/Programs/Shared/EpicGames.BuildGraph/BgLabelDef.cs
+++ b/Engine/Source/Programs/Shared/EpicGames.BuildGraph/BgLabelDef.cs
@@ -30,27 +30,27 @@ namespace EpicGames.BuildGraph
 		/// <summary>
 		/// Name of this badge
 		/// </summary>
-		public string? DashboardName { get; }
+		public string? DashboardName { get; set; }
 
 		/// <summary>
 		/// Category for this label
 		/// </summary>
-		public string? DashboardCategory { get; }
+		public string? DashboardCategory { get; set; }
 
 		/// <summary>
 		/// Name of the badge in UGS
 		/// </summary>
-		public string? UgsBadge { get; }
+		public string? UgsBadge { get; set; }
 
 		/// <summary>
 		/// Path to the project folder in UGS
 		/// </summary>
-		public string? UgsProject { get; }
+		public string? UgsProject { get; set; }
 
 		/// <summary>
 		/// Which change to show the badge for
 		/// </summary>
-		public BgLabelChange Change { get; }
+		public BgLabelChange Change { get; set; }
 
 		/// <summary>
 		/// Set of nodes that must be run for this label to be shown.
@@ -62,6 +62,13 @@ namespace EpicGames.BuildGraph
 		/// </summary>
 		public HashSet<BgNodeDef> IncludedNodes { get; } = new HashSet<BgNodeDef>();
 
+		/// <summary>
+		/// Constructor
+		/// </summary>
+		public BgLabelDef()
+		{
+		}
+
 		/// <summary>
 		/// Constructor
 		/// </summary>
@@ -70,7 +77,7 @@ namespace EpicGames.BuildGraph
 		/// <param name="inUgsBadge">The UGS badge name</param>
 		/// <param name="inUgsProject">Project to display this badge for</param>
 		/// <param name="inChange">The change to show this badge on in UGS</param>
-		public BgLabelDef(string? inDashboardName, string? inDashboardCategory, string? inUgsBadge, string? inUgsProject, BgLabelChange inChange)
+		public BgLabelDef(string? inDashboardName = null, string? inDashboardCategory = null, string? inUgsBadge = null, string? inUgsProject = null, BgLabelChange inChange = BgLabelChange.Current)
 		{
 			DashboardName = inDashboardName;
 			DashboardCategory = inDashboardCategory;
diff --git a/Engine/Source/Programs/Shared/EpicGames.BuildGraph/BgNodeDef.cs b/Engine/Source/Programs/Shared/EpicGames.BuildGraph/BgNodeDef.cs
--- a/Engine/Source/Programs/Shared/EpicGames.BuildGraph/BgNodeDef.cs
+++ b/Engine/Source/Programs/Shared/EpicGames.BuildGraph/BgNodeDef.cs
@@ -48,6 +48,59 @@ namespace EpicGames.BuildGraph
 		}
 	}
 
+	/// <summary>
+	/// Describes a dependency on a node output
+	/// </summary>
+	[BgObject(typeof(BgNodeOutputExprDefSerializer))]
+	public class BgNodeOutputExprDef
+	{
+		/// <summary>
+		/// The producing node
+		/// </summary>
+		public BgNodeExpressionDef ProducingNode { get; }
+
+		/// <summary>
+		/// The output index. -1 means all inputs and outputs for the node.
+		/// </summary>
+		public int OutputIndex { get; }
+
+		/// <summary>
+		/// Constructor
+		/// </summary>
+		/// <param name="producingNode"></param>
+		/// <param name="outputIndex"></param>
+		public BgNodeOutputExprDef(BgNodeExpressionDef producingNode, int outputIndex)
+		{
+			ProducingNode = producingNode;
+			OutputIndex = outputIndex;
+		}
+
+		/// <summary>
+		/// Flattens this expression to a list of outputs
+		/// </summary>
+		/// <returns></returns>
+		public IEnumerable<BgNodeOutput> Flatten()
+		{
+			if (OutputIndex == -1)
+			{
+				return ProducingNode.InputDependencies.SelectMany(x => x.Outputs).Concat(ProducingNode.Outputs).ToArray();
+			}
+			else
+			{
+				return new[] { ProducingNode.Outputs[OutputIndex] };
+			}
+		}
+	}
+
+	class BgNodeOutputExprDefSerializer : BgObjectSerializer<BgNodeOutputExprDef>
+	{
+		/// <inheritdoc/>
+		public override BgNodeOutputExprDef Deserialize(BgObjectDef<BgNodeOutputExprDef> obj)
+		{
+			return new BgNodeOutputExprDef(obj.Get(x => x.ProducingNode, null!), obj.Get(x => x.OutputIndex, -1));
+		}
+	}
+
 	/// <summary>
 	/// Defines a node, a container for tasks and the smallest unit of execution that can be run as part of a build graph.
 	/// </summary>
@@ -58,10 +111,15 @@ namespace EpicGames.BuildGraph
 		/// </summary>
 		public string Name { get; }
 
+		/// <summary>
+		/// Thunk to execute this node.
+		/// </summary>
+		public BgThunkDef? Thunk { get; }
+
 		/// <summary>
 		/// Array of inputs which this node requires to run
 		/// </summary>
-		public IReadOnlyList<BgNodeOutput> Inputs { get; }
+		public List<BgNodeOutput> Inputs { get; } = new List<BgNodeOutput>();
 
 		/// <summary>
 		/// Array of outputs produced by this node
@@ -71,17 +129,17 @@ namespace EpicGames.BuildGraph
 		/// <summary>
 		/// Nodes which this node has input dependencies on
 		/// </summary>
-		public IReadOnlyList<BgNodeDef> InputDependencies { get; set; }
+		public List<BgNodeDef> InputDependencies { get; } = new List<BgNodeDef>();
 
 		/// <summary>
 		/// Nodes which this node needs to run after
 		/// </summary>
-		public IReadOnlyList<BgNodeDef> OrderDependencies { get; set; }
+		public List<BgNodeDef> OrderDependencies { get; } = new List<BgNodeDef>();
 
 		/// <summary>
 		/// Tokens which must be acquired for this node to run
 		/// </summary>
-		public IReadOnlyList<FileReference> RequiredTokens { get; }
+		public List<FileReference> RequiredTokens { get; } = new List<FileReference>();
 
 		/// <summary>
 		/// List of email addresses to notify if this node fails.
@@ -113,6 +171,20 @@ namespace EpicGames.BuildGraph
 		/// </summary>
 		public List<BgDiagnosticDef> Diagnostics { get; } = new List<BgDiagnosticDef>();
 
+		/// <summary>
+		/// Constructor
+		/// </summary>
+		public BgNodeDef(string name, BgThunkDef? thunk, IReadOnlyList<string> outputNames)
+		{
+			Name = name;
+			Thunk = thunk;
+
+			List<BgNodeOutput> allOutputs = new List<BgNodeOutput>();
+			allOutputs.Add(new BgNodeOutput(this, "#" + Name));
+			allOutputs.AddRange(outputNames.Where(x => String.Compare(x, Name, StringComparison.InvariantCultureIgnoreCase) != 0).Select(x => new BgNodeOutput(this, x)));
+			Outputs = allOutputs.ToArray();
+		}
+
 		/// <summary>
 		/// Constructor
 		/// </summary>
@@ -123,18 +195,13 @@ namespace EpicGames.BuildGraph
 		/// <param name="orderDependencies">Nodes which this node needs to run after. Should include all input dependencies.</param>
 		/// <param name="requiredTokens">Optional tokens which must be required for this node to run</param>
 		public BgNodeDef(string name, IReadOnlyList<BgNodeOutput> inputs, IReadOnlyList<string> outputNames, IReadOnlyList<BgNodeDef> inputDependencies, IReadOnlyList<BgNodeDef> orderDependencies, IReadOnlyList<FileReference> requiredTokens)
+			: this(name, null, outputNames)
 		{
 			Name = name;
-			Inputs = inputs;
-
-			List<BgNodeOutput> allOutputs = new List<BgNodeOutput>();
-			allOutputs.Add(new BgNodeOutput(this, "#" + Name));
-			allOutputs.AddRange(outputNames.Where(x => String.Compare(x, Name, StringComparison.InvariantCultureIgnoreCase) != 0).Select(x => new BgNodeOutput(this, x)));
-			Outputs = allOutputs.ToArray();
-
-			InputDependencies = inputDependencies;
-			OrderDependencies = orderDependencies;
-			RequiredTokens = requiredTokens;
+			Inputs.AddRange(inputs);
+			InputDependencies.AddRange(inputDependencies);
+			OrderDependencies.AddRange(orderDependencies);
+			RequiredTokens.AddRange(requiredTokens);
 		}
 
 		/// <summary>
@@ -183,7 +250,8 @@ namespace EpicGames.BuildGraph
 	/// <summary>
 	/// Node constructed from a bytecode expression
 	/// </summary>
-	public class BgExpressionNodeDef : BgNodeDef
+	[BgObject(typeof(BgNodeExpressionDefSerializer))]
+	public class BgNodeExpressionDef : BgNodeDef
 	{
 		/// <summary>
 		/// Agent declaring this node
@@ -191,30 +259,57 @@ namespace EpicGames.BuildGraph
 		public BgAgentDef Agent { get; }
 
 		/// <summary>
-		/// Offset in the bytecode of this node. Used to lookup method calling information.
+		/// Labels to add this node to
 		/// </summary>
-		public BgMethod Method { get; }
+		public List<BgLabelDef> Labels { get; } = new List<BgLabelDef>();
 
 		/// <summary>
-		/// Expression arguments to the method
+		/// Input expressions
 		/// </summary>
-		public IReadOnlyList<object> Arguments { get; }
+		public List<BgNodeOutputExprDef> InputExprs { get; } = new List<BgNodeOutputExprDef>();
 
 		/// <summary>
-		/// Labels to add this node to
+		/// Number of outputs from this node
 		/// </summary>
-		public IReadOnlyList<BgLabelDef> Labels { get; }
+		public int OutputCount { get; }
 
 		/// <summary>
 		/// Constructor
 		/// </summary>
-		public BgExpressionNodeDef(BgAgentDef agent, string name, BgMethod method, IReadOnlyList<object> arguments, IReadOnlyList<BgNodeOutput> inputs, IReadOnlyList<string> outputNames, IReadOnlyList<BgNodeDef> inputDependencies, IReadOnlyList<BgNodeDef> orderDependencies, IReadOnlyList<FileReference> requiredTokens, IReadOnlyList<BgLabelDef> labels)
-			: base(name, inputs, outputNames, inputDependencies, orderDependencies, requiredTokens)
+		public BgNodeExpressionDef(BgAgentDef agent, string name, BgThunkDef thunk, int outputCount)
+			: base(name, thunk, GetOutputNames(name, outputCount))
 		{
 			Agent = agent;
-			Method = method;
-			Arguments = arguments;
-			Labels = labels;
+			OutputCount = outputCount;
+		}
+
+		static string[] GetOutputNames(string name, int numOutputs)
+		{
+			return Enumerable.Range(0, numOutputs).Select(x => BgNode.GetDefaultTagName(name, x)).ToArray();
+		}
+	}
+
+	class BgNodeExpressionDefSerializer : BgObjectSerializer<BgNodeExpressionDef>
+	{
+		/// <inheritdoc/>
+		public override BgNodeExpressionDef Deserialize(BgObjectDef<BgNodeExpressionDef> obj)
+		{
+			BgNodeExpressionDef node = new BgNodeExpressionDef(obj.Get(x => x.Agent, null!), obj.Get(x => x.Name, ""), obj.Get(x => x.Thunk!, null!), obj.Get(x => x.OutputCount, 0));
+			obj.CopyTo(node);
+
+			node.Inputs.AddRange(node.InputExprs.SelectMany(x => x.Flatten()));
+
+			HashSet<BgNodeDef> inputDependencies = new HashSet<BgNodeDef>();
+			foreach (BgNodeOutput input in node.Inputs)
+			{
+				inputDependencies.Add(input.ProducingNode);
+				inputDependencies.UnionWith(input.ProducingNode.InputDependencies);
+			}
+
+			inputDependencies.ExceptWith(node.InputDependencies);
+			node.InputDependencies.AddRange(inputDependencies);
+
+			return node;
 		}
 	}
 }
diff --git a/Engine/Source/Programs/Shared/EpicGames.BuildGraph/BgObjectDef.cs b/Engine/Source/Programs/Shared/EpicGames.BuildGraph/BgObjectDef.cs
new file mode 100644
--- /dev/null
+++ b/Engine/Source/Programs/Shared/EpicGames.BuildGraph/BgObjectDef.cs
@@ -0,0 +1,361 @@
+// Copyright Epic Games, Inc. All Rights Reserved.
+
+using System;
+using System.Collections;
+using System.Collections.Concurrent;
+using System.Collections.Generic;
+using System.Collections.Immutable;
+using System.Linq;
+using System.Linq.Expressions;
+using System.Reflection;
+using System.Text;
+
+namespace EpicGames.BuildGraph
+{
+	/// <summary>
+	/// Value of an object
+	/// </summary>
+	public class BgObjectDef
+	{
+		/// <summary>
+		/// Deserialized object value
+		/// </summary>
+		private object? _value;
+
+		/// <summary>
+		/// Properties for the object
+		/// </summary>
+		public ImmutableDictionary<string, object?> Properties { get; }
+
+		/// <summary>
+		/// Constructor
+		/// </summary>
+		public BgObjectDef()
+		{
+			Properties = ImmutableDictionary<string, object?>.Empty.WithComparers(StringComparer.Ordinal);
+		}
+
+		/// <summary>
+		/// Constructor
+		/// </summary>
+		public BgObjectDef(ImmutableDictionary<string, object?> properties)
+		{
+			Properties = properties;
+		}
+
+		/// <summary>
+		/// Sets a property value
+		/// </summary>
+		/// <param name="name">Name of the property</param>
+		/// <param name="value">New value for the property</param>
+		/// <returns>Object definition with the new properties</returns>
+		public BgObjectDef Set(string name, object? value)
+		{
+			return new BgObjectDef(Properties.SetItem(name, value));
+		}
+
+		/// <summary>
+		/// Gets a property value
+		/// </summary>
+		/// <param name="name">Name of the property</param>
+		/// <param name="defaultValue">Default value to return if the property is not set</param>
+		/// <returns></returns>
+		public object? Get(string name, object? defaultValue)
+		{
+			object? value;
+			if(!Properties.TryGetValue(name, out value))
+			{
+				value = defaultValue;
+			}
+			return value;
+		}
+
+		/// <summary>
+		/// Creates a strongly typed object instance
+		/// </summary>
+		/// <typeparam name="T">Type of the object</typeparam>
+		/// <returns></returns>
+		public BgObjectDef<T> WithType<T>() => new BgObjectDef<T>(Properties);
+
+		/// <summary>
+		/// Cache of serializer instances
+		/// </summary>
+		static readonly ConcurrentDictionary<Type, BgObjectSerializer> s_typeToSerializer = new ConcurrentDictionary<Type, BgObjectSerializer>();
+
+		/// <summary>
+		/// Deserialize an object of the given type, using the default serializer
+		/// </summary>
+		/// <param name="type">Type of the object to create</param>
+		/// <returns>Deserialized instance of the type</returns>
+		public object Deserialize(Type type)
+		{
+			if (_value == null)
+			{
+				BgObjectSerializer? serializer;
+				if (!s_typeToSerializer.TryGetValue(type, out serializer))
+				{
+					Type serializerType = type.GetCustomAttribute<BgObjectAttribute>()?.SerializerType ?? typeof(BgDefaultObjectSerializer<>).MakeGenericType(type);
+					serializer = (BgObjectSerializer)Activator.CreateInstance(serializerType)!;
+					serializer = s_typeToSerializer.GetOrAdd(type, serializer);
+				}
+				_value = serializer.Deserialize(this);
+			}
+			return _value;
+		}
+
+		/// <summary>
+		/// Deserialize an object of the given type
+		/// </summary>
+		/// <typeparam name="T">Type to deserialize</typeparam>
+		/// <returns>Deserialized instance of the type</returns>
+		public T Deserialize<T>() => (T)Deserialize(typeof(T));
+	}
+
+	/// <summary>
+	/// Strongly typed instance of <see cref="BgObjectDef"/>
+	/// </summary>
+	/// <typeparam name="T"></typeparam>
+	public class BgObjectDef<T> : BgObjectDef
+	{
+		class PropertySetter
+		{
+			protected readonly PropertyInfo _property;
+
+			public PropertySetter(PropertyInfo property) => _property = property;
+
+			public virtual void SetValue(object instance, object? value) => _property.SetValue(instance, ConvertValue(value, _property.PropertyType));
+		}
+
+		class CollectionPropertySetter<TElement> : PropertySetter
+		{
+			public CollectionPropertySetter(PropertyInfo property)
+				: base(property)
+			{
+			}
+
+			public override void SetValue(object instance, object? value)
+			{
+				ICollection<TElement>? list = (ICollection<TElement>?)_property.GetValue(instance)!;
+				if (list == null)
+				{
+					list = (ICollection<TElement>)Activator.CreateInstance(_property.PropertyType)!;
+					_property.SetValue(instance, list);
+				}
+
+				IEnumerable<object> elements = (IEnumerable<object>)value!;
+				foreach (object element in elements)
+				{
+					if (element is IEnumerable<TElement> range)
+					{
+						foreach (object? rangeElement in range)
+						{
+							list.Add((TElement)ConvertValue(rangeElement, typeof(TElement))!);
+						}
+					}
+					else
+					{
+						list.Add((TElement)ConvertValue(element, typeof(TElement))!);
+					}
+				}
+			}
+		}
+
+		static readonly IReadOnlyDictionary<string, PropertySetter> s_nameToSetter = GetPropertyMap();
+
+		/// <summary>
+		/// Constructor
+		/// </summary>
+		public BgObjectDef()
+		{
+		}
+
+		/// <summary>
+		/// Constructor
+		/// </summary>
+		public BgObjectDef(ImmutableDictionary<string, object?> properties)
+			: base(properties)
+		{
+		}
+
+		/// <summary>
+		/// Gets a property value
+		/// </summary>
+		/// <typeparam name="TValue"></typeparam>
+		/// <param name="property"></param>
+		/// <param name="defaultValue"></param>
+		/// <returns></returns>
+		public TValue Get<TValue>(Expression<Func<T, TValue>> property, TValue defaultValue)
+		{
+			MemberExpression member = ((MemberExpression)property.Body);
+			PropertyInfo propertyInfo = (PropertyInfo)member.Member;
+			string name = propertyInfo.GetCustomAttribute<BgPropertyAttribute>()?.Name ?? propertyInfo.Name;
+			return (TValue)ConvertValue(Get(name, defaultValue), typeof(TValue))!;
+		}
+
+		/// <summary>
+		/// Copies properties from the object value to an instance
+		/// </summary>
+		/// <param name="instance"></param>
+		public void CopyTo(object instance)
+		{
+			foreach ((string name, object? value) in Properties)
+			{
+				if (s_nameToSetter.TryGetValue(name, out PropertySetter? propertySetter))
+				{
+					propertySetter.SetValue(instance, value);
+				}
+			}
+		}
+
+		/// <summary>
+		/// Convert a value to a particular type
+		/// </summary>
+		/// <param name="value">Value to convert</param>
+		/// <param name="type">Target type</param>
+		/// <returns>The converted value</returns>
+		static object? ConvertValue(object? value, Type type)
+		{
+			if (value == null)
+			{
+				return null;
+			}
+			else if (value.GetType() == type)
+			{
+				return value;
+			}
+			else if (type.IsClass && type != typeof(string) && type != typeof(BgNodeOutput))
+			{
+				return ((BgObjectDef)value).Deserialize(type);
+			}
+			else
+			{
+				return value;
+			}
+		}
+
+		/// <summary>
+		/// Creates a map from name to property info
+		/// </summary>
+		/// <returns></returns>
+		static IReadOnlyDictionary<string, PropertySetter> GetPropertyMap()
+		{
+			Dictionary<string, PropertySetter> nameToProperty = new Dictionary<string, PropertySetter>(StringComparer.Ordinal);
+
+			PropertyInfo[] properties = typeof(T).GetProperties(BindingFlags.Public | BindingFlags.Instance);
+			foreach (PropertyInfo property in properties)
+			{
+				PropertySetter? setter = CreateSetter(property);
+				if (setter != null)
+				{
+					BgPropertyAttribute? attribute = property.GetCustomAttribute<BgPropertyAttribute>();
+					string name = attribute?.Name ?? property.Name;
+					nameToProperty[name] = setter;
+				}
+			}
+
+			return nameToProperty;
+		}
+
+		static PropertySetter? CreateSetter(PropertyInfo property)
+		{
+			Type? collectionInterface = property.PropertyType.GetInterfaces().FirstOrDefault(x => x.IsGenericType && x.GetGenericTypeDefinition() == typeof(ICollection<>));
+			if (collectionInterface != null)
+			{
+				Type setterType = typeof(CollectionPropertySetter<>).MakeGenericType(typeof(T), collectionInterface.GetGenericArguments()[0]);
+				return (PropertySetter)Activator.CreateInstance(setterType, property)!;
+			}
+
+			if (property.SetMethod != null)
+			{
+				return new PropertySetter(property);
+			}
+
+			return null;
+		}
+	}
+
+	/// <summary>
+	/// Attribute marking that a property should be serialized to BuildGraph
+	/// </summary>
+	[AttributeUsage(AttributeTargets.Property)]
+	public sealed class BgPropertyAttribute : Attribute
+	{
+		/// <summary>
+		/// Name of the property. If unspecified, the property name will be used.
+		/// </summary>
+		public string? Name { get; }
+
+		/// <summary>
+		/// Constructor
+		/// </summary>
+		public BgPropertyAttribute(string? name = null)
+		{
+			Name = name;
+		}
+	}
+
+	/// <summary>
+	/// Attribute marking that a property should be serialized to BuildGraph
+	/// </summary>
+	public sealed class BgObjectAttribute : Attribute
+	{
+		/// <summary>
+		/// The serailizer to use for this object
+		/// </summary>
+		public Type SerializerType { get; }
+
+		/// <summary>
+		/// Constructor
+		/// </summary>
+		/// <param name="serializerType">The serializer type</param>
+		public BgObjectAttribute(Type serializerType)
+		{
+			SerializerType = serializerType;
+		}
+	}
+
+	/// <summary>
+	/// Base class for deserializing objects from untyped <see cref="BgObjectDef"/> instances
+	/// </summary>
+	public abstract class BgObjectSerializer
+	{
+		/// <summary>
+		/// Deserialize an object from the given property bag
+		/// </summary>
+		/// <param name="obj">Properties for the object</param>
+		/// <returns>New object instance</returns>
+		public abstract object Deserialize(BgObjectDef obj);
+	}
+
+	/// <summary>
+	/// Strongly typed base class for deserializing objects
+	/// </summary>
+	/// <typeparam name="T">The object type</typeparam>
+	public abstract class BgObjectSerializer<T> : BgObjectSerializer
+	{
+		/// <inheritdoc/>
+		public sealed override object Deserialize(BgObjectDef obj) => Deserialize(obj.WithType<T>())!;
+
+		/// <summary>
+		/// Typed deserialization method
+		/// </summary>
+		/// <param name="obj">Properties for the object</param>
+		/// <returns>Object instance</returns>
+		public abstract T Deserialize(BgObjectDef<T> obj);
+	}
+
+	/// <summary>
+	/// Default serializer for objects with a default constructor
+	/// </summary>
+	/// <typeparam name="T">Object type</typeparam>
+	public class BgDefaultObjectSerializer<T> : BgObjectSerializer<T> where T : new()
+	{
+		/// <inheritdoc/>
+		public override T Deserialize(BgObjectDef<T> obj)
+		{
+			T result = new T();
+			obj.CopyTo(result);
+			return result;
+		}
+	}
+}
diff --git a/Engine/Source/Programs/Shared/EpicGames.BuildGraph/BgOpcode.cs b/Engine/Source/Programs/Shared/EpicGames.BuildGraph/BgOpcode.cs
--- a/Engine/Source/Programs/Shared/EpicGames.BuildGraph/BgOpcode.cs
+++ b/Engine/Source/Programs/Shared/EpicGames.BuildGraph/BgOpcode.cs
@@ -260,26 +260,12 @@ namespace EpicGames.BuildGraph
 		ListLazy = 0x4c,
 
 		/// <summary>
-		/// An option providing a list of strings (see <see cref="BgListOptionSpec"/>
+		/// An option providing a list of strings (see <see cref="BgListOption"/>
 		/// </summary>
 		ListOption = 0x4d,
 
 		#endregion
 
-		#region Fileset opcodes
-
-		/// <summary>
-		/// Fileset constructed from all the inputs and outputs of a node
-		/// </summary>
-		FileSetFromNode = 0x50,
-
-		/// <summary>
-		/// Fileset for an individual numbered output from a node
-		/// </summary>
-		FileSetFromNodeOutput = 0x51,
-
-		#endregion
-
 		#region Object opcodes
 
 		/// <summary>
@@ -335,34 +321,10 @@ namespace EpicGames.BuildGraph
 		/// </summary>
 		Null = 0xe2,
 
-		#endregion
-
-		#region Graph structures
-
-		/// <summary>
-		/// Defines a node (see <see cref="BgNode"/>)
-		/// </summary>
-		Node = 0xf0,
-
-		/// <summary>
-		/// Defines an agent to execute work on (see <see cref="BgAgent"/>)
-		/// </summary>
-		Agent = 0xf1,
-		
-		/// <summary>
-		/// Defines a label (see <see cref="BgLabel"/>)
-		/// </summary>
-		Label = 0xf2,
-
-		/// <summary>
-		/// Defines an aggregate (see <see cref="BgAggregate"/>)
-		/// </summary>
-		Aggregate = 0xf3,
-
 		/// <summary>
-		/// Defines a graph (see <see cref="BgGraph"/>)
+		/// A native method invocation, represented as an index into the method table
 		/// </summary>
-		Graph = 0xf4,
+		Thunk = 0xe3,
 
 		#endregion
 	}
diff --git a/Engine/Source/Programs/Shared/EpicGames.BuildGraph/BgOption.cs b/Engine/Source/Programs/Shared/EpicGames.BuildGraph/BgOption.cs
deleted file mode 100644
diff --git a/Engine/Source/Programs/Shared/EpicGames.BuildGraph/BgOptionDef.cs b/Engine/Source/Programs/Shared/EpicGames.BuildGraph/BgOptionDef.cs
new file mode 100644
--- /dev/null
+++ b/Engine/Source/Programs/Shared/EpicGames.BuildGraph/BgOptionDef.cs
@@ -0,0 +1,287 @@
+// Copyright Epic Games, Inc. All Rights Reserved.
+
+using System;
+using System.Collections.Generic;
+using System.Text;
+
+namespace EpicGames.BuildGraph
+{
+	/// <summary>
+	/// Represents a graph option. These are expanded during preprocessing, but are retained in order to display help messages.
+	/// </summary>
+	public abstract class BgOptionDef
+	{
+		/// <summary>
+		/// Name of this option
+		/// </summary>
+		public string Name { get; }
+
+		/// <summary>
+		/// Label for the option in the dashboard
+		/// </summary>
+		public string? Label { get; set; }
+
+		/// <summary>
+		/// Description for this option
+		/// </summary>
+		public string? Description { get; set; }
+
+		/// <summary>
+		/// Constructor
+		/// </summary>
+		/// <param name="name">The name of this option</param>
+		public BgOptionDef(string name)
+		{
+			Name = name;
+		}
+
+		/// <summary>
+		/// Returns the default argument value
+		/// </summary>
+		/// <returns>Default argument value</returns>
+		public abstract string? GetDefaultArgument();
+
+		/// <summary>
+		/// Returns a name of this option for debugging
+		/// </summary>
+		/// <returns>Name of the option</returns>
+		public override string ToString()
+		{
+			return Name;
+		}
+	}
+
+	/// <summary>
+	/// Definition of a boolean option
+	/// </summary>
+	[BgObject(typeof(BgBoolOptionSerializer))]
+	public class BgBoolOptionDef : BgOptionDef
+	{
+		/// <summary>
+		/// Default value for the option
+		/// </summary>
+		public bool DefaultValue { get; set; }
+
+		/// <summary>
+		/// Constructor
+		/// </summary>
+		/// <param name="name">The name of this option</param>
+		public BgBoolOptionDef(string name)
+			: base(name)
+		{
+		}
+
+		/// <inheritdoc/>
+		public override string? GetDefaultArgument() => DefaultValue.ToString();
+	}
+
+	/// <summary>
+	/// Serializer for bool options
+	/// </summary>
+	class BgBoolOptionSerializer : BgObjectSerializer<BgBoolOptionDef>
+	{
+		/// <inheritdoc/>
+		public override BgBoolOptionDef Deserialize(BgObjectDef<BgBoolOptionDef> obj)
+		{
+			BgBoolOptionDef option = new BgBoolOptionDef(obj.Get(x => x.Name, null!));
+			obj.CopyTo(option);
+			return option;
+		}
+	}
+
+	/// <summary>
+	/// Definition of an integer option
+	/// </summary>
+	[BgObject(typeof(BgIntOptionSerializer))]
+	public class BgIntOptionDef : BgOptionDef
+	{
+		/// <summary>
+		/// Default value for the option
+		/// </summary>
+		public int DefaultValue { get; set; }
+
+		/// <summary>
+		/// Minimum allowed value
+		/// </summary>
+		public int? MinValue { get; set; }
+
+		/// <summary>
+		/// Maximum allowed value
+		/// </summary>
+		public int? MaxValue { get; set; }
+
+		/// <summary>
+		/// Constructor
+		/// </summary>
+		/// <param name="name">The name of this option</param>
+		public BgIntOptionDef(string name)
+			: base(name)
+		{
+		}
+
+		/// <inheritdoc/>
+		public override string? GetDefaultArgument() => DefaultValue.ToString();
+	}
+
+	/// <summary>
+	/// Serializer for int options
+	/// </summary>
+	class BgIntOptionSerializer : BgObjectSerializer<BgIntOptionDef>
+	{
+		/// <inheritdoc/>
+		public override BgIntOptionDef Deserialize(BgObjectDef<BgIntOptionDef> obj)
+		{
+			BgIntOptionDef option = new BgIntOptionDef(obj.Get(x => x.Name, null!));
+			obj.CopyTo(option);
+			return option;
+		}
+	}
+
+	/// <summary>
+	/// Style for a string option
+	/// </summary>
+	public enum BgStringOptionStyle
+	{
+		/// <summary>
+		/// Free-form text entry
+		/// </summary>
+		Text,
+
+		/// <summary>
+		/// List of options
+		/// </summary>
+		DropList,
+	}
+
+	/// <summary>
+	/// Definition of a string option
+	/// </summary>
+	[BgObject(typeof(BgStringOptionSerializer))]
+	public class BgStringOptionDef : BgOptionDef
+	{
+		/// <summary>
+		/// Default value for the option
+		/// </summary>
+		public string DefaultValue { get; set; } = String.Empty;
+
+		/// <summary>
+		/// Style for this option
+		/// </summary>
+		public BgStringOptionStyle Style { get; }
+
+		/// <summary>
+		/// Regex for validating values for the option
+		/// </summary>
+		public string? Pattern { get; set; }
+
+		/// <summary>
+		/// Message to display if validation fails
+		/// </summary>
+		public string? PatternFailed { get; set; }
+
+		/// <summary>
+		/// List of values to choose from
+		/// </summary>
+		public List<string>? Values { get; set; }
+
+		/// <summary>
+		/// Matching list of descriptions for each value
+		/// </summary>
+		public List<string>? ValueDescriptions { get; set; }
+
+		/// <summary>
+		/// Constructor
+		/// </summary>
+		/// <param name="name">Name of the option</param>
+		public BgStringOptionDef(string name)
+			: base(name)
+		{
+		}
+
+		/// <inheritdoc/>
+		public override string? GetDefaultArgument() => DefaultValue;
+	}
+
+	/// <summary>
+	/// Serializer for string options
+	/// </summary>
+	class BgStringOptionSerializer : BgObjectSerializer<BgStringOptionDef>
+	{
+		/// <inheritdoc/>
+		public override BgStringOptionDef Deserialize(BgObjectDef<BgStringOptionDef> obj)
+		{
+			BgStringOptionDef option = new BgStringOptionDef(obj.Get(x => x.Name, null!));
+			obj.CopyTo(option);
+			return option;
+		}
+	}
+
+	/// <summary>
+	/// Style for a list option
+	/// </summary>
+	public enum BgListOptionStyle
+	{
+		/// <summary>
+		/// List of checkboxes
+		/// </summary>
+		CheckList = 0,
+
+		/// <summary>
+		/// Tag picker
+		/// </summary>
+		TagPicker = 1,
+	}
+
+	/// <summary>
+	/// A list option definition
+	/// </summary>
+	[BgObject(typeof(BgListOptionSerializer))]
+	public class BgListOptionDef : BgOptionDef
+	{
+		/// <summary>
+		/// Style for this list box
+		/// </summary>
+		public BgListOptionStyle Style { get; }
+
+		/// <summary>
+		/// Default value for the option
+		/// </summary>
+		public string? DefaultValue { get; set; }
+
+		/// <summary>
+		/// List of values to choose from
+		/// </summary>
+		public List<string>? Values { get; set; }
+
+		/// <summary>
+		/// Matching list of descriptions for each value
+		/// </summary>
+		public List<string>? ValueDescriptions { get; set; }
+
+		/// <summary>
+		/// Constructor
+		/// </summary>
+		/// <param name="name">Name of the option</param>
+		public BgListOptionDef(string name)
+			: base(name)
+		{
+		}
+
+		/// <inheritdoc/>
+		public override string? GetDefaultArgument() => DefaultValue;
+	}
+
+	/// <summary>
+	/// Serializer for string options
+	/// </summary>
+	class BgListOptionSerializer : BgObjectSerializer<BgListOptionDef>
+	{
+		/// <inheritdoc/>
+		public override BgListOptionDef Deserialize(BgObjectDef<BgListOptionDef> obj)
+		{
+			BgListOptionDef option = new BgListOptionDef(obj.Get(x => x.Name, null!));
+			obj.CopyTo(option);
+			return option;
+		}
+	}
+}
diff --git a/Engine/Source/Programs/Shared/EpicGames.BuildGraph/BgThunkDef.cs b/Engine/Source/Programs/Shared/EpicGames.BuildGraph/BgThunkDef.cs
new file mode 100644
--- /dev/null
+++ b/Engine/Source/Programs/Shared/EpicGames.BuildGraph/BgThunkDef.cs
@@ -0,0 +1,62 @@
+// Copyright Epic Games, Inc. All Rights Reserved.
+
+using System;
+using System.Collections.Generic;
+using System.Linq;
+using System.Reflection;
+using System.Text;
+
+namespace EpicGames.BuildGraph
+{
+	/// <summary>
+	/// Information about the method bound to execute a node
+	/// </summary>
+	public class BgThunkDef
+	{
+		/// <summary>
+		/// Method to call
+		/// </summary>
+		public MethodInfo Method { get; }
+
+		/// <summary>
+		/// Arguments to the method
+		/// </summary>
+		public IReadOnlyList<object?> Arguments { get; }
+
+		/// <summary>
+		/// Constructor
+		/// </summary>
+		public BgThunkDef(MethodInfo method, IReadOnlyList<object?> arguments)
+		{
+			Method = method;
+			Arguments = arguments;
+		}
+	}
+
+	/// <summary>
+	/// Outputs from a thunk
+	/// </summary>
+	public class BgThunkOutputDef
+	{
+		/// <summary>
+		/// The thunk definition
+		/// </summary>
+		public BgThunkDef Thunk { get; }
+
+		/// <summary>
+		/// Output index from the thunk
+		/// </summary>
+		public int Index { get; }
+
+		/// <summary>
+		/// Constructor
+		/// </summary>
+		/// <param name="thunk"></param>
+		/// <param name="index"></param>
+		public BgThunkOutputDef(BgThunkDef thunk, int index)
+		{
+			Thunk = thunk;
+			Index = index;
+		}
+	}
+}
diff --git a/Engine/Source/Programs/Shared/EpicGames.BuildGraph/Expressions/BgAgent.cs b/Engine/Source/Programs/Shared/EpicGames.BuildGraph/Expressions/BgAgent.cs
--- a/Engine/Source/Programs/Shared/EpicGames.BuildGraph/Expressions/BgAgent.cs
+++ b/Engine/Source/Programs/Shared/EpicGames.BuildGraph/Expressions/BgAgent.cs
@@ -29,7 +29,7 @@ namespace EpicGames.BuildGraph.Expressions
 		/// Constructor
 		/// </summary>
 		public BgAgent(BgString name, BgString type)
-			: this(name, BgList<BgString>.Create(type))
+			: this(name, BgList.Create(type))
 		{
 		}
 
@@ -46,9 +46,10 @@ namespace EpicGames.BuildGraph.Expressions
 		/// <inheritdoc/>
 		public override void Write(BgBytecodeWriter writer)
 		{
-			writer.WriteOpcode(BgOpcode.Agent);
-			writer.WriteExpr(Name);
-			writer.WriteExpr(Types);
+			BgObject<BgAgentDef> obj = BgObject<BgAgentDef>.Empty;
+			obj = obj.Set(x => x.Name, Name);
+			obj = obj.Set(x => x.PossibleTypes, Types);
+			writer.WriteExpr(obj);
 		}
 
 		/// <inheritdoc/>
diff --git a/Engine/Source/Programs/Shared/EpicGames.BuildGraph/Expressions/BgAggregate.cs b/Engine/Source/Programs/Shared/EpicGames.BuildGraph/Expressions/BgAggregate.cs
--- a/Engine/Source/Programs/Shared/EpicGames.BuildGraph/Expressions/BgAggregate.cs
+++ b/Engine/Source/Programs/Shared/EpicGames.BuildGraph/Expressions/BgAggregate.cs
@@ -17,7 +17,7 @@ namespace EpicGames.BuildGraph.Expressions
 		/// <summary>
 		/// Outputs required for the aggregate
 		/// </summary>
-		public BgList<BgFileSet> Requires { get; }
+		public BgList<BgNode> Requires { get; }
 
 		/// <summary>
 		/// Label to apply to this aggregate
@@ -27,23 +27,23 @@ namespace EpicGames.BuildGraph.Expressions
 		/// <summary>
 		/// Constructor.
 		/// </summary>
-		public BgAggregate(BgString name, params BgFileSet[] requires)
-			: this(name, BgList<BgFileSet>.Create(requires))
+		public BgAggregate(BgString name, params BgNode[] requires)
+			: this(name, BgList.Create(requires))
 		{
 		}
 
 		/// <summary>
 		/// Constructor.
 		/// </summary>
-		public BgAggregate(BgString name, params BgList<BgFileSet>[] requires)
-			: this(name, BgList<BgFileSet>.Concat(requires))
+		public BgAggregate(BgString name, params BgList<BgNode>[] requires)
+			: this(name, BgList.Concat(requires))
 		{
 		}
 
 		/// <summary>
 		/// Constructor.
 		/// </summary>
-		public BgAggregate(BgString name, BgList<BgFileSet> requires, string label)
+		public BgAggregate(BgString name, BgList<BgNode> requires, string label)
 			: this(name, requires, new BgLabel(label))
 		{
 		}
@@ -51,7 +51,7 @@ namespace EpicGames.BuildGraph.Expressions
 		/// <summary>
 		/// Constructor.
 		/// </summary>
-		public BgAggregate(BgString name, BgList<BgFileSet> requires, BgLabel? label = null)
+		public BgAggregate(BgString name, BgList<BgNode> requires, BgLabel? label = null)
 			: base(BgExprFlags.None)
 		{
 			Name = name;
@@ -62,10 +62,10 @@ namespace EpicGames.BuildGraph.Expressions
 		/// <inheritdoc/>
 		public override void Write(BgBytecodeWriter writer)
 		{
-			writer.WriteOpcode(BgOpcode.Aggregate);
-			writer.WriteExpr(Name);
-			writer.WriteExpr(Requires);
-			writer.WriteExpr(Label ?? BgExpr.Null);
+			BgObject<BgAggregateDef> obj = BgObject<BgAggregateDef>.Empty;
+			obj = obj.Set(x => x.Name, Name);
+			obj = obj.Set(x => x.RequiredNodes, Requires);
+			writer.WriteExpr(obj);
 		}
 
 		/// <inheritdoc/>
diff --git a/Engine/Source/Programs/Shared/EpicGames.BuildGraph/Expressions/BgBool.cs b/Engine/Source/Programs/Shared/EpicGames.BuildGraph/Expressions/BgBool.cs
--- a/Engine/Source/Programs/Shared/EpicGames.BuildGraph/Expressions/BgBool.cs
+++ b/Engine/Source/Programs/Shared/EpicGames.BuildGraph/Expressions/BgBool.cs
@@ -214,10 +214,26 @@ namespace EpicGames.BuildGraph.Expressions
 		public override void Write(BgBytecodeWriter writer)
 		{
 			writer.WriteOpcode(BgOpcode.BoolOption);
-			writer.WriteExpr(Name);
-			writer.WriteExpr(Label ?? BgString.Empty);
-			writer.WriteExpr(Description ?? BgString.Empty);
-			writer.WriteExpr(DefaultValue ?? BgBool.False);
+			writer.WriteExpr(CreateOptionsObject());
+		}
+
+		BgObject<BgBoolOptionDef> CreateOptionsObject()
+		{
+			BgObject<BgBoolOptionDef> option = BgObject<BgBoolOptionDef>.Empty;
+			option = option.Set(x => x.Name, Name);
+			if (!(Label is null))
+			{
+				option = option.Set(x => x.Label, Label);
+			}
+			if (!(Description is null))
+			{
+				option = option.Set(x => x.Description, Description);
+			}
+			if (!(DefaultValue is null))
+			{
+				option = option.Set(x => x.DefaultValue, DefaultValue);
+			}
+			return option;
 		}
 	}
 }
diff --git a/Engine/Source/Programs/Shared/EpicGames.BuildGraph/Expressions/BgEnum.cs b/Engine/Source/Programs/Shared/EpicGames.BuildGraph/Expressions/BgEnum.cs
--- a/Engine/Source/Programs/Shared/EpicGames.BuildGraph/Expressions/BgEnum.cs
+++ b/Engine/Source/Programs/Shared/EpicGames.BuildGraph/Expressions/BgEnum.cs
@@ -17,12 +17,12 @@ namespace EpicGames.BuildGraph.Expressions
 		/// <summary>
 		/// Names of this enum
 		/// </summary>
-		public static BgList<BgString> Names { get; } = BgList<BgString>.Create(Enum.GetNames(typeof(TEnum)).Select(x => (BgString)x));
+		public static BgList<BgString> Names { get; } = BgList.Create(Enum.GetNames(typeof(TEnum)).Select(x => (BgString)x));
 
 		/// <summary>
 		/// Values of this enum
 		/// </summary>
-		public static BgList<BgInt> Values { get; } = BgList<BgInt>.Create(((TEnum[])Enum.GetValues(typeof(TEnum))).Select(x => (BgInt)(int)(object)x));
+		public static BgList<BgInt> Values { get; } = BgList.Create(((TEnum[])Enum.GetValues(typeof(TEnum))).Select(x => (BgInt)(int)(object)x));
 
 		/// <summary>
 		/// Constructor
diff --git a/Engine/Source/Programs/Shared/EpicGames.BuildGraph/Expressions/BgFileSet.cs b/Engine/Source/Programs/Shared/EpicGames.BuildGraph/Expressions/BgFileSet.cs
--- a/Engine/Source/Programs/Shared/EpicGames.BuildGraph/Expressions/BgFileSet.cs
+++ b/Engine/Source/Programs/Shared/EpicGames.BuildGraph/Expressions/BgFileSet.cs
@@ -2,6 +2,7 @@
 
 using System;
 using System.Collections.Generic;
+using System.Linq;
 using EpicGames.Core;
 
 namespace EpicGames.BuildGraph.Expressions
@@ -45,7 +46,12 @@ namespace EpicGames.BuildGraph.Expressions
 	class BgFileSetType : BgType<BgFileSet>
 	{
 		/// <inheritdoc/>
-		public override BgFileSet Constant(object value) => new BgFileSetInputExpr((BgNodeOutput[])value);
+		public override BgFileSet Constant(object value)
+		{
+			BgObjectDef obj = (BgObjectDef)value;
+			BgNodeOutput[] outputs = obj.Deserialize<BgNodeOutputExprDef>().Flatten().ToArray();
+			return new BgFileSetInputExpr(outputs);
+		}
 
 		/// <inheritdoc/>
 		public override BgFileSet Wrap(BgExpr expr) => new BgFileSetWrappedExpr(expr);
@@ -108,8 +114,9 @@ namespace EpicGames.BuildGraph.Expressions
 
 		public override void Write(BgBytecodeWriter writer)
 		{
-			writer.WriteOpcode(BgOpcode.FileSetFromNode);
-			writer.WriteExpr(Node);
+			BgObject<BgNodeOutputExprDef> obj = BgObject<BgNodeOutputExprDef>.Empty;
+			obj = obj.Set(x => x.ProducingNode, Node);
+			writer.WriteExpr(obj);
 		}
 	}
 
@@ -127,9 +134,10 @@ namespace EpicGames.BuildGraph.Expressions
 
 		public override void Write(BgBytecodeWriter writer)
 		{
-			writer.WriteOpcode(BgOpcode.FileSetFromNodeOutput);
-			writer.WriteExpr(Node);
-			writer.WriteUnsignedInteger(OutputIndex);
+			BgObject<BgNodeOutputExprDef> obj = BgObject<BgNodeOutputExprDef>.Empty;
+			obj = obj.Set(x => x.ProducingNode, Node);
+			obj = obj.Set(x => x.OutputIndex, (BgInt)OutputIndex);
+			writer.WriteExpr(obj);
 		}
 	}
 
diff --git a/Engine/Source/Programs/Shared/EpicGames.BuildGraph/Expressions/BgGraph.cs b/Engine/Source/Programs/Shared/EpicGames.BuildGraph/Expressions/BgGraph.cs
--- a/Engine/Source/Programs/Shared/EpicGames.BuildGraph/Expressions/BgGraph.cs
+++ b/Engine/Source/Programs/Shared/EpicGames.BuildGraph/Expressions/BgGraph.cs
@@ -66,9 +66,10 @@ namespace EpicGames.BuildGraph.Expressions
 		/// <inheritdoc/>
 		public override void Write(BgBytecodeWriter writer)
 		{
-			writer.WriteOpcode(BgOpcode.Graph);
-			writer.WriteExpr(Nodes);
-			writer.WriteExpr(Aggregates);
+			BgObject<BgGraphExpressionDef> obj = BgObject<BgGraphExpressionDef>.Empty;
+			obj = obj.Set(x => x.Nodes, Nodes);
+			obj = obj.Set(x => x.Aggregates, Aggregates);
+			writer.WriteExpr(obj);
 		}
 
 		/// <inheritdoc/>
diff --git a/Engine/Source/Programs/Shared/EpicGames.BuildGraph/Expressions/BgInt.cs b/Engine/Source/Programs/Shared/EpicGames.BuildGraph/Expressions/BgInt.cs
--- a/Engine/Source/Programs/Shared/EpicGames.BuildGraph/Expressions/BgInt.cs
+++ b/Engine/Source/Programs/Shared/EpicGames.BuildGraph/Expressions/BgInt.cs
@@ -280,12 +280,34 @@ namespace EpicGames.BuildGraph.Expressions
 		public override void Write(BgBytecodeWriter writer)
 		{
 			writer.WriteOpcode(BgOpcode.IntOption);
-			writer.WriteExpr(Name);
-			writer.WriteExpr(Label ?? BgString.Empty);
-			writer.WriteExpr(Description ?? BgString.Empty);
-			writer.WriteExpr(DefaultValue ?? (BgInt)0);
-			writer.WriteExpr(MinValue ?? (BgInt)(-1));
-			writer.WriteExpr(MaxValue ?? (BgInt)(-1));
+			writer.WriteExpr(CreateOptionsObject());
+		}
+
+		BgObject<BgIntOptionDef> CreateOptionsObject()
+		{
+			BgObject<BgIntOptionDef> option = BgObject<BgIntOptionDef>.Empty;
+			option = option.Set(x => x.Name, Name);
+			if (!(Label is null))
+			{
+				option = option.Set(x => x.Label, Label);
+			}
+			if (!(Description is null))
+			{
+				option = option.Set(x => x.Description, Description);
+			}
+			if (!(DefaultValue is null))
+			{
+				option = option.Set(x => x.DefaultValue, DefaultValue);
+			}
+			if (!(MinValue is null))
+			{
+				option = option.Set(x => x.MinValue, MinValue);
+			}
+			if (!(MaxValue is null))
+			{
+				option = option.Set(x => x.MaxValue, MaxValue);
+			}
+			return option;
 		}
 	}
 }
diff --git a/Engine/Source/Programs/Shared/EpicGames.BuildGraph/Expressions/BgLabel.cs b/Engine/Source/Programs/Shared/EpicGames.BuildGraph/Expressions/BgLabel.cs
--- a/Engine/Source/Programs/Shared/EpicGames.BuildGraph/Expressions/BgLabel.cs
+++ b/Engine/Source/Programs/Shared/EpicGames.BuildGraph/Expressions/BgLabel.cs
@@ -33,12 +33,12 @@ namespace EpicGames.BuildGraph.Expressions
 		/// <summary>
 		/// Which change to show the badge for
 		/// </summary>
-		public BgString? Change { get; }
+		public BgEnum<BgLabelChange>? Change { get; }
 
 		/// <summary>
 		/// Constructor
 		/// </summary>
-		public BgLabel(BgString? name = null, BgString? category = null, BgString? ugsBadge = null, BgString? ugsProject = null, BgString? change = null)
+		public BgLabel(BgString? name = null, BgString? category = null, BgString? ugsBadge = null, BgString? ugsProject = null, BgEnum<BgLabelChange>? change = null)
 			: base(BgExprFlags.ForceFragment)
 		{
 			DashboardName = name;
@@ -51,12 +51,28 @@ namespace EpicGames.BuildGraph.Expressions
 		/// <inheritdoc/>
 		public override void Write(BgBytecodeWriter writer)
 		{
-			writer.WriteOpcode(BgOpcode.Label);
-			writer.WriteExpr(DashboardName ?? BgString.Empty);
-			writer.WriteExpr(DashboardCategory ?? BgString.Empty);
-			writer.WriteExpr(UgsBadge ?? BgString.Empty);
-			writer.WriteExpr(UgsProject ?? BgString.Empty);
-			writer.WriteExpr(Change ?? BgString.Empty);
+			BgObject<BgLabelDef> obj = BgObject<BgLabelDef>.Empty;
+			if (!(DashboardName is null))
+			{
+				obj = obj.Set(x => x.DashboardName, DashboardName);
+			}
+			if (!(DashboardCategory is null))
+			{
+				obj = obj.Set(x => x.DashboardCategory, DashboardCategory);
+			}
+			if (!(UgsBadge is null))
+			{
+				obj = obj.Set(x => x.UgsBadge, UgsBadge);
+			}
+			if (!(UgsProject is null))
+			{
+				obj = obj.Set(x => x.UgsProject, UgsProject);
+			}
+			if (!(Change is null))
+			{
+				obj = obj.Set(x => x.Change, Change);
+			}
+			writer.WriteExpr(obj);
 		}
 
 		/// <inheritdoc/>
diff --git a/Engine/Source/Programs/Shared/EpicGames.BuildGraph/Expressions/BgList.cs b/Engine/Source/Programs/Shared/EpicGames.BuildGraph/Expressions/BgList.cs
--- a/Engine/Source/Programs/Shared/EpicGames.BuildGraph/Expressions/BgList.cs
+++ b/Engine/Source/Programs/Shared/EpicGames.BuildGraph/Expressions/BgList.cs
@@ -6,6 +6,62 @@ using System.Linq;
 
 namespace EpicGames.BuildGraph.Expressions
 {
+	/// <summary>
+	/// Utility methods for lists
+	/// </summary>
+	public static class BgList
+	{
+		/// <summary>
+		/// Gets an empty list of the given type
+		/// </summary>
+		/// <typeparam name="T"></typeparam>
+		/// <returns></returns>
+		public static BgList<T> Empty<T>() where T : BgExpr => BgList<T>.Empty;
+
+		/// <summary>
+		/// Crates a list from an array of values
+		/// </summary>
+		/// <param name="items">Sequence to construct from</param>
+		/// <returns></returns>
+		public static BgList<BgString> Create(IEnumerable<string> items) => BgList<BgString>.Create(items.Select(x => (BgString)x));
+
+		/// <summary>
+		/// Crates a list from an array of values
+		/// </summary>
+		/// <param name="items">Sequence to construct from</param>
+		/// <returns></returns>
+		public static BgList<BgString> Create(params string[] items) => BgList<BgString>.Create(items.Select(x => (BgString)x));
+
+		/// <summary>
+		/// Crates a list from an array of values
+		/// </summary>
+		/// <param name="items">Sequence to construct from</param>
+		/// <returns></returns>
+		public static BgList<T> Create<T>(IEnumerable<T> items) where T : BgExpr => BgList<T>.Create(items);
+
+		/// <summary>
+		/// Crates a list from an array of values
+		/// </summary>
+		/// <param name="items">Sequence to construct from</param>
+		/// <returns></returns>
+		public static BgList<T> Create<T>(params T[] items) where T : BgExpr => BgList<T>.Create(items);
+
+		/// <summary>
+		/// Concatenates two lists together
+		/// </summary>
+		/// <param name="lhs"></param>
+		/// <param name="rhs"></param>
+		/// <returns></returns>
+		public static BgList<T> Concat<T>(BgList<T> lhs, BgList<T> rhs) where T : BgExpr => BgList<T>.Concat(lhs, rhs);
+
+		/// <summary>
+		/// Concatenates two lists together
+		/// </summary>
+		/// <param name="others"></param>
+		/// <returns></returns>
+		public static BgList<T> Concat<T>(params BgList<T>[] others) where T : BgExpr => BgList<T>.Concat(others);
+	}
+
 	/// <summary>
 	/// Abstract base class for expressions returning an immutable list of values
 	/// </summary>
@@ -171,7 +227,19 @@ namespace EpicGames.BuildGraph.Expressions
 	class BgListType<T> : BgType<BgList<T>> where T : BgExpr
 	{
 		/// <inheritdoc/>
-		public override BgList<T> Constant(object value) => new BgListConstantExpr<T>((List<object>)value);
+		public override BgList<T> Constant(object value)
+		{
+			IEnumerable<object> elements = (IEnumerable<object>)value;
+
+			List<T> items = new List<T>();
+			foreach (object element in elements)
+			{
+				items.Add(BgType.Constant<T>(element));
+			}
+
+			return new BgListConstantExpr<T>(items);
+		}
+//		=> new BgListConstantExpr<T>(((IEnumerable<object>)value).Select(x => BgType.Constant<T>(x)).ToList());
 
 		/// <inheritdoc/>
 		public override BgList<T> Wrap(BgExpr expr) => new BgListWrappedExpr<T>(expr);
@@ -194,12 +262,12 @@ namespace EpicGames.BuildGraph.Expressions
 
 	class BgListConstantExpr<T> : BgList<T> where T : BgExpr
 	{
-		public List<object> Value { get; }
+		public IReadOnlyList<T> Value { get; }
 
-		public BgListConstantExpr(List<object> value)
+		public BgListConstantExpr(IEnumerable<T> value)
 			: base(BgExprFlags.NotInterned)
 		{
-			Value = value;
+			Value = value.ToArray();
 		}
 
 		public override void Write(BgBytecodeWriter writer)
@@ -445,26 +513,10 @@ namespace EpicGames.BuildGraph.Expressions
 
 	#endregion
 
-	/// <summary>
-	/// Style for a list option
-	/// </summary>
-	public enum BgListOptionStyle
-	{
-		/// <summary>
-		/// List of checkboxes
-		/// </summary>
-		CheckList = 0,
-
-		/// <summary>
-		/// Tag picker
-		/// </summary>
-		TagPicker = 1,
-	}
-
 	/// <summary>
 	/// A list option expression
 	/// </summary>
-	public class BgListOptionSpec : BgList<BgString>
+	public class BgListOption : BgList<BgString>
 	{
 		/// <summary>
 		/// Name of the option
@@ -484,23 +536,13 @@ namespace EpicGames.BuildGraph.Expressions
 		/// <summary>
 		/// Style for this list box
 		/// </summary>
-		public BgListOptionStyle Style { get; }
+		public BgEnum<BgListOptionStyle> Style { get; }
 
 		/// <summary>
 		/// Default value for the option
 		/// </summary>
 		public BgString? DefaultValue { get; set; }
 
-		/// <summary>
-		/// Regex for validating values for the option
-		/// </summary>
-		public BgString? Pattern { get; set; }
-
-		/// <summary>
-		/// Message to display if validation fails
-		/// </summary>
-		public BgString? PatternFailed { get; set; }
-
 		/// <summary>
 		/// List of values to choose from
 		/// </summary>
@@ -514,7 +556,7 @@ namespace EpicGames.BuildGraph.Expressions
 		/// <summary>
 		/// Constructor
 		/// </summary>
-		public BgListOptionSpec(string name, BgString? label = null, BgString? description = null, BgString? defaultValue = null, BgListOptionStyle style = BgListOptionStyle.CheckList, BgString? pattern = null, BgString? patternFailed = null, BgList<BgString>? values = null, BgList<BgString>? valueDescriptions = null)
+		public BgListOption(string name, BgString? description = null, BgString? defaultValue = null, BgListOptionStyle style = BgListOptionStyle.CheckList, BgList<BgString>? values = null, BgList<BgString>? valueDescriptions = null, BgString? label = null)
 			: base(BgExprFlags.None)
 		{
 			Name = name;
@@ -522,8 +564,6 @@ namespace EpicGames.BuildGraph.Expressions
 			Description = description;
 			DefaultValue = defaultValue;
 			Style = style;
-			Pattern = pattern;
-			PatternFailed = patternFailed;
 			Values = values;
 			ValueDescriptions = valueDescriptions;
 		}
@@ -532,15 +572,38 @@ namespace EpicGames.BuildGraph.Expressions
 		public override void Write(BgBytecodeWriter writer)
 		{
 			writer.WriteOpcode(BgOpcode.ListOption);
-			writer.WriteExpr(Name);
-			writer.WriteExpr(Label ?? BgString.Empty);
-			writer.WriteExpr(Description ?? BgString.Empty);
-			writer.WriteUnsignedInteger((int)Style);
-			writer.WriteExpr(DefaultValue ?? BgString.Empty);
-			writer.WriteExpr(Pattern ?? BgString.Empty);
-			writer.WriteExpr(PatternFailed ?? BgString.Empty);
-			writer.WriteExpr(Values ?? BgList<BgString>.Empty);
-			writer.WriteExpr(ValueDescriptions ?? BgList<BgString>.Empty);
+			writer.WriteExpr(CreateOptionsObject());
+		}
+
+		BgObject<BgListOptionDef> CreateOptionsObject()
+		{
+			BgObject<BgListOptionDef> option = BgObject<BgListOptionDef>.Empty;
+			option = option.Set(x => x.Name, Name);
+			if (!(Label is null))
+			{
+				option = option.Set(x => x.Label, Label);
+			}
+			if (!(Description is null))
+			{
+				option = option.Set(x => x.Description, Description);
+			}
+			if (!(DefaultValue is null))
+			{
+				option = option.Set(x => x.DefaultValue, DefaultValue);
+			}
+			if (!(Style is null))
+			{
+				option = option.Set(x => x.Style, Style);
+			}
+			if (!(Values is null))
+			{
+				option = option.Set(x => x.Values, Values);
+			}
+			if (!(ValueDescriptions is null))
+			{
+				option = option.Set(x => x.ValueDescriptions, ValueDescriptions);
+			}
+			return option;
 		}
 	}
 }
diff --git a/Engine/Source/Programs/Shared/EpicGames.BuildGraph/Expressions/BgNode.cs b/Engine/Source/Programs/Shared/EpicGames.BuildGraph/Expressions/BgNode.cs
--- a/Engine/Source/Programs/Shared/EpicGames.BuildGraph/Expressions/BgNode.cs
+++ b/Engine/Source/Programs/Shared/EpicGames.BuildGraph/Expressions/BgNode.cs
@@ -46,31 +46,6 @@ namespace EpicGames.BuildGraph.Expressions
 		}
 	}
 
-	/// <summary>
-	/// Information about the method bound to execute a node
-	/// </summary>
-	public class BgMethod
-	{
-		/// <summary>
-		/// Method to call
-		/// </summary>
-		public MethodInfo Method { get; }
-
-		/// <summary>
-		/// Arguments to the method
-		/// </summary>
-		public IReadOnlyList<object?> Arguments { get; }
-
-		/// <summary>
-		/// Constructor
-		/// </summary>
-		public BgMethod(MethodInfo method, IReadOnlyList<object?> arguments)
-		{
-			Method = method;
-			Arguments = arguments;
-		}
-	}
-
 	/// <summary>
 	/// Specification for a node to execute
 	/// </summary>
@@ -82,25 +57,20 @@ namespace EpicGames.BuildGraph.Expressions
 		public BgString Name { get; }
 
 		/// <summary>
-		/// Method to execute
+		/// Thunk to native code to execute the node
 		/// </summary>
-		protected BgMethod Method { get; }
+		public BgThunk Thunk { get; }
 
 		/// <summary>
-		/// Expression arguments from the bound function.
+		/// Number of outputs from this node
 		/// </summary>
-		protected BgExpr[] ExpressionArguments { get; }
+		public int OutputCount { get; protected set; }
 
 		/// <summary>
 		/// The default output of this node. Includes all other outputs. 
 		/// </summary>
 		public BgFileSet DefaultOutput { get; }
 
-		/// <summary>
-		/// Explicitly tagged outputs from this node
-		/// </summary>
-		public IReadOnlyList<BgExpr> TaggedOutputs { get; protected set; } = Array.Empty<BgExpr>();
-
 		/// <summary>
 		/// Agent for the node to be run on
 		/// </summary>
@@ -109,285 +79,123 @@ namespace EpicGames.BuildGraph.Expressions
 		/// <summary>
 		/// Tokens for inputs of this node
 		/// </summary>
-		public BgList<BgFileSet> Inputs { get; }
+		public BgList<BgFileSet> Inputs { get; private set; } = BgList.Empty<BgFileSet>();
 
 		/// <summary>
 		/// Weak dependency on outputs that must be generated for the node to run, without making those dependencies inputs.
 		/// </summary>
-		public BgList<BgFileSet> Fences { get; }
+		public BgList<BgNode> Fences { get; private set; } = BgList.Empty<BgNode>();
 
 		/// <summary>
 		/// Whether this node should start running as soon as its dependencies are ready, even if other nodes in the same agent are not.
 		/// </summary>
-		public BgBool RunEarly { get; } = BgBool.False;
+		public BgBool RunEarly { get; private set; } = BgBool.False;
 
 		/// <summary>
 		/// Labels that this node contributes to
 		/// </summary>
-		public BgList<BgLabel> Labels { get; }
+		public BgList<BgLabel> Labels { get; private set; } = BgList.Empty<BgLabel>();
 
 		/// <summary>
 		/// Constructor
 		/// </summary>
-		public BgNode(BgString name, BgMethod method, BgAgent agent, BgList<BgFileSet> inputs, BgList<BgFileSet> fences, BgBool runEarly, BgList<BgLabel> labels)
+		public BgNode(BgThunk thunk, BgAgent agent)
 			: base(BgExprFlags.ForceFragment)
 		{
-			Name = name;
-
-			ParameterInfo[] parameters = method.Method.GetParameters();
-			ExpressionArguments = new BgExpr[parameters.Length];
-
-			object?[] constantArguments = new object?[method.Arguments.Count];
-			for (int idx = 0; idx < parameters.Length; idx++)
-			{
-				if (typeof(BgExpr).IsAssignableFrom(parameters[idx].ParameterType))
-				{
-					ExpressionArguments[idx] = (BgExpr)method.Arguments[idx]!;
-					constantArguments[idx] = null;
-				}
-				else
-				{
-					ExpressionArguments[idx] = BgBool.False;
-					constantArguments[idx] = method.Arguments[idx];
-				}
-			}
-			Method = new BgMethod(method.Method, constantArguments);
-
+			Name = GetDefaultNodeName(thunk);
+			Thunk = thunk;
 			Agent = agent;
 			DefaultOutput = new BgFileSetFromNodeOutputExpr(this, 0);
-			Inputs = inputs;
-			Fences = fences;
-			RunEarly = runEarly;
-			Labels = labels;
-		}
-
-		int written = 0;
-
-		/// <inheritdoc/>
-		public override void Write(BgBytecodeWriter writer)
-		{
-			writer.WriteOpcode(BgOpcode.Node);
-
-			writer.WriteExpr(Name ?? BgString.Empty);
-			writer.WriteExpr(Agent);
-			writer.WriteMethod(Method);
-
-			writer.WriteUnsignedInteger((uint)ExpressionArguments.Length);
-			foreach (BgExpr argument in ExpressionArguments)
-			{
-				writer.WriteExpr(argument);
-			}
-
-			writer.WriteExpr(Inputs);
-			writer.WriteExpr(Fences);
-			writer.WriteUnsignedInteger(TaggedOutputs.Count);
-			writer.WriteExpr(RunEarly);
-			writer.WriteExpr(Labels);
-			written++;
 		}
 
 		/// <summary>
-		/// Gets the default tag name for the numbered output index
-		/// </summary>
-		/// <param name="name">Name of the node</param>
-		/// <param name="index">Index of the output. Index zero is the default, others are explicit.</param>
-		/// <returns></returns>
-		internal static string GetDefaultTagName(string name, int index)
-		{
-			return $"#{name}${index}";
-		}
-
-		/// <summary>
-		/// Implicit conversion to a fileset
+		/// Copy constructor
 		/// </summary>
 		/// <param name="node"></param>
-		public static implicit operator BgFileSet(BgNode node)
+		public BgNode(BgNode node)
+			: base(BgExprFlags.ForceFragment)
 		{
-			return new BgFileSetFromNodeExpr(node);
+			Name = node.Name;
+			Thunk = node.Thunk;
+			Agent = node.Agent;
+			DefaultOutput = node.DefaultOutput;
+			Inputs = node.Inputs;
+			Fences = node.Fences;
+			RunEarly = node.RunEarly;
+			Labels = node.Labels;
 		}
 
-		/// <summary>
-		/// Implicit conversion to a fileset
-		/// </summary>
-		/// <param name="node"></param>
-		public static implicit operator BgList<BgFileSet>(BgNode node)
+		/// <inheritdoc/>
+		public override void Write(BgBytecodeWriter writer)
 		{
-			return (BgFileSet)node;
+			BgObject<BgNodeExpressionDef> obj = BgObject<BgNodeExpressionDef>.Empty;
+			obj = obj.Set(x => x.Name, Name);
+			obj = obj.Set(x => x.Agent, Agent);
+			obj = obj.Set(x => x.Thunk, Thunk);
+			obj = obj.Set(x => x.OutputCount, (BgInt)OutputCount);
+			obj = obj.Set(x => x.InputExprs, Inputs);
+			obj = obj.Set(x => x.OrderDependencies, Fences);
+			obj = obj.Set(x => x.RunEarly, RunEarly);
+			obj = obj.Set(x => x.Labels, Labels);
+			writer.WriteExpr(obj);
 		}
 
-		/// <inheritdoc/>
-		public override BgString ToBgString() => Name ?? BgString.Empty;
-	}
-
-	/// <summary>
-	/// Nodespec with a typed return value
-	/// </summary>
-	/// <typeparam name="T"></typeparam>
-	public class BgNode<T> : BgNode
-	{
-		/// <summary>
-		/// Output from this node
-		/// </summary>
-		public T Output { get; }
-
 		/// <summary>
-		/// Constructor
+		/// Creates a copy of this node and updates the given parameters
 		/// </summary>
-		public BgNode(BgString name, BgMethod handler, BgAgent agent, BgList<BgFileSet> inputs, BgList<BgFileSet> fences, BgBool runEarly, BgList<BgLabel> labels)
-			: base(name, handler, agent, inputs, fences, runEarly, labels)
-		{
-			Output = CreateOutput();
-		}
-
-		T CreateOutput()
-		{
-			Type type = typeof(T);
-			if (IsValueTuple(type))
-			{
-				BgExpr[] outputs = CreateOutputExprs(type.GetGenericArguments());
-				TaggedOutputs = outputs;
-				return (T)Activator.CreateInstance(type, outputs)!;
-			}
-			else
-			{
-				BgExpr[] outputs = CreateOutputExprs(new[] { type });
-				TaggedOutputs = outputs;
-				return (T)(object)outputs[0];
-			}
-		}
-
-		BgExpr[] CreateOutputExprs(Type[] types)
+		/// <param name="inputs"></param>
+		/// <param name="fences"></param>
+		/// <param name="runEarly"></param>
+		/// <param name="labels"></param>
+		/// <returns></returns>
+		internal BgNode Modify(BgList<BgFileSet>? inputs = null, BgList<BgNode>? fences = null, BgBool? runEarly = null, BgList<BgLabel>? labels = null)
 		{
-			BgExpr[] outputs = new BgExpr[types.Length];
-			for (int idx = 0; idx < types.Length; idx++)
+			BgNode node = (BgNode)MemberwiseClone();
+			if (inputs is object)
 			{
-				outputs[idx] = CreateOutputExpr(types[idx], idx);
+				node.Inputs = inputs;
 			}
-			return outputs;
-		}
-
-		BgExpr CreateOutputExpr(Type type, int index)
-		{
-			if (type == typeof(BgFileSet))
+			if (fences is object)
 			{
-				return new BgFileSetFromNodeOutputExpr(this, index);
+				node.Fences = fences;
 			}
-			else
+			if (runEarly is object)
 			{
-				throw new NotImplementedException();
+				node.RunEarly = runEarly;
 			}
-		}
-
-		internal static bool IsValueTuple(Type returnType)
-		{
-			if (returnType.IsGenericType)
+			if (labels is object)
 			{
-				Type genericType = returnType.GetGenericTypeDefinition();
-				if (genericType.FullName != null && genericType.FullName.StartsWith("System.ValueTuple`", StringComparison.Ordinal))
-				{
-					return true;
-				}
+				node.Labels = labels;
 			}
-			return false;
+			return node;
 		}
-	}
-
-	/// <summary>
-	/// Class that allows constructing a nodespec
-	/// </summary>
-	public class BgNodeSpecBuilder
-	{
-		/// <inheritdoc cref="BgNode.Name"/>
-		public BgString Name { get; }
-
-		/// <inheritdoc cref="BgNode.Method"/>
-		protected BgMethod Handler { get; }
-
-		/// <inheritdoc cref="BgNode.Agent"/>
-		public BgAgent Agent { get; }
-
-		/// <inheritdoc cref="BgNode.Inputs"/>
-		public BgList<BgFileSet> Inputs { get; set; } = BgList<BgFileSet>.Empty;
-
-		/// <inheritdoc cref="BgNode.Fences"/>
-		public BgList<BgFileSet> Fences { get; set; } = BgList<BgFileSet>.Empty;
-
-		/// <inheritdoc cref="BgNode.RunEarly"/>
-		public BgBool RunEarly { get; set; } = BgBool.False;
-
-		/// <inheritdoc cref="BgNode.Labels"/>
-		public BgList<BgLabel> Labels { get; set; } = BgList<BgLabel>.Empty;
 
 		/// <summary>
-		/// Constructor
+		/// Gets the default tag name for the numbered output index
 		/// </summary>
-		public BgNodeSpecBuilder(BgString? name, MethodCallExpression call, BgAgent agent)
-		{
-			try
-			{
-				object?[] arguments = CreateArgumentExprs(call);
-
-				Name = name ?? GetDefaultNodeName(call.Method, arguments);
-				Handler = new BgMethod(call.Method, arguments);
-				Agent = agent;
-			}
-			catch (Exception ex)
-			{
-				ExceptionUtils.AddContext(ex, $"while calling method {call.Method}");
-				throw;
-			}
-		}
-
-		static object?[] CreateArgumentExprs(MethodCallExpression call)
-		{
-			object?[] args = new object?[call.Arguments.Count];
-			for (int idx = 0; idx < call.Arguments.Count; idx++)
-			{
-				Expression argumentExpr = call.Arguments[idx];
-				if (argumentExpr is ParameterExpression parameterExpr)
-				{
-					if (parameterExpr.Type != typeof(BgContext))
-					{
-						throw new BgNodeException($"Unable to determine type of parameter '{parameterExpr.Name}'");
-					}
-				}
-				else
-				{
-					Delegate compiled = Expression.Lambda(argumentExpr).Compile();
-					args[idx] = compiled.DynamicInvoke();
-				}
-			}
-			return args;
-		}
-
-		internal static bool IsValueTuple(Type returnType)
+		/// <param name="name">Name of the node</param>
+		/// <param name="index">Index of the output. Index zero is the default, others are explicit.</param>
+		/// <returns></returns>
+		internal static string GetDefaultTagName(string name, int index)
 		{
-			if (returnType.IsGenericType)
-			{
-				Type genericType = returnType.GetGenericTypeDefinition();
-				if (genericType.FullName != null && genericType.FullName.StartsWith("System.ValueTuple`", StringComparison.Ordinal))
-				{
-					return true;
-				}
-			}
-			return false;
+			return $"#{name}${index}";
 		}
 
-		static BgString GetDefaultNodeName(MethodInfo methodInfo, object?[] args)
+		static BgString GetDefaultNodeName(BgThunk thunk)
 		{
 			// Check if it's got an attribute override for the node name
-			BgNodeNameAttribute? nameAttr = methodInfo.GetCustomAttribute<BgNodeNameAttribute>();
+			BgNodeNameAttribute? nameAttr = thunk.Method.GetCustomAttribute<BgNodeNameAttribute>();
 			if (nameAttr != null)
 			{
-				return GetNodeNameFromTemplate(nameAttr.Template, methodInfo.GetParameters(), args);
+				return GetNodeNameFromTemplate(nameAttr.Template, thunk.Method.GetParameters(), thunk.Arguments);
 			}
 			else
 			{
-				return GetNodeNameFromMethodName(methodInfo.Name);
+				return GetNodeNameFromMethodName(thunk.Method.Name);
 			}
 		}
 
-		static BgString GetNodeNameFromTemplate(string template, ParameterInfo[] parameters, object?[] args)
+		static BgString GetNodeNameFromTemplate(string template, ParameterInfo[] parameters, IReadOnlyList<object?> arguments)
 		{
 			// Create a list of lazily computed string fragments which comprise the evaluated name
 			List<BgString> fragments = new List<BgString>();
@@ -421,17 +229,14 @@ namespace EpicGames.BuildGraph.Expressions
 							throw new BgNodeException($"Unable to find parameter named {paramName} in {template}");
 						}
 
-						object? arg = args[paramIdx];
-						if (arg != null)
+						object? arg = arguments[paramIdx];
+						if (typeof(BgExpr).IsAssignableFrom(parameters[paramIdx].ParameterType))
 						{
-							if (arg is BgExpr expr)
-							{
-								fragments.Add(expr.ToBgString());
-							}
-							else
-							{
-								fragments.Add(arg.ToString() ?? String.Empty);
-							}
+							fragments.Add(((BgExpr)arg!).ToBgString());
+						}
+						else if (arg != null)
+						{
+							fragments.Add(arg.ToString() ?? String.Empty);
 						}
 
 						lastIdx = nextIdx = endIdx + 1;
@@ -497,36 +302,106 @@ namespace EpicGames.BuildGraph.Expressions
 		}
 
 		/// <summary>
-		/// Construct a nodespec from the current parameters
+		/// Implicit conversion to a fileset
 		/// </summary>
-		/// <returns>New nodespec</returns>
-		public BgNode Construct()
+		/// <param name="node"></param>
+		public static implicit operator BgFileSet(BgNode node)
 		{
-			return new BgNode(Name, Handler, Agent, Inputs, Fences, RunEarly, Labels);
+			return new BgFileSetFromNodeExpr(node);
+		}
+
+		/// <summary>
+		/// Implicit conversion to a fileset
+		/// </summary>
+		/// <param name="node"></param>
+		public static implicit operator BgList<BgFileSet>(BgNode node)
+		{
+			return (BgFileSet)node;
 		}
+
+		/// <inheritdoc/>
+		public override BgString ToBgString() => Name ?? BgString.Empty;
 	}
 
 	/// <summary>
-	/// Specification for a node to execute that returns one or more tagged outputs
+	/// Nodespec with a typed return value
 	/// </summary>
-	/// <typeparam name="T">Return type from the node. Must consist of latent Bg* types (eg. BgToken)</typeparam>
-	public class BgNodeSpecBuilder<T> : BgNodeSpecBuilder
+	/// <typeparam name="T"></typeparam>
+	public class BgNode<T> : BgNode
 	{
+		/// <summary>
+		/// Output from this node
+		/// </summary>
+		public T Output { get; }
+
 		/// <summary>
 		/// Constructor
 		/// </summary>
-		internal BgNodeSpecBuilder(BgString? name, MethodCallExpression call, BgAgent agent)
-			: base(name, call, agent)
+		public BgNode(BgThunk<T> thunk, BgAgent agent)
+			: base(thunk, agent)
 		{
+			Output = CreateOutput();
 		}
 
 		/// <summary>
-		/// Construct a nodespec from the current parameters
+		/// Copy constructor
 		/// </summary>
-		/// <returns>New nodespec</returns>
-		public new BgNode<T> Construct()
+		public BgNode(BgNode<T> other)
+			: base(other)
+		{
+			Output = other.Output;
+		}
+
+		T CreateOutput()
+		{
+			Type type = typeof(T);
+			if (IsValueTuple(type))
+			{
+				BgExpr[] outputs = CreateOutputExprs(type.GetGenericArguments());
+				OutputCount = outputs.Length;
+				return (T)Activator.CreateInstance(type, outputs)!;
+			}
+			else
+			{
+				BgExpr[] outputs = CreateOutputExprs(new[] { type });
+				OutputCount = outputs.Length;
+				return (T)(object)outputs[0];
+			}
+		}
+
+		BgExpr[] CreateOutputExprs(Type[] types)
 		{
-			return new BgNode<T>(Name, Handler, Agent, Inputs, Fences, RunEarly, Labels);
+			BgExpr[] outputs = new BgExpr[types.Length];
+			for (int idx = 0; idx < types.Length; idx++)
+			{
+				outputs[idx] = CreateOutputExpr(types[idx], idx);
+			}
+			return outputs;
+		}
+
+		BgExpr CreateOutputExpr(Type type, int index)
+		{
+			if (type == typeof(BgFileSet))
+			{
+				return new BgFileSetFromNodeOutputExpr(this, index);
+			}
+			else
+			{
+				throw new NotImplementedException();
+			}
+		}
+
+		internal static bool IsValueTuple(Type returnType)
+		{
+			if (returnType.IsGenericType)
+			{
+				Type genericType = returnType.GetGenericTypeDefinition();
+				if (genericType.FullName != null && genericType.FullName.StartsWith("System.ValueTuple`", StringComparison.Ordinal))
+				{
+					return true;
+				}
+			}
+			return false;
 		}
 	}
 
@@ -541,10 +416,10 @@ namespace EpicGames.BuildGraph.Expressions
 		/// <param name="agent">Agent to run the node</param>
 		/// <param name="func">Function to execute</param>
 		/// <returns>Node builder</returns>
-		public static BgNodeSpecBuilder AddNode(this BgAgent agent, Expression<Func<BgContext, Task>> func)
+		public static BgNode AddNode(this BgAgent agent, Expression<Func<BgContext, Task>> func)
 		{
-			MethodCallExpression call = (MethodCallExpression)func.Body;
-			return new BgNodeSpecBuilder(null, call, agent);
+			BgThunk thunk = BgThunk.Create(func);
+			return new BgNode(thunk, agent);
 		}
 
 		/// <summary>
@@ -553,59 +428,55 @@ namespace EpicGames.BuildGraph.Expressions
 		/// <param name="agent">Agent to run the node</param>
 		/// <param name="func">Function to execute</param>
 		/// <returns>Node builder</returns>
-		public static BgNodeSpecBuilder<T> AddNode<T>(this BgAgent agent, Expression<Func<BgContext, Task<T>>> func)
+		public static BgNode<T> AddNode<T>(this BgAgent agent, Expression<Func<BgContext, Task<T>>> func)
 		{
-			MethodCallExpression call = (MethodCallExpression)func.Body;
-			return new BgNodeSpecBuilder<T>(null, call, agent);
+			BgThunk<T> thunk = BgThunk.Create(func);
+			return new BgNode<T>(thunk, agent);
 		}
 
 		/// <summary>
 		/// Add dependencies onto other nodes or outputs. Outputs from the given tokens will be copied to the current machine before execution of the node.
 		/// </summary>
-		/// <param name="builder">The node builder</param>
+		/// <param name="node">The node to modify</param>
 		/// <param name="inputs">Files to add as inputs</param>
 		/// <returns>The current node spec, to allow chaining calls</returns>
-		public static T Requires<T>(this T builder, params BgNode[] inputs) where T : BgNodeSpecBuilder
+		public static T Requires<T>(this T node, params BgNode[] inputs) where T : BgNode
 		{
-			builder.Inputs = builder.Inputs.Add(inputs.Select(x => (BgFileSet)x));
-			return builder;
+			return (T)node.Modify(inputs: node.Inputs.Add(inputs.Select(x => (BgFileSet)x)));
 		}
 
 		/// <summary>
 		/// Add dependencies onto other nodes or outputs. Outputs from the given tokens will be copied to the current machine before execution of the node.
 		/// </summary>
-		/// <param name="builder">The node builder</param>
+		/// <param name="node">The node to modify</param>
 		/// <param name="inputs">Files to add as inputs</param>
 		/// <returns>The current node spec, to allow chaining calls</returns>
-		public static T Requires<T>(this T builder, params BgFileSet[] inputs) where T : BgNodeSpecBuilder
+		public static T Requires<T>(this T node, params BgFileSet[] inputs) where T : BgNode
 		{
-			builder.Inputs = builder.Inputs.Add(inputs);
-			return builder;
+			return (T)node.Modify(inputs: node.Inputs.Add(inputs));
 		}
 
 		/// <summary>
 		/// Add dependencies onto other nodes or outputs. Outputs from the given tokens will be copied to the current machine before execution of the node.
 		/// </summary>
-		/// <param name="builder">The node builder</param>
+		/// <param name="node">The node to modify</param>
 		/// <param name="inputs">Files to add as inputs</param>
 		/// <returns>The current node spec, to allow chaining calls</returns>
-		public static T Requires<T>(this T builder, BgList<BgFileSet> inputs) where T : BgNodeSpecBuilder
+		public static T Requires<T>(this T node, BgList<BgFileSet> inputs) where T : BgNode
 		{
-			builder.Inputs = builder.Inputs.Add(inputs);
-			return builder;
+			return (T)node.Modify(inputs: node.Inputs.Add(inputs));
 		}
 
 		/// <summary>
 		/// Add weak dependencies onto other nodes or outputs. The producing nodes must complete successfully if they are part of the graph, but outputs from them will not be 
 		/// transferred to the machine running this node.
 		/// </summary>
-		/// <param name="builder">The node builder</param>
+		/// <param name="node">The node to modify</param>
 		/// <param name="inputs">Files to add as inputs</param>
 		/// <returns>The current node spec, to allow chaining calls</returns>
-		public static T After<T>(this T builder, params BgFileSet[] inputs) where T : BgNodeSpecBuilder
+		public static T After<T>(this T node, params BgNode[] inputs) where T : BgNode
 		{
-			builder.Fences = builder.Fences.Add(inputs);
-			return builder;
+			return (T)node.Modify(fences: node.Fences.Add(inputs));
 		}
 	}
 }
diff --git a/Engine/Source/Programs/Shared/EpicGames.BuildGraph/Expressions/BgObject.cs b/Engine/Source/Programs/Shared/EpicGames.BuildGraph/Expressions/BgObject.cs
--- a/Engine/Source/Programs/Shared/EpicGames.BuildGraph/Expressions/BgObject.cs
+++ b/Engine/Source/Programs/Shared/EpicGames.BuildGraph/Expressions/BgObject.cs
@@ -14,13 +14,8 @@ namespace EpicGames.BuildGraph.Expressions
 	/// </summary>
 	/// <typeparam name="T">The native type which mirrors this object</typeparam>
 	[BgType(typeof(BgObjectType<>))]
-	public abstract class BgObject<T> : BgExpr where T : new()
+	public abstract class BgObject<T> : BgExpr
 	{
-		/// <summary>
-		/// Default instance of the type, for determining default values
-		/// </summary>
-		internal static T DefaultInstance { get; } = new T();
-
 		/// <summary>
 		/// Constant representation of an empty object
 		/// </summary>
@@ -53,7 +48,8 @@ namespace EpicGames.BuildGraph.Expressions
 		{
 			MemberExpression member = ((MemberExpression)property.Body);
 			PropertyInfo propertyInfo = (PropertyInfo)member.Member;
-			return Set(propertyInfo.Name, value);
+			string name = propertyInfo.GetCustomAttribute<BgPropertyAttribute>()?.Name ?? propertyInfo.Name;
+			return Set(name, value);
 		}
 
 		/// <summary>
@@ -68,13 +64,14 @@ namespace EpicGames.BuildGraph.Expressions
 		/// Gets the value of a field in the object
 		/// </summary>
 		/// <param name="property">Expression indicating the property to retrieve</param>
+		/// <param name="defaultValue">Default value for the property</param>
 		/// <returns>Value of the field</returns>
-		public TExpr Get<TExpr, TNative>(Expression<Func<T, TNative>> property) where TExpr : BgExpr
+		public TExpr Get<TExpr, TNative>(Expression<Func<T, TNative>> property, TExpr defaultValue) where TExpr : BgExpr
 		{
 			MemberExpression member = ((MemberExpression)property.Body);
 			PropertyInfo propertyInfo = (PropertyInfo)member.Member;
-			TExpr defaultValue = BgType.Constant<TExpr>(propertyInfo.GetValue(DefaultInstance)!);
-			return Get(propertyInfo.Name, defaultValue);
+			string name = propertyInfo.GetCustomAttribute<BgPropertyAttribute>()?.Name ?? propertyInfo.Name;
+			return Get(name, defaultValue);
 		}
 
 		/// <inheritdoc/>
@@ -95,7 +92,7 @@ namespace EpicGames.BuildGraph.Expressions
 
 	#region Expression classes
 
-	class BgObjectEmptyExpr<T> : BgObject<T> where T : new()
+	class BgObjectEmptyExpr<T> : BgObject<T>
 	{
 		public BgObjectEmptyExpr()
 			: base(BgExprFlags.NotInterned)
@@ -108,7 +105,7 @@ namespace EpicGames.BuildGraph.Expressions
 		}
 	}
 
-	class BgObjectConstantExpr<T> : BgObject<T> where T : new()
+	class BgObjectConstantExpr<T> : BgObject<T>
 	{
 		public ImmutableDictionary<string, object> Fields { get; }
 
@@ -124,7 +121,7 @@ namespace EpicGames.BuildGraph.Expressions
 		}
 	}
 
-	class BgObjectSetExpr<T> : BgObject<T> where T : new()
+	class BgObjectSetExpr<T> : BgObject<T>
 	{
 		public BgObject<T> Object { get; }
 		public string Name { get; }
@@ -142,12 +139,12 @@ namespace EpicGames.BuildGraph.Expressions
 		{
 			writer.WriteOpcode(BgOpcode.ObjSet);
 			writer.WriteExpr(Object);
-			writer.WriteString(Name);
+			writer.WriteName(Name);
 			writer.WriteExpr(Value);
 		}
 	}
 
-	class BgObjectGetExpr<T> : BgObject<T> where T : new()
+	class BgObjectGetExpr<T> : BgObject<T>
 	{
 		public BgObject<T> Object { get; }
 		public string Name { get; }
@@ -165,12 +162,12 @@ namespace EpicGames.BuildGraph.Expressions
 		{
 			writer.WriteOpcode(BgOpcode.ObjGet);
 			writer.WriteExpr(Object);
-			writer.WriteString(Name);
+			writer.WriteName(Name);
 			writer.WriteExpr(DefaultValue);
 		}
 	}
 
-	class BgObjectWrappedExpr<T> : BgObject<T> where T : new()
+	class BgObjectWrappedExpr<T> : BgObject<T>
 	{
 		BgExpr Expr { get; }
 
diff --git a/Engine/Source/Programs/Shared/EpicGames.BuildGraph/Expressions/BgString.cs b/Engine/Source/Programs/Shared/EpicGames.BuildGraph/Expressions/BgString.cs
--- a/Engine/Source/Programs/Shared/EpicGames.BuildGraph/Expressions/BgString.cs
+++ b/Engine/Source/Programs/Shared/EpicGames.BuildGraph/Expressions/BgString.cs
@@ -341,7 +341,7 @@ namespace EpicGames.BuildGraph.Expressions
 		/// <summary>
 		/// Style for this option
 		/// </summary>
-		public BgStringOptionStyle Style { get; }
+		public BgEnum<BgStringOptionStyle>? Style { get; }
 
 		/// <summary>
 		/// Regex for validating values for the option
@@ -366,7 +366,7 @@ namespace EpicGames.BuildGraph.Expressions
 		/// <summary>
 		/// Constructor
 		/// </summary>
-		public BgStringOption(string name, BgString? label = null, BgString? description = null, BgString? defaultValue = null, BgStringOptionStyle style = BgStringOptionStyle.Text, BgString? pattern = null, BgString? patternFailed = null, BgList<BgString>? values = null, BgList<BgString>? valueDescriptions = null)
+		public BgStringOption(string name, BgString? description = null, BgString? defaultValue = null, BgEnum<BgStringOptionStyle>? style = null, BgString? pattern = null, BgString? patternFailed = null, BgList<BgString>? values = null, BgList<BgString>? valueDescriptions = null, BgString? label = null)
 			: base(BgExprFlags.None)
 		{
 			Name = name;
@@ -384,15 +384,46 @@ namespace EpicGames.BuildGraph.Expressions
 		public override void Write(BgBytecodeWriter writer)
 		{
 			writer.WriteOpcode(BgOpcode.StrOption);
-			writer.WriteExpr(Name);
-			writer.WriteExpr(Label ?? BgString.Empty);
-			writer.WriteExpr(Description ?? BgString.Empty);
-			writer.WriteUnsignedInteger((int)Style);
-			writer.WriteExpr(DefaultValue ?? BgString.Empty);
-			writer.WriteExpr(Pattern ?? BgString.Empty);
-			writer.WriteExpr(PatternFailed ?? BgString.Empty);
-			writer.WriteExpr(Values ?? BgList<BgString>.Empty);
-			writer.WriteExpr(ValueDescriptions ?? BgList<BgString>.Empty);
+			writer.WriteExpr(CreateOptionsObject());
+		}
+
+		BgObject<BgStringOptionDef> CreateOptionsObject()
+		{
+			BgObject<BgStringOptionDef> option = BgObject<BgStringOptionDef>.Empty;
+			option = option.Set(x => x.Name, Name);
+			if (!(Label is null))
+			{
+				option = option.Set(x => x.Label, Label);
+			}
+			if (!(Description is null))
+			{
+				option = option.Set(x => x.Description, Description);
+			}
+			if (!(DefaultValue is null))
+			{
+				option = option.Set(x => x.DefaultValue, DefaultValue);
+			}
+			if (!(Style is null))
+			{
+				option = option.Set(x => x.Style, Style);
+			}
+			if (!(Pattern is null))
+			{
+				option = option.Set(x => x.Pattern, Pattern);
+			}
+			if (!(PatternFailed is null))
+			{
+				option = option.Set(x => x.PatternFailed, PatternFailed);
+			}
+			if (!(Values is null))
+			{
+				option = option.Set(x => x.Values, Values);
+			}
+			if (!(ValueDescriptions is null))
+			{
+				option = option.Set(x => x.ValueDescriptions, ValueDescriptions);
+			}
+			return option;
 		}
 	}
 }
diff --git a/Engine/Source/Programs/Shared/EpicGames.BuildGraph/Expressions/BgThunk.cs b/Engine/Source/Programs/Shared/EpicGames.BuildGraph/Expressions/BgThunk.cs
new file mode 100644
--- /dev/null
+++ b/Engine/Source/Programs/Shared/EpicGames.BuildGraph/Expressions/BgThunk.cs
@@ -0,0 +1,149 @@
+// Copyright Epic Games, Inc. All Rights Reserved.
+
+using System;
+using System.Collections.Generic;
+using System.Linq.Expressions;
+using System.Reflection;
+using System.Text;
+using System.Threading.Tasks;
+
+namespace EpicGames.BuildGraph.Expressions
+{
+	/// <summary>
+	/// Thunks to a native method from within an expression. These objects contain an index into a thunk table that is not persisted to the bytecode.
+	/// </summary>
+	public class BgThunk : BgExpr
+	{
+		/// <summary>
+		/// Method to call
+		/// </summary>
+		public MethodInfo Method { get; }
+
+		/// <summary>
+		/// Arguments for the method
+		/// </summary>
+		public IReadOnlyList<object?> Arguments { get; }
+
+		/// <summary>
+		/// Creates a method closure from the given expression
+		/// </summary>
+		/// <param name="expr">Method call expression</param>
+		/// <returns></returns>
+		public static BgThunk Create(Expression<Func<Task>> expr)
+		{
+			MethodCallExpression call = (MethodCallExpression)expr.Body;
+			return new BgThunk(call);
+		}
+
+		/// <summary>
+		/// Creates a method closure from the given expression
+		/// </summary>
+		/// <typeparam name="TRet">Type of the return value</typeparam>
+		/// <param name="expr">Method call expression</param>
+		/// <returns></returns>
+		public static BgThunk<TRet> Create<TRet>(Expression<Func<Task<TRet>>> expr)
+		{
+			MethodCallExpression call = (MethodCallExpression)expr.Body;
+			return new BgThunk<TRet>(call);
+		}
+
+		/// <summary>
+		/// Creates a method closure from the given expression
+		/// </summary>
+		/// <typeparam name="TArg">Parameter to the method</typeparam>
+		/// <param name="expr">Method call expression</param>
+		/// <returns></returns>
+		public static BgThunk Create<TArg>(Expression<Func<TArg, Task>> expr)
+		{
+			MethodCallExpression call = (MethodCallExpression)expr.Body;
+			return new BgThunk(call);
+		}
+
+		/// <summary>
+		/// Creates a method closure from the given expression
+		/// </summary>
+		/// <typeparam name="TArg">Parameter to the method</typeparam>
+		/// <typeparam name="TRet">Type of the return value</typeparam>
+		/// <param name="expr">Method call expression</param>
+		/// <returns></returns>
+		public static BgThunk<TRet> Create<TArg, TRet>(Expression<Func<TArg, Task<TRet>>> expr)
+		{
+			MethodCallExpression call = (MethodCallExpression)expr.Body;
+			return new BgThunk<TRet>(call);
+		}
+
+		/// <summary>
+		/// Constructor
+		/// </summary>
+		/// <param name="call">Call expression</param>
+		public BgThunk(MethodCallExpression call)
+			: base(BgExprFlags.None)
+		{
+			Method = call.Method;
+			Arguments = GetArguments(call);
+		}
+
+		static object?[] GetArguments(MethodCallExpression call)
+		{
+			object?[] args = new object?[call.Arguments.Count];
+			for (int idx = 0; idx < call.Arguments.Count; idx++)
+			{
+				Expression argumentExpr = call.Arguments[idx];
+				if (argumentExpr is ParameterExpression parameterExpr)
+				{
+					if (parameterExpr.Type != typeof(BgContext))
+					{
+						throw new BgNodeException($"Unable to determine type of parameter '{parameterExpr.Name}'");
+					}
+				}
+				else
+				{
+					Delegate compiled = Expression.Lambda(argumentExpr).Compile();
+					args[idx] = compiled.DynamicInvoke();
+				}
+			}
+			return args;
+		}
+
+		/// <inheritdoc/>
+		public override void Write(BgBytecodeWriter writer)
+		{
+			writer.WriteOpcode(BgOpcode.Thunk);
+			writer.WriteThunk(new BgThunkDef(Method, Arguments));
+
+			ParameterInfo[] parameters = Method.GetParameters();
+			writer.WriteUnsignedInteger(parameters.Length);
+
+			for (int idx = 0; idx < parameters.Length; idx++)
+			{
+				if (typeof(BgExpr).IsAssignableFrom(parameters[idx].ParameterType))
+				{
+					writer.WriteExpr((BgExpr)Arguments[idx]!);
+				}
+				else
+				{
+					writer.WriteOpcode(BgOpcode.Null);
+				}
+			}
+		}
+
+		/// <inheritdoc/>
+		public override BgString ToBgString() => "{Method}";
+	}
+
+	/// <summary>
+	/// Wraps a native method that returns a value
+	/// </summary>
+	/// <typeparam name="T"></typeparam>
+	public class BgThunk<T> : BgThunk
+	{
+		/// <summary>
+		/// Constructor
+		/// </summary>
+		/// <param name="expr"></param>
+		public BgThunk(MethodCallExpression expr)
+			: base(expr)
+		{
+		}	
+	}
+}
