Horde: Fix corrupt offsets into compressed packet data for bundles.

diff --git a/Engine/Source/Programs/Shared/EpicGames.Horde/Storage/Bundles/BundleStore.cs b/Engine/Source/Programs/Shared/EpicGames.Horde/Storage/Bundles/BundleStore.cs
--- a/Engine/Source/Programs/Shared/EpicGames.Horde/Storage/Bundles/BundleStore.cs
+++ b/Engine/Source/Programs/Shared/EpicGames.Horde/Storage/Bundles/BundleStore.cs
@@ -493,6 +493,8 @@ namespace EpicGames.Horde.Storage.Bundles
 
 					nodeOffset += header.Exports[exportIdx].Length;
 				}
+
+				packetOffset += packet.EncodedLength;
 			}
 
 			return nodes[^1];
@@ -869,7 +871,7 @@ namespace EpicGames.Horde.Storage.Bundles
 		}
 
 		/// <inheritdoc/>
-		public async Task WriteTreeAsync(RefId refName, T root, bool flush, CancellationToken cancellationToken = default)
+		public async Task WriteTreeAsync(RefId refName, T root, bool flush = true, CancellationToken cancellationToken = default)
 		{
 			ITreeBlob rootBlob = await TreeNode.SerializeAsync(root, this, cancellationToken);
 			await base.WriteTreeAsync(refName, rootBlob, flush, cancellationToken);
