Unzip optimization for archives with lots of small files
- Remove redunant exists check when creating directories.  Directory.CreateDirectory already does that.
- Use streams directly instead of calling ZipFileExtensions.ExtractToFile.

diff --git a/Engine/Source/Programs/AutomationTool/AutomationUtils/CommandUtils.cs b/Engine/Source/Programs/AutomationTool/AutomationUtils/CommandUtils.cs
--- a/Engine/Source/Programs/AutomationTool/AutomationUtils/CommandUtils.cs
+++ b/Engine/Source/Programs/AutomationTool/AutomationUtils/CommandUtils.cs
@@ -2285,10 +2285,7 @@ namespace AutomationTool
 						continue;
 					}
 					FileReference OutputFile = FileReference.Combine(BaseDirectory, Entry.FullName);
-					if(!DirectoryReference.Exists(OutputFile.Directory))
-					{
-						DirectoryReference.CreateDirectory(OutputFile.Directory);
-					}
+					DirectoryReference.CreateDirectory(OutputFile.Directory);
 					Entry.ExtractToFile_CrossPlatform(OutputFile.FullName, true);
 					OutputFiles.Add(OutputFile);
 				}
diff --git a/Engine/Source/Programs/Shared/EpicGames.Core/ZipArchiveExtensions.cs b/Engine/Source/Programs/Shared/EpicGames.Core/ZipArchiveExtensions.cs
--- a/Engine/Source/Programs/Shared/EpicGames.Core/ZipArchiveExtensions.cs
+++ b/Engine/Source/Programs/Shared/EpicGames.Core/ZipArchiveExtensions.cs
@@ -1,6 +1,7 @@
 // Copyright Epic Games, Inc. All Rights Reserved.
 
 using System;
+using System.IO;
 using System.IO.Compression;
 using System.Reflection;
 using System.Runtime.InteropServices;
@@ -55,7 +56,15 @@ namespace EpicGames.Core
 		/// <param name="overwrite"></param>
 		public static void ExtractToFile_CrossPlatform(this ZipArchiveEntry entry, string targetFileName, bool overwrite)
 		{
-			ZipFileExtensions.ExtractToFile(entry, targetFileName, overwrite);
+			// This seems to be consistently at least 35% faster than ZipFileExtensions.ExtractToFile(Entry, TargetFileName, Overwrite) when the
+			// archive contains many small files.
+			using (FileStream OutputStream = new FileStream(targetFileName, overwrite ? FileMode.Create : FileMode.CreateNew, FileAccess.Write))
+			{
+				using (Stream InStream = entry.Open())
+				{
+					InStream.CopyTo(OutputStream);
+				}
+			}
 
 			if (RuntimeInformation.IsOSPlatform(OSPlatform.Linux))
 			{
