Fix runtime issue where internal UHT crashes due to a fix for VS2019.

#rnx

diff --git a/Engine/Source/Programs/Shared/EpicGames.UHT/Parsers/UhtPropertyParser.cs b/Engine/Source/Programs/Shared/EpicGames.UHT/Parsers/UhtPropertyParser.cs
--- a/Engine/Source/Programs/Shared/EpicGames.UHT/Parsers/UhtPropertyParser.cs
+++ b/Engine/Source/Programs/Shared/EpicGames.UHT/Parsers/UhtPropertyParser.cs
@@ -338,7 +338,7 @@ namespace EpicGames.UHT.Parsers
 		/// <summary>
 		/// Request a context
 		/// </summary>
-		public UhtThreadBorrower(bool _ = false) // argument needed for vs2019
+		public UhtThreadBorrower(bool _) // argument needed for vs2019
 		{
 			List<T> cache = s_tls.Value!;
 			if (cache.Count == 0)
@@ -487,7 +487,7 @@ namespace EpicGames.UHT.Parsers
 		public UhtPropertyParser Parse(UhtParsingScope topScope, EPropertyFlags disallowPropertyFlags, UhtPropertyParseOptions options, UhtPropertyCategory category, UhtPropertyDelegate propertyDelegate)
 		{
 			// Initialize the property context
-			using UhtThreadBorrower<UhtPropertySpecifierContext> borrower = new();
+			using UhtThreadBorrower<UhtPropertySpecifierContext> borrower = new(true);
 			UhtPropertySpecifierContext specifierContext = borrower.Instance;
 			specifierContext.Type = topScope.ScopeType;
 			specifierContext.TokenReader = topScope.TokenReader;
@@ -625,7 +625,7 @@ namespace EpicGames.UHT.Parsers
 
 			// Save the token reader state.  We need this to restore back to the start when invoking the resolve methods.
 			{
-				using UhtThreadBorrower<UhtPropertySpecifierContext> borrower = new();
+				using UhtThreadBorrower<UhtPropertySpecifierContext> borrower = new(true);
 				UhtPropertySpecifierContext specifierContext = borrower.Instance;
 				specifierContext.Type = parentPropertySettings.Outer;
 				specifierContext.TokenReader = tokenReader;
