C# implementation of CompactBinaryPackage (cbpkg)

diff --git a/Engine/Source/Programs/Shared/EpicGames.Serialization.Tests/CbPackageTests.cs b/Engine/Source/Programs/Shared/EpicGames.Serialization.Tests/CbPackageTests.cs
new file mode 100644
--- /dev/null
+++ b/Engine/Source/Programs/Shared/EpicGames.Serialization.Tests/CbPackageTests.cs
@@ -0,0 +1,128 @@
+// Copyright Epic Games, Inc. All Rights Reserved.
+
+using System;
+using System.IO;
+using System.Text;
+using System.Threading.Tasks;
+using EpicGames.Core;
+using Microsoft.VisualStudio.TestTools.UnitTesting;
+
+namespace EpicGames.Serialization.Tests
+{
+    class TestObject
+    {
+        [CbField]
+        public CbBinaryAttachment BinaryAttachment { get; set; }
+
+        [CbField]
+        public CbObjectAttachment ObjectAttachment { get; set; }
+    }
+
+    class SimpleObject
+    {
+        [CbField] public int A { get; set; } = 0;
+
+        [CbField]
+        public string B { get; set; } = String.Empty;
+    }
+
+    [TestClass]
+    public class CbPackageTests
+    {
+        [TestMethod]
+        public async Task RoundTrip()
+        {
+            TestObject root = new TestObject()
+            {
+
+            };
+            CbObject rootObject = CbSerializer.Serialize(root);
+            IoHash rootHash = IoHash.Compute(rootObject.GetView().Span);
+            
+            CbPackageBuilder builder = new();
+            await builder.AddAttachment(rootHash, CbPackageAttachmentFlags.IsObject, rootObject.GetView().ToArray());
+            byte[] bytes = builder.ToByteArray();
+        
+            using MemoryStream ms = new MemoryStream(bytes);
+            CbPackageReader reader = await CbPackageReader.Create(ms);
+
+            Assert.AreEqual(rootHash, reader.RootHash);
+            Assert.AreEqual(rootObject, reader.RootObject);
+
+            await foreach ((CbPackageAttachmentEntry, byte[]) _ in reader.IterateAttachments())
+            {
+                Assert.Fail("No Attachments expected");
+            }
+        }
+
+        [TestMethod]
+        public async Task RoundTripComplex()
+        {
+            byte[] blob = Encoding.UTF8.GetBytes("This is blob contents");
+            IoHash blobHash = IoHash.Compute(blob);
+
+            SimpleObject simple = new SimpleObject();
+            CbObject simpleObject = CbSerializer.Serialize(simple);
+            IoHash simpleHash = IoHash.Compute(simpleObject.GetView().Span);
+
+            TestObject root = new TestObject()
+            {
+                BinaryAttachment = blobHash,
+                ObjectAttachment = simpleHash,
+            };
+            CbObject rootObject = CbSerializer.Serialize(root);
+            IoHash rootHash = IoHash.Compute(rootObject.GetView().Span);
+
+            CbPackageBuilder builder = new();
+            await builder.AddAttachment(rootHash, CbPackageAttachmentFlags.IsObject, rootObject.GetView().ToArray());
+            await builder.AddAttachment(simpleHash, CbPackageAttachmentFlags.IsObject, simpleObject.GetView().ToArray());
+            await builder.AddAttachment(blobHash, 0, blob);
+
+            byte[] bytes = builder.ToByteArray();
+
+            using MemoryStream ms = new MemoryStream(bytes);
+            CbPackageReader reader = await CbPackageReader.Create(ms);
+
+            Assert.AreEqual(rootHash, reader.RootHash);
+            Assert.AreEqual(rootObject, reader.RootObject);
+
+            int countOfAttachments = 0;
+            await foreach ((CbPackageAttachmentEntry entry, byte[] attachmentBytes) in reader.IterateAttachments())
+            {
+                countOfAttachments++;
+                if (entry.AttachmentHash == blobHash)
+                {
+                    Assert.AreEqual((ulong)blob.Length, entry.PayloadSize);
+                    CollectionAssert.AreEqual(blob, attachmentBytes);
+                    continue;
+                }
+
+                if (entry.AttachmentHash == simpleHash)
+                {
+                    Assert.AreEqual((ulong)simpleObject.GetView().Length, entry.PayloadSize);
+
+                    CbObject roundTrippedObject = new CbObject(attachmentBytes);
+                    Assert.AreEqual(simpleObject, roundTrippedObject);
+                    continue;
+                }
+
+                Assert.Fail($"Unknown attachment {entry.AttachmentHash}");
+            }
+            Assert.AreEqual(2, countOfAttachments);
+        }
+
+        [TestMethod]
+        public void ReadNoAttachments()
+        {
+            using MemoryStream ms = new MemoryStream();
+            CbPackageHeader header = new CbPackageHeader(0);
+            header.Write(ms);
+            ms.Position = 0;
+
+            Assert.ThrowsException<Exception>(async () =>
+            {
+                CbPackageReader reader = await CbPackageReader.Create(ms);
+            });
+        }
+    }
+}
diff --git a/Engine/Source/Programs/Shared/EpicGames.Serialization/CbPackage.cs b/Engine/Source/Programs/Shared/EpicGames.Serialization/CbPackage.cs
new file mode 100644
--- /dev/null
+++ b/Engine/Source/Programs/Shared/EpicGames.Serialization/CbPackage.cs
@@ -0,0 +1,283 @@
+// Copyright Epic Games, Inc. All Rights Reserved.
+
+using System;
+using System.Collections.Generic;
+using System.IO;
+using System.Text;
+using System.Threading.Tasks;
+using EpicGames.Core;
+
+namespace EpicGames.Serialization
+{
+    internal static class CbPackageConstants
+    {
+        public const uint CbPackageHeaderMagic = 0xaa77aacc;
+    }
+
+    /// <summary>
+    /// Header in the CbPackage format
+    /// </summary>
+    /// <remarks>see CbPackage definition in Zen https://github.com/EpicGames/zen/blob/main/zenhttp/include/zenhttp/httpshared.h</remarks>
+    public struct CbPackageHeader
+    {
+        public CbPackageHeader(uint attachmentCount, uint reserved1 = 0, uint reserved2 = 0)
+        {
+            HeaderMagic = CbPackageConstants.CbPackageHeaderMagic;
+            AttachmentCount = attachmentCount;
+
+            Reserved1 = reserved1;
+            Reserved2 = reserved2;
+        }
+
+        public uint HeaderMagic { get; set; }
+        public uint AttachmentCount { get; set; }
+        public uint Reserved1 { get; set; }
+        public uint Reserved2 { get; set; }
+
+        private const int Length = sizeof(uint) * 4;
+
+        public static async Task<CbPackageHeader> Read(Stream s)
+        {
+            // using async reading of stream because Asp.net requires this
+            byte[] buf = new byte[Length];
+            int readBytes = await s.ReadAsync(buf, 0, Length);
+            if (readBytes != Length)
+            {
+                throw new Exception($"Was not able to read package header, not enough bytes. Bytes read: {readBytes}");
+            }
+
+            using MemoryStream ms = new MemoryStream(buf);
+            using BinaryReader reader = new BinaryReader(ms);
+            uint headerMagic = reader.ReadUInt32();
+            if (headerMagic != CbPackageConstants.CbPackageHeaderMagic)
+            {
+                throw new Exception($"Magic did not match, expected {CbPackageConstants.CbPackageHeaderMagic} got {headerMagic}");
+            }
+
+            uint attachmentCount = reader.ReadUInt32();
+            uint reserved1 = reader.ReadUInt32();
+            uint reserved2 = reader.ReadUInt32();
+
+            return new CbPackageHeader(attachmentCount, reserved1, reserved2);
+        }
+
+        public void Write(Stream stream)
+        {
+            using BinaryWriter writer = new BinaryWriter(stream, Encoding.ASCII, leaveOpen: true);
+            writer.Write(HeaderMagic);
+            writer.Write(AttachmentCount);
+            writer.Write(Reserved1);
+            writer.Write(Reserved2);
+        }
+    }
+
+    /// <summary>
+    /// Per attachment entry flags
+    /// </summary>
+    [Flags]
+    public enum CbPackageAttachmentFlags : uint
+    {
+        IsCompressed = (1u << 0),  // Is marshaled using compressed buffer storage format
+        IsObject = (1u << 1),  // Is compact binary object
+        IsError = (1u << 2),   // Is error (compact binary formatted) object
+    }
+
+    /// <summary>
+    /// Header for each attachment in the package
+    /// </summary>
+    public struct CbPackageAttachmentEntry
+    { 
+        public ulong PayloadSize { get; set; }
+        public CbPackageAttachmentFlags Flags { get; set; }
+        public IoHash AttachmentHash { get; set; }
+
+        public static async Task<CbPackageAttachmentEntry> Read(Stream s)
+        {
+            byte[] buf = new byte[Length];
+            int readBytes = await s.ReadAsync(buf, 0, Length);
+            if (readBytes != Length)
+            {
+                throw new Exception($"Was not able to read package attachment entry, not enough bytes. Bytes read: {readBytes}");
+            }
+            using MemoryStream ms = new MemoryStream(buf);
+            using BinaryReader reader = new BinaryReader(ms);
+            ulong payloadSize = reader.ReadUInt64();
+
+            uint flags = reader.ReadUInt32();
+            byte[] hashBytes = reader.ReadBytes(20);
+
+            return new CbPackageAttachmentEntry
+            {
+                PayloadSize = payloadSize,
+                Flags = (CbPackageAttachmentFlags)flags,
+                AttachmentHash = new IoHash(hashBytes),
+            };
+        }
+
+        private const int Length = IoHash.NumBytes + sizeof(ulong) + sizeof(uint);
+
+        public void Write(Stream stream)
+        {
+            using BinaryWriter writer = new BinaryWriter(stream, Encoding.ASCII, leaveOpen: true);
+            writer.Write(PayloadSize);
+            writer.Write((uint)Flags);
+            writer.Write(AttachmentHash.ToByteArray());
+        }
+    }
+
+    /// <summary>
+    /// Reads a CbPackage (a bundle of both a cb object and its attachments)
+    /// </summary>
+    public class CbPackageReader
+    {
+        private readonly CbPackageHeader _header;
+        private readonly Stream _attachmentsStream;
+
+        public CbPackageReader(CbPackageHeader header, CbObject rootObject, IoHash rootHash, Stream attachmentsStream)
+        {
+            _header = header;
+            _attachmentsStream = attachmentsStream;
+            RootObject = rootObject;
+            RootHash = rootHash;
+        }
+
+        private static async Task<byte[]> ReadStreamAsync(Stream s, int count)
+        {
+            int index = 0;
+            byte[] bytes = new byte[count];
+            do
+            {
+                int n = await s.ReadAsync(new Memory<byte>(bytes, index, count - index))
+                    .ConfigureAwait(false);
+                if (n == 0)
+                {
+                    throw new Exception("Unexpected end of stream");
+                }
+
+                index += n;
+            } while (index < count);
+
+            return bytes;
+        }
+
+        public static async Task<CbPackageReader> Create(Stream s)
+        {
+            CbPackageHeader header = await CbPackageHeader.Read(s);
+
+            if (header.AttachmentCount < 1)
+            {
+                throw new Exception("At least 1 attachment was expected in a CbPackage");
+            }
+
+            // the first object is assumed to be a CbObject, the root object which references the other attachments
+            CbPackageAttachmentEntry entry = await CbPackageAttachmentEntry.Read(s);
+            if (!entry.Flags.HasFlag(CbPackageAttachmentFlags.IsObject))
+            {
+                throw new Exception("First attachment must be a CbObject for a package");
+            }
+            if (entry.Flags.HasFlag(CbPackageAttachmentFlags.IsError))
+            {
+                throw new Exception("First attachment was a error object");
+            }
+            if (entry.Flags.HasFlag(CbPackageAttachmentFlags.IsCompressed))
+            {
+                // TODO: We could support this being a compressed buffer and just remove it
+                throw new Exception("First attachment must not be compressed");
+            }
+
+            if (entry.PayloadSize > Int32.MaxValue)
+            {
+                throw new Exception($"Package attachments larger then {Int32.MaxValue} not supported");
+            }
+            byte[] rootObjectBytes = await ReadStreamAsync(s, (int)entry.PayloadSize);
+
+            return new CbPackageReader(header, new CbObject(rootObjectBytes), entry.AttachmentHash, s);
+        }
+
+        public async IAsyncEnumerable<(CbPackageAttachmentEntry, byte[])> IterateAttachments()
+        {
+            // close the stream after we have iterated the attachments as there should be nothing left in it
+            await using Stream s = _attachmentsStream;
+            // we have already read the first attachment as that is the root object
+            for (int i = 1; i < _header.AttachmentCount; i++)
+            {
+                CbPackageAttachmentEntry entry = await CbPackageAttachmentEntry.Read(_attachmentsStream);
+
+                if (entry.PayloadSize > Int32.MaxValue)
+                {
+                    throw new Exception($"Package attachments larger then {Int32.MaxValue} not supported");
+                }
+                byte[] blob = await ReadStreamAsync(s, (int)entry.PayloadSize);
+                yield return (entry, blob);
+            }
+        }
+
+        public CbObject RootObject { get; }
+        public IoHash RootHash { get; }
+    }
+
+    /// <summary>
+    /// Builds a in-memory representation of a CbPackage
+    /// </summary>
+    public class CbPackageBuilder
+    {
+        private readonly MemoryStream _packageBuffer;
+        private uint _countOfWrittenAttachments;
+
+        public CbPackageBuilder()
+        {
+            _packageBuffer = new MemoryStream();
+
+            // we will overwrite this header again after building the package builder
+            CbPackageHeader header = new CbPackageHeader
+            {
+                HeaderMagic = CbPackageConstants.CbPackageHeaderMagic,
+                AttachmentCount = 0,
+                Reserved1 = 0,
+                Reserved2 = 0,
+            };
+            header.Write(_packageBuffer);
+        }
+        public async Task AddAttachment(IoHash attachmentHash, CbPackageAttachmentFlags flags, byte[] blobMemory)
+        {
+            CbPackageAttachmentEntry entry = new CbPackageAttachmentEntry()
+            {
+                AttachmentHash = attachmentHash,
+                Flags = flags,
+                PayloadSize = (ulong)blobMemory.LongLength
+            };
+            entry.Write(_packageBuffer);
+            await _packageBuffer.WriteAsync(blobMemory, 0, blobMemory.Length);
+            _countOfWrittenAttachments++;
+        }
+
+        public async Task AddAttachment(IoHash attachmentHash, CbPackageAttachmentFlags flags, Stream stream, ulong length)
+        {
+            CbPackageAttachmentEntry entry = new CbPackageAttachmentEntry()
+            {
+                AttachmentHash = attachmentHash,
+                Flags = flags,
+                PayloadSize = length
+            };
+            entry.Write(_packageBuffer);
+            await stream.CopyToAsync(_packageBuffer);
+            _countOfWrittenAttachments++;
+        }
+
+        public byte[] ToByteArray()
+        {
+            // update the package with the correct header
+            _packageBuffer.Position = 0;
+            CbPackageHeader header = new CbPackageHeader
+            {
+                HeaderMagic = CbPackageConstants.CbPackageHeaderMagic,
+                AttachmentCount = _countOfWrittenAttachments,
+                Reserved1 = 0,
+                Reserved2 = 0,
+            };
+            header.Write(_packageBuffer);
+
+            return _packageBuffer.ToArray();
+        }
+    }
+}
