Fixed failing CbPackage test, Assert.ThrowsException does not support async lambda.

diff --git a/Engine/Source/Programs/Shared/EpicGames.Serialization.Tests/CbPackageTests.cs b/Engine/Source/Programs/Shared/EpicGames.Serialization.Tests/CbPackageTests.cs
--- a/Engine/Source/Programs/Shared/EpicGames.Serialization.Tests/CbPackageTests.cs
+++ b/Engine/Source/Programs/Shared/EpicGames.Serialization.Tests/CbPackageTests.cs
@@ -112,17 +112,21 @@ namespace EpicGames.Serialization.Tests
         }
 
         [TestMethod]
-        public void ReadNoAttachments()
+        public async Task ReadNoAttachments()
         {
             using MemoryStream ms = new MemoryStream();
             CbPackageHeader header = new CbPackageHeader(0);
             header.Write(ms);
             ms.Position = 0;
 
-            Assert.ThrowsException<Exception>(async () =>
+            try
             {
-                CbPackageReader reader = await CbPackageReader.Create(ms);
-            });
+                await CbPackageReader.Create(ms);
+                Assert.Fail("Exception should be thrown");
+            }
+            catch (Exception)
+            {
+            }
         }
     }
 }
