EpicGames.Core: Add DateTime serializers for IMemoryReader/IMemoryWriter.

diff --git a/Engine/Source/Programs/Shared/EpicGames.Core/MemoryReader.cs b/Engine/Source/Programs/Shared/EpicGames.Core/MemoryReader.cs
--- a/Engine/Source/Programs/Shared/EpicGames.Core/MemoryReader.cs
+++ b/Engine/Source/Programs/Shared/EpicGames.Core/MemoryReader.cs
@@ -188,6 +188,16 @@ namespace EpicGames.Core
 			return value;
 		}
 
+		/// <summary>
+		/// Reads a DateTime from the memory buffer
+		/// </summary>
+		/// <param name="reader">Reader to deserialize from</param>
+		public static DateTime ReadDateTime(this IMemoryReader reader)
+		{
+			ulong encoded = reader.ReadUnsignedVarInt();
+			return new DateTime((long)encoded >> 2, (DateTimeKind)(encoded & 3UL));
+		}
+
 		/// <summary>
 		/// Reads a sequence of bytes from the buffer
 		/// </summary>
diff --git a/Engine/Source/Programs/Shared/EpicGames.Core/MemoryWriter.cs b/Engine/Source/Programs/Shared/EpicGames.Core/MemoryWriter.cs
--- a/Engine/Source/Programs/Shared/EpicGames.Core/MemoryWriter.cs
+++ b/Engine/Source/Programs/Shared/EpicGames.Core/MemoryWriter.cs
@@ -179,6 +179,17 @@ namespace EpicGames.Core
 			writer.Advance(sizeof(ulong));
 		}
 
+		/// <summary>
+		/// Writes a DateTime to the memory writer
+		/// </summary>
+		/// <param name="writer">Writer to serialize to</param>
+		/// <param name="value">Value to write</param>
+		public static void WriteDateTime(this IMemoryWriter writer, DateTime value)
+		{
+			ulong encoded = ((ulong)value.Ticks << 2) | (ulong)value.Kind;
+			writer.WriteUnsignedVarInt(encoded);
+		}
+
 		/// <summary>
 		/// Appends a sequence of bytes to the buffer
 		/// </summary>
