Horde: Remove accessor methods for data/references in blobs; expose them as properties instead.

diff --git a/Engine/Source/Programs/Shared/EpicGames.Horde.Tests/BundleStoreTests.cs b/Engine/Source/Programs/Shared/EpicGames.Horde.Tests/BundleStoreTests.cs
--- a/Engine/Source/Programs/Shared/EpicGames.Horde.Tests/BundleStoreTests.cs
+++ b/Engine/Source/Programs/Shared/EpicGames.Horde.Tests/BundleStoreTests.cs
@@ -1,7 +1,7 @@
 // Copyright Epic Games, Inc. All Rights Reserved.
 
 using EpicGames.Horde.Storage;
-using EpicGames.Horde.Storage.Impl;
+using EpicGames.Horde.Storage.Backends;
 using EpicGames.Horde.Storage.Bundles;
 using EpicGames.Horde.Storage.Nodes;
 using Microsoft.Extensions.Caching.Memory;
diff --git a/Engine/Source/Programs/Shared/EpicGames.Horde/Storage/Backends/InMemoryBlobStore.cs b/Engine/Source/Programs/Shared/EpicGames.Horde/Storage/Backends/InMemoryBlobStore.cs
--- a/Engine/Source/Programs/Shared/EpicGames.Horde/Storage/Backends/InMemoryBlobStore.cs
+++ b/Engine/Source/Programs/Shared/EpicGames.Horde/Storage/Backends/InMemoryBlobStore.cs
@@ -7,31 +7,13 @@ using System.Collections.Generic;
 using System.Threading;
 using System.Threading.Tasks;
 
-namespace EpicGames.Horde.Storage.Impl
+namespace EpicGames.Horde.Storage.Backends
 {
 	/// <summary>
 	/// Implementation of <see cref="IBlobStore"/> which stores data in memory. Not intended for production use.
 	/// </summary>
 	public class InMemoryBlobStore : IBlobStore
 	{
-		class Blob : IBlob
-		{
-			readonly ReadOnlyMemory<byte> _data;
-			readonly IReadOnlyList<BlobId> _references;
-
-			public Blob(ReadOnlyMemory<byte> data, IReadOnlyList<BlobId> references)
-			{
-				_data = data;
-				_references = references;
-			}
-
-			/// <inheritdoc/>
-			public ValueTask<ReadOnlyMemory<byte>> GetDataAsync() => new ValueTask<ReadOnlyMemory<byte>>(_data);
-
-			/// <inheritdoc/>
-			public ValueTask<IReadOnlyList<BlobId>> GetReferencesAsync() => new ValueTask<IReadOnlyList<BlobId>>(_references);
-		}
-
 		/// <summary>
 		/// Map of blob id to blob data
 		/// </summary>
@@ -51,14 +33,17 @@ namespace EpicGames.Horde.Storage.Impl
 		#region Blobs
 
 		/// <inheritdoc/>
-		public Task<IBlob> ReadBlobAsync(BlobId blobId, CancellationToken cancellationToken = default) => Task.FromResult(_blobs[blobId]);
+		public Task<IBlob?> TryReadBlobAsync(BlobId blobId, CancellationToken cancellationToken = default)
+		{
+			_blobs.TryGetValue(blobId, out IBlob? blob);
+			return Task.FromResult(blob);
+		}
 
 		/// <inheritdoc/>
 		public Task<BlobId> WriteBlobAsync(ReadOnlySequence<byte> data, IReadOnlyList<BlobId> references, CancellationToken cancellationToken = default)
 		{
 			BlobId blobId = new BlobId(Guid.NewGuid().ToString());
-			Blob blob = new Blob(data.ToArray(), references);
-			_blobs[blobId] = blob;
+			_blobs[blobId] = BlobUtils.FromMemory(data.ToArray(), references);
 			return Task.FromResult(blobId);
 		}
 
@@ -73,12 +58,16 @@ namespace EpicGames.Horde.Storage.Impl
 		public Task<bool> HasRefAsync(RefId id, CancellationToken cancellationToken) => Task.FromResult(_refs.ContainsKey(id));
 
 		/// <inheritdoc/>
-		public Task<IBlob> ReadRefAsync(RefId id, CancellationToken cancellationToken) => Task.FromResult(_refs[id]);
+		public Task<IBlob?> TryReadRefAsync(RefId id, CancellationToken cancellationToken)
+		{
+			_refs.TryGetValue(id, out IBlob? blob);
+			return Task.FromResult(blob);
+		}
 
 		/// <inheritdoc/>
 		public Task WriteRefAsync(RefId id, ReadOnlySequence<byte> data, IReadOnlyList<BlobId> references, CancellationToken cancellationToken)
 		{
-			_refs[id] = new Blob(data.ToArray(), references);
+			_refs[id] = BlobUtils.FromMemory(data.ToArray(), references);
 			return Task.CompletedTask;
 		}
 
diff --git a/Engine/Source/Programs/Shared/EpicGames.Horde/Storage/Bundles/Bundle.cs b/Engine/Source/Programs/Shared/EpicGames.Horde/Storage/Bundles/Bundle.cs
--- a/Engine/Source/Programs/Shared/EpicGames.Horde/Storage/Bundles/Bundle.cs
+++ b/Engine/Source/Programs/Shared/EpicGames.Horde/Storage/Bundles/Bundle.cs
@@ -348,7 +348,7 @@ namespace EpicGames.Horde.Storage.Bundles
 		public static async Task<Bundle> ReadBundleAsync(this IBlobStore store, BlobId id, CancellationToken cancellationToken)
 		{
 			IBlob blob = await store.ReadBlobAsync(id, cancellationToken);
-			return await ReadBundleAsync(blob);
+			return ReadBundle(blob);
 		}
 
 		/// <summary>
@@ -361,7 +361,7 @@ namespace EpicGames.Horde.Storage.Bundles
 		public static async Task<Bundle> ReadBundleAsync(this IBlobStore store, RefId id, CancellationToken cancellationToken)
 		{
 			IBlob blob = await store.ReadRefAsync(id, cancellationToken);
-			return await ReadBundleAsync(blob);
+			return ReadBundle(blob);
 		}
 
 		/// <summary>
@@ -392,9 +392,9 @@ namespace EpicGames.Horde.Storage.Bundles
 			return store.WriteBlobAsync(data, imports, cancellationToken);
 		}
 
-		static async Task<Bundle> ReadBundleAsync(IBlob blob)
+		static Bundle ReadBundle(IBlob blob)
 		{
-			ReadOnlyMemory<byte> data = await blob.GetDataAsync();
+			ReadOnlyMemory<byte> data = blob.Data;
 			MemoryReader reader = new MemoryReader(data);
 			Bundle bundle = new Bundle(reader);
 			reader.CheckEmpty();
diff --git a/Engine/Source/Programs/Shared/EpicGames.Horde/Storage/IBlobStore.cs b/Engine/Source/Programs/Shared/EpicGames.Horde/Storage/IBlobStore.cs
--- a/Engine/Source/Programs/Shared/EpicGames.Horde/Storage/IBlobStore.cs
+++ b/Engine/Source/Programs/Shared/EpicGames.Horde/Storage/IBlobStore.cs
@@ -5,6 +5,7 @@ using System.Buffers;
 using System.Collections.Generic;
 using System.Threading;
 using System.Threading.Tasks;
+using EpicGames.Core;
 
 namespace EpicGames.Horde.Storage
 {
@@ -17,13 +18,13 @@ namespace EpicGames.Horde.Storage
 		/// Gets the data for this blob
 		/// </summary>
 		/// <returns></returns>
-		ValueTask<ReadOnlyMemory<byte>> GetDataAsync();
+		ReadOnlyMemory<byte> Data { get; }
 
 		/// <summary>
 		/// Find the outward references for a node
 		/// </summary>
 		/// <returns></returns>
-		ValueTask<IReadOnlyList<BlobId>> GetReferencesAsync();
+		IReadOnlyList<BlobId> References { get; }
 	}
 
 	/// <summary>
@@ -39,7 +40,7 @@ namespace EpicGames.Horde.Storage
 		/// <param name="id">The blob identifier</param>
 		/// <param name="cancellationToken">Cancellation token for the operation</param>
 		/// <returns></returns>
-		Task<IBlob> ReadBlobAsync(BlobId id, CancellationToken cancellationToken = default);
+		Task<IBlob?> TryReadBlobAsync(BlobId id, CancellationToken cancellationToken = default);
 
 		/// <summary>
 		/// Writes a new object to the store
@@ -68,7 +69,7 @@ namespace EpicGames.Horde.Storage
 		/// <param name="id">The blob identifier</param>
 		/// <param name="cancellationToken">Cancellation token for the operation</param>
 		/// <returns></returns>
-		Task<IBlob> ReadRefAsync(RefId id, CancellationToken cancellationToken = default);
+		Task<IBlob?> TryReadRefAsync(RefId id, CancellationToken cancellationToken = default);
 
 		/// <summary>
 		/// Writes a new blob to the store
@@ -90,4 +91,135 @@ namespace EpicGames.Horde.Storage
 
 		#endregion
 	}
+
+	/// <summary>
+	/// Extension methods for <see cref="IBlobStore"/>
+	/// </summary>
+	static class BlobStoreExtensions
+	{
+		/// <summary>
+		/// Reads a blob from the store, throwing an exception if it does not exist
+		/// </summary>
+		/// <param name="store">The store instance to read from</param>
+		/// <param name="id">Id for the blob</param>
+		/// <param name="cancellationToken">Cancellation token for the operation</param>
+		/// <returns>The blob instance</returns>
+		public static async Task<IBlob> ReadBlobAsync(this IBlobStore store, BlobId id, CancellationToken cancellationToken = default)
+		{
+			IBlob? blob = await store.TryReadBlobAsync(id, cancellationToken);
+			if (blob == null)
+			{
+				throw new KeyNotFoundException($"Unable to find blob {id}");
+			}
+			return blob;
+		}
+
+		/// <summary>
+		/// Reads a ref from the store, throwing an exception if it does not exist
+		/// </summary>
+		/// <param name="store">The store instance to read from</param>
+		/// <param name="id">Id for the ref</param>
+		/// <param name="cancellationToken">Cancellation token for the operation</param>
+		/// <returns>The blob instance</returns>
+		public static async Task<IBlob> ReadRefAsync(this IBlobStore store, RefId id, CancellationToken cancellationToken = default)
+		{
+			IBlob? blob = await store.TryReadRefAsync(id, cancellationToken);
+			if (blob == null)
+			{
+				throw new KeyNotFoundException($"Unable to find ref {id}");
+			}
+			return blob;
+		}
+	}
+
+	/// <summary>
+	/// Utility methods for blobs
+	/// </summary>
+	static class BlobUtils
+	{
+		class InMemoryBlob : IBlob
+		{
+			/// <inheritdoc/>
+			public ReadOnlyMemory<byte> Data { get; }
+
+			/// <inheritdoc/>
+			public IReadOnlyList<BlobId> References { get; }
+
+			public InMemoryBlob(ReadOnlyMemory<byte> data, IReadOnlyList<BlobId> references)
+			{
+				Data = data;
+				References = references;
+			}
+		}
+
+		/// <summary>
+		/// Create a blob from memory
+		/// </summary>
+		/// <param name="data">Payload for the blob</param>
+		/// <param name="references">References to other blobs</param>
+		/// <returns>Blob instance</returns>
+		public static IBlob FromMemory(ReadOnlyMemory<byte> data, IReadOnlyList<BlobId> references) => new InMemoryBlob(data, references);
+
+		/// <summary>
+		/// Serialize a blob into a flat memory buffer
+		/// </summary>
+		/// <param name="blob">The blob to serialize</param>
+		/// <returns>Serialized data</returns>
+		public static ReadOnlySequence<byte> Serialize(IBlob blob)
+		{
+			return Serialize(blob.Data, blob.References);
+		}
+
+		/// <summary>
+		/// Serialize a blob into a flat memory buffer
+		/// </summary>
+		/// <param name="data">Data for the blob</param>
+		/// <param name="references">List of references to other blobs</param>
+		/// <returns>Serialized data</returns>
+		public static ReadOnlySequence<byte> Serialize(ReadOnlyMemory<byte> data, IReadOnlyList<BlobId> references)
+		{
+			return Serialize(new ReadOnlySequence<byte>(data), references);
+		}
+
+		/// <summary>
+		/// Serialize a blob into a flat memory buffer
+		/// </summary>
+		/// <param name="data">Data for the blob</param>
+		/// <param name="references">List of references to other blobs</param>
+		/// <returns>Serialized data</returns>
+		public static ReadOnlySequence<byte> Serialize(ReadOnlySequence<byte> data, IReadOnlyList<BlobId> references)
+		{
+			ByteArrayBuilder writer = new ByteArrayBuilder();
+			writer.WriteVariableLengthArray(references, x => writer.WriteBlobId(x));
+			writer.WriteUnsignedVarInt((ulong)data.Length);
+
+			ReadOnlySequenceBuilder<byte> builder = new ReadOnlySequenceBuilder<byte>();
+			builder.Append(writer.ToByteArray());
+			builder.Append(data);
+
+			return builder.Construct();
+		}
+
+		/// <summary>
+		/// Deserialize a blob from a block of memory
+		/// </summary>
+		/// <param name="memory">Memory to deserialize from</param>
+		/// <returns>Deserialized blob data. May reference the supplied memory.</returns>
+		public static IBlob Deserialize(ReadOnlyMemory<byte> memory) => Deserialize(new ReadOnlySequence<byte>(memory));
+
+		/// <summary>
+		/// Deserialize a blob from a block of memory
+		/// </summary>
+		/// <param name="sequence">Sequence to deserialize from</param>
+		/// <returns>Deserialized blob data. May reference the supplied memory.</returns>
+		public static IBlob Deserialize(ReadOnlySequence<byte> sequence)
+		{
+			MemoryReader reader = new MemoryReader(sequence.First);
+			IReadOnlyList<BlobId> references = reader.ReadVariableLengthArray(() => reader.ReadBlobId());
+			long length = (long)reader.ReadUnsignedVarInt();
+
+			ReadOnlyMemory<byte> data = sequence.Slice(sequence.First.Length - reader.Memory.Length).AsSingleSegment();
+			return new InMemoryBlob(data, references);
+		}
+	}
 }
