Fix missing XML documentation

diff --git a/Engine/Source/Programs/Shared/EpicGames.Core/FileTransaction.cs b/Engine/Source/Programs/Shared/EpicGames.Core/FileTransaction.cs
--- a/Engine/Source/Programs/Shared/EpicGames.Core/FileTransaction.cs
+++ b/Engine/Source/Programs/Shared/EpicGames.Core/FileTransaction.cs
@@ -89,6 +89,7 @@ namespace EpicGames.Core
 		/// </summary>
 		/// <param name="file">File to write to</param>
 		/// <param name="data">Data to be written</param>
+		/// <param name="cancellationToken">Cancellation token for the operation</param>
 		public static Task WriteAllBytesAsync(FileReference file, ReadOnlyMemory<byte> data, CancellationToken cancellationToken = default) => WriteAllBytesAsync(file, new ReadOnlySequence<byte>(data), cancellationToken);
 
 		/// <summary>
@@ -96,6 +97,7 @@ namespace EpicGames.Core
 		/// </summary>
 		/// <param name="file">File to write to</param>
 		/// <param name="data">Data to be written</param>
+		/// <param name="cancellationToken">Cancellation token for the operation</param>
 		public static async Task WriteAllBytesAsync(FileReference file, ReadOnlySequence<byte> data, CancellationToken cancellationToken = default)
 		{
 			using (FileTransactionStream stream = OpenWrite(file))
diff --git a/Engine/Source/Programs/Shared/EpicGames.Serialization/CbPackage.cs b/Engine/Source/Programs/Shared/EpicGames.Serialization/CbPackage.cs
--- a/Engine/Source/Programs/Shared/EpicGames.Serialization/CbPackage.cs
+++ b/Engine/Source/Programs/Shared/EpicGames.Serialization/CbPackage.cs
@@ -20,6 +20,12 @@ namespace EpicGames.Serialization
     /// <remarks>see CbPackage definition in Zen https://github.com/EpicGames/zen/blob/main/zenhttp/include/zenhttp/httpshared.h</remarks>
     public struct CbPackageHeader
     {
+		/// <summary>
+		/// Constructor
+		/// </summary>
+		/// <param name="attachmentCount"></param>
+		/// <param name="reserved1"></param>
+		/// <param name="reserved2"></param>
         public CbPackageHeader(uint attachmentCount, uint reserved1 = 0, uint reserved2 = 0)
         {
             HeaderMagic = CbPackageConstants.CbPackageHeaderMagic;
