Fix duplicated data packet when copying compressed data from Perforce response streams.

diff --git a/Engine/Source/Programs/Shared/EpicGames.Perforce.Native/Client.cpp b/Engine/Source/Programs/Shared/EpicGames.Perforce.Native/Client.cpp
--- a/Engine/Source/Programs/Shared/EpicGames.Perforce.Native/Client.cpp
+++ b/Engine/Source/Programs/Shared/EpicGames.Perforce.Native/Client.cpp
@@ -633,6 +633,7 @@ public:
 		if (GzipInst && mode == FOM_WRITE && GzipInst->os > Buffer)
 		{
 			FFileSys::Write(Buffer, (int)(GzipInst->os - Buffer), e);
+			GzipInst->os = Buffer;
 		}
 
 		FFileSys::Close(e);
