Epicgames.Build: Use a Lazy FileSystemInfo to ensure info is initialized before use while still deferring the Refresh until accessed

#rnx

diff --git a/Engine/Source/Programs/Shared/EpicGames.Build/System/DirectoryItem.cs b/Engine/Source/Programs/Shared/EpicGames.Build/System/DirectoryItem.cs
--- a/Engine/Source/Programs/Shared/EpicGames.Build/System/DirectoryItem.cs
+++ b/Engine/Source/Programs/Shared/EpicGames.Build/System/DirectoryItem.cs
@@ -5,9 +5,6 @@ using System.Collections.Concurrent;
 using System.Collections.Generic;
 using System.Diagnostics.CodeAnalysis;
 using System.IO;
-using System.Linq;
-using System.Text;
-using System.Threading.Tasks;
 using EpicGames.Core;
 
 namespace UnrealBuildBase
@@ -25,7 +22,7 @@ namespace UnrealBuildBase
 		/// <summary>
 		/// Cached value for whether the directory exists
 		/// </summary>
-		DirectoryInfo Info;
+		Lazy<DirectoryInfo> Info;
 
 		/// <summary>
 		/// Cached map of name to subdirectory item
@@ -50,53 +47,45 @@ namespace UnrealBuildBase
 		private DirectoryItem(DirectoryReference Location, DirectoryInfo Info)
 		{
 			this.Location = Location;
-			this.Info = Info;
+			this.Info = new Lazy<DirectoryInfo>(() =>
+			{
+				Info.Refresh();
+				return Info;
+			});
 		}
 
 		/// <summary>
 		/// The name of this directory
 		/// </summary>
-		public string Name
-		{
-			get { return Info.Name; }
-		}
+		public string Name => Info.Value.Name;
 
 		/// <summary>
 		/// The full name of this directory
 		/// </summary>
-		public string FullName
-		{
-			get { return Location.FullName; }
-		}
+		public string FullName => Location.FullName;
 
 		/// <summary>
 		/// Whether the directory exists or not
 		/// </summary>
-		public bool Exists
-		{
-			get { return Info.Exists; }
-		}
+		public bool Exists => Info.Value.Exists;
 
 		/// <summary>
 		/// The last write time of the file.
 		/// </summary>
-		public DateTime LastWriteTimeUtc
-		{
-			get { return Info.LastWriteTimeUtc; }
-		}
+		public DateTime LastWriteTimeUtc => Info.Value.LastWriteTimeUtc;
 
 		/// <summary>
 		/// Gets the parent directory item
 		/// </summary>
 		public DirectoryItem? GetParentDirectoryItem()
 		{
-			if(Info.Parent == null)
+			if (Info.Value.Parent == null)
 			{
 				return null;
 			}
 			else
 			{
-				return GetItemByDirectoryInfo(Info.Parent);
+				return GetItemByDirectoryInfo(Info.Value.Parent);
 			}
 		}
 
@@ -128,20 +117,7 @@ namespace UnrealBuildBase
 		/// <returns>The directory item for this location</returns>
 		public static DirectoryItem GetItemByDirectoryReference(DirectoryReference Location)
 		{
-			DirectoryItem? Result;
-			if(!LocationToItem.TryGetValue(Location, out Result))
-			{
-				DirectoryItem NewItem = new DirectoryItem(Location, new DirectoryInfo(Location.FullName));
-				if(LocationToItem.TryAdd(Location, NewItem))
-				{
-					Result = NewItem;
-				}
-				else
-				{
-					Result = LocationToItem[Location];
-				}
-			}
-			return Result;
+			return LocationToItem.GetOrAdd(Location, _ => new DirectoryItem(Location, new DirectoryInfo(Location.FullName)));
 		}
 
 		/// <summary>
@@ -152,21 +128,7 @@ namespace UnrealBuildBase
 		public static DirectoryItem GetItemByDirectoryInfo(DirectoryInfo Info)
 		{
 			DirectoryReference Location = new DirectoryReference(Info);
-
-			DirectoryItem? Result;
-			if(!LocationToItem.TryGetValue(Location, out Result))
-			{
-				DirectoryItem NewItem = new DirectoryItem(Location, Info);
-				if(LocationToItem.TryAdd(Location, NewItem))
-				{
-					Result = NewItem;
-				}
-				else
-				{
-					Result = LocationToItem[Location];
-				}
-			}
-			return Result;
+			return LocationToItem.GetOrAdd(Location, _ => new DirectoryItem(Location, Info));
 		}
 
 		/// <summary>
@@ -174,12 +136,17 @@ namespace UnrealBuildBase
 		/// </summary>
 		public void ResetCachedInfo()
 		{
-			Info = new DirectoryInfo(Info.FullName);
+			Info = new Lazy<DirectoryInfo>(() =>
+			{
+				DirectoryInfo Info = Location.ToDirectoryInfo();
+				Info.Refresh();
+				return Info;
+			});
 
 			Dictionary<string, DirectoryItem>? PrevDirectories = Directories;
-			if(PrevDirectories != null)
+			if (PrevDirectories != null)
 			{
-				foreach(DirectoryItem SubDirectory in PrevDirectories.Values)
+				foreach (DirectoryItem SubDirectory in PrevDirectories.Values)
 				{
 					SubDirectory.ResetCachedInfo();
 				}
@@ -187,9 +154,9 @@ namespace UnrealBuildBase
 			}
 
 			Dictionary<string, FileItem>? PrevFiles = Files;
-			if(PrevFiles != null)
+			if (PrevFiles != null)
 			{
-				foreach(FileItem File in PrevFiles.Values)
+				foreach (FileItem File in PrevFiles.Values)
 				{
 					File.ResetCachedInfo();
 				}
@@ -202,9 +169,14 @@ namespace UnrealBuildBase
 		/// </summary>
 		public static void ResetAllCachedInfo_SLOW()
 		{
-			foreach(DirectoryItem Item in LocationToItem.Values)
+			foreach (DirectoryItem Item in LocationToItem.Values)
 			{
-				Item.Info = new DirectoryInfo(Item.Info.FullName);
+				Item.Info = new Lazy<DirectoryInfo>(() =>
+				{
+					DirectoryInfo Info = Item.Location.ToDirectoryInfo();
+					Info.Refresh();
+					return Info;
+				});
 				Item.Directories = null;
 				Item.Files = null;
 			}
@@ -216,18 +188,18 @@ namespace UnrealBuildBase
 		/// </summary>
 		public void CacheDirectories()
 		{
-			if(Directories == null)
+			if (Directories == null)
 			{
 				Dictionary<string, DirectoryItem> NewDirectories = new Dictionary<string, DirectoryItem>(DirectoryReference.Comparer);
-				if(Info.Exists)
+				if (Info.Value.Exists)
 				{
-					foreach(DirectoryInfo SubDirectoryInfo in Info.EnumerateDirectories())
+					foreach (DirectoryInfo SubDirectoryInfo in Info.Value.EnumerateDirectories())
 					{
-						if(SubDirectoryInfo.Name.Length == 1 && SubDirectoryInfo.Name[0] == '.')
+						if (SubDirectoryInfo.Name.Length == 1 && SubDirectoryInfo.Name[0] == '.')
 						{
 							continue;
 						}
-						else if(SubDirectoryInfo.Name.Length == 2 && SubDirectoryInfo.Name[0] == '.' && SubDirectoryInfo.Name[1] == '.')
+						else if (SubDirectoryInfo.Name.Length == 2 && SubDirectoryInfo.Name[0] == '.' && SubDirectoryInfo.Name[1] == '.')
 						{
 							continue;
 						}
@@ -259,14 +231,14 @@ namespace UnrealBuildBase
 		/// <returns>True if the file exists, false otherwise</returns>
 		public bool TryGetDirectory(string Name, [NotNullWhen(true)] out DirectoryItem? OutDirectory)
 		{
-			if(Name.Length > 0 && Name[0] == '.')
+			if (Name.Length > 0 && Name[0] == '.')
 			{
-				if(Name.Length == 1)
+				if (Name.Length == 1)
 				{
 					OutDirectory = this;
 					return true;
 				}
-				else if(Name.Length == 2 && Name[1] == '.')
+				else if (Name.Length == 2 && Name[1] == '.')
 				{
 					OutDirectory = GetParentDirectoryItem();
 					return OutDirectory != null;
@@ -282,12 +254,12 @@ namespace UnrealBuildBase
 		/// </summary>
 		public void CacheFiles()
 		{
-			if(Files == null)
+			if (Files == null)
 			{
 				Dictionary<string, FileItem> NewFiles = new Dictionary<string, FileItem>(FileReference.Comparer);
-				if(Info.Exists)
+				if (Info.Value.Exists)
 				{
-					foreach(FileInfo FileInfo in Info.EnumerateFiles())
+					foreach (FileInfo FileInfo in Info.Value.EnumerateFiles())
 					{
 						FileItem FileItem = FileItem.GetItemByFileInfo(FileInfo);
 						FileItem.UpdateCachedDirectory(this);
diff --git a/Engine/Source/Programs/Shared/EpicGames.Build/System/FileItem.cs b/Engine/Source/Programs/Shared/EpicGames.Build/System/FileItem.cs
--- a/Engine/Source/Programs/Shared/EpicGames.Build/System/FileItem.cs
+++ b/Engine/Source/Programs/Shared/EpicGames.Build/System/FileItem.cs
@@ -1,14 +1,11 @@
 // Copyright Epic Games, Inc. All Rights Reserved.
 
 using System;
-using System.Collections.Generic;
-using System.Text;
-using System.IO;
+using System.Collections.Concurrent;
 using System.Diagnostics;
+using System.IO;
 using System.Threading;
-using System.Runtime.Serialization;
 using EpicGames.Core;
-using System.Collections.Concurrent;
 using Microsoft.Extensions.Logging;
 
 namespace UnrealBuildBase
@@ -32,7 +29,7 @@
 		/// <summary>
 		/// The information about the file.
 		/// </summary>
-		FileInfo Info;
+		Lazy<FileInfo> Info;
 
 		/// <summary>
 		/// A case-insensitive dictionary that's used to map each unique file name to a single FileItem object.
@@ -47,32 +44,27 @@
 		private FileItem(FileReference Location, FileInfo Info)
 		{
 			this.Location = Location;
-			this.Info = Info;
+			this.Info = new Lazy<FileInfo>(() =>
+			{
+				Info.Refresh();
+				return Info;
+			});
 		}
 
 		/// <summary>
 		/// Name of this file
 		/// </summary>
-		public string Name
-		{
-			get { return Info.Name; }
-		}
+		public string Name => Info.Value.Name;
 
 		/// <summary>
 		/// Full name of this file
 		/// </summary>
-		public string FullName
-		{
-			get { return Location.FullName; }
-		}
+		public string FullName => Location.FullName;
 
 		/// <summary>
 		/// Accessor for the absolute path to the file
 		/// </summary>
-		public string AbsolutePath
-		{
-			get { return Location.FullName; }
-		}
+		public string AbsolutePath => Location.FullName;
 
 		/// <summary>
 		/// Gets the directory that this file is in
@@ -81,7 +73,7 @@
 		{
 			get
 			{
-				if(CachedDirectory == null)
+				if (CachedDirectory == null)
 				{
 					CachedDirectory = DirectoryItem.GetItemByDirectoryReference(Location.Directory);
 				}
@@ -92,34 +84,22 @@
 		/// <summary>
 		/// Whether the file exists.
 		/// </summary>
-		public bool Exists
-		{
-			get { return Info.Exists; }
-		}
+		public bool Exists => Info.Value.Exists;
 
 		/// <summary>
 		/// Size of the file if it exists, otherwise -1
 		/// </summary>
-		public long Length
-		{
-			get { return Info.Length; }
-		}
+		public long Length => Info.Value.Length;
 
 		/// <summary>
 		/// The attributes for this file
 		/// </summary>
-		public FileAttributes Attributes
-		{
-			get { return Info.Attributes; }
-		}
+		public FileAttributes Attributes => Info.Value.Attributes;
 
 		/// <summary>
 		/// The last write time of the file.
 		/// </summary>
-		public DateTime LastWriteTimeUtc
-		{
-			get { return Info.LastWriteTimeUtc; }
-		}
+		public DateTime LastWriteTimeUtc => Info.Value.LastWriteTimeUtc;
 
 		/// <summary>
 		/// Determines if the file has the given extension
@@ -168,21 +148,7 @@
 		public static FileItem GetItemByFileInfo(FileInfo Info)
 		{
 			FileReference Location = new FileReference(Info);
-
-			FileItem? Result;
-			if (!UniqueSourceFileMap.TryGetValue(Location, out Result))
-			{
-				FileItem NewFileItem = new FileItem(Location, Info);
-				if(UniqueSourceFileMap.TryAdd(Location, NewFileItem))
-				{
-					Result = NewFileItem;
-				}
-				else
-				{
-					Result = UniqueSourceFileMap[Location];
-				}
-			}
-			return Result;
+			return UniqueSourceFileMap.GetOrAdd(Location, _ => new FileItem(Location, Info));
 		}
 
 		/// <summary>
@@ -192,20 +158,7 @@
 		/// <returns>The FileItem that represents the given a full file path.</returns>
 		public static FileItem GetItemByFileReference(FileReference Location)
 		{
-			FileItem? Result;
-			if (!UniqueSourceFileMap.TryGetValue(Location, out Result))
-			{
-				FileItem NewFileItem = new FileItem(Location, Location.ToFileInfo());
-				if(UniqueSourceFileMap.TryAdd(Location, NewFileItem))
-				{
-					Result = NewFileItem;
-				}
-				else
-				{
-					Result = UniqueSourceFileMap[Location];
-				}
-			}
-			return Result;
+			return UniqueSourceFileMap.GetOrAdd(Location, _ => new FileItem(Location, Location.ToFileInfo()));
 		}
 
 		/// <summary>
@@ -260,7 +213,12 @@
 		/// </summary>
 		public void ResetCachedInfo()
 		{
-			Info = Location.ToFileInfo();
+			Info = new Lazy<FileInfo>(() =>
+			{
+				FileInfo Info = Location.ToFileInfo();
+				Info.Refresh();
+				return Info;
+			});
 		}
 
 		/// <summary>
@@ -268,7 +226,7 @@
 		/// </summary>
 		public static void ResetAllCachedInfo_SLOW()
 		{
-			foreach(FileItem Item in UniqueSourceFileMap.Values)
+			foreach (FileItem Item in UniqueSourceFileMap.Values)
 			{
 				Item.ResetCachedInfo();
 			}
@@ -321,8 +279,8 @@
 			string Name = Reader.ReadString()!;
 
 			FileItem FileItem = FileItem.GetItemByFileReference(FileReference.Combine(Directory.Location, Name));
-            FileItem.UpdateCachedDirectory(Directory);
-            return FileItem;
+			FileItem.UpdateCachedDirectory(Directory);
+			return FileItem;
 		}
 
 		/// <summary>
@@ -344,7 +302,11 @@
 		public static void WriteCompactFileItem(this BinaryArchiveWriter Writer, FileItem FileItem)
 		{
 			// Use lambda that doesn't require anything to be captured thus eliminating an allocation.
-			Writer.WriteObjectReference<FileItem>(FileItem, (BinaryArchiveWriter Writer, FileItem FileItem) => { Writer.WriteDirectoryItem(FileItem.GetDirectoryItem()); Writer.WriteString(FileItem.Name); });
+			Writer.WriteObjectReference<FileItem>(FileItem, (BinaryArchiveWriter Writer, FileItem FileItem) =>
+			{
+				Writer.WriteDirectoryItem(FileItem.GetDirectoryItem());
+				Writer.WriteString(FileItem.Name);
+			});
 		}
 	}
 }
