Horde: Remove legacy bundle implementation, and fix up tests to work with new implementation.

diff --git a/Engine/Source/Programs/Horde/Horde.Agent/Execution/ReplicatedExecutor.cs b/Engine/Source/Programs/Horde/Horde.Agent/Execution/ReplicatedExecutor.cs
--- a/Engine/Source/Programs/Horde/Horde.Agent/Execution/ReplicatedExecutor.cs
+++ b/Engine/Source/Programs/Horde/Horde.Agent/Execution/ReplicatedExecutor.cs
@@ -1,11 +1,10 @@
 // Copyright Epic Games, Inc. All Rights Reserved.
 
+using System;
 using System.Collections.Generic;
 using System.Threading;
 using System.Threading.Tasks;
 using EpicGames.Core;
-using EpicGames.Horde.Bundles;
-using EpicGames.Horde.Bundles.Nodes;
 using EpicGames.Horde.Storage;
 using Horde.Agent.Utility;
 using HordeCommon.Rpc;
@@ -37,6 +36,9 @@ namespace Horde.Agent.Execution
 			_workspaceDir = DirectoryReference.Combine(_rootDir, "Replicated");
 			FileUtils.ForceDeleteDirectoryContents(_workspaceDir);
 
+			throw new NotImplementedException("Needs refactoring for new bundle/store logic");
+
+/*
 			using (IMemoryCache cache = new MemoryCache(new MemoryCacheOptions { SizeLimit = 1024 * 1024 * 1000 }))
 			{
 				BucketId bucketId = new BucketId(_job.StreamId.ToString());
@@ -45,7 +47,7 @@ namespace Horde.Agent.Execution
 			}
 
 			// Get all the environment variables for jobs
-			_envVars["IsBuildMachine"] = "1";
+			_envVars["IsBuildMachine"] = "1";*/
 		}
 
 		protected override Task<bool> SetupAsync(BeginStepResponse step, ILogger logger, CancellationToken cancellationToken)
diff --git a/Engine/Source/Programs/Shared/EpicGames.Horde.Tests/BundleStoreTests.cs b/Engine/Source/Programs/Shared/EpicGames.Horde.Tests/BundleStoreTests.cs
--- a/Engine/Source/Programs/Shared/EpicGames.Horde.Tests/BundleStoreTests.cs
+++ b/Engine/Source/Programs/Shared/EpicGames.Horde.Tests/BundleStoreTests.cs
@@ -109,7 +109,7 @@ namespace EpicGames.Horde.Tests
 			InMemoryBlobStore blobStore = new InMemoryBlobStore();
 
 			// Generate a tree
-			using (BundleStore<DirectoryNode> store = new BundleStore<DirectoryNode>(blobStore, new BundleOptions(), cache))
+			using (BundleStore store = new BundleStore(blobStore, new BundleOptions(), cache))
 			{
 				DirectoryNode root = new DirectoryNode();
 				DirectoryNode hello = root.AddDirectory("hello");
@@ -120,9 +120,9 @@ namespace EpicGames.Horde.Tests
 			}
 
 			// Check we can read it back in
-			using (BundleStore<DirectoryNode> store = new BundleStore<DirectoryNode>(blobStore, new BundleOptions(), cache))
+			using (BundleStore store = new BundleStore(blobStore, new BundleOptions(), cache))
 			{
-				DirectoryNode root = await store.ReadTreeAsync(new RefId("test"));
+				DirectoryNode root = await store.ReadTreeAsync<DirectoryNode>(new RefId("test"));
 				await CheckDirectoryTreeAsync(root);
 			}
 		}
@@ -147,12 +147,13 @@ namespace EpicGames.Horde.Tests
 			InMemoryBlobStore blobStore = new InMemoryBlobStore();
 
 			// Generate a tree
-			using (BundleStore<DirectoryNode> store = new BundleStore<DirectoryNode>(blobStore, new BundleOptions(), cache))
+			using (BundleStore store = new BundleStore(blobStore, new BundleOptions(), cache))
 			{
 				DirectoryNode root = new DirectoryNode();
 				DirectoryNode hello = root.AddDirectory("hello");
 
-				FileNode world = await hello.AddFileAsync("world", FileEntryFlags.None, Encoding.UTF8.GetBytes("world"), new ChunkingOptions(), CancellationToken.None);
+				FileEntry world = hello.AddFile("world", FileEntryFlags.None);
+				await world.AppendAsync(Encoding.UTF8.GetBytes("world"), new ChunkingOptions(), CancellationToken.None);
 				
 				await store.WriteTreeAsync(new RefId("test"), root, true, CancellationToken.None);
 
@@ -160,9 +161,9 @@ namespace EpicGames.Horde.Tests
 			}
 
 			// Check we can read it back in
-			using (BundleStore<DirectoryNode> store = new BundleStore<DirectoryNode>(blobStore, new BundleOptions(), cache))
+			using (BundleStore store = new BundleStore(blobStore, new BundleOptions(), cache))
 			{
-				DirectoryNode root = await store.ReadTreeAsync(new RefId("test"));
+				DirectoryNode root = await store.ReadTreeAsync<DirectoryNode>(new RefId("test"));
 				await CheckFileTreeAsync(root);
 			}
 		}
@@ -237,7 +238,7 @@ namespace EpicGames.Horde.Tests
 			}
 
 			// Generate a tree
-			using (BundleStore<DirectoryNode> store = new BundleStore<DirectoryNode>(blobStore, new BundleOptions { MaxInlineBlobSize = 1024, MaxBlobSize = 1024 }, cache))
+			using (BundleStore store = new BundleStore(blobStore, new BundleOptions { MaxInlineBlobSize = 1024, MaxBlobSize = 1024 }, cache))
 			{
 				DirectoryNode root = new DirectoryNode();
 
@@ -246,7 +247,8 @@ namespace EpicGames.Horde.Tests
 				options.LeafOptions.TargetSize = 128;
 				options.LeafOptions.MaxSize = 64 * 1024;
 
-				FileNode file = await root.AddFileAsync("test", FileEntryFlags.None, data, options, CancellationToken.None);
+				FileEntry file = root.AddFile("test", FileEntryFlags.None);
+				await file.AppendAsync(data, options, CancellationToken.None);
 
 				await store.WriteTreeAsync(new RefId("test"), root, true, CancellationToken.None);
 
@@ -254,9 +256,9 @@ namespace EpicGames.Horde.Tests
 			}
 
 			// Check we can read it back in
-			using (BundleStore<DirectoryNode> store = new BundleStore<DirectoryNode>(blobStore, new BundleOptions(), cache))
+			using (BundleStore store = new BundleStore(blobStore, new BundleOptions(), cache))
 			{
-				DirectoryNode root = await store.ReadTreeAsync(new RefId("test"));
+				DirectoryNode root = await store.ReadTreeAsync<DirectoryNode>(new RefId("test"));
 				await CheckLargeFileTreeAsync(root, data);
 
 				TreeNodeRef<FileNode> file = root.GetFileEntry("test");
diff --git a/Engine/Source/Programs/Shared/EpicGames.Horde.Tests/BundleTests.cs b/Engine/Source/Programs/Shared/EpicGames.Horde.Tests/BundleTests.cs
--- a/Engine/Source/Programs/Shared/EpicGames.Horde.Tests/BundleTests.cs
+++ b/Engine/Source/Programs/Shared/EpicGames.Horde.Tests/BundleTests.cs
@@ -1,17 +1,21 @@
 // Copyright Epic Games, Inc. All Rights Reserved.
 
 using EpicGames.Core;
-using EpicGames.Horde.Bundles;
-using EpicGames.Horde.Bundles.Nodes;
 using EpicGames.Horde.Storage;
+using EpicGames.Horde.Storage.Backends;
+using EpicGames.Horde.Storage.Bundles;
 using EpicGames.Horde.Storage.Impl;
+using EpicGames.Horde.Storage.Nodes;
 using EpicGames.Serialization;
 using Microsoft.Extensions.Caching.Memory;
 using Microsoft.VisualStudio.TestTools.UnitTesting;
 using System;
+using System.Buffers;
+using System.Collections.Concurrent;
 using System.Collections.Generic;
 using System.Linq;
 using System.Text;
+using System.Threading;
 using System.Threading.Tasks;
 
 namespace EpicGames.Horde.Tests
@@ -20,18 +24,21 @@ namespace EpicGames.Horde.Tests
 	public sealed class BundleTests : IDisposable
 	{
 		readonly MemoryCache _cache;
-		readonly MemoryStorageClient _storageClient = new MemoryStorageClient();
-		readonly NamespaceId _namespaceId = new NamespaceId("namespace");
-		readonly BucketId _bucketId = new BucketId("bucket");
+		readonly InMemoryBlobStore _blobStore;
+		readonly ITreeStore _treeStore;
 
 		public BundleTests()
 		{
 			_cache = new MemoryCache(new MemoryCacheOptions { SizeLimit = 50 * 1024 * 1024 });
+
+			_blobStore = new InMemoryBlobStore();
+			_treeStore = new BundleStore(_blobStore, new BundleOptions(), _cache);
 		}
 
 		public void Dispose()
 		{
 			_cache.Dispose();
+			_treeStore.Dispose();
 		}
 
 		[TestMethod]
@@ -69,121 +76,133 @@ namespace EpicGames.Horde.Tests
 			ChunkingOptions options = new ChunkingOptions();
 			options.LeafOptions = new ChunkingOptionsForNodeType(8, 8, 8);
 
-			FileNode node = new FileNode();
-			node.Append(new byte[7], options);
-			Assert.AreEqual(0, node.Depth);
-			Assert.AreEqual(7, node.Payload.Length);
+			FileNode node = new LeafFileNode();
+			node = node.Append(new byte[7], options);
+			Assert.IsTrue(node is LeafFileNode);
+			Assert.AreEqual(7, ((LeafFileNode)node).Data.Length);
 
-			node = new FileNode();
-			node.Append(new byte[8], options);
-			Assert.AreEqual(0, node.Depth);
-			Assert.AreEqual(8, node.Payload.Length);
+			node = new LeafFileNode();
+			node = node.Append(new byte[8], options);
+			Assert.IsTrue(node is LeafFileNode);
+			Assert.AreEqual(8, ((LeafFileNode)node).Data.Length);
 
-			node = new FileNode();
-			node.Append(new byte[9], options);
-			Assert.AreEqual(1, node.Depth);
-			Assert.AreEqual(2, node.Children.Count);
+			node = new LeafFileNode();
+			node = node.Append(new byte[9], options);
+			Assert.IsTrue(node is InteriorFileNode);
+			Assert.AreEqual(2, node.GetReferences().Count);
 
-			FileNode? childNode1 = node.Children[0].Node;
+			FileNode? childNode1 = (FileNode?)node.GetReferences()[0].Node;
 			Assert.IsNotNull(childNode1);
-			Assert.AreEqual(0, childNode1!.Depth);
-			Assert.AreEqual(8, childNode1!.Payload.Length);
+			Assert.IsTrue(childNode1 is LeafFileNode);
+			Assert.AreEqual(8, ((LeafFileNode)childNode1!).Data.Length);
 
-			FileNode? childNode2 = node.Children[1].Node;
+			FileNode? childNode2 = (FileNode?)node.GetReferences()[1].Node;
 			Assert.IsNotNull(childNode2);
-			Assert.AreEqual(0, childNode2!.Depth);
-			Assert.AreEqual(1, childNode2!.Payload.Length);
+			Assert.IsTrue(childNode2 is LeafFileNode);
+			Assert.AreEqual(1, ((LeafFileNode)childNode2!).Data.Length);
 		}
 
 		[TestMethod]
 		public void SerializationTests()
 		{
-			List<BundleImport> imports = new List<BundleImport>();
-			imports.Add(new BundleImport(IoHash.Compute(Encoding.UTF8.GetBytes("import1")), 0, 123));
-			imports.Add(new BundleImport(IoHash.Compute(Encoding.UTF8.GetBytes("import2")), 1, 456));
+			BundleHeader oldHeader;
+			{
+				List<BundleImport> imports = new List<BundleImport>();
+				imports.Add(new BundleImport(new BlobId("import1"), 10, new (int, IoHash)[] { (5, IoHash.Compute(Encoding.UTF8.GetBytes("blob1"))) }));
+				imports.Add(new BundleImport(new BlobId("import2"), 20, new (int, IoHash)[] { (6, IoHash.Compute(Encoding.UTF8.GetBytes("blob2"))) }));
 
-			BundleObject oldObject = new BundleObject();
-			oldObject.ImportObjects.Add(new BundleImportObject(new CbObjectAttachment(IoHash.Compute(new byte[] { 1, 2, 3 })), 12345, imports));
-			oldObject.Exports.Add(new BundleExport(IoHash.Compute(Encoding.UTF8.GetBytes("export1")), 2, new BundleCompressionPacket(0) { EncodedLength = 20, DecodedLength = 40 }, 0, 40, new int[] { 1, 2 }));
-			oldObject.Exports.Add(new BundleExport(IoHash.Compute(Encoding.UTF8.GetBytes("export2")), 3, new BundleCompressionPacket(20) { EncodedLength = 10, DecodedLength = 20 }, 0, 20, new int[] { -1 }));
+				List<BundleExport> exports = new List<BundleExport>();
+				exports.Add(new BundleExport(IoHash.Compute(Encoding.UTF8.GetBytes("export1")), 2, new int[] { 1, 2 }));
+				exports.Add(new BundleExport(IoHash.Compute(Encoding.UTF8.GetBytes("export2")), 3, new int[] { 3 }));
 
-			CbObject serializedObject = CbSerializer.Serialize(oldObject);
-			ReadOnlyMemory<byte> serializedData = serializedObject.GetView();
+				List<BundlePacket> packets = new List<BundlePacket>();
+				packets.Add(new BundlePacket(20, 40));
+				packets.Add(new BundlePacket(10, 20));
+
+				oldHeader = new BundleHeader(imports, exports, packets);
+			}
 
-			BundleObject newObject = CbSerializer.Deserialize<BundleObject>(serializedData);
+			ByteArrayBuilder writer = new ByteArrayBuilder();
+			oldHeader.Write(writer);
+			byte[] serializedData = writer.ToByteArray();
 
-			Assert.AreEqual(oldObject.ImportObjects.Count, newObject.ImportObjects.Count);
-			for (int idx = 0; idx < oldObject.ImportObjects.Count; idx++)
+			BundleHeader newHeader = new BundleHeader(new MemoryReader(serializedData));
+				
+			Assert.AreEqual(oldHeader.Imports.Count, newHeader.Imports.Count);
+			for (int idx = 0; idx < oldHeader.Imports.Count; idx++)
 			{
-				BundleImportObject oldObjectImport = oldObject.ImportObjects[idx];
-				BundleImportObject newObjectImport = newObject.ImportObjects[idx];
+				BundleImport oldImport = oldHeader.Imports[idx];
+				BundleImport newImport = newHeader.Imports[idx];
 
-				Assert.AreEqual(oldObjectImport.Object, newObjectImport.Object);
-				Assert.AreEqual(oldObjectImport.TotalCost, newObjectImport.TotalCost);
-				Assert.AreEqual(oldObjectImport.Imports.Count, newObjectImport.Imports.Count);
+				Assert.AreEqual(oldImport.BlobId, newImport.BlobId);
+				Assert.AreEqual(oldImport.ExportCount, newImport.ExportCount);
+				Assert.AreEqual(oldImport.Exports.Count, newImport.Exports.Count);
 
-				for (int importIdx = 0; importIdx < oldObjectImport.Imports.Count; importIdx++)
+				for (int importIdx = 0; importIdx < oldImport.Exports.Count; importIdx++)
 				{
-					BundleImport oldImport = oldObjectImport.Imports[importIdx];
-					BundleImport newImport = newObjectImport.Imports[importIdx];
+					(int Index, IoHash Hash) oldExport = oldImport.Exports[importIdx];
+					(int Index, IoHash Hash) newExport = newImport.Exports[importIdx];
 
-					Assert.AreEqual(oldImport.Hash, newImport.Hash);
-					Assert.AreEqual(oldImport.Rank, newImport.Rank);
-					Assert.AreEqual(oldImport.Length, newImport.Length);
+					Assert.AreEqual(oldExport.Hash, newExport.Hash);
+					Assert.AreEqual(oldExport.Index, newExport.Index);
 				}
 			}
 
-			Assert.AreEqual(oldObject.Exports.Count, newObject.Exports.Count);
-			for (int idx = 0; idx < oldObject.Exports.Count; idx++)
+			Assert.AreEqual(oldHeader.Exports.Count, newHeader.Exports.Count);
+			for (int idx = 0; idx < oldHeader.Exports.Count; idx++)
 			{
-				BundleExport oldExport = oldObject.Exports[idx];
-				BundleExport newExport = newObject.Exports[idx];
+				BundleExport oldExport = oldHeader.Exports[idx];
+				BundleExport newExport = newHeader.Exports[idx];
 
 				Assert.AreEqual(oldExport.Hash, newExport.Hash);
-				Assert.AreEqual(oldExport.Rank, newExport.Rank);
-				Assert.AreEqual(oldExport.Length, newExport.Length);
-				Assert.AreEqual(oldExport.Packet.Offset, newExport.Packet.Offset);
-				Assert.AreEqual(oldExport.Packet.EncodedLength, newExport.Packet.EncodedLength);
-				Assert.AreEqual(oldExport.Packet.DecodedLength, newExport.Packet.DecodedLength);
-				Assert.AreEqual(oldExport.Offset, newExport.Offset);
 				Assert.AreEqual(oldExport.Length, newExport.Length);
 				Assert.IsTrue(oldExport.References.SequenceEqual(newExport.References));
 			}
+
+			Assert.AreEqual(oldHeader.Packets.Count, newHeader.Packets.Count);
+			for (int idx = 0; idx < oldHeader.Packets.Count; idx++)
+			{
+				BundlePacket oldPacket = oldHeader.Packets[idx];
+				BundlePacket newPacket = newHeader.Packets[idx];
+				Assert.AreEqual(oldPacket.DecodedLength, newPacket.DecodedLength);
+				Assert.AreEqual(oldPacket.EncodedLength, newPacket.EncodedLength);
+			}
 		}
 
 		[TestMethod]
 		public async Task BasicTestDirectory()
 		{
-			using Bundle<DirectoryNode> newBundle = Bundle.Create<DirectoryNode>(_storageClient, _namespaceId, new DirectoryNode(), new BundleOptions(), _cache);
-			DirectoryNode node = newBundle.Root.AddDirectory("hello");
+			InMemoryBlobStore blobStore = new InMemoryBlobStore();
+			using ITreeStore treeStore = new BundleStore(blobStore, new BundleOptions(), _cache);
+
+			DirectoryNode root = new DirectoryNode();
+			DirectoryNode node = root.AddDirectory("hello");
 			DirectoryNode node2 = node.AddDirectory("world");
 
 			RefId refId = new RefId("testref");
-			await newBundle.WriteAsync(_bucketId, refId, CbObject.Empty, false);
+			await treeStore.WriteTreeAsync(refId, root, true);
 
 			// Should be stored inline
-			Assert.AreEqual(1, _storageClient.Refs.Count);
-			Assert.AreEqual(0, _storageClient.Blobs.Count);
+			Assert.AreEqual(1, blobStore.Refs.Count);
+			Assert.AreEqual(0, blobStore.Blobs.Count);
 
 			// Check the ref
-			BundleRoot root = await _storageClient.GetRefAsync<BundleRoot>(_namespaceId, _bucketId, refId);
-			BundleObject rootObject = root.Object;
-			Assert.AreEqual(3, rootObject.Exports.Count);
-			Assert.AreEqual(0, rootObject.Exports[0].Rank);
-			Assert.AreEqual(1, rootObject.Exports[1].Rank);
-			Assert.AreEqual(2, rootObject.Exports[2].Rank);
+			IBlob blob = await blobStore.ReadRefAsync(refId);
+			Bundle bundle = new Bundle(new MemoryReader(blob.Data));
+			Assert.AreEqual(0, bundle.Header.Imports.Count);
+			Assert.AreEqual(3, bundle.Header.Exports.Count);
 
 			// Create a new bundle and read it back in again
-			using Bundle<DirectoryNode> newBundle2 = await Bundle.ReadAsync<DirectoryNode>(_storageClient, _namespaceId, _bucketId, refId, new BundleOptions(), _cache);
+			DirectoryNode newRoot = await treeStore.ReadTreeAsync<DirectoryNode>(refId);
 
-			Assert.AreEqual(0, newBundle2.Root.Files.Count);
-			Assert.AreEqual(1, newBundle2.Root.Directories.Count);
-			DirectoryNode? outputNode = await newBundle.FindDirectoryAsync(newBundle2.Root, "hello");
+			Assert.AreEqual(0, newRoot.Files.Count);
+			Assert.AreEqual(1, newRoot.Directories.Count);
+			DirectoryNode? outputNode = await newRoot.FindDirectoryAsync("hello", CancellationToken.None);
 			Assert.IsNotNull(outputNode);
 
 			Assert.AreEqual(0, outputNode!.Files.Count);
 			Assert.AreEqual(1, outputNode!.Directories.Count);
-			DirectoryNode? outputNode2 = await newBundle.FindDirectoryAsync(outputNode, "world");
+			DirectoryNode? outputNode2 = await outputNode.FindDirectoryAsync("world", CancellationToken.None);
 			Assert.IsNotNull(outputNode2);
 
 			Assert.AreEqual(0, outputNode2!.Files.Count);
@@ -197,20 +216,22 @@ namespace EpicGames.Horde.Tests
 			options.MaxBlobSize = 1;
 			options.MaxInlineBlobSize = 1;
 
-			using Bundle<DirectoryNode> newBundle = Bundle.Create<DirectoryNode>(_storageClient, _namespaceId, new DirectoryNode(), options, _cache);
+			InMemoryBlobStore blobStore = new InMemoryBlobStore();
+			using ITreeStore treeStore = new BundleStore(blobStore, options, _cache);
 
-			newBundle.Root.AddDirectory("node1");
-			newBundle.Root.AddDirectory("node2");
-			newBundle.Root.AddDirectory("node3");
+			DirectoryNode root = new DirectoryNode();
+			root.AddDirectory("node1");
+			root.AddDirectory("node2");
+			root.AddDirectory("node3");
 
-			Assert.AreEqual(0, _storageClient.Refs.Count);
-			Assert.AreEqual(0, _storageClient.Blobs.Count);
+			Assert.AreEqual(0, blobStore.Refs.Count);
+			Assert.AreEqual(0, blobStore.Blobs.Count);
 
 			RefId refId = new RefId("ref");
-			await newBundle.WriteAsync(_bucketId, refId, CbObject.Empty, false);
+			await treeStore.WriteTreeAsync(refId, root, true);
 
-			Assert.AreEqual(1, _storageClient.Refs.Count);
-			Assert.AreEqual(1, _storageClient.Blobs.Count);
+			Assert.AreEqual(1, blobStore.Refs.Count);
+			Assert.AreEqual(1, blobStore.Blobs.Count);
 		}
 
 		[TestMethod]
@@ -220,34 +241,43 @@ namespace EpicGames.Horde.Tests
 			options.MaxBlobSize = 1;
 			options.MaxInlineBlobSize = 1;
 
-			using Bundle<DirectoryNode> initialBundle = Bundle.Create<DirectoryNode>(_storageClient, _namespaceId, new DirectoryNode(), options, _cache);
-
-			DirectoryNode node1 = initialBundle.Root.AddDirectory("node1");
-			DirectoryNode node2 = node1.AddDirectory("node2");
-			DirectoryNode node3 = node2.AddDirectory("node3");
-			DirectoryNode node4 = node3.AddDirectory("node4");
+			InMemoryBlobStore blobStore = new InMemoryBlobStore();
 
 			RefId refId = new RefId("ref");
-			await initialBundle.WriteAsync(_bucketId, refId, CbObject.Empty, false);
 
-			Assert.AreEqual(1, _storageClient.Refs.Count);
-			Assert.AreEqual(4, _storageClient.Blobs.Count);
+			using (ITreeStore oldBundle = new BundleStore(blobStore, options, _cache))
+			{
+				DirectoryNode root = new DirectoryNode();
+
+				DirectoryNode node1 = root.AddDirectory("node1");
+				DirectoryNode node2 = node1.AddDirectory("node2");
+				DirectoryNode node3 = node2.AddDirectory("node3");
+				DirectoryNode node4 = node3.AddDirectory("node4");
+
+				await oldBundle.WriteTreeAsync(refId, root, true);
 
-			using Bundle<DirectoryNode> newBundle = await Bundle.ReadAsync<DirectoryNode>(_storageClient, _namespaceId, _bucketId, refId, options, _cache);
+				Assert.AreEqual(1, blobStore.Refs.Count);
+				Assert.AreEqual(4, blobStore.Blobs.Count);
+			}
 
-			DirectoryNode? newNode1 = await newBundle.FindDirectoryAsync(newBundle.Root, "node1");
-			Assert.IsNotNull(newNode1);
+			using (ITreeStore newBundle = new BundleStore(blobStore, options, _cache))
+			{
+				DirectoryNode root = await newBundle.ReadTreeAsync<DirectoryNode>(refId);
 
-			DirectoryNode? newNode2 = await newBundle.FindDirectoryAsync(newNode1!, "node2");
-			Assert.IsNotNull(newNode2);
+				DirectoryNode? newNode1 = await root.FindDirectoryAsync("node1", CancellationToken.None);
+				Assert.IsNotNull(newNode1);
 
-			DirectoryNode? newNode3 = await newBundle.FindDirectoryAsync(newNode2!, "node3");
-			Assert.IsNotNull(newNode3);
+				DirectoryNode? newNode2 = await newNode1!.FindDirectoryAsync("node2", CancellationToken.None);
+				Assert.IsNotNull(newNode2);
 
-			DirectoryNode? newNode4 = await newBundle.FindDirectoryAsync(newNode3!, "node4");
-			Assert.IsNotNull(newNode4);
-		}
+				DirectoryNode? newNode3 = await newNode2!.FindDirectoryAsync("node3", CancellationToken.None);
+				Assert.IsNotNull(newNode3);
 
+				DirectoryNode? newNode4 = await newNode3!.FindDirectoryAsync("node4", CancellationToken.None);
+				Assert.IsNotNull(newNode4);
+			}
+		}
+		/*
 		[TestMethod]
 		public async Task CompactTest()
 		{
@@ -316,7 +346,7 @@ namespace EpicGames.Horde.Tests
 			Assert.AreEqual(1, leafObject3.Exports.Count);
 			Assert.AreEqual(0, leafObject3.ImportObjects.Count);
 		}
-
+*/
 		[TestMethod]
 		public async Task CoreAppendTest()
 		{
@@ -326,12 +356,12 @@ namespace EpicGames.Horde.Tests
 			ChunkingOptions options = new ChunkingOptions();
 			options.LeafOptions = new ChunkingOptionsForNodeType(16, 64, 256);
 
-			FileNode node = new FileNode();
+			FileNode node = new LeafFileNode();
 			for (int idx = 0; idx < data.Length; idx++)
 			{
-				node.Append(data.AsMemory(idx, 1), options);
+				node = node.Append(data.AsMemory(idx, 1), options);
 
-				byte[] outputData = await node.ToByteArrayAsync(null!);
+				byte[] outputData = await node.ToByteArrayAsync(CancellationToken.None);
 				Assert.IsTrue(data.AsMemory(0, idx + 1).Span.SequenceEqual(outputData.AsSpan(0, idx + 1)));
 			}
 		}
@@ -366,16 +396,15 @@ namespace EpicGames.Horde.Tests
 				data[idx] = (byte)idx;
 			}
 
-			using Bundle<FileNode> newBundle = Bundle.Create<FileNode>(_storageClient, _namespaceId, new FileNode(), new BundleOptions(), _cache);
-			FileNode root = newBundle.Root;
+			FileNode root = new LeafFileNode();
 
 			const int NumIterations = 100;
 			for (int idx = 0; idx < NumIterations; idx++)
 			{
-				root.Append(data, options);
+				root = root.Append(data, options);
 			}
 
-			byte[] result = await root.ToByteArrayAsync(null!);
+			byte[] result = await root.ToByteArrayAsync(CancellationToken.None);
 			Assert.AreEqual(NumIterations * data.Length, result.Length);
 
 			for (int idx = 0; idx < NumIterations; idx++)
@@ -384,26 +413,28 @@ namespace EpicGames.Horde.Tests
 				Assert.IsTrue(spanData.Span.SequenceEqual(data));
 			}
 
-			await CheckSizes(newBundle, root, options, true);
+			await CheckSizes(root, options, true);
 		}
 
-		async Task CheckSizes(Bundle bundle, FileNode node, ChunkingOptions options, bool rightmost)
+		async Task CheckSizes(FileNode node, ChunkingOptions options, bool rightmost)
 		{
-			if (node.Depth == 0)
+			if (node is LeafFileNode leafNode)
 			{
-				Assert.IsTrue(rightmost || node.Payload.Length >= options.LeafOptions.MinSize);
-				Assert.IsTrue(node.Payload.Length <= options.LeafOptions.MaxSize);
+				Assert.IsTrue(rightmost || leafNode.Data.Length >= options.LeafOptions.MinSize);
+				Assert.IsTrue(leafNode.Data.Length <= options.LeafOptions.MaxSize);
 			}
 			else
 			{
-				Assert.IsTrue(rightmost || node.Payload.Length >= options.InteriorOptions.MinSize);
-				Assert.IsTrue(node.Payload.Length <= options.InteriorOptions.MaxSize);
+				InteriorFileNode interiorNode = (InteriorFileNode)node;
 
-				int childCount = node.Children.Count;
+				Assert.IsTrue(rightmost || interiorNode.Children.Count >= options.InteriorOptions.MinSize);
+				Assert.IsTrue(interiorNode.Children.Count <= options.InteriorOptions.MaxSize);
+
+				int childCount = interiorNode.Children.Count;
 				for (int idx = 0; idx < childCount; idx++)
 				{
-					FileNode childNode = await bundle.GetAsync(node.Children[idx]);
-					await CheckSizes(bundle, childNode, options, idx == childCount - 1);
+					FileNode childNode = await interiorNode.Children[idx].ExpandAsync(CancellationToken.None);
+					await CheckSizes(childNode, options, idx == childCount - 1);
 				}
 			}
 		}
@@ -414,8 +445,10 @@ namespace EpicGames.Horde.Tests
 			BundleOptions options = new BundleOptions();
 			options.MaxBlobSize = 1;
 
-			using Bundle<DirectoryNode> newBundle = Bundle.Create<DirectoryNode>(_storageClient, _namespaceId, new DirectoryNode(), options, _cache);
-			DirectoryNode root = newBundle.Root;
+			using BundleStore treeStore = new BundleStore(_blobStore, options, _cache);
+
+			RefId refId = new RefId("ref");
+			DirectoryNode root = new DirectoryNode();
 
 			long totalLength = 0;
 			for (int idxA = 0; idxA < 10; idxA++)
@@ -429,37 +462,36 @@ namespace EpicGames.Horde.Tests
 						DirectoryNode nodeC = nodeB.AddDirectory($"{idxC}");
 						for (int idxD = 0; idxD < 10; idxD++)
 						{
-							FileNode file = nodeC.AddFile($"{idxD}", 0);
+							FileEntry file = nodeC.AddFile($"{idxD}", FileEntryFlags.None);
 							byte[] data = Encoding.UTF8.GetBytes($"This is file {idxA}/{idxB}/{idxC}/{idxD}");
 							totalLength += data.Length;
-							file.Append(data, new ChunkingOptions());
+							await file.AppendAsync(data, new ChunkingOptions(), CancellationToken.None);
 						}
 					}
 
-					int oldWorkingSetSize = GetWorkingSetSize(newBundle.Root);
-					await newBundle.TrimAsync(20);
-					int newWorkingSetSize = GetWorkingSetSize(newBundle.Root);
+					int oldWorkingSetSize = GetWorkingSetSize(root);
+					await treeStore.WriteTreeAsync(refId, root, false);
+					int newWorkingSetSize = GetWorkingSetSize(root);
 					Assert.IsTrue(newWorkingSetSize <= oldWorkingSetSize);
 					Assert.IsTrue(newWorkingSetSize <= 20);
 				}
 			}
 
-			Assert.IsTrue(_storageClient.Blobs.Count > 0);
-			Assert.IsTrue(_storageClient.Refs.Count == 0);
+			Assert.IsTrue(_blobStore.Blobs.Count > 0);
+			Assert.IsTrue(_blobStore.Refs.Count == 0);
 
-			RefId refId = new RefId("ref");
-			await newBundle.WriteAsync(_bucketId, refId, CbObject.Empty, true);
+			await treeStore.WriteTreeAsync(refId, root, true, CancellationToken.None);
 
 			Assert.AreEqual(totalLength, root.Length);
 
-			Assert.IsTrue(_storageClient.Blobs.Count > 0);
-			Assert.IsTrue(_storageClient.Refs.Count == 1);
+			Assert.IsTrue(_blobStore.Blobs.Count > 0);
+			Assert.IsTrue(_blobStore.Refs.Count == 1);
 		}
 
-		int GetWorkingSetSize(BundleNode node)
+		int GetWorkingSetSize(TreeNode node)
 		{
 			int size = 0;
-			foreach (BundleNodeRef nodeRef in node.GetReferences())
+			foreach (TreeNodeRef nodeRef in node.GetReferences())
 			{
 				if (nodeRef.Node != null)
 				{
diff --git a/Engine/Source/Programs/Shared/EpicGames.Horde.Tests/NodeTests.cs b/Engine/Source/Programs/Shared/EpicGames.Horde.Tests/NodeTests.cs
deleted file mode 100644
diff --git a/Engine/Source/Programs/Shared/EpicGames.Horde/Bundles/Bundle.cs b/Engine/Source/Programs/Shared/EpicGames.Horde/Bundles/Bundle.cs
deleted file mode 100644
diff --git a/Engine/Source/Programs/Shared/EpicGames.Horde/Bundles/BundleNode.cs b/Engine/Source/Programs/Shared/EpicGames.Horde/Bundles/BundleNode.cs
deleted file mode 100644
diff --git a/Engine/Source/Programs/Shared/EpicGames.Horde/Bundles/BundleObject.cs b/Engine/Source/Programs/Shared/EpicGames.Horde/Bundles/BundleObject.cs
deleted file mode 100644
diff --git a/Engine/Source/Programs/Shared/EpicGames.Horde/Bundles/BundleTools.cs b/Engine/Source/Programs/Shared/EpicGames.Horde/Bundles/BundleTools.cs
deleted file mode 100644
diff --git a/Engine/Source/Programs/Shared/EpicGames.Horde/Bundles/Nodes/DirectoryNode.cs b/Engine/Source/Programs/Shared/EpicGames.Horde/Bundles/Nodes/DirectoryNode.cs
deleted file mode 100644
diff --git a/Engine/Source/Programs/Shared/EpicGames.Horde/Bundles/Nodes/FileNode.cs b/Engine/Source/Programs/Shared/EpicGames.Horde/Bundles/Nodes/FileNode.cs
deleted file mode 100644
diff --git a/Engine/Source/Programs/Shared/EpicGames.Horde/Storage/BlobId.cs b/Engine/Source/Programs/Shared/EpicGames.Horde/Storage/BlobId.cs
--- a/Engine/Source/Programs/Shared/EpicGames.Horde/Storage/BlobId.cs
+++ b/Engine/Source/Programs/Shared/EpicGames.Horde/Storage/BlobId.cs
@@ -35,7 +35,7 @@ namespace EpicGames.Horde.Storage
 		}
 
 		/// <inheritdoc/>
-		public override bool Equals(object? obj) => obj is RefId refId && Equals(refId);
+		public override bool Equals(object? obj) => obj is BlobId blobId && Equals(blobId);
 
 		/// <inheritdoc/>
 		public override int GetHashCode() => Inner.GetHashCode();
diff --git a/Engine/Source/Programs/Shared/EpicGames.Horde/Storage/Nodes/FileNode.cs b/Engine/Source/Programs/Shared/EpicGames.Horde/Storage/Nodes/FileNode.cs
--- a/Engine/Source/Programs/Shared/EpicGames.Horde/Storage/Nodes/FileNode.cs
+++ b/Engine/Source/Programs/Shared/EpicGames.Horde/Storage/Nodes/FileNode.cs
@@ -143,17 +143,6 @@ namespace EpicGames.Horde.Storage.Nodes
 
 				// Increase the height of the tree by pushing the contents of this node into a new child node
 				root = new InteriorFileNode(root.Length, root);
-
-/*
-				FileNode newNode = new InteriorFileNode(root.Length, root);
-				if (_data is InternalNodeData internalData)
-				{
-					foreach (TreeNodeRef<FileNode> childNodeRef in internalData._children)
-					{
-						childNodeRef.Reparent(newNode);
-					}
-				}
-				_data = new InternalNodeData(newNode.Length, new TreeNodeRef<FileNode>(this, newNode));*/
 			}
 			return root;
 		}
@@ -442,6 +431,11 @@ namespace EpicGames.Horde.Storage.Nodes
 		long _length;
 		readonly List<TreeNodeRef<FileNode>> _children = new List<TreeNodeRef<FileNode>>();
 
+		/// <summary>
+		/// Child nodes
+		/// </summary>
+		public IReadOnlyList<TreeNodeRef<FileNode>> Children => _children;
+
 		/// <inheritdoc/>
 		public override long Length => _length;
 
