BuildGraph: Fix compile errors in tests.

diff --git a/Engine/Source/Programs/Shared/EpicGames.BuildGraph.Tests/GraphTests.cs b/Engine/Source/Programs/Shared/EpicGames.BuildGraph.Tests/GraphTests.cs
--- a/Engine/Source/Programs/Shared/EpicGames.BuildGraph.Tests/GraphTests.cs
+++ b/Engine/Source/Programs/Shared/EpicGames.BuildGraph.Tests/GraphTests.cs
@@ -32,11 +32,11 @@ namespace EpicGames.BuildGraph.Tests
 		public void NodeTest()
 		{
 			BgAgent agent = new BgAgent("name", "type");
-			BgNode nodeSpec1 = agent.AddNode(x => UpdateVersionFiles()).Construct();
+			BgNode nodeSpec1 = agent.AddNode(x => UpdateVersionFiles());
 
-			BgNode nodeSpec2 = agent.AddNode(x => CompileShooterGameWin64()).Requires(BgList<BgFileSet>.Create(nodeSpec1)).Construct();
+			BgNode nodeSpec2 = agent.AddNode(x => CompileShooterGameWin64()).Requires(BgList<BgFileSet>.Create(nodeSpec1));
 
-			BgNode nodeSpec3 = agent.AddNode(x => CookShooterGameWin64()).Requires(BgList<BgFileSet>.Create(nodeSpec2.DefaultOutput)).Construct();
+			BgNode nodeSpec3 = agent.AddNode(x => CookShooterGameWin64()).Requires(BgList<BgFileSet>.Create(nodeSpec2.DefaultOutput));
 
 			BgNodeDef node3 = ((BgNodeDef)Evaluate(nodeSpec3));
 			Assert.AreEqual(node3.Name, "Cook Shooter Game Win64");
@@ -77,8 +77,8 @@ namespace EpicGames.BuildGraph.Tests
 		{
 			BgAgent agent = new BgAgent("test", "test");
 
-			BgNode nodeSpec1 = agent.AddNode(x => UpdateVersionFiles()).Construct();
-			BgNode nodeSpec2 = agent.AddNode(x => CompileShooterGameWin64()).Construct();
+			BgNode nodeSpec1 = agent.AddNode(x => UpdateVersionFiles());
+			BgNode nodeSpec2 = agent.AddNode(x => CompileShooterGameWin64());
 
 			BgAggregate aggregateSpec = new BgAggregate("All nodes", BgList<BgNode>.Create(nodeSpec1, nodeSpec2));
 
@@ -110,8 +110,8 @@ namespace EpicGames.BuildGraph.Tests
 		{
 			BgAgent agent = new BgAgent("test", "test");
 
-			BgNode nodeSpec1 = agent.AddNode(x => UpdateVersionFiles()).Construct();
-			BgNode nodeSpec2 = agent.AddNode(x => CompileShooterGameWin64()).Construct();
+			BgNode nodeSpec1 = agent.AddNode(x => UpdateVersionFiles());
+			BgNode nodeSpec2 = agent.AddNode(x => CompileShooterGameWin64());
 
 			BgAggregate aggregateSpec = new BgAggregate("All nodes", BgList<BgNode>.Create(nodeSpec1, nodeSpec2));
 			BgGraph graphSpec = new BgGraph(BgList.Create(nodeSpec1, nodeSpec2), BgList.Create(aggregateSpec));
