BuildGraph: Fix test cases with recent changes to storing graph structures as BG objects.

diff --git a/Engine/Source/Programs/Shared/EpicGames.BuildGraph.Tests/GraphTests.cs b/Engine/Source/Programs/Shared/EpicGames.BuildGraph.Tests/GraphTests.cs
--- a/Engine/Source/Programs/Shared/EpicGames.BuildGraph.Tests/GraphTests.cs
+++ b/Engine/Source/Programs/Shared/EpicGames.BuildGraph.Tests/GraphTests.cs
@@ -34,11 +34,11 @@ namespace EpicGames.BuildGraph.Tests
 			BgAgent agent = new BgAgent("name", "type");
 			BgNode nodeSpec1 = agent.AddNode(x => UpdateVersionFiles());
 
-			BgNode nodeSpec2 = agent.AddNode(x => CompileShooterGameWin64()).Requires(BgList<BgFileSet>.Create(nodeSpec1));
+			BgNode nodeSpec2 = agent.AddNode(x => CompileShooterGameWin64()).Requires(nodeSpec1);
 
 			BgNode nodeSpec3 = agent.AddNode(x => CookShooterGameWin64()).Requires(BgList<BgFileSet>.Create(nodeSpec2.DefaultOutput));
 
-			BgNodeDef node3 = ((BgNodeDef)Evaluate(nodeSpec3));
+			BgNodeDef node3 = ((BgObjectDef)Evaluate(nodeSpec3)).Deserialize<BgNodeExpressionDef>();
 			Assert.AreEqual(node3.Name, "Cook Shooter Game Win64");
 			Assert.AreEqual(2, node3.InputDependencies.Count);
 			Assert.AreEqual(1, node3.Inputs.Count);
@@ -82,7 +82,7 @@ namespace EpicGames.BuildGraph.Tests
 
 			BgAggregate aggregateSpec = new BgAggregate("All nodes", BgList<BgNode>.Create(nodeSpec1, nodeSpec2));
 
-			BgAggregateDef aggregate = ((BgObjectDef)Evaluate(aggregateSpec)).Deserialize<BgAggregateDef>();
+			BgAggregateDef aggregate = ((BgObjectDef)Evaluate(aggregateSpec)).Deserialize<BgAggregateExpressionDef>().ToAggregateDef();
 			Assert.AreEqual("All nodes", aggregate.Name);
 
 			BgNodeDef node1 = aggregate.RequiredNodes.OrderBy(x => x.Name).ElementAt(0);
diff --git a/Engine/Source/Programs/Shared/EpicGames.BuildGraph/Expressions/BgList.cs b/Engine/Source/Programs/Shared/EpicGames.BuildGraph/Expressions/BgList.cs
--- a/Engine/Source/Programs/Shared/EpicGames.BuildGraph/Expressions/BgList.cs
+++ b/Engine/Source/Programs/Shared/EpicGames.BuildGraph/Expressions/BgList.cs
@@ -2,6 +2,7 @@
 
 using System;
 using System.Collections.Generic;
+using System.Diagnostics;
 using System.Linq;
 
 namespace EpicGames.BuildGraph.Expressions
@@ -69,11 +70,13 @@ namespace EpicGames.BuildGraph.Expressions
 	public abstract class BgList<T> : BgExpr where T : BgExpr
 	{
 		/// <inheritdoc/>
+		[DebuggerBrowsable(DebuggerBrowsableState.Never)]
 		public Type ElementType => typeof(T);
 
 		/// <summary>
 		/// Constant representation of an empty list
 		/// </summary>
+		[DebuggerBrowsable(DebuggerBrowsableState.Never)]
 		public static BgList<T> Empty { get; } = new BgListEmptyExpr<T>();
 
 		/// <summary>
@@ -143,6 +146,7 @@ namespace EpicGames.BuildGraph.Expressions
 		/// <summary>
 		/// Gets the length of this list
 		/// </summary>
+		[DebuggerBrowsable(DebuggerBrowsableState.Never)]
 		public BgInt Count => new BgListCountExpr<T>(this);
 
 		/// <summary>
diff --git a/Engine/Source/Programs/Shared/EpicGames.BuildGraph/Expressions/BgNode.cs b/Engine/Source/Programs/Shared/EpicGames.BuildGraph/Expressions/BgNode.cs
--- a/Engine/Source/Programs/Shared/EpicGames.BuildGraph/Expressions/BgNode.cs
+++ b/Engine/Source/Programs/Shared/EpicGames.BuildGraph/Expressions/BgNode.cs
@@ -118,7 +118,7 @@ namespace EpicGames.BuildGraph.Expressions
 			Name = node.Name;
 			Thunk = node.Thunk;
 			Agent = node.Agent;
-			DefaultOutput = node.DefaultOutput;
+			DefaultOutput = new BgFileSetFromNodeOutputExpr(this, 0);
 			Inputs = node.Inputs;
 			Fences = node.Fences;
 			RunEarly = node.RunEarly;
@@ -150,7 +150,7 @@ namespace EpicGames.BuildGraph.Expressions
 		/// <returns></returns>
 		internal BgNode Modify(BgList<BgFileSet>? inputs = null, BgList<BgNode>? fences = null, BgBool? runEarly = null, BgList<BgLabel>? labels = null)
 		{
-			BgNode node = (BgNode)MemberwiseClone();
+			BgNode node = Clone();
 			if (inputs is object)
 			{
 				node.Inputs = inputs;
@@ -170,6 +170,12 @@ namespace EpicGames.BuildGraph.Expressions
 			return node;
 		}
 
+		/// <summary>
+		/// Clone this node
+		/// </summary>
+		/// <returns>Clone of this node</returns>
+		protected virtual BgNode Clone() => new BgNode(this);
+
 		/// <summary>
 		/// Gets the default tag name for the numbered output index
 		/// </summary>
@@ -349,9 +355,15 @@ namespace EpicGames.BuildGraph.Expressions
 		public BgNode(BgNode<T> other)
 			: base(other)
 		{
-			Output = other.Output;
+			Output = CreateOutput();
 		}
 
+		/// <summary>
+		/// Clone this node
+		/// </summary>
+		/// <returns>Clone of this node</returns>
+		protected override BgNode Clone() => new BgNode<T>(this);
+
 		T CreateOutput()
 		{
 			Type type = typeof(T);
diff --git a/Engine/Source/Programs/Shared/EpicGames.BuildGraph/Expressions/BgString.cs b/Engine/Source/Programs/Shared/EpicGames.BuildGraph/Expressions/BgString.cs
--- a/Engine/Source/Programs/Shared/EpicGames.BuildGraph/Expressions/BgString.cs
+++ b/Engine/Source/Programs/Shared/EpicGames.BuildGraph/Expressions/BgString.cs
@@ -65,10 +65,10 @@ namespace EpicGames.BuildGraph.Expressions
 		public BgString Append(BgString other) => this + other;
 
 		/// <inheritdoc/>
-		public BgBool Match(string pattern) => new BgStringMatchExpr(this, pattern);
+		public BgBool Match(BgString pattern) => new BgStringMatchExpr(this, pattern);
 
 		/// <inheritdoc/>
-		public BgString Replace(string pattern, BgString replace) => new BgStringReplaceExpr(this, pattern, replace);
+		public BgString Replace(BgString pattern, BgString replace) => new BgStringReplaceExpr(this, pattern, replace);
 
 		/// <inheritdoc/>
 		public override bool Equals(object? obj) => throw new InvalidOperationException();
@@ -188,9 +188,9 @@ namespace EpicGames.BuildGraph.Expressions
 	class BgStringMatchExpr : BgBool
 	{
 		public BgString Input { get; }
-		public string Pattern { get; }
+		public BgString Pattern { get; }
 
-		public BgStringMatchExpr(BgString input, string pattern)
+		public BgStringMatchExpr(BgString input, BgString pattern)
 			: base(input.Flags & BgExprFlags.Eager)
 		{
 			Input = input;
@@ -201,17 +201,17 @@ namespace EpicGames.BuildGraph.Expressions
 		{
 			writer.WriteOpcode(BgOpcode.StrIsMatch);
 			writer.WriteExpr(Input);
-			writer.WriteString(Pattern);
+			writer.WriteExpr(Pattern);
 		}
 	}
 
 	class BgStringReplaceExpr : BgString
 	{
 		public BgString Input { get; }
-		public string Pattern { get; }
+		public BgString Pattern { get; }
 		public BgString Replacement { get; }
 
-		public BgStringReplaceExpr(BgString input, string pattern, BgString replacement)
+		public BgStringReplaceExpr(BgString input, BgString pattern, BgString replacement)
 			: base(input.Flags & replacement.Flags & BgExprFlags.Eager)
 		{
 			Input = input;
@@ -223,7 +223,7 @@ namespace EpicGames.BuildGraph.Expressions
 		{
 			writer.WriteOpcode(BgOpcode.StrReplace);
 			writer.WriteExpr(Input);
-			writer.WriteString(Pattern);
+			writer.WriteExpr(Pattern);
 			writer.WriteExpr(Replacement);
 		}
 	}
