Added support for UE_5_1_OR_LATER style define constants for UBT/UHT C# plugins.

#rnx

diff --git a/Engine/Source/Programs/Shared/EpicGames.Build/System/CompileScriptModules.cs b/Engine/Source/Programs/Shared/EpicGames.Build/System/CompileScriptModules.cs
--- a/Engine/Source/Programs/Shared/EpicGames.Build/System/CompileScriptModules.cs
+++ b/Engine/Source/Programs/Shared/EpicGames.Build/System/CompileScriptModules.cs
@@ -114,7 +114,7 @@ namespace UnrealBuildBase
 				GameFolders: Unreal.IsEngineInstalled() ? GameDirectories : new List<DirectoryReference>(), 
 				ForeignPlugins: null, AdditionalSearchPaths: AdditionalDirectories.Concat(GameBuildDirectories).ToList()));
 
-			return GetTargetPaths(Build(RulesFileType, FoundProjects, BaseDirectories, bForceCompile, bNoCompile, bUseBuildRecords, out bBuildSuccess, OnBuildingProjects, Logger));
+			return GetTargetPaths(Build(RulesFileType, FoundProjects, BaseDirectories, null, bForceCompile, bNoCompile, bUseBuildRecords, out bBuildSuccess, OnBuildingProjects, Logger));
 		}
 
 		/// <summary>
@@ -149,6 +149,7 @@ namespace UnrealBuildBase
 		/// <param name="RulesFileType"></param>
 		/// <param name="FoundProjects">Projects to be compiled</param>
 		/// <param name="BaseDirectories">Base directories for all the projects</param>
+		/// <param name="DefineConstants">Collection of constants to be added to the project</param>
 		/// <param name="bForceCompile"></param>
 		/// <param name="bNoCompile"></param>
 		/// <param name="bUseBuildRecords"></param>
@@ -157,7 +158,8 @@ namespace UnrealBuildBase
 		/// <param name="Logger"></param>
 		/// <returns>Collection of all the projects.  They will have been compiled.</returns>
 		public static Dictionary<FileReference, (CsProjBuildRecord, FileReference)> Build(Rules.RulesFileType RulesFileType,
-			HashSet<FileReference> FoundProjects, List<DirectoryReference> BaseDirectories, bool bForceCompile, bool bNoCompile, bool bUseBuildRecords,
+			HashSet<FileReference> FoundProjects, List<DirectoryReference> BaseDirectories, List<string>? DefineConstants, 
+			bool bForceCompile, bool bNoCompile, bool bUseBuildRecords,
 			out bool bBuildSuccess, Action<int> OnBuildingProjects, ILogger Logger)
 		{
 			CsProjBuildHook Hook = new Hook(Logger);
@@ -244,7 +246,8 @@ namespace UnrealBuildBase
 			}
 
 			// Fall back to the slower approach: use msbuild to load csproj files & build as necessary
-			return Build(FoundProjects, bForceCompile || !bUseBuildRecords, out bBuildSuccess, Hook, BaseDirectories, OnBuildingProjects, Logger);
+			return Build(FoundProjects, bForceCompile || !bUseBuildRecords, out bBuildSuccess, Hook, BaseDirectories,
+				DefineConstants, OnBuildingProjects, Logger);
 		}
 
 		/// <summary>
@@ -254,9 +257,9 @@ namespace UnrealBuildBase
 		/// </summary>
 		static Dictionary<FileReference, (CsProjBuildRecord, FileReference)> Build(HashSet<FileReference> FoundProjects,
 			bool bForceCompile, out bool bBuildSuccess, CsProjBuildHook Hook, List<DirectoryReference> BaseDirectories,
-			Action<int> OnBuildingProjects, ILogger Logger)
+			List<string>? DefineConstants, Action<int> OnBuildingProjects, ILogger Logger)
 		{
-			return CsProjBuilder.Build(FoundProjects, bForceCompile, out bBuildSuccess, Hook, BaseDirectories, OnBuildingProjects, Logger);
+			return CsProjBuilder.Build(FoundProjects, bForceCompile, out bBuildSuccess, Hook, BaseDirectories, DefineConstants ?? (new()), OnBuildingProjects, Logger);
 		}
 
 		/// <summary>
diff --git a/Engine/Source/Programs/Shared/EpicGames.Build/System/EnumeratePlugins.cs b/Engine/Source/Programs/Shared/EpicGames.Build/System/EnumeratePlugins.cs
--- a/Engine/Source/Programs/Shared/EpicGames.Build/System/EnumeratePlugins.cs
+++ b/Engine/Source/Programs/Shared/EpicGames.Build/System/EnumeratePlugins.cs
@@ -133,11 +133,12 @@
 		/// </summary>
 		/// <param name="ProjectFile">Path to the project file</param>
 		/// <param name="FoundPlugins">Collection of plugins to build</param>
+		/// <param name="DefineConstants">Collection of constants to add to the projects</param>
 		/// <param name="Logger">Logger for output</param>
 		/// <param name="Plugins">Collection of built plugins</param>
 		/// <returns>True if the plugins compiled</returns>
-		public static bool BuildUbtPlugins(FileReference? ProjectFile, IEnumerable<FileReference> FoundPlugins, ILogger Logger,
-			out (FileReference ProjectFile, FileReference TargetAssembly)[]? Plugins)
+		public static bool BuildUbtPlugins(FileReference? ProjectFile, IEnumerable<FileReference> FoundPlugins,
+			List<string>? DefineConstants, ILogger Logger, out (FileReference ProjectFile, FileReference TargetAssembly)[]? Plugins)
 		{
 			bool bBuildSuccess = true;
 			Plugins = null;
@@ -153,7 +154,7 @@
 			{
 				Dictionary<FileReference, (CsProjBuildRecord BuildRecord, FileReference BuildRecordFile)>? BuiltPlugins =
 					CompileScriptModule.Build(UnrealBuildBase.Rules.RulesFileType.UbtPlugin, new HashSet<FileReference>(NotBuilt),
-					GetBaseDirectories(ProjectFile), false, false, true, out bBuildSuccess,
+					GetBaseDirectories(ProjectFile), DefineConstants, false, false, true, out bBuildSuccess,
 					Count =>
 					{
 						if (Log.OutputFile != null)
diff --git a/Engine/Source/Programs/Shared/EpicGames.MsBuild/CsProjBuilder.cs b/Engine/Source/Programs/Shared/EpicGames.MsBuild/CsProjBuilder.cs
--- a/Engine/Source/Programs/Shared/EpicGames.MsBuild/CsProjBuilder.cs
+++ b/Engine/Source/Programs/Shared/EpicGames.MsBuild/CsProjBuilder.cs
@@ -12,6 +12,7 @@ using System.Collections.Generic;
 using System.Diagnostics;
 using System.IO;
 using System.Linq;
+using System.Text;
 using System.Text.Json;
 
 namespace EpicGames.MsBuild
@@ -151,19 +152,40 @@ namespace EpicGames.MsBuild
 		/// <summary>
 		/// Builds multiple projects
 		/// </summary>
+		/// <param name="FoundProjects">Collection of project to be built</param>
+		/// <param name="bForceCompile">If true, force the compilation of the projects</param>
+		/// <param name="bBuildSuccess">Set to true/false depending on if all projects compiled or are up-to-date</param>
+		/// <param name="Hook">Interface to fetch data about the building environment</param>
+		/// <param name="BaseDirectories">Base directories of the engine and project</param>
+		/// <param name="DefineConstants">Collection of constants to be defined while building projects</param>
+		/// <param name="OnBuildingProjects">Action invoked to notify caller regarding the number of projects being built</param>
+		/// <param name="Logger">Destination logger</param>
 		public static Dictionary<FileReference, (CsProjBuildRecord, FileReference)> Build(HashSet<FileReference> FoundProjects,
-			bool bForceCompile, out bool bBuildSuccess, CsProjBuildHook Hook, List<DirectoryReference> BaseDirectories, Action<int> OnBuildingProjects, ILogger Logger)
+			bool bForceCompile, out bool bBuildSuccess, CsProjBuildHook Hook, List<DirectoryReference> BaseDirectories, 
+			List<string> DefineConstants, Action<int> OnBuildingProjects, ILogger Logger)
 		{
 
 			// Register the MS build path prior to invoking the internal routine.  By not having the internal routine
 			// inline, we avoid having the issue of the Microsoft.Build libraries being resolved prior to the build path
 			// being set.
 			RegisterMsBuildPath(Hook);
-			return BuildInternal(FoundProjects, bForceCompile, out bBuildSuccess, Hook, BaseDirectories, OnBuildingProjects, Logger);
+			return BuildInternal(FoundProjects, bForceCompile, out bBuildSuccess, Hook, BaseDirectories, DefineConstants, OnBuildingProjects, Logger);
 		}
 
+		/// <summary>
+		/// Builds multiple projects.  This is the internal implementation invoked after the MS build path is set
+		/// </summary>
+		/// <param name="FoundProjects">Collection of project to be built</param>
+		/// <param name="bForceCompile">If true, force the compilation of the projects</param>
+		/// <param name="bBuildSuccess">Set to true/false depending on if all projects compiled or are up-to-date</param>
+		/// <param name="Hook">Interface to fetch data about the building environment</param>
+		/// <param name="BaseDirectories">Base directories of the engine and project</param>
+		/// <param name="DefineConstants">Collection of constants to be defined while building projects</param>
+		/// <param name="OnBuildingProjects">Action invoked to notify caller regarding the number of projects being built</param>
+		/// <param name="Logger">Destination logger</param>
 		private static Dictionary<FileReference, (CsProjBuildRecord, FileReference)> BuildInternal(HashSet<FileReference> FoundProjects,
-			bool bForceCompile, out bool bBuildSuccess, CsProjBuildHook Hook, List<DirectoryReference> BaseDirectories, Action<int> OnBuildingProjects, ILogger Logger)
+			bool bForceCompile, out bool bBuildSuccess, CsProjBuildHook Hook, List<DirectoryReference> BaseDirectories, List<string> DefineConstants,
+			Action<int> OnBuildingProjects, ILogger Logger)
 		{
 			Dictionary<string, string> GlobalProperties = new Dictionary<string, string>
 			{
@@ -175,6 +197,11 @@ namespace EpicGames.MsBuild
 #endif
 			};
 
+			if (DefineConstants.Count > 0)
+			{
+				GlobalProperties.Add("DefineConstants", String.Join(';', DefineConstants));
+			}
+
 			Dictionary<FileReference, (CsProjBuildRecord BuildRecord, FileReference BuildRecordPath)> BuildRecords = new Dictionary<FileReference, (CsProjBuildRecord, FileReference)>();
 
 			using ProjectCollection ProjectCollection = new ProjectCollection(GlobalProperties);
diff --git a/Engine/Source/Programs/UnrealBuildTool/Configuration/UEBuildTarget.cs b/Engine/Source/Programs/UnrealBuildTool/Configuration/UEBuildTarget.cs
--- a/Engine/Source/Programs/UnrealBuildTool/Configuration/UEBuildTarget.cs
+++ b/Engine/Source/Programs/UnrealBuildTool/Configuration/UEBuildTarget.cs
@@ -3662,7 +3662,7 @@ namespace UnrealBuildTool
 				EnabledUbtPlugins.SortBy(P => P.FullName);
 
 				// Build the plugins
-				bool bCompiled = PluginsBase.BuildUbtPlugins(ProjectFile, EnabledUbtPlugins, Logger, out BuiltPlugins);
+				bool bCompiled = PluginsBase.BuildUbtPlugins(ProjectFile, EnabledUbtPlugins, RulesAssembly.PreprocessorDefines, Logger, out BuiltPlugins);
 				if (!bCompiled)
 				{
 					throw new BuildException("Not all plugins compiled");
