BuildGraph: Fix labels not getting passed through from the aggregate into the bytecode.

diff --git a/Engine/Source/Programs/Shared/EpicGames.BuildGraph.Tests/GraphTests.cs b/Engine/Source/Programs/Shared/EpicGames.BuildGraph.Tests/GraphTests.cs
--- a/Engine/Source/Programs/Shared/EpicGames.BuildGraph.Tests/GraphTests.cs
+++ b/Engine/Source/Programs/Shared/EpicGames.BuildGraph.Tests/GraphTests.cs
@@ -80,9 +80,11 @@ namespace EpicGames.BuildGraph.Tests
 			BgNode nodeSpec1 = agent.AddNode(x => UpdateVersionFiles());
 			BgNode nodeSpec2 = agent.AddNode(x => CompileShooterGameWin64());
 
-			BgAggregate aggregateSpec = new BgAggregate("All nodes", BgList<BgNode>.Create(nodeSpec1, nodeSpec2));
+			BgLabel labelSpec = new BgLabel("name", "category");
+			BgAggregate aggregateSpec = new BgAggregate("All nodes", BgList.Create(nodeSpec1, nodeSpec2), labelSpec);
 
-			BgAggregateDef aggregate = ((BgObjectDef)Evaluate(aggregateSpec)).Deserialize<BgAggregateExpressionDef>().ToAggregateDef();
+			BgAggregateExpressionDef aggregateExpr = ((BgObjectDef)Evaluate(aggregateSpec)).Deserialize<BgAggregateExpressionDef>();
+			BgAggregateDef aggregate = aggregateExpr.ToAggregateDef();
 			Assert.AreEqual("All nodes", aggregate.Name);
 
 			BgNodeDef node1 = aggregate.RequiredNodes.OrderBy(x => x.Name).ElementAt(0);
@@ -90,6 +92,11 @@ namespace EpicGames.BuildGraph.Tests
 
 			BgNodeDef node2 = aggregate.RequiredNodes.OrderBy(x => x.Name).ElementAt(1);
 			Assert.AreEqual("Update Version Files", node2.Name);
+
+			BgLabelDef? label = aggregateExpr.Label;
+			Assert.IsNotNull(label);
+			Assert.AreEqual(label!.DashboardName, "name");
+			Assert.AreEqual(label!.DashboardCategory, "category");
 		}
 
 		[TestMethod]
diff --git a/Engine/Source/Programs/Shared/EpicGames.BuildGraph/Expressions/BgAggregate.cs b/Engine/Source/Programs/Shared/EpicGames.BuildGraph/Expressions/BgAggregate.cs
--- a/Engine/Source/Programs/Shared/EpicGames.BuildGraph/Expressions/BgAggregate.cs
+++ b/Engine/Source/Programs/Shared/EpicGames.BuildGraph/Expressions/BgAggregate.cs
@@ -62,9 +62,13 @@ namespace EpicGames.BuildGraph.Expressions
 		/// <inheritdoc/>
 		public override void Write(BgBytecodeWriter writer)
 		{
-			BgObject<BgAggregateDef> obj = BgObject<BgAggregateDef>.Empty;
+			BgObject<BgAggregateExpressionDef> obj = BgObject<BgAggregateExpressionDef>.Empty;
 			obj = obj.Set(x => x.Name, Name);
 			obj = obj.Set(x => x.RequiredNodes, Requires);
+			if (Label != null)
+			{
+				obj = obj.Set(x => x.Label, Label);
+			}
 			writer.WriteExpr(obj);
 		}
 
