BuildGraph: Fix labels having empty list of included nodes.

diff --git a/Engine/Source/Programs/Shared/EpicGames.BuildGraph.Tests/GraphTests.cs b/Engine/Source/Programs/Shared/EpicGames.BuildGraph.Tests/GraphTests.cs
--- a/Engine/Source/Programs/Shared/EpicGames.BuildGraph.Tests/GraphTests.cs
+++ b/Engine/Source/Programs/Shared/EpicGames.BuildGraph.Tests/GraphTests.cs
@@ -120,12 +120,25 @@ namespace EpicGames.BuildGraph.Tests
 			BgNode nodeSpec1 = agent.AddNode(x => UpdateVersionFiles());
 			BgNode nodeSpec2 = agent.AddNode(x => CompileShooterGameWin64());
 
-			BgAggregate aggregateSpec = new BgAggregate("All nodes", BgList<BgNode>.Create(nodeSpec1, nodeSpec2));
+			BgAggregate aggregateSpec = new BgAggregate("All nodes", BgList.Create(nodeSpec1, nodeSpec2), new BgLabel("name", "category"));
 			BgGraph graphSpec = new BgGraph(BgList.Create(nodeSpec1, nodeSpec2), BgList.Create(aggregateSpec));
 
-			BgGraphExpressionDef graph = ((BgObjectDef)Evaluate(graphSpec)).Deserialize<BgGraphExpressionDef>();
-			Assert.AreEqual(2, graph.Nodes.Count);
-			Assert.AreEqual(1, graph.Aggregates.Count);
+			BgGraphDef graph = ((BgObjectDef)Evaluate(graphSpec)).Deserialize<BgGraphExpressionDef>().ToGraphDef();
+			Assert.AreEqual(2, graph.NameToNode.Count);
+			Assert.AreEqual(1, graph.NameToAggregate.Count);
+			Assert.AreEqual(1, graph.Labels.Count);
+
+			BgNodeDef node1 = graph.NameToNode["Update Version Files"];
+
+			BgNodeDef node2 = graph.NameToNode["Compile Shooter Game Win64"];
+
+			BgLabelDef label = graph.Labels[0];
+			Assert.AreEqual(label!.DashboardName, "name");
+			Assert.AreEqual(label!.DashboardCategory, "category");
+			Assert.IsTrue(label!.RequiredNodes.Contains(node1));
+			Assert.IsTrue(label!.RequiredNodes.Contains(node2));
+			Assert.IsTrue(label!.IncludedNodes.Contains(node1));
+			Assert.IsTrue(label!.IncludedNodes.Contains(node2));
 		}
 	}
 }
diff --git a/Engine/Source/Programs/Shared/EpicGames.BuildGraph/BgGraphDef.cs b/Engine/Source/Programs/Shared/EpicGames.BuildGraph/BgGraphDef.cs
--- a/Engine/Source/Programs/Shared/EpicGames.BuildGraph/BgGraphDef.cs
+++ b/Engine/Source/Programs/Shared/EpicGames.BuildGraph/BgGraphDef.cs
@@ -705,8 +705,9 @@ namespace EpicGames.BuildGraph
 			{
 				foreach (BgLabelDef label in node.Labels)
 				{
-					labels.Add(label);
 					label.RequiredNodes.Add(node);
+					label.IncludedNodes.Add(node);
+					labels.Add(label);
 				}
 			}
 			foreach (BgAggregateExpressionDef aggregate in Aggregates)
@@ -714,6 +715,7 @@ namespace EpicGames.BuildGraph
 				if (aggregate.Label != null)
 				{
 					aggregate.Label.RequiredNodes.UnionWith(aggregate.RequiredNodes);
+					aggregate.Label.IncludedNodes.UnionWith(aggregate.RequiredNodes);
 					labels.Add(aggregate.Label);
 				}
 			}
