- Fixed up 3 Linux ContentWorker issues that came from the dotnet 6 upgrade (or other recent UE5 changes):
* Staged the wrong Dotnet directory when making CW from Windows (it would stage Dotnet/6.20.00/windows)
* The ScriptModule code that used the Automation.json files to discover the precompiled UAT script dlls didn't handle \'s on Linux in the paths, and since Windows wrote out the files, there are \'s aplenty. I fixed up the known issue, but possibly all paths loaded from the json files should be fixed up, but it's not easy to do automatically
* We moved to using RunUAT.sh instead of calling dotnet directly, however because we make the CW on Windows, dotnet was not executable, and the code in RunUAT.sh to find the dotnet executable and make it executable (chmod `which dotnet`) only works if dotnet is already executable! So, I moved the +x code into SetupDotnet.sh, on Linux and Mac

diff --git a/Engine/Build/BatchFiles/Linux/SetupDotnet.sh b/Engine/Build/BatchFiles/Linux/SetupDotnet.sh
--- a/Engine/Build/BatchFiles/Linux/SetupDotnet.sh
+++ b/Engine/Build/BatchFiles/Linux/SetupDotnet.sh
@@ -32,6 +32,7 @@ if [ $IS_DOTNET_INSTALLED -eq 0 ]; then
 	echo Setting up bundled DotNet SDK
 	CUR_DIR=`pwd`
 	export UE_DOTNET_DIR="$CUR_DIR/../../../Binaries/ThirdParty/DotNet/6.0.200/linux"
+	chmod u+x "$UE_DOTNET_DIR/dotnet"
 	export PATH="$UE_DOTNET_DIR:$PATH"
 	export DOTNET_ROOT="$UE_DOTNET_DIR"
 
diff --git a/Engine/Build/BatchFiles/Mac/SetupDotnet.sh b/Engine/Build/BatchFiles/Mac/SetupDotnet.sh
--- a/Engine/Build/BatchFiles/Mac/SetupDotnet.sh
+++ b/Engine/Build/BatchFiles/Mac/SetupDotnet.sh
@@ -37,6 +37,7 @@ if [ $IS_DOTNET_INSTALLED -eq 0 ]; then
 	[ $(uname -m) == "arm64" ] && ARCH=arm64 
 	
 	export UE_DOTNET_DIR=$CUR_DIR/../../../Binaries/ThirdParty/DotNet/6.0.200/mac-$ARCH
+	chmod u+x "$UE_DOTNET_DIR/dotnet"
 	echo $UE_DOTNET_DIR
 	export PATH=$UE_DOTNET_DIR:$PATH
 	export DOTNET_ROOT=$UE_DOTNET_DIR
diff --git a/Engine/Build/BatchFiles/RunUAT.sh b/Engine/Build/BatchFiles/RunUAT.sh
--- a/Engine/Build/BatchFiles/RunUAT.sh
+++ b/Engine/Build/BatchFiles/RunUAT.sh
@@ -119,9 +119,6 @@ else
 	LogDir="$uebp_LogFolder"
 fi
 
-# make sure it's executable (if it was zipped it may not be)
-chmod u+x `which dotnet`
-
 # if we are running under UE, we need to run this with the term handler (otherwise canceling a UAT job from the editor
 # can leave mono, etc running in the background, which means we need the PID so we 
 # run it in the background
diff --git a/Engine/Source/Programs/AutomationTool/CookedEditor/CookedEditor.Automation.cs b/Engine/Source/Programs/AutomationTool/CookedEditor/CookedEditor.Automation.cs
--- a/Engine/Source/Programs/AutomationTool/CookedEditor/CookedEditor.Automation.cs
+++ b/Engine/Source/Programs/AutomationTool/CookedEditor/CookedEditor.Automation.cs
@@ -565,7 +565,7 @@ public class MakeCookedEditor : BuildCommand
 		if (Context.IniPlatformName == "Linux")
 		{
 			// linux needs dotnet runtime
-			SC.StageFiles(StagedFileType.NonUFS, Unreal.DotnetDirectory, StageFilesSearch.AllDirectories);
+			SC.StageFiles(StagedFileType.NonUFS, Unreal.FindDotnetDirectoryForPlatform(RuntimePlatform.Type.Linux), StageFilesSearch.AllDirectories);
 		}
 
 		// not sure if we need this or not now
diff --git a/Engine/Source/Programs/Shared/EpicGames.Build/System/CompileScriptModules.cs b/Engine/Source/Programs/Shared/EpicGames.Build/System/CompileScriptModules.cs
--- a/Engine/Source/Programs/Shared/EpicGames.Build/System/CompileScriptModules.cs
+++ b/Engine/Source/Programs/Shared/EpicGames.Build/System/CompileScriptModules.cs
@@ -298,7 +298,9 @@ namespace UnrealBuildBase
 
 					if (BuildRecord != null && BuildRecord.ProjectPath != null)
 					{
-						LoadedBuildRecords.Add(FileReference.FromString(Path.GetFullPath(BuildRecord.ProjectPath, JsonFile.Directory.FullName)), (BuildRecord, JsonFile));
+						// Path.Combine, etc don't work so well with improper dir separators (a build record made on Windows will fail on Linux with the \ characters)
+						string Combined = Path.Combine(JsonFile.Directory.FullName, BuildRecord.ProjectPath.Replace('\\', Path.DirectorySeparatorChar).Replace('/', Path.DirectorySeparatorChar));
+						LoadedBuildRecords.Add(FileReference.FromString(Path.GetFullPath(Combined)), (BuildRecord, JsonFile));
 					}
 					else
                     {
diff --git a/Engine/Source/Programs/Shared/EpicGames.Build/Unreal.cs b/Engine/Source/Programs/Shared/EpicGames.Build/Unreal.cs
--- a/Engine/Source/Programs/Shared/EpicGames.Build/Unreal.cs
+++ b/Engine/Source/Programs/Shared/EpicGames.Build/Unreal.cs
@@ -81,10 +81,10 @@
 
 		static private string DotnetVersionDirectory = "6.0.200";
 			
-		static private string FindRelativeDotnetDirectory()
+		static private string FindRelativeDotnetDirectory(RuntimePlatform.Type HostPlatform)
 		{
 			string HostDotNetDirectoryName;
-			switch (RuntimePlatform.Current)
+			switch (HostPlatform)
 			{
 				case RuntimePlatform.Type.Windows: HostDotNetDirectoryName = "windows"; break;
 				case RuntimePlatform.Type.Mac:
@@ -103,6 +103,8 @@
 			return Path.Combine("Binaries", "ThirdParty", "DotNet", DotnetVersionDirectory, HostDotNetDirectoryName);
 		}
 
+		static private string FindRelativeDotnetDirectory() => FindRelativeDotnetDirectory(RuntimePlatform.Current);
+
 		/// <summary>
 		/// Relative path to the dotnet executable from EngineDir
 		/// </summary>
@@ -153,6 +155,16 @@
 		static private bool? bIsEngineInstalled;
 
 		/// <summary>
+		/// Returns where another platform's Dotnet is located
+		/// </summary>
+		/// <param name="HostPlatform"></param>
+		/// <returns></returns>
+		static public DirectoryReference FindDotnetDirectoryForPlatform(RuntimePlatform.Type HostPlatform)
+		{
+			return DirectoryReference.Combine(EngineDirectory, FindRelativeDotnetDirectory(HostPlatform));
+		}
+
+		/// <summary>
 		/// Returns true if UnrealBuildTool is running using installed Engine components
 		/// </summary>
 		/// <returns>True if running using installed Engine components</returns>
