Updating C# UHT to support the new Iris replication.

#rnx

diff --git a/Engine/Source/Programs/Shared/EpicGames.UHT/Exporters/CodeGen/UhtCodeGenerator.cs b/Engine/Source/Programs/Shared/EpicGames.UHT/Exporters/CodeGen/UhtCodeGenerator.cs
--- a/Engine/Source/Programs/Shared/EpicGames.UHT/Exporters/CodeGen/UhtCodeGenerator.cs
+++ b/Engine/Source/Programs/Shared/EpicGames.UHT/Exporters/CodeGen/UhtCodeGenerator.cs
@@ -26,49 +26,54 @@ namespace EpicGames.UHT.Exporters.CodeGen
 
 		public struct PackageInfo
 		{
-			public string _strippedName;
-			public string _api;
+			public string StrippedName { get; set; }
+			public string Api { get; set; }
 		}
-		public PackageInfo[] _packageInfos;
+		public PackageInfo[] PackageInfos { get; set; }
 
 		public struct HeaderInfo
 		{
-			public Task? _task;
-			public string _includePath;
-			public string _fileId;
-			public uint _bodyHash;
-			public bool _needsPushModelHeaders;
+			public Task? Task { get; set; }
+			public string IncludePath { get; set; }
+			public string FileId { get; set; }
+			public uint BodyHash { get; set; }
+			public bool NeedsPushModelHeaders { get; set; }
+			public bool NeedsFastArrayHeaders { get; set; }
 		}
-		public HeaderInfo[] _headerInfos;
+		public HeaderInfo[] HeaderInfos { get; set; }
 
 		public struct ObjectInfo
 		{
-			public string _registeredSingletonName;
-			public string _unregisteredSingletonName;
-			public string _regsiteredCrossReference;
-			public string _unregsiteredCrossReference;
-			public string _regsiteredExternalDecl;
-			public string _unregisteredExternalDecl;
-			public UhtClass? _nativeInterface;
-			public uint _hash;
+			public string RegisteredSingletonName { get; set; }
+			public string UnregisteredSingletonName { get; set; }
+			public string RegsiteredCrossReference { get; set; }
+			public string UnregsiteredCrossReference { get; set; }
+			public string RegsiteredExternalDecl { get; set; }
+			public string UnregisteredExternalDecl { get; set; }
+			public UhtClass? NativeInterface { get; set; }
+			public UhtProperty? FastArrayProperty { get; set; }
+			public uint Hash { get; set; }
 		}
-		public ObjectInfo[] _objectInfos;
+		public ObjectInfo[] ObjectInfos { get; set; }
 
 		public readonly IUhtExportFactory Factory;
 		public UhtSession Session => Factory.Session;
+		public UhtScriptStruct? FastArraySerializer { get; set; } = null;
 
 		private UhtCodeGenerator(IUhtExportFactory factory)
 		{
 			this.Factory = factory;
-			this._headerInfos = new HeaderInfo[this.Factory.Session.HeaderFileTypeCount];
-			this._objectInfos = new ObjectInfo[this.Factory.Session.ObjectTypeCount];
-			this._packageInfos = new PackageInfo[this.Factory.Session.PackageTypeCount];
+			this.HeaderInfos = new HeaderInfo[this.Factory.Session.HeaderFileTypeCount];
+			this.ObjectInfos = new ObjectInfo[this.Factory.Session.ObjectTypeCount];
+			this.PackageInfos = new PackageInfo[this.Factory.Session.PackageTypeCount];
 		}
 
 		private void Generate()
 		{
 			List<Task?> prereqs = new();
 
+			this.FastArraySerializer = Session.FindType(null, UhtFindOptions.SourceName | UhtFindOptions.ScriptStruct, "FFastArraySerializer") as UhtScriptStruct;
+
 			// Perform some startup initialization to compute things we need over and over again
 			if (this.Session.GoWide)
 			{
@@ -98,11 +103,11 @@ namespace EpicGames.UHT.Exporters.CodeGen
 					{
 						if (headerFile != referenced)
 						{
-							prereqs.Add(this._headerInfos[referenced.HeaderFileTypeIndex]._task);
+							prereqs.Add(this.HeaderInfos[referenced.HeaderFileTypeIndex].Task);
 						}
 					}
 
-					this._headerInfos[headerFile.HeaderFileTypeIndex]._task = Factory.CreateTask(prereqs,
+					this.HeaderInfos[headerFile.HeaderFileTypeIndex].Task = Factory.CreateTask(prereqs,
 						(IUhtExportFactory factory) =>
 						{
 							new UhtHeaderCodeGeneratorHFile(this, package, headerFile).Generate(factory);
@@ -121,7 +126,7 @@ namespace EpicGames.UHT.Exporters.CodeGen
 				prereqs.Clear();
 				foreach (UhtHeaderFile headerFile in package.Children)
 				{
-					prereqs.Add(this._headerInfos[headerFile.HeaderFileTypeIndex]._task);
+					prereqs.Add(this.HeaderInfos[headerFile.HeaderFileTypeIndex].Task);
 					if (!writeHeader)
 					{
 						foreach (UhtType type in headerFile.Children)
@@ -177,7 +182,7 @@ namespace EpicGames.UHT.Exporters.CodeGen
 			{
 				return "nullptr";
 			}
-			return registered ? this._objectInfos[obj.ObjectTypeIndex]._registeredSingletonName : this._objectInfos[obj.ObjectTypeIndex]._unregisteredSingletonName;
+			return registered ? this.ObjectInfos[obj.ObjectTypeIndex].RegisteredSingletonName : this.ObjectInfos[obj.ObjectTypeIndex].UnregisteredSingletonName;
 		}
 
 		/// <summary>
@@ -199,7 +204,7 @@ namespace EpicGames.UHT.Exporters.CodeGen
 		/// <returns>External declaration</returns>
 		public string GetExternalDecl(int objectIndex, bool registered)
 		{
-			return registered ? this._objectInfos[objectIndex]._regsiteredExternalDecl : this._objectInfos[objectIndex]._unregisteredExternalDecl;
+			return registered ? this.ObjectInfos[objectIndex].RegsiteredExternalDecl : this.ObjectInfos[objectIndex].UnregisteredExternalDecl;
 		}
 
 		/// <summary>
@@ -221,7 +226,7 @@ namespace EpicGames.UHT.Exporters.CodeGen
 		/// <returns>Cross reference</returns>
 		public string GetCrossReference(int objectIndex, bool registered)
 		{
-			return registered ? this._objectInfos[objectIndex]._regsiteredCrossReference : this._objectInfos[objectIndex]._unregsiteredCrossReference;
+			return registered ? this.ObjectInfos[objectIndex].RegsiteredCrossReference : this.ObjectInfos[objectIndex].UnregsiteredCrossReference;
 		}
 		#endregion
 
@@ -230,17 +235,17 @@ namespace EpicGames.UHT.Exporters.CodeGen
 		{
 			StringBuilder builder = new();
 
-			ref PackageInfo packageInfo = ref this._packageInfos[package.PackageTypeIndex];
-			packageInfo._strippedName = package.SourceName.Replace('/', '_');
-			packageInfo._api = $"{package.ShortName.ToString().ToUpper()}_API ";
+			ref PackageInfo packageInfo = ref this.PackageInfos[package.PackageTypeIndex];
+			packageInfo.StrippedName = package.SourceName.Replace('/', '_');
+			packageInfo.Api = $"{package.ShortName.ToString().ToUpper()}_API ";
 
 			// Construct the names used commonly during export
-			ref ObjectInfo objectInfo = ref this._objectInfos[package.ObjectTypeIndex];
+			ref ObjectInfo objectInfo = ref this.ObjectInfos[package.ObjectTypeIndex];
 			builder.Append("Z_Construct_UPackage_");
-			builder.Append(packageInfo._strippedName);
-			objectInfo._unregisteredSingletonName = objectInfo._registeredSingletonName = builder.ToString();
-			objectInfo._unregisteredExternalDecl = objectInfo._regsiteredExternalDecl = $"\t{packageInfo._api}_API UPackage* {objectInfo._registeredSingletonName}();\r\n"; //COMPATIBILITY-TODO remove the extra _API
-			objectInfo._unregsiteredCrossReference = objectInfo._regsiteredCrossReference = $"\tUPackage* {objectInfo._registeredSingletonName}();\r\n";
+			builder.Append(packageInfo.StrippedName);
+			objectInfo.UnregisteredSingletonName = objectInfo.RegisteredSingletonName = builder.ToString();
+			objectInfo.UnregisteredExternalDecl = objectInfo.RegsiteredExternalDecl = $"\t{packageInfo.Api}_API UPackage* {objectInfo.RegisteredSingletonName}();\r\n"; //COMPATIBILITY-TODO remove the extra _API
+			objectInfo.UnregsiteredCrossReference = objectInfo.RegsiteredCrossReference = $"\tUPackage* {objectInfo.RegisteredSingletonName}();\r\n";
 
 			foreach (UhtHeaderFile headerFile in package.Children)
 			{
@@ -250,9 +255,9 @@ namespace EpicGames.UHT.Exporters.CodeGen
 
 		private void InitHeaderInfo(StringBuilder builder, UhtPackage package, ref PackageInfo packageInfo, UhtHeaderFile headerFile)
 		{
-			ref HeaderInfo headerInfo = ref this._headerInfos[headerFile.HeaderFileTypeIndex];
+			ref HeaderInfo headerInfo = ref this.HeaderInfos[headerFile.HeaderFileTypeIndex];
 
-			headerInfo._includePath = Path.GetRelativePath(package.Module.IncludeBase, headerFile.FilePath).Replace('\\', '/');
+			headerInfo.IncludePath = Path.GetRelativePath(package.Module.IncludeBase, headerFile.FilePath).Replace('\\', '/');
 
 			// Convert the file path to a C identifier
 			string filePath = headerFile.FilePath;
@@ -293,7 +298,7 @@ namespace EpicGames.UHT.Exporters.CodeGen
 			{
 				outFilePath[index + 4] = UhtFCString.IsAlnum(filePath[index]) ? filePath[index] : '_';
 			}
-			headerInfo._fileId = new string(outFilePath);
+			headerInfo.FileId = new string(outFilePath);
 
 			foreach (UhtObject obj in headerFile.Children)
 			{
@@ -303,7 +308,7 @@ namespace EpicGames.UHT.Exporters.CodeGen
 
 		private void InitObjectInfo(StringBuilder builder, UhtPackage package, ref PackageInfo packageInfo, ref HeaderInfo headerInfo, UhtObject obj)
 		{
-			ref ObjectInfo objectInfo = ref this._objectInfos[obj.ObjectTypeIndex];
+			ref ObjectInfo objectInfo = ref this.ObjectInfos[obj.ObjectTypeIndex];
 
 			builder.Clear();
 
@@ -320,13 +325,40 @@ namespace EpicGames.UHT.Exporters.CodeGen
 				}
 				if (classObj.ClassExportFlags.HasExactFlags(UhtClassExportFlags.HasReplciatedProperties, UhtClassExportFlags.SelfHasReplicatedProperties))
 				{
-					headerInfo._needsPushModelHeaders = true;
+					headerInfo.NeedsPushModelHeaders = true;
 				}
 				if (classObj.ClassType == UhtClassType.NativeInterface)
 				{
 					if (classObj.AlternateObject != null)
 					{
-						this._objectInfos[classObj.AlternateObject.ObjectTypeIndex]._nativeInterface = classObj;
+						this.ObjectInfos[classObj.AlternateObject.ObjectTypeIndex].NativeInterface = classObj;
+					}
+				}
+			}
+			else if (obj is UhtScriptStruct scriptStruct)
+			{
+				// Check to see if we are a FastArraySerializer and should try to deduce the FastArraySerializerItemType
+				// To fulfill that requirement the struct should be derived from FFastArraySerializer and have a single replicated TArrayProperty
+				if (scriptStruct.IsChildOf(FastArraySerializer))
+				{
+					foreach (UhtType child in scriptStruct.Children)
+					{
+						if (child is UhtProperty property)
+						{
+							if (!property.PropertyFlags.HasAnyFlags(EPropertyFlags.RepSkip) && property is UhtArrayProperty)
+							{
+								if (objectInfo.FastArrayProperty != null)
+								{
+									objectInfo.FastArrayProperty = null;
+									break;
+								}
+								objectInfo.FastArrayProperty = property;
+							}
+						}
+					}
+					if (objectInfo.FastArrayProperty != null)
+					{
+						headerInfo.NeedsFastArrayHeaders = true;
 					}
 				}
 			}
@@ -338,22 +370,22 @@ namespace EpicGames.UHT.Exporters.CodeGen
 
 			if (isNonIntrinsicClass)
 			{
-				objectInfo._registeredSingletonName = builder.ToString();
+				objectInfo.RegisteredSingletonName = builder.ToString();
 				builder.Append("_NoRegister");
-				objectInfo._unregisteredSingletonName = builder.ToString();
+				objectInfo.UnregisteredSingletonName = builder.ToString();
 
-				objectInfo._unregisteredExternalDecl = $"\t{packageInfo._api}U{engineClassName}* {objectInfo._unregisteredSingletonName}();\r\n";
-				objectInfo._regsiteredExternalDecl = $"\t{packageInfo._api}U{engineClassName}* {objectInfo._registeredSingletonName}();\r\n";
+				objectInfo.UnregisteredExternalDecl = $"\t{packageInfo.Api}U{engineClassName}* {objectInfo.UnregisteredSingletonName}();\r\n";
+				objectInfo.RegsiteredExternalDecl = $"\t{packageInfo.Api}U{engineClassName}* {objectInfo.RegisteredSingletonName}();\r\n";
 			}
 			else
 			{
-				objectInfo._unregisteredSingletonName = objectInfo._registeredSingletonName = builder.ToString();
-				objectInfo._unregisteredExternalDecl = objectInfo._regsiteredExternalDecl = $"\t{packageInfo._api}U{engineClassName}* {objectInfo._registeredSingletonName}();\r\n";
+				objectInfo.UnregisteredSingletonName = objectInfo.RegisteredSingletonName = builder.ToString();
+				objectInfo.UnregisteredExternalDecl = objectInfo.RegsiteredExternalDecl = $"\t{packageInfo.Api}U{engineClassName}* {objectInfo.RegisteredSingletonName}();\r\n";
 			}
 
 			//COMPATIBILITY-TODO - The cross reference string should match the extern decl string always.  But currently, it is different for packages.
-			objectInfo._unregsiteredCrossReference = objectInfo._unregisteredExternalDecl;
-			objectInfo._regsiteredCrossReference = objectInfo._regsiteredExternalDecl;
+			objectInfo.UnregsiteredCrossReference = objectInfo.UnregisteredExternalDecl;
+			objectInfo.RegsiteredCrossReference = objectInfo.RegsiteredExternalDecl;
 
 			// Init the children
 			foreach (UhtType child in obj.Children)
diff --git a/Engine/Source/Programs/Shared/EpicGames.UHT/Exporters/CodeGen/UhtHeaderCodeGenerator.cs b/Engine/Source/Programs/Shared/EpicGames.UHT/Exporters/CodeGen/UhtHeaderCodeGenerator.cs
--- a/Engine/Source/Programs/Shared/EpicGames.UHT/Exporters/CodeGen/UhtHeaderCodeGenerator.cs
+++ b/Engine/Source/Programs/Shared/EpicGames.UHT/Exporters/CodeGen/UhtHeaderCodeGenerator.cs
@@ -35,7 +35,7 @@ namespace EpicGames.UHT.Exporters.CodeGen
 		#endregion
 
 		public readonly UhtHeaderFile HeaderFile;
-		public string FileId => this.HeaderInfos[this.HeaderFile.HeaderFileTypeIndex]._fileId;
+		public string FileId => this.HeaderInfos[this.HeaderFile.HeaderFileTypeIndex].FileId;
 
 		public UhtHeaderCodeGenerator(UhtCodeGenerator codeGenerator, UhtPackage package, UhtHeaderFile headerFile)
 			: base(codeGenerator, package)
diff --git a/Engine/Source/Programs/Shared/EpicGames.UHT/Exporters/CodeGen/UhtHeaderCodeGeneratorCppFile.cs b/Engine/Source/Programs/Shared/EpicGames.UHT/Exporters/CodeGen/UhtHeaderCodeGeneratorCppFile.cs
--- a/Engine/Source/Programs/Shared/EpicGames.UHT/Exporters/CodeGen/UhtHeaderCodeGeneratorCppFile.cs
+++ b/Engine/Source/Programs/Shared/EpicGames.UHT/Exporters/CodeGen/UhtHeaderCodeGeneratorCppFile.cs
@@ -37,7 +37,11 @@ namespace EpicGames.UHT.Exporters.CodeGen
 
 				builder.Append(HeaderCopyright);
 				builder.Append(RequiredCPPIncludes);
-				builder.Append("#include \"").Append(headerInfo._includePath).Append("\"\r\n");
+				builder.Append("#include \"").Append(headerInfo.IncludePath).Append("\"\r\n");
+				if (headerInfo.NeedsFastArrayHeaders)
+				{
+					builder.Append("#include \"Net/Serialization/FastArraySerializerImplementation.h\"\r\n");
+				}
 
 				bool addedStructuredArchiveFromArchiveHeader = false;
 				bool addedArchiveUObjectFromStructuredArchiveHeader = false;
@@ -51,7 +55,7 @@ namespace EpicGames.UHT.Exporters.CodeGen
 						{
 							if (addedIncludes.Add(classObj.ClassWithin.HeaderFile))
 							{
-								builder.Append("#include \"").Append(this.HeaderInfos[classObj.ClassWithin.HeaderFile.HeaderFileTypeIndex]._includePath).Append("\"\r\n");
+								builder.Append("#include \"").Append(this.HeaderInfos[classObj.ClassWithin.HeaderFile.HeaderFileTypeIndex].IncludePath).Append("\"\r\n");
 							}
 						}
 
@@ -133,7 +137,7 @@ namespace EpicGames.UHT.Exporters.CodeGen
 
 				if (hasRegisteredEnums || hasRegisteredScriptStructs || hasRegisteredClasses)
 				{
-					string name = $"Z_CompiledInDeferFile_{headerInfo._fileId}";
+					string name = $"Z_CompiledInDeferFile_{headerInfo.FileId}";
 					string staticsName = $"{name}_Statics";
 
 					builder.Append("\tstruct ").Append(staticsName).Append("\r\n");
@@ -179,7 +183,7 @@ namespace EpicGames.UHT.Exporters.CodeGen
 								{
 									builder.Append("#if WITH_EDITORONLY_DATA\r\n");
 								}
-								uint hash = this.ObjectInfos[enumObj.ObjectTypeIndex]._hash;
+								uint hash = this.ObjectInfos[enumObj.ObjectTypeIndex].Hash;
 								builder
 									.Append("\t\t{ ")
 									.Append(enumObj.SourceName)
@@ -211,7 +215,7 @@ namespace EpicGames.UHT.Exporters.CodeGen
 							{
 								if (scriptStruct.ScriptStructFlags.HasAnyFlags(EStructFlags.Native))
 								{
-									uint hash = this.ObjectInfos[scriptStruct.ObjectTypeIndex]._hash;
+									uint hash = this.ObjectInfos[scriptStruct.ObjectTypeIndex].Hash;
 									builder
 										.Append("\t\t{ ")
 										.Append(scriptStruct.SourceName)
@@ -240,7 +244,7 @@ namespace EpicGames.UHT.Exporters.CodeGen
 							{
 								if (!classObj.ClassFlags.HasAnyFlags(EClassFlags.Intrinsic))
 								{
-									uint hash = this.ObjectInfos[classObj.ObjectTypeIndex]._hash;
+									uint hash = this.ObjectInfos[classObj.ObjectTypeIndex].Hash;
 									builder
 										.Append("\t\t{ Z_Construct_UClass_")
 										.Append(classObj.SourceName)
@@ -298,7 +302,7 @@ namespace EpicGames.UHT.Exporters.CodeGen
 					}
 
 					// Save the hash of the generated body 
-					this.HeaderInfos[this.HeaderFile.HeaderFileTypeIndex]._bodyHash = UhtHash.GenenerateTextHash(generatedBody.Span[generatedBodyStart..generatedBodyEnd]);
+					this.HeaderInfos[this.HeaderFile.HeaderFileTypeIndex].BodyHash = UhtHash.GenenerateTextHash(generatedBody.Span[generatedBodyStart..generatedBodyEnd]);
 				}
 			}
 		}
@@ -416,7 +420,7 @@ namespace EpicGames.UHT.Exporters.CodeGen
 
 				{
 					using UhtBorrowBuffer borrowBuffer = new(builder, hashCodeBlockStart, builder.Length - hashCodeBlockStart);
-					this.ObjectInfos[enumObj.ObjectTypeIndex]._hash = UhtHash.GenenerateTextHash(borrowBuffer.Buffer.Memory.Span);
+					this.ObjectInfos[enumObj.ObjectTypeIndex].Hash = UhtHash.GenenerateTextHash(borrowBuffer.Buffer.Memory.Span);
 				}
 			}
 			return builder;
@@ -519,6 +523,12 @@ namespace EpicGames.UHT.Exporters.CodeGen
 				builder.Append("{\r\n");
 				builder.Append("\treturn ").Append(scriptStruct.SourceName).Append("::StaticStruct();\r\n");
 				builder.Append("}\r\n");
+
+				// Inject implementation needed to support auto bindings of fast arrays
+				if (this.ObjectInfos[scriptStruct.ObjectTypeIndex].FastArrayProperty != null)
+				{
+					builder.Append("UE_NET_IMPLEMENT_FASTARRAY(").Append(scriptStruct.SourceName).Append(");\r\n");
+				}
 			}
 
 			// Everything from this point on will be part of the definition hash
@@ -608,7 +618,7 @@ namespace EpicGames.UHT.Exporters.CodeGen
 
 			using (UhtBorrowBuffer borrowBuffer = new(builder, hashCodeBlockStart, builder.Length - hashCodeBlockStart))
 			{
-				this.ObjectInfos[scriptStruct.ObjectTypeIndex]._hash = UhtHash.GenenerateTextHash(borrowBuffer.Buffer.Memory.Span);
+				this.ObjectInfos[scriptStruct.ObjectTypeIndex].Hash = UhtHash.GenenerateTextHash(borrowBuffer.Buffer.Memory.Span);
 			}
 
 			// if this struct has RigVM methods we need to implement both the 
@@ -773,7 +783,7 @@ namespace EpicGames.UHT.Exporters.CodeGen
 
 			using (UhtBorrowBuffer borrowBuffer = new(builder, hashCodeBlockStart, builder.Length - hashCodeBlockStart))
 			{
-				this.ObjectInfos[function.ObjectTypeIndex]._hash = UhtHash.GenenerateTextHash(borrowBuffer.Buffer.Memory.Span);
+				this.ObjectInfos[function.ObjectTypeIndex].Hash = UhtHash.GenenerateTextHash(borrowBuffer.Buffer.Memory.Span);
 			}
 			return builder;
 		}
@@ -1430,7 +1440,7 @@ namespace EpicGames.UHT.Exporters.CodeGen
 				uint baseClassHash = 0;
 				if (classObj.SuperClass != null && !classObj.SuperClass.ClassFlags.HasAnyFlags(EClassFlags.Intrinsic))
 				{
-					baseClassHash = this.ObjectInfos[classObj.SuperClass.ObjectTypeIndex]._hash;
+					baseClassHash = this.ObjectInfos[classObj.SuperClass.ObjectTypeIndex].Hash;
 				}
 				hashBuilder.Append($"\r\n// {baseClassHash}\r\n");
 
@@ -1472,7 +1482,7 @@ namespace EpicGames.UHT.Exporters.CodeGen
 				// Calculate generated class initialization code hash so that we know when it changes after hot-reload
 				{
 					using UhtBorrowBuffer borrowBuffer = new(hashBuilder);
-					this.ObjectInfos[classObj.ObjectTypeIndex]._hash = UhtHash.GenenerateTextHash(borrowBuffer.Buffer.Memory.Span);
+					this.ObjectInfos[classObj.ObjectTypeIndex].Hash = UhtHash.GenenerateTextHash(borrowBuffer.Buffer.Memory.Span);
 				}
 			}
 
@@ -1768,7 +1778,7 @@ namespace EpicGames.UHT.Exporters.CodeGen
 
 			public uint GetTypeHash(UhtObject obj)
 			{
-				return this._codeGenerator._objectInfos[obj.ObjectTypeIndex]._hash;
+				return this._codeGenerator.ObjectInfos[obj.ObjectTypeIndex].Hash;
 			}
 		}
 	}
diff --git a/Engine/Source/Programs/Shared/EpicGames.UHT/Exporters/CodeGen/UhtHeaderCodeGeneratorHFile.cs b/Engine/Source/Programs/Shared/EpicGames.UHT/Exporters/CodeGen/UhtHeaderCodeGeneratorHFile.cs
--- a/Engine/Source/Programs/Shared/EpicGames.UHT/Exporters/CodeGen/UhtHeaderCodeGeneratorHFile.cs
+++ b/Engine/Source/Programs/Shared/EpicGames.UHT/Exporters/CodeGen/UhtHeaderCodeGeneratorHFile.cs
@@ -42,7 +42,7 @@ namespace EpicGames.UHT.Exporters.CodeGen
 				builder.Append(HeaderCopyright);
 				builder.Append("#include \"UObject/ObjectMacros.h\"\r\n");
 				builder.Append("#include \"UObject/ScriptMacros.h\"\r\n");
-				if (headerInfo._needsPushModelHeaders)
+				if (headerInfo.NeedsPushModelHeaders)
 				{
 					builder.Append("#include \"Net/Core/PushModel/PushModelMacros.h\"\r\n");
 				}
@@ -99,7 +99,7 @@ namespace EpicGames.UHT.Exporters.CodeGen
 				}
 
 				builder.Append("#undef CURRENT_FILE_ID\r\n");
-				builder.Append("#define CURRENT_FILE_ID ").Append(headerInfo._fileId).Append("\r\n\r\n\r\n");
+				builder.Append("#define CURRENT_FILE_ID ").Append(headerInfo.FileId).Append("\r\n\r\n\r\n");
 
 				foreach (UhtField field in this.HeaderFile.References.ExportTypes)
 				{
@@ -245,7 +245,22 @@ namespace EpicGames.UHT.Exporters.CodeGen
 
 					if (scriptStruct.SuperScriptStruct != null)
 					{
-						builder.Append("\ttypedef ").Append(scriptStruct.SuperScriptStruct.SourceName).Append(" Super;\\\\r\n");
+						builder.Append("\ttypedef ").Append(scriptStruct.SuperScriptStruct.SourceName).Append(" Super; \\\r\n");
+					}
+					UhtProperty? FastArrayProperty = this.ObjectInfos[scriptStruct.ObjectTypeIndex].FastArrayProperty;
+					if (FastArrayProperty != null)
+					{
+						builder
+							.Append("\tUE_NET_DECLARE_FASTARRAY(")
+							.Append(scriptStruct.SourceName)
+							.Append(", ")
+							.Append(FastArrayProperty.SourceName)
+							.Append(", ");
+						if (!scriptStruct.ScriptStructFlags.HasAnyFlags(EStructFlags.RequiredAPI))
+						{
+							builder.Append(this.PackageApi);
+						}
+						builder.Append("); \\\r\n");
 					}
 				}
 
@@ -321,7 +336,7 @@ namespace EpicGames.UHT.Exporters.CodeGen
 		private StringBuilder AppendClass(StringBuilder builder, UhtClass classObj)
 		{
 			string api = classObj.ClassFlags.HasAnyFlags(EClassFlags.MinimalAPI) ? this.PackageApi : "NO_API ";
-			UhtClass? nativeInterface = this.ObjectInfos[classObj.ObjectTypeIndex]._nativeInterface;
+			UhtClass? nativeInterface = this.ObjectInfos[classObj.ObjectTypeIndex].NativeInterface;
 
 			// Write the spare declarations
 			AppendSparseDeclarations(builder, classObj);
@@ -1007,7 +1022,7 @@ namespace EpicGames.UHT.Exporters.CodeGen
 			builder.Append("\tstatic void StaticRegisterNatives").Append(classObj.SourceName).Append("(); \\\r\n");
 			if (!classObj.ClassExportFlags.HasAnyFlags(UhtClassExportFlags.NoExport))
 			{
-				builder.Append("\tfriend struct ").Append(this.ObjectInfos[classObj.ObjectTypeIndex]._registeredSingletonName).Append("_Statics; \\\r\n");
+				builder.Append("\tfriend struct ").Append(this.ObjectInfos[classObj.ObjectTypeIndex].RegisteredSingletonName).Append("_Statics; \\\r\n");
 			}
 			builder.Append("public: \\\r\n");
 
diff --git a/Engine/Source/Programs/Shared/EpicGames.UHT/Exporters/CodeGen/UhtPackageCodeGenerator.cs b/Engine/Source/Programs/Shared/EpicGames.UHT/Exporters/CodeGen/UhtPackageCodeGenerator.cs
--- a/Engine/Source/Programs/Shared/EpicGames.UHT/Exporters/CodeGen/UhtPackageCodeGenerator.cs
+++ b/Engine/Source/Programs/Shared/EpicGames.UHT/Exporters/CodeGen/UhtPackageCodeGenerator.cs
@@ -28,11 +28,11 @@ namespace EpicGames.UHT.Exporters.CodeGen
 		public bool SaveExportedHeaders => this.Package.Module.SaveExportedHeaders;
 
 		public Utils.UhtSession Session => this.CodeGenerator.Session;
-		public UhtCodeGenerator.PackageInfo[] PackageInfos => this.CodeGenerator._packageInfos;
-		public UhtCodeGenerator.HeaderInfo[] HeaderInfos => this.CodeGenerator._headerInfos;
-		public UhtCodeGenerator.ObjectInfo[] ObjectInfos => this.CodeGenerator._objectInfos;
-		public string PackageApi => this.PackageInfos[this.Package.PackageTypeIndex]._api;
-		public string PackageSingletonName => this.ObjectInfos[this.Package.ObjectTypeIndex]._registeredSingletonName;
+		public UhtCodeGenerator.PackageInfo[] PackageInfos => this.CodeGenerator.PackageInfos;
+		public UhtCodeGenerator.HeaderInfo[] HeaderInfos => this.CodeGenerator.HeaderInfos;
+		public UhtCodeGenerator.ObjectInfo[] ObjectInfos => this.CodeGenerator.ObjectInfos;
+		public string PackageApi => this.PackageInfos[this.Package.PackageTypeIndex].Api;
+		public string PackageSingletonName => this.ObjectInfos[this.Package.ObjectTypeIndex].RegisteredSingletonName;
 
 		public UhtPackageCodeGenerator(UhtCodeGenerator codeGenerator, UhtPackage package)
 		{
diff --git a/Engine/Source/Programs/Shared/EpicGames.UHT/Exporters/CodeGen/UhtPackageCodeGeneratorCppFile.cs b/Engine/Source/Programs/Shared/EpicGames.UHT/Exporters/CodeGen/UhtPackageCodeGeneratorCppFile.cs
--- a/Engine/Source/Programs/Shared/EpicGames.UHT/Exporters/CodeGen/UhtPackageCodeGeneratorCppFile.cs
+++ b/Engine/Source/Programs/Shared/EpicGames.UHT/Exporters/CodeGen/UhtPackageCodeGeneratorCppFile.cs
@@ -49,7 +49,7 @@ namespace EpicGames.UHT.Exporters.CodeGen
 
 					singletons.AddRange(headerFile.References.Singletons);
 
-					uint bodyHash = this.HeaderInfos[headerFile.HeaderFileTypeIndex]._bodyHash;
+					uint bodyHash = this.HeaderInfos[headerFile.HeaderFileTypeIndex].BodyHash;
 					if (bodiesHash == 0)
 					{
 						// Don't combine in the first case because it keeps GUID backwards compatibility
@@ -68,7 +68,7 @@ namespace EpicGames.UHT.Exporters.CodeGen
 				}
 				uint declarationsHash = UhtHash.GenenerateTextHash(declarations.ToString());
 
-				string strippedName = this.PackageInfos[this.Package.PackageTypeIndex]._strippedName;
+				string strippedName = this.PackageInfos[this.Package.PackageTypeIndex].StrippedName;
 				string singletonName = this.PackageSingletonName;
 				singletons.Sort((UhtField lhs, UhtField rhs) =>
 				{
@@ -79,8 +79,8 @@ namespace EpicGames.UHT.Exporters.CodeGen
 						return !lhsIsDel ? -1 : 1;
 					}
 					return StringComparerUE.OrdinalIgnoreCase.Compare(
-						this.ObjectInfos[lhs.ObjectTypeIndex]._registeredSingletonName,
-						this.ObjectInfos[rhs.ObjectTypeIndex]._registeredSingletonName);
+						this.ObjectInfos[lhs.ObjectTypeIndex].RegisteredSingletonName,
+						this.ObjectInfos[rhs.ObjectTypeIndex].RegisteredSingletonName);
 				});
 
 				builder.Append(HeaderCopyright);
@@ -102,7 +102,7 @@ namespace EpicGames.UHT.Exporters.CodeGen
 				foreach (UhtObject obj in singletons)
 				{
 					ref UhtCodeGenerator.ObjectInfo info = ref this.ObjectInfos[obj.ObjectTypeIndex];
-					builder.Append(info._regsiteredExternalDecl);
+					builder.Append(info.RegsiteredExternalDecl);
 				}
 
 				builder.AppendMetaDataDef(this.Package, null, MetaDataParamsName, 3);
@@ -118,7 +118,7 @@ namespace EpicGames.UHT.Exporters.CodeGen
 					builder.Append("\t\t\tstatic UObject* (*const SingletonFuncArray[])() = {\r\n");
 					foreach (UhtField field in singletons)
 					{
-						builder.Append("\t\t\t\t(UObject* (*)())").Append(this.ObjectInfos[field.ObjectTypeIndex]._registeredSingletonName).Append(",\r\n");
+						builder.Append("\t\t\t\t(UObject* (*)())").Append(this.ObjectInfos[field.ObjectTypeIndex].RegisteredSingletonName).Append(",\r\n");
 					}
 					builder.Append("\t\t\t};\r\n");
 				}
diff --git a/Engine/Source/Programs/Shared/EpicGames.UHT/Exporters/CodeGen/UhtPackageCodeGeneratorHFile.cs b/Engine/Source/Programs/Shared/EpicGames.UHT/Exporters/CodeGen/UhtPackageCodeGeneratorHFile.cs
--- a/Engine/Source/Programs/Shared/EpicGames.UHT/Exporters/CodeGen/UhtPackageCodeGeneratorHFile.cs
+++ b/Engine/Source/Programs/Shared/EpicGames.UHT/Exporters/CodeGen/UhtPackageCodeGeneratorHFile.cs
@@ -55,7 +55,7 @@ namespace EpicGames.UHT.Exporters.CodeGen
 				{
 					if (headerFile.HeaderFileType == UhtHeaderFileType.Classes)
 					{
-						builder.Append("#include \"").Append(this.HeaderInfos[headerFile.HeaderFileTypeIndex]._includePath).Append("\"\r\n");
+						builder.Append("#include \"").Append(this.HeaderInfos[headerFile.HeaderFileTypeIndex].IncludePath).Append("\"\r\n");
 					}
 				}
 
