Fixed issue with UAT/UBT script module rebuild tracking.  If MSBuild doesn't update the target DLL due to there being no real changes, then the simple check against target time stamps being out of order would fail.  Changed to save the reference time stamp and then compare them to test to see if they don't match.

Minor refactor to how valid and invalid build records are tracked.

diff --git a/Engine/Source/Programs/Shared/EpicGames.Build/System/CompileScriptModules.cs b/Engine/Source/Programs/Shared/EpicGames.Build/System/CompileScriptModules.cs
--- a/Engine/Source/Programs/Shared/EpicGames.Build/System/CompileScriptModules.cs
+++ b/Engine/Source/Programs/Shared/EpicGames.Build/System/CompileScriptModules.cs
@@ -19,10 +19,22 @@ namespace UnrealBuildBase
 		{
 			private ILogger Logger;
 			private Dictionary<string, DateTime> WriteTimes = new Dictionary<string, DateTime>();
+			private string BuildRecordDirectory;
 
-			public Hook(ILogger InLogger)
+			public Hook(ILogger InLogger, Rules.RulesFileType InRulesFileType)
 			{
 				Logger = InLogger;
+				switch (InRulesFileType)
+				{
+					case Rules.RulesFileType.AutomationModule:
+						BuildRecordDirectory = "ScriptModules";
+						break;
+					case Rules.RulesFileType.UbtPlugin:
+						BuildRecordDirectory = "ScriptModules.UBT";
+						break;
+					default:
+						throw new ArgumentException("Unexpected rules file type");
+				}
 			}
 
 			public DateTime GetLastWriteTime(DirectoryReference BasePath, string RelativeFilePath)
@@ -40,20 +52,20 @@ namespace UnrealBuildBase
 				return WriteTime;
 			}
 
-			public void ValidateRecursively(
-				Dictionary<FileReference, (CsProjBuildRecord, FileReference)> ValidBuildRecords,
-				Dictionary<FileReference, (CsProjBuildRecord, FileReference)> InvalidBuildRecords,
-				Dictionary<FileReference, (CsProjBuildRecord, FileReference)> BuildRecords,
-				FileReference ProjectPath)
+			public DirectoryReference GetBuildRecordDirectory(DirectoryReference BasePath)
 			{
-				CompileScriptModule.ValidateBuildRecordRecursively(ValidBuildRecords, InvalidBuildRecords, BuildRecords, ProjectPath, this, Logger);
+				return DirectoryReference.Combine(BasePath, "Intermediate", BuildRecordDirectory);
+			}
+
+			public void ValidateRecursively(Dictionary<FileReference, CsProjBuildRecordEntry> BuildRecords, FileReference ProjectPath)
+			{
+				CompileScriptModule.ValidateBuildRecordRecursively(BuildRecords, ProjectPath, this, Logger);
 			}
 
 			public bool HasWildcards(string FileSpec)
 			{
 				return FileMatcher.HasWildcards(FileSpec);
 			}
-
 			DirectoryReference CsProjBuildHook.EngineDirectory => Unreal.EngineDirectory;
 
 			DirectoryReference CsProjBuildHook.DotnetDirectory => Unreal.DotnetDirectory;
@@ -66,7 +78,7 @@ namespace UnrealBuildBase
 		/// </summary>
 		/// <param name="BuildRecords">Input build records</param>
 		/// <returns>Set of target files</returns>
-		public static HashSet<FileReference> GetTargetPaths(IReadOnlyDictionary<FileReference, (CsProjBuildRecord BuildRecord, FileReference BuildRecordPath)> BuildRecords)
+		public static HashSet<FileReference> GetTargetPaths(IReadOnlyDictionary<FileReference, CsProjBuildRecordEntry> BuildRecords)
 		{
 			return new HashSet<FileReference>(BuildRecords.Select(x => GetTargetPath(x)));
 		}
@@ -76,7 +88,7 @@ namespace UnrealBuildBase
 		/// </summary>
 		/// <param name="BuildRecord">Build record</param>
 		/// <returns>File reference for the target</returns>
-		public static FileReference GetTargetPath(KeyValuePair<FileReference, (CsProjBuildRecord BuildRecord, FileReference BuildRecordPath)> BuildRecord)
+		public static FileReference GetTargetPath(KeyValuePair<FileReference, CsProjBuildRecordEntry> BuildRecord)
 		{
 			return FileReference.Combine(BuildRecord.Key.Directory, BuildRecord.Value.BuildRecord.TargetPath!);
 		}
@@ -120,26 +132,24 @@ namespace UnrealBuildBase
 		/// <summary>
 		/// Test to see if all the given projects are up-to-date
 		/// </summary>
+		/// <param name="RulesFileType">Type of script modules being tested</param>
 		/// <param name="FoundProjects">Collection of projects to test</param>
 		/// <param name="BaseDirectories">Base directories of the projects</param>
 		/// <param name="Logger">Logger for output</param>
 		/// <returns>True if all of the projects are up to date</returns>
-		public static bool AreScriptModulesUpToDate(HashSet<FileReference> FoundProjects, List<DirectoryReference> BaseDirectories, ILogger Logger)
+		public static bool AreScriptModulesUpToDate(Rules.RulesFileType RulesFileType, HashSet<FileReference> FoundProjects, List<DirectoryReference> BaseDirectories, ILogger Logger)
 		{
-			CsProjBuildHook Hook = new Hook(Logger);
+			CsProjBuildHook Hook = new Hook(Logger, RulesFileType);
 
 			// Load existing build records, validating them only if (re)compiling script projects is an option
-			Dictionary<FileReference, (CsProjBuildRecord BuildRecord, FileReference BuildRecordPath)> ExistingBuildRecords = LoadExistingBuildRecords(BaseDirectories, Logger);
-
-			Dictionary<FileReference, (CsProjBuildRecord BuildRecord, FileReference BuildRecordPath)> ValidBuildRecords = new Dictionary<FileReference, (CsProjBuildRecord, FileReference)>(ExistingBuildRecords.Count);
-			Dictionary<FileReference, (CsProjBuildRecord BuildRecord, FileReference BuildRecordPath)> InvalidBuildRecords = new Dictionary<FileReference, (CsProjBuildRecord, FileReference)>(ExistingBuildRecords.Count);
-
+			Dictionary<FileReference, CsProjBuildRecordEntry> BuildRecords = LoadExistingBuildRecords(Hook, BaseDirectories, Logger);
 			foreach (FileReference Project in FoundProjects)
 			{
-				ValidateBuildRecordRecursively(ValidBuildRecords, InvalidBuildRecords, ExistingBuildRecords, Project, Hook, Logger);
+				ValidateBuildRecordRecursively(BuildRecords, Project, Hook, Logger);
 			}
 
 			// If all found records are valid, we can return their targets directly
+			Dictionary<FileReference, CsProjBuildRecordEntry> ValidBuildRecords = new(BuildRecords.Where(x => x.Value.Status == CsProjBuildRecordStatus.Valid));
 			return FoundProjects.All(x => ValidBuildRecords.ContainsKey(x));
 		}
 
@@ -157,26 +167,23 @@ namespace UnrealBuildBase
 		/// <param name="OnBuildingProjects">Action to invoke when projects get built</param>
 		/// <param name="Logger"></param>
 		/// <returns>Collection of all the projects.  They will have been compiled.</returns>
-		public static Dictionary<FileReference, (CsProjBuildRecord, FileReference)> Build(Rules.RulesFileType RulesFileType,
+		public static Dictionary<FileReference, CsProjBuildRecordEntry> Build(Rules.RulesFileType RulesFileType,
 			HashSet<FileReference> FoundProjects, List<DirectoryReference> BaseDirectories, List<string>? DefineConstants, 
 			bool bForceCompile, bool bNoCompile, bool bUseBuildRecords,
 			out bool bBuildSuccess, Action<int> OnBuildingProjects, ILogger Logger)
 		{
-			CsProjBuildHook Hook = new Hook(Logger);
+			CsProjBuildHook Hook = new Hook(Logger, RulesFileType);
 
 			bool bUseBuildRecordsOnlyForProjectDiscovery = bNoCompile || Unreal.IsEngineInstalled();
 
 			// Load existing build records, validating them only if (re)compiling script projects is an option
-			Dictionary<FileReference, (CsProjBuildRecord BuildRecord, FileReference BuildRecordPath)> ExistingBuildRecords = LoadExistingBuildRecords(BaseDirectories, Logger);
-
-			Dictionary<FileReference, (CsProjBuildRecord BuildRecord, FileReference BuildRecordPath)> ValidBuildRecords = new Dictionary<FileReference, (CsProjBuildRecord, FileReference)>(ExistingBuildRecords.Count);
-			Dictionary<FileReference, (CsProjBuildRecord BuildRecord, FileReference BuildRecordPath)> InvalidBuildRecords = new Dictionary<FileReference, (CsProjBuildRecord, FileReference)>(ExistingBuildRecords.Count);
+			Dictionary<FileReference, CsProjBuildRecordEntry> BuildRecords = LoadExistingBuildRecords(Hook, BaseDirectories, Logger);
 
 			if (bUseBuildRecords)
 			{
 				foreach (FileReference Project in FoundProjects)
 				{
-					ValidateBuildRecordRecursively(ValidBuildRecords, InvalidBuildRecords, ExistingBuildRecords, Project, Hook, Logger);
+					ValidateBuildRecordRecursively(BuildRecords, Project, Hook, Logger);
 				}
 			}
 
@@ -194,9 +201,8 @@ namespace UnrealBuildBase
 					default:
 						throw new Exception("Unsupported rules file type");
 				}
-				Dictionary<FileReference, (CsProjBuildRecord, FileReference)> OutRecords = new Dictionary<FileReference, (CsProjBuildRecord, FileReference)>(ExistingBuildRecords.Count);
-				foreach (KeyValuePair<FileReference, (CsProjBuildRecord BuildRecord, FileReference BuildRecordPath)> Record in
-					ExistingBuildRecords.Where(x => x.Value.BuildRecordPath.HasExtension(FilterExtension)))
+				Dictionary<FileReference, CsProjBuildRecordEntry> OutRecords = new Dictionary<FileReference, CsProjBuildRecordEntry>(BuildRecords.Count);
+				foreach (KeyValuePair<FileReference, CsProjBuildRecordEntry> Record in BuildRecords.Where(x => x.Value.BuildRecordFile.HasExtension(FilterExtension)))
 				{
 					FileReference TargetPath = FileReference.Combine(Record.Key.Directory, Record.Value.BuildRecord.TargetPath!);
 
@@ -209,12 +215,12 @@ namespace UnrealBuildBase
 						if (bNoCompile)
 						{
 							// when -NoCompile is on the command line, try to run with whatever is available
-							Logger.LogWarning("Script module \"{TargetPath}\" not found for record \"{BuildRecordPath}\"", TargetPath, Record.Value.BuildRecordPath);
+							Logger.LogWarning("Script module \"{TargetPath}\" not found for record \"{BuildRecordPath}\"", TargetPath, Record.Value.BuildRecordFile);
 						}
 						else
 						{
 							// when the engine is installed, expect to find a built target assembly for every record that was found
-							throw new Exception($"Script module \"{TargetPath}\" not found for record \"{Record.Value.BuildRecordPath}\"");
+							throw new Exception($"Script module \"{TargetPath}\" not found for record \"{Record.Value.BuildRecordFile}\"");
 						}
 
 					}
@@ -225,12 +231,12 @@ namespace UnrealBuildBase
 			else
 			{
 				// when the engine is not installed, delete any build record .json file that is not valid
-				foreach ((CsProjBuildRecord _, FileReference BuildRecordPath) in InvalidBuildRecords.Values)
+				foreach (CsProjBuildRecordEntry Entry in BuildRecords.Values.Where(x => x.Status == CsProjBuildRecordStatus.Invalid))
 				{
-					if (BuildRecordPath != null)
+					if (Entry.BuildRecordFile != null)
 					{
-						Logger.LogDebug("Deleting invalid build record \"{BuildRecordPath}\"", BuildRecordPath);
-						FileReference.Delete(BuildRecordPath);
+						Logger.LogDebug("Deleting invalid build record \"{BuildRecordPath}\"", Entry.BuildRecordFile);
+						FileReference.Delete(Entry.BuildRecordFile);
 					}
 				}
 			}
@@ -238,10 +244,11 @@ namespace UnrealBuildBase
 			if (!bForceCompile && bUseBuildRecords)
 			{
 				// If all found records are valid, we can return their targets directly
+				Dictionary<FileReference, CsProjBuildRecordEntry> ValidBuildRecords = new(BuildRecords.Where(x => x.Value.Status == CsProjBuildRecordStatus.Valid));
 				if (FoundProjects.All(x => ValidBuildRecords.ContainsKey(x)))
 				{
 					bBuildSuccess = true;
-					return new Dictionary<FileReference, (CsProjBuildRecord, FileReference)>(ValidBuildRecords.Where(x => FoundProjects.Contains(x.Key)));
+					return new Dictionary<FileReference, CsProjBuildRecordEntry>(ValidBuildRecords.Where(x => FoundProjects.Contains(x.Key)));
 				}
 			}
 
@@ -255,7 +262,7 @@ namespace UnrealBuildBase
 		/// If it is placed in the caller directly, then when the caller is invoked, the assembly will be loaded resulting
 		/// in the possible Microsoft.Build.Framework load issue later on is this method isn't invoked.
 		/// </summary>
-		static Dictionary<FileReference, (CsProjBuildRecord, FileReference)> Build(HashSet<FileReference> FoundProjects,
+		static Dictionary<FileReference, CsProjBuildRecordEntry> Build(HashSet<FileReference> FoundProjects,
 			bool bForceCompile, out bool bBuildSuccess, CsProjBuildHook Hook, List<DirectoryReference> BaseDirectories,
 			List<string>? DefineConstants, Action<int> OnBuildingProjects, ILogger Logger)
 		{
@@ -265,16 +272,17 @@ namespace UnrealBuildBase
 		/// <summary>
 		/// Find and load existing build record .json files from any Intermediate/ScriptModules found in the provided lists
 		/// </summary>
+		/// <param name="Hook">Hook object used to query directory</param>
 		/// <param name="BaseDirectories"></param>
 		/// <param name="Logger"></param>
 		/// <returns></returns>
-		static Dictionary<FileReference, (CsProjBuildRecord, FileReference)> LoadExistingBuildRecords(List<DirectoryReference> BaseDirectories, ILogger Logger)
+		static Dictionary<FileReference, CsProjBuildRecordEntry> LoadExistingBuildRecords(CsProjBuildHook Hook, List<DirectoryReference> BaseDirectories, ILogger Logger)
         {
-			Dictionary<FileReference, (CsProjBuildRecord, FileReference)> LoadedBuildRecords = new Dictionary<FileReference, (CsProjBuildRecord, FileReference)>();
+			Dictionary<FileReference, CsProjBuildRecordEntry> LoadedBuildRecords = new Dictionary<FileReference, CsProjBuildRecordEntry>();
 
 			foreach (DirectoryReference Directory in BaseDirectories)
 			{
-				DirectoryReference IntermediateDirectory = DirectoryReference.Combine(Directory, "Intermediate", "ScriptModules");
+				DirectoryReference IntermediateDirectory = Hook.GetBuildRecordDirectory(Directory);
 				if (!DirectoryReference.Exists(IntermediateDirectory))
 				{
 					continue;
@@ -298,9 +306,11 @@ namespace UnrealBuildBase
 
 					if (BuildRecord != null && BuildRecord.ProjectPath != null)
 					{
-						// Path.Combine, etc don't work so well with improper dir separators (a build record made on Windows will fail on Linux with the \ characters)
-						string Combined = Path.Combine(JsonFile.Directory.FullName, BuildRecord.ProjectPath.Replace('\\', Path.DirectorySeparatorChar).Replace('/', Path.DirectorySeparatorChar));
-						LoadedBuildRecords.Add(FileReference.FromString(Path.GetFullPath(Combined)), (BuildRecord, JsonFile));
+						CsProjBuildRecordEntry Entry = new CsProjBuildRecordEntry(
+							FileReference.FromString(Path.GetFullPath(BuildRecord.ProjectPath, JsonFile.Directory.FullName)),
+							JsonFile,
+							BuildRecord);
+						LoadedBuildRecords.Add(Entry.ProjectFile, Entry);
 					}
 					else
                     {
@@ -434,21 +444,20 @@ namespace UnrealBuildBase
 		}
 
 		static void ValidateBuildRecordRecursively(
-			Dictionary<FileReference, (CsProjBuildRecord BuildRecord, FileReference)> ValidBuildRecords,
-			Dictionary<FileReference, (CsProjBuildRecord, FileReference)> InvalidBuildRecords,
-			Dictionary<FileReference, (CsProjBuildRecord, FileReference)> BuildRecords,
+			Dictionary<FileReference, CsProjBuildRecordEntry> BuildRecords,
 			FileReference ProjectPath, CsProjBuildHook Hook, ILogger Logger)
 		{
-			if (ValidBuildRecords.ContainsKey(ProjectPath) || InvalidBuildRecords.ContainsKey(ProjectPath))
+			// Was a build record loaded for this project path? (relevant when considering referenced projects)
+			if (!BuildRecords.TryGetValue(ProjectPath, out CsProjBuildRecordEntry? Entry))
 			{
-				// Project validity has already been determined
+				Logger.LogDebug("Found project {ProjectPath} with no existing build record", ProjectPath);
 				return;
 			}
 
-			// Was a build record loaded for this project path? (relevant when considering referenced projects)
-			if (!BuildRecords.TryGetValue(ProjectPath, out (CsProjBuildRecord BuildRecord, FileReference BuildRecordPath) Entry))
+			// Ignore if the status is known
+			if (Entry.Status != CsProjBuildRecordStatus.Unknown)
 			{
-				Logger.LogDebug("Found project {ProjectPath} with no existing build record", ProjectPath);
+				// Project validity has already been determined
 				return;
 			}
 
@@ -457,38 +466,60 @@ namespace UnrealBuildBase
 			{
 				string ProjectRelativePath = Path.GetRelativePath(Unreal.EngineDirectory.FullName, ProjectPath.FullName);
 				Logger.LogDebug("[{ProjectRelativePath}] {ValidationFailureMessage}", ProjectRelativePath, ValidationFailureMessage);
-
-				InvalidBuildRecords.Add(ProjectPath, Entry);
+				Entry.Status = CsProjBuildRecordStatus.Invalid;
 				return;
 			}
 
 			// Are all referenced build records valid?
-			foreach (string ReferencedProjectPath in Entry.BuildRecord.ProjectReferences)
+			foreach (CsProjBuildRecordRef ReferencedProject in Entry.BuildRecord.ProjectReferencesAndTimes)
 			{
-				FileReference FullProjectPath = FileReference.FromString(Path.GetFullPath(ReferencedProjectPath, ProjectPath.Directory.FullName));
-				ValidateBuildRecordRecursively(ValidBuildRecords, InvalidBuildRecords, BuildRecords, FullProjectPath, Hook, Logger);
+				FileReference FullProjectPath = FileReference.FromString(Path.GetFullPath(ReferencedProject.ProjectPath, ProjectPath.Directory.FullName));
+				ValidateBuildRecordRecursively(BuildRecords, FullProjectPath, Hook, Logger);
 
-				if (!ValidBuildRecords.ContainsKey(FullProjectPath))
+				if (!BuildRecords.TryGetValue(FullProjectPath, out CsProjBuildRecordEntry? RefEntry))
 				{
 					string ProjectRelativePath = Path.GetRelativePath(Unreal.EngineDirectory.FullName, ProjectPath.FullName);
 					string DependencyRelativePath = Path.GetRelativePath(Unreal.EngineDirectory.FullName, FullProjectPath.FullName);
-					Logger.LogDebug("[{ProjectRelativePath}] Existing output is not valid because dependency {DependencyRelativePath} is not valid", ProjectRelativePath, DependencyRelativePath);
-					InvalidBuildRecords.Add(ProjectPath, Entry);
+					Logger.LogDebug("[{ProjectRelativePath}] Existing output is not valid because dependency {DependencyRelativePath} could not be found", ProjectRelativePath, DependencyRelativePath);
+					Entry.Status = CsProjBuildRecordStatus.Invalid;
 					return;
 				}
 
-				// Ensure that the dependency was not built more recently than the project
-				if (Entry.BuildRecord.TargetBuildTime < ValidBuildRecords[FullProjectPath].BuildRecord.TargetBuildTime)
+				if (RefEntry.Status != CsProjBuildRecordStatus.Valid)
 				{
 					string ProjectRelativePath = Path.GetRelativePath(Unreal.EngineDirectory.FullName, ProjectPath.FullName);
 					string DependencyRelativePath = Path.GetRelativePath(Unreal.EngineDirectory.FullName, FullProjectPath.FullName);
-					Logger.LogDebug("[{ProjectRelativePath}] Existing output is not valid because dependency {DependencyRelativePath} is newer", ProjectRelativePath, DependencyRelativePath);
-					InvalidBuildRecords.Add(ProjectPath, Entry);
+					Logger.LogDebug("[{ProjectRelativePath}] Existing output is not valid because dependency {DependencyRelativePath} is not valid", ProjectRelativePath, DependencyRelativePath);
+					Entry.Status = CsProjBuildRecordStatus.Invalid;
 					return;
 				}
+
+				// Ensure that the dependency was not built more recently than the project
+				if (ReferencedProject.TargetBuildTime == DateTime.MinValue)
+				{
+					if (Entry.BuildRecord.TargetBuildTime < RefEntry.BuildRecord.TargetBuildTime)
+					{
+						string ProjectRelativePath = Path.GetRelativePath(Unreal.EngineDirectory.FullName, ProjectPath.FullName);
+						string DependencyRelativePath = Path.GetRelativePath(Unreal.EngineDirectory.FullName, FullProjectPath.FullName);
+						Logger.LogDebug("[{ProjectRelativePath}] Existing output is not valid because dependency {DependencyRelativePath} is newer", ProjectRelativePath, DependencyRelativePath);
+						Entry.Status = CsProjBuildRecordStatus.Invalid;
+						return;
+					}
+				}
+				else
+				{
+					if (ReferencedProject.TargetBuildTime != RefEntry.BuildRecord.TargetBuildTime)
+					{
+						string ProjectRelativePath = Path.GetRelativePath(Unreal.EngineDirectory.FullName, ProjectPath.FullName);
+						string DependencyRelativePath = Path.GetRelativePath(Unreal.EngineDirectory.FullName, FullProjectPath.FullName);
+						Logger.LogDebug("[{ProjectRelativePath}] Existing dependency output time stamp for {DependencyRelativePath} does not match expected value", ProjectRelativePath, DependencyRelativePath);
+						Entry.Status = CsProjBuildRecordStatus.Invalid;
+						return;
+					}
+				}
 			}
 
-			ValidBuildRecords.Add(ProjectPath, Entry);
+			Entry.Status = CsProjBuildRecordStatus.Valid;
 		}
 
 		static List<DirectoryReference> GetGameDirectories(string? ScriptsForProjectFileName, ILogger Logger)
diff --git a/Engine/Source/Programs/Shared/EpicGames.Build/System/EnumeratePlugins.cs b/Engine/Source/Programs/Shared/EpicGames.Build/System/EnumeratePlugins.cs
--- a/Engine/Source/Programs/Shared/EpicGames.Build/System/EnumeratePlugins.cs
+++ b/Engine/Source/Programs/Shared/EpicGames.Build/System/EnumeratePlugins.cs
@@ -152,7 +152,7 @@ namespace UnrealBuildBase
 			// If we have anything left to build, then build it
 			if (NotBuilt.Any())
 			{
-				Dictionary<FileReference, (CsProjBuildRecord BuildRecord, FileReference BuildRecordFile)>? BuiltPlugins =
+				Dictionary<FileReference, CsProjBuildRecordEntry>? BuiltPlugins =
 					CompileScriptModule.Build(UnrealBuildBase.Rules.RulesFileType.UbtPlugin, new HashSet<FileReference>(NotBuilt),
 					GetBaseDirectories(ProjectFile), DefineConstants, false, false, true, out bBuildSuccess,
 					Count =>
@@ -169,7 +169,7 @@ namespace UnrealBuildBase
 				);
 				
 				// Add any built plugins back into the cache
-				foreach (KeyValuePair<FileReference, (CsProjBuildRecord BuildRecord, FileReference BuildRecordFile)> BuiltPlugin in BuiltPlugins)
+				foreach (KeyValuePair<FileReference, CsProjBuildRecordEntry> BuiltPlugin in BuiltPlugins)
 				{
 					Built.Add(BuiltPlugin.Key);
 					UbtPluginFileCache.Add(BuiltPlugin.Key, CompileScriptModule.GetTargetPath(BuiltPlugin));
diff --git a/Engine/Source/Programs/Shared/EpicGames.Core/CsProjBuildHook.cs b/Engine/Source/Programs/Shared/EpicGames.Core/CsProjBuildHook.cs
--- a/Engine/Source/Programs/Shared/EpicGames.Core/CsProjBuildHook.cs
+++ b/Engine/Source/Programs/Shared/EpicGames.Core/CsProjBuildHook.cs
@@ -33,19 +33,18 @@ namespace EpicGames.Core
 		DateTime GetLastWriteTime(string BasePath, string RelativeFilePath);
 
 		/// <summary>
-		/// Validate the given build records splitting them up into valid and invalid records.
-		/// The invalid ones will need to be build.
+		/// Return the build record directory for the given base path (i.e. engine dir or project dir)
 		/// </summary>
-		/// <param name="ValidBuildRecords">Output collection of up-to-date build records</param>
-		/// <param name="InvalidBuildRecords">Output collection of build records needing to be built.</param>
-		/// <param name="BuildRecords">Top level build records to test.  Dependencies are also checked so the
-		/// of valid and invalid build records will usually be larger than the number of input records</param>
-		/// <param name="ProjectPath">Path of the project</param>
-		void ValidateRecursively(
-			Dictionary<FileReference, (CsProjBuildRecord, FileReference)> ValidBuildRecords,
-			Dictionary<FileReference, (CsProjBuildRecord, FileReference)> InvalidBuildRecords,
-			Dictionary<FileReference, (CsProjBuildRecord, FileReference)> BuildRecords,
-			FileReference ProjectPath);
+		/// <param name="BasePath">The base path for the directory</param>
+		/// <returns>Directory for the build records</returns>
+		DirectoryReference GetBuildRecordDirectory(DirectoryReference BasePath);
+
+		/// <summary>
+		/// Validate the given build records for the project
+		/// </summary>
+		/// <param name="buildRecords">Build records being validated.  This also includes build records for dependencies.</param>
+		/// <param name="projectPath">Path of the project</param>
+		void ValidateRecursively(Dictionary<FileReference, CsProjBuildRecordEntry> buildRecords, FileReference projectPath);
 
 		/// <summary>
 		/// Test to see if the given file spec has any wild cards
diff --git a/Engine/Source/Programs/Shared/EpicGames.Core/CsProjBuildRecord.cs b/Engine/Source/Programs/Shared/EpicGames.Core/CsProjBuildRecord.cs
--- a/Engine/Source/Programs/Shared/EpicGames.Core/CsProjBuildRecord.cs
+++ b/Engine/Source/Programs/Shared/EpicGames.Core/CsProjBuildRecord.cs
@@ -6,6 +6,20 @@ using System.Text;
 
 namespace EpicGames.Core
 {
+
+	public class CsProjBuildRecordRef
+	{
+		/// <summary>
+		/// Name of the referenced project
+		/// </summary>
+		public string ProjectPath { get; set; } = String.Empty;
+
+		/// <summary>
+		/// The time that the target assembly was built (read from the file after the build)
+		/// </summary>
+		public DateTime TargetBuildTime { get; set; } = DateTime.Now;
+	}
+
 	/// <summary>
 	/// Acceleration structure:
 	/// used to encapsulate a full set of dependencies for an msbuild project - explicit and globbed
@@ -16,12 +30,12 @@ namespace EpicGames.Core
 		/// <summary>
 		/// Version number making it possible to quickly invalidate written records.
 		/// </summary>
-		public static readonly int CurrentVersion = 4;
+		public static readonly int CurrentVersion = 5;
 
 		/// <summary>
 		/// Version number for this record
 		/// </summary>
-		public int Version { get; set; } // what value does this get if deserialized from a file with no value for this field? 
+		public int Version { get; set; } = -1;
 
 		/// <summary>
 		/// Path to the .csproj project file, relative to the location of the build record .json file
@@ -33,7 +47,6 @@ namespace EpicGames.Core
 		/// </summary>
 		public DateTime TargetBuildTime { get; set; }
 
-
 		// all following paths are relative to the project directory, the directory containing ProjectPath 
 
 		/// <summary>
@@ -42,9 +55,9 @@ namespace EpicGames.Core
 		public string? TargetPath { get; set; }
 
 		/// <summary>
-		/// Paths of referenced projects
+		/// Paths to referenced projects and the target build time in that project
 		/// </summary>
-		public HashSet<string> ProjectReferences { get; set; } = new HashSet<string>();
+		public List<CsProjBuildRecordRef> ProjectReferencesAndTimes { get; set; } = new List<CsProjBuildRecordRef>();
 
 		/// <summary>
 		/// file dependencies from non-glob sources
@@ -87,4 +100,69 @@ namespace EpicGames.Core
 		/// </summary>
 		public List<Glob> Globs { get; set; } = new List<Glob>();
 	}
+
+	/// <summary>
+	/// Represents the status of a build record.  When read from disk, the status will start
+	/// as being unknown.  Then through the validation process, it will either become valid
+	/// or invalid.
+	/// </summary>
+	public enum CsProjBuildRecordStatus
+	{
+		/// <summary>
+		/// Status of the build record is unknown
+		/// </summary>
+		Unknown,
+
+		/// <summary>
+		/// The build record is valid and doesn't require rebuild
+		/// </summary>
+		Valid,
+
+		/// <summary>
+		/// The build record is invalid
+		/// </summary>
+		Invalid,
+	};
+
+	/// <summary>
+	/// Represents a runtime build record entry that isn't persisted.  Aside from the build record, it 
+	/// contains the project file and build record file reference.  Also the "need to build" status
+	/// of the record.
+	/// </summary>
+	public class CsProjBuildRecordEntry
+	{
+
+		/// <summary>
+		/// Location of the csproj file
+		/// </summary>
+		public FileReference ProjectFile { get; }
+
+		/// <summary>
+		/// Location of the build record
+		/// </summary>
+		public FileReference BuildRecordFile { get; }
+
+		/// <summary>
+		/// The persisted build record
+		/// </summary>
+		public CsProjBuildRecord BuildRecord { get; }
+
+		/// <summary>
+		/// Status of the build record
+		/// </summary>
+		public CsProjBuildRecordStatus Status { get; set; } = CsProjBuildRecordStatus.Unknown;
+
+		/// <summary>
+		/// Construct a new entry
+		/// </summary>
+		/// <param name="projectFile">Project file reference</param>
+		/// <param name="buildRecordFile">Build record file reference</param>
+		/// <param name="buildRecord">Persisted build record</param>
+		public CsProjBuildRecordEntry(FileReference projectFile, FileReference buildRecordFile, CsProjBuildRecord buildRecord)
+		{
+			ProjectFile = projectFile;
+			BuildRecordFile = buildRecordFile;
+			BuildRecord = buildRecord;
+		}
+	}
 }
diff --git a/Engine/Source/Programs/Shared/EpicGames.MsBuild/CsProjBuilder.cs b/Engine/Source/Programs/Shared/EpicGames.MsBuild/CsProjBuilder.cs
--- a/Engine/Source/Programs/Shared/EpicGames.MsBuild/CsProjBuilder.cs
+++ b/Engine/Source/Programs/Shared/EpicGames.MsBuild/CsProjBuilder.cs
@@ -125,7 +125,7 @@ namespace EpicGames.MsBuild
 			}
 		}
 
-		static FileReference ConstructBuildRecordPath(FileReference ProjectPath, List<DirectoryReference> BaseDirectories)
+		static FileReference ConstructBuildRecordPath(CsProjBuildHook Hook, FileReference ProjectPath, List<DirectoryReference> BaseDirectories)
 		{
 			DirectoryReference BasePath = null;
 
@@ -143,7 +143,7 @@ namespace EpicGames.MsBuild
 				throw new Exception($"Unable to map csproj {ProjectPath} to Engine, game, or an additional script folder. Candidates were:{Environment.NewLine} {String.Join(Environment.NewLine, BaseDirectories)}");
 			}
 
-			DirectoryReference BuildRecordDirectory = DirectoryReference.Combine(BasePath!, "Intermediate", "ScriptModules");
+			DirectoryReference BuildRecordDirectory = Hook.GetBuildRecordDirectory(BasePath);
 			DirectoryReference.CreateDirectory(BuildRecordDirectory);
 
 			return FileReference.Combine(BuildRecordDirectory, ProjectPath.GetFileName()).ChangeExtension(".json");
@@ -160,7 +160,7 @@ namespace EpicGames.MsBuild
 		/// <param name="DefineConstants">Collection of constants to be defined while building projects</param>
 		/// <param name="OnBuildingProjects">Action invoked to notify caller regarding the number of projects being built</param>
 		/// <param name="Logger">Destination logger</param>
-		public static Dictionary<FileReference, (CsProjBuildRecord, FileReference)> Build(HashSet<FileReference> FoundProjects,
+		public static Dictionary<FileReference, CsProjBuildRecordEntry> Build(HashSet<FileReference> FoundProjects,
 			bool bForceCompile, out bool bBuildSuccess, CsProjBuildHook Hook, List<DirectoryReference> BaseDirectories, 
 			List<string> DefineConstants, Action<int> OnBuildingProjects, ILogger Logger)
 		{
@@ -183,7 +183,7 @@ namespace EpicGames.MsBuild
 		/// <param name="DefineConstants">Collection of constants to be defined while building projects</param>
 		/// <param name="OnBuildingProjects">Action invoked to notify caller regarding the number of projects being built</param>
 		/// <param name="Logger">Destination logger</param>
-		private static Dictionary<FileReference, (CsProjBuildRecord, FileReference)> BuildInternal(HashSet<FileReference> FoundProjects,
+		private static Dictionary<FileReference, CsProjBuildRecordEntry> BuildInternal(HashSet<FileReference> FoundProjects,
 			bool bForceCompile, out bool bBuildSuccess, CsProjBuildHook Hook, List<DirectoryReference> BaseDirectories, List<string> DefineConstants,
 			Action<int> OnBuildingProjects, ILogger Logger)
 		{
@@ -202,7 +202,7 @@ namespace EpicGames.MsBuild
 				GlobalProperties.Add("DefineConstants", String.Join(';', DefineConstants));
 			}
 
-			Dictionary<FileReference, (CsProjBuildRecord BuildRecord, FileReference BuildRecordPath)> BuildRecords = new Dictionary<FileReference, (CsProjBuildRecord, FileReference)>();
+			Dictionary<FileReference, CsProjBuildRecordEntry> BuildRecords = new();
 
 			using ProjectCollection ProjectCollection = new ProjectCollection(GlobalProperties);
 			Dictionary<string, Project> Projects = new Dictionary<string, Project>();
@@ -281,7 +281,7 @@ namespace EpicGames.MsBuild
 				string TargetPath = Path.GetRelativePath(Project.DirectoryPath, Project.GetPropertyValue("TargetPath"));
 
 				FileReference ProjectPath = FileReference.FromString(Project.FullPath);
-				FileReference BuildRecordPath = ConstructBuildRecordPath(ProjectPath, BaseDirectories);
+				FileReference BuildRecordPath = ConstructBuildRecordPath(Hook, ProjectPath, BaseDirectories);
 
 				CsProjBuildRecord BuildRecord = new CsProjBuildRecord()
 				{
@@ -318,7 +318,7 @@ namespace EpicGames.MsBuild
 
 				foreach (ProjectItem ReferencedProjectItem in Project.GetItems("ProjectReference"))
 				{
-					BuildRecord.ProjectReferences.Add(ReferencedProjectItem.EvaluatedInclude);
+					BuildRecord.ProjectReferencesAndTimes.Add(new CsProjBuildRecordRef { ProjectPath = ReferencedProjectItem.EvaluatedInclude });
 				}
 
 				foreach (ProjectItem CompileItem in Project.GetItems("Compile"))
@@ -418,7 +418,8 @@ namespace EpicGames.MsBuild
 					});
 				}
 
-				BuildRecords.Add(ProjectPath, (BuildRecord, BuildRecordPath));
+				CsProjBuildRecordEntry Entry = new CsProjBuildRecordEntry(ProjectPath, BuildRecordPath, BuildRecord);
+				BuildRecords.Add(Entry.ProjectFile, Entry);
 			}
 
 			// Potential optimization: Constructing the ProjectGraph here gives the full graph of dependencies - which is nice,
@@ -439,15 +440,13 @@ namespace EpicGames.MsBuild
 			}
 			else
 			{
-				Dictionary<FileReference, (CsProjBuildRecord, FileReference)> ValidBuildRecords = new Dictionary<FileReference, (CsProjBuildRecord, FileReference)>(BuildRecords.Count);
-				Dictionary<FileReference, (CsProjBuildRecord, FileReference)> InvalidBuildRecords = new Dictionary<FileReference, (CsProjBuildRecord, FileReference)>(BuildRecords.Count);
-
 				foreach (ProjectGraphNode Project in InputProjectGraph.ProjectNodesTopologicallySorted)
 				{
-					Hook.ValidateRecursively(ValidBuildRecords, InvalidBuildRecords, BuildRecords, FileReference.FromString(Project.ProjectInstance.FullPath));
+					Hook.ValidateRecursively(BuildRecords, FileReference.FromString(Project.ProjectInstance.FullPath));
 				}
 
 				// Select the projects that have been found to be out of date
+				Dictionary<FileReference, CsProjBuildRecordEntry> InvalidBuildRecords = new(BuildRecords.Where(x => x.Value.Status == CsProjBuildRecordStatus.Invalid));
 				HashSet<ProjectGraphNode> OutOfDateProjects = new HashSet<ProjectGraphNode>(InputProjectGraph.ProjectNodes.Where(x => InvalidBuildRecords.ContainsKey(FileReference.FromString(x.ProjectInstance.FullPath))));
 
 				if (OutOfDateProjects.Count > 0)
@@ -466,28 +465,46 @@ namespace EpicGames.MsBuild
 				bBuildSuccess = true;
 			}
 
-			// write all build records
+			// Update the target times
 			foreach (ProjectGraphNode ProjectNode in InputProjectGraph.ProjectNodes)
 			{
 				FileReference ProjectPath = FileReference.FromString(ProjectNode.ProjectInstance.FullPath);
+				CsProjBuildRecordEntry Entry = BuildRecords[ProjectPath];
+				FileReference FullPath = FileReference.Combine(ProjectPath.Directory, Entry.BuildRecord.TargetPath);
+				Entry.BuildRecord.TargetBuildTime = FileReference.GetLastWriteTime(FullPath);
+			}
 
-				(CsProjBuildRecord BuildRecord, FileReference BuildRecordPath) = BuildRecords[ProjectPath];
-
-				// update target build times into build records to ensure everything is up-to-date
-				FileReference FullPath = FileReference.Combine(ProjectPath.Directory, BuildRecord.TargetPath);
-				BuildRecord.TargetBuildTime = FileReference.GetLastWriteTime(FullPath);
+			// Update the project reference target times
+			foreach (ProjectGraphNode ProjectNode in InputProjectGraph.ProjectNodes)
+			{
+				FileReference ProjectPath = FileReference.FromString(ProjectNode.ProjectInstance.FullPath);
+				CsProjBuildRecordEntry Entry = BuildRecords[ProjectPath];
+				foreach (CsProjBuildRecordRef ReferencedProject in Entry.BuildRecord.ProjectReferencesAndTimes)
+				{
+					FileReference RefProjectPath = FileReference.FromString(Path.GetFullPath(ReferencedProject.ProjectPath, ProjectPath.Directory.FullName));
+					if (BuildRecords.TryGetValue(RefProjectPath, out CsProjBuildRecordEntry RefEntry))
+					{
+						ReferencedProject.TargetBuildTime = RefEntry.BuildRecord.TargetBuildTime;
+					}
+				}
+			}
 
-				if (FileReference.WriteAllTextIfDifferent(BuildRecordPath,
-					JsonSerializer.Serialize<CsProjBuildRecord>(BuildRecord, new JsonSerializerOptions { WriteIndented = true })))
+			// write all build records
+			foreach (ProjectGraphNode ProjectNode in InputProjectGraph.ProjectNodes)
+			{
+				FileReference ProjectPath = FileReference.FromString(ProjectNode.ProjectInstance.FullPath);
+				CsProjBuildRecordEntry Entry = BuildRecords[ProjectPath];
+				if (FileReference.WriteAllTextIfDifferent(Entry.BuildRecordFile,
+					JsonSerializer.Serialize<CsProjBuildRecord>(Entry.BuildRecord, new JsonSerializerOptions { WriteIndented = true })))
 				{
-					Logger.LogDebug("Wrote script module build record to {BuildRecordPath}", BuildRecordPath);
+					Logger.LogDebug("Wrote script module build record to {BuildRecordPath}", Entry.BuildRecordFile);
 				}
 			}
 
 			// todo: re-verify build records after a build to verify that everything is actually up to date
 
 			// even if only a subset was built, this function returns the full list of target assembly paths
-			Dictionary<FileReference, (CsProjBuildRecord BuildRecord, FileReference BuildRecordPath)> OutDict = new Dictionary<FileReference, (CsProjBuildRecord BuildRecord, FileReference BuildRecordPath)>();
+			Dictionary<FileReference, CsProjBuildRecordEntry> OutDict = new();
 			foreach (ProjectGraphNode EntryPointNode in InputProjectGraph.EntryPointNodes)
 			{
 				FileReference ProjectPath = FileReference.FromString(EntryPointNode.ProjectInstance.FullPath);
diff --git a/Engine/Source/Programs/UnrealBuildTool/System/TargetMakefile.cs b/Engine/Source/Programs/UnrealBuildTool/System/TargetMakefile.cs
--- a/Engine/Source/Programs/UnrealBuildTool/System/TargetMakefile.cs
+++ b/Engine/Source/Programs/UnrealBuildTool/System/TargetMakefile.cs
@@ -618,7 +618,7 @@ namespace UnrealBuildTool
 						BaseDirectories.Add(ProjectFile.Directory);
 					}
 					HashSet<FileReference> EnabledPlugins = new HashSet<FileReference>(Makefile.EnabledUbtPlugins);
-					if (!CompileScriptModule.AreScriptModulesUpToDate(EnabledPlugins, BaseDirectories, Logger))
+					if (!CompileScriptModule.AreScriptModulesUpToDate(Rules.RulesFileType.UbtPlugin, EnabledPlugins, BaseDirectories, Logger))
 					{
 						Logger.LogDebug("Enabled UBT plugins need to be recompiled");
 						ReasonNotLoaded = "Enabled UBT plugins need to be recompiled";
