Fixing missing XML documentation

#rnx

diff --git a/Engine/Source/Programs/Shared/EpicGames.Core/CsProjBuildRecord.cs b/Engine/Source/Programs/Shared/EpicGames.Core/CsProjBuildRecord.cs
--- a/Engine/Source/Programs/Shared/EpicGames.Core/CsProjBuildRecord.cs
+++ b/Engine/Source/Programs/Shared/EpicGames.Core/CsProjBuildRecord.cs
@@ -7,6 +7,10 @@ using System.Text;
 namespace EpicGames.Core
 {
 
+	/// <summary>
+	/// The list of references projects needs to be both the path to the project and the
+	/// target build time from that project.
+	/// </summary>
 	public class CsProjBuildRecordRef
 	{
 		/// <summary>
