Remove unacceptable string from Epicgames.Build
#rnx

diff --git a/Engine/Source/Programs/Shared/EpicGames.Build/UEBuildPlatformSDK.cs b/Engine/Source/Programs/Shared/EpicGames.Build/UEBuildPlatformSDK.cs
--- a/Engine/Source/Programs/Shared/EpicGames.Build/UEBuildPlatformSDK.cs
+++ b/Engine/Source/Programs/Shared/EpicGames.Build/UEBuildPlatformSDK.cs
@@ -971,7 +971,7 @@ namespace EpicGames.Core
 		}
 
 		/// <summary>
-		/// Gets the version number of the SDK setup script itself.  The version in the base should ALWAYS be the master revision from the last refactor.
+		/// Gets the version number of the SDK setup script itself.  The version in the base should ALWAYS be the primary revision from the last refactor.
 		/// If you need to force a rebuild for a given platform, modify the version file.
 		/// </summary>
 		/// <returns>Setup script version</returns>
