Horde: Fix blob store tests.

diff --git a/Engine/Source/Programs/Horde/Horde.Build.Tests/TestSetup.cs b/Engine/Source/Programs/Horde/Horde.Build.Tests/TestSetup.cs
--- a/Engine/Source/Programs/Horde/Horde.Build.Tests/TestSetup.cs
+++ b/Engine/Source/Programs/Horde/Horde.Build.Tests/TestSetup.cs
@@ -226,7 +226,7 @@ namespace Horde.Build.Tests
 
 			services.AddSingleton<ITreeStore<CommitService>>(sp =>
 			{
-				IBlobStore blobStore = new BlobStore(new TransientStorageBackend(), new BlobStoreOptions());
+				IBlobStore blobStore = new BlobStore(new TransientStorageBackend());
 				return new BundleStore(blobStore, new BundleOptions(), sp.GetRequiredService<IMemoryCache>()).ForType<CommitService>();
 			});
 
diff --git a/Engine/Source/Programs/Horde/Horde.Build/ServerSettings.cs b/Engine/Source/Programs/Horde/Horde.Build/ServerSettings.cs
--- a/Engine/Source/Programs/Horde/Horde.Build/ServerSettings.cs
+++ b/Engine/Source/Programs/Horde/Horde.Build/ServerSettings.cs
@@ -89,13 +89,8 @@ namespace Horde.Build
 	/// <summary>
 	/// Options for configuring a blob store
 	/// </summary>
-	public class BlobStoreOptions : StorageBackendOptions, IBlobStoreOptions
+	public class BlobStoreOptions : StorageBackendOptions
 	{
-		/// <inheritdoc/>
-		public string BlobPrefix { get; set; } = "blobs/";
-
-		/// <inheritdoc/>
-		public string RefPrefix { get; set; } = "refs/";
 	}
 
 	/// <summary>
diff --git a/Engine/Source/Programs/Horde/Horde.Build/Startup.cs b/Engine/Source/Programs/Horde/Horde.Build/Startup.cs
--- a/Engine/Source/Programs/Horde/Horde.Build/Startup.cs
+++ b/Engine/Source/Programs/Horde/Horde.Build/Startup.cs
@@ -279,7 +279,7 @@ namespace Horde.Build
 		static IBlobStore CreateBlobStore(IServiceProvider sp, BlobStoreOptions options)
 		{
 			IStorageBackend backend = CreateStorageBackend(sp, options);
-			return new BlobStore(backend, options);
+			return new BlobStore(backend);
 		}
 
 		static ITreeStore CreateTreeStore(IServiceProvider sp, TreeStoreOptions options)
diff --git a/Engine/Source/Programs/Horde/Horde.Build/Storage/Blobs/BlobStore.cs b/Engine/Source/Programs/Horde/Horde.Build/Storage/Blobs/BlobStore.cs
--- a/Engine/Source/Programs/Horde/Horde.Build/Storage/Blobs/BlobStore.cs
+++ b/Engine/Source/Programs/Horde/Horde.Build/Storage/Blobs/BlobStore.cs
@@ -10,39 +10,20 @@ using EpicGames.Horde.Storage;
 
 namespace Horde.Build.Storage
 {
-	/// <summary>
-	/// Options for configuring the default blob store implementation
-	/// </summary>
-	public interface IBlobStoreOptions
-	{
-		/// <summary>
-		/// Prefix for storing blobs
-		/// </summary>
-		string BlobPrefix { get; }
-
-		/// <summary>
-		/// Prefix for storing refs
-		/// </summary>
-		string RefPrefix { get; }
-	}
-
 	/// <summary>
 	/// Implementation of <see cref="IBlobStore"/> for AWS
 	/// </summary>
 	public class BlobStore : IBlobStore
 	{
 		readonly IStorageBackend _backend;
-		readonly IBlobStoreOptions _options;
 
 		/// <summary>
 		/// Constructor
 		/// </summary>
 		/// <param name="backend">Backend to use for storing data</param>
-		/// <param name="options">Options for the blob store</param>
-		public BlobStore(IStorageBackend backend, IBlobStoreOptions options)
+		public BlobStore(IStorageBackend backend)
 		{
 			_backend = backend;
-			_options = options;
 		}
 
 		async Task<IBlob?> TryReadAsync(string path, CancellationToken cancellationToken = default)
@@ -63,7 +44,7 @@ namespace Horde.Build.Storage
 
 		#region Blobs
 
-		string GetBlobPath(BlobId id) => $"{_options.BlobPrefix}{id}";
+		static string GetBlobPath(BlobId id) => $"{id}.blob";
 
 		/// <inheritdoc/>
 		public Task<IBlob?> TryReadBlobAsync(BlobId id, CancellationToken cancellationToken = default) => TryReadAsync(GetBlobPath(id), cancellationToken);
@@ -80,7 +61,7 @@ namespace Horde.Build.Storage
 
 		#region Refs
 
-		string GetRefPath(RefName name) => $"{_options.RefPrefix}{name}";
+		static string GetRefPath(RefName name) => $"{name}.ref";
 
 		/// <inheritdoc/>
 		public async Task DeleteRefAsync(RefName name, CancellationToken cancellationToken = default)
diff --git a/Engine/Source/Programs/Shared/EpicGames.Horde.Tests/BundleStoreTests.cs b/Engine/Source/Programs/Shared/EpicGames.Horde.Tests/BundleStoreTests.cs
--- a/Engine/Source/Programs/Shared/EpicGames.Horde.Tests/BundleStoreTests.cs
+++ b/Engine/Source/Programs/Shared/EpicGames.Horde.Tests/BundleStoreTests.cs
@@ -55,7 +55,7 @@ namespace EpicGames.Horde.Tests
 				ITreeBlob node4 = await collection.WriteBlobAsync(new ReadOnlySequence<byte>(new byte[] { 4 }), Array.Empty<ITreeBlob>(), CancellationToken.None);
 				ITreeBlob node5 = await collection.WriteBlobAsync(new ReadOnlySequence<byte>(new byte[] { 5 }), new[] { node4, node3 }, CancellationToken.None);
 
-				await collection.WriteTreeAsync(new RefId("test"), node5);
+				await collection.WriteTreeAsync(new RefName("test"), node5);
 
 				ITreeBlob root = node5;
 				await CheckTree(root);
@@ -64,7 +64,7 @@ namespace EpicGames.Horde.Tests
 			// Check we can read it back in
 			using (BundleStore collection = new BundleStore(blobStore, options, cache))
 			{
-				ITreeBlob root = await collection.ReadTreeAsync(new RefId("test"));
+				ITreeBlob root = await collection.ReadTreeAsync(new RefName("test"));
 				await CheckTree(root);
 			}
 		}
@@ -114,7 +114,7 @@ namespace EpicGames.Horde.Tests
 				DirectoryNode root = new DirectoryNode();
 				DirectoryNode hello = root.AddDirectory("hello");
 				DirectoryNode world = hello.AddDirectory("world");
-				await store.WriteTreeAsync(new RefId("test"), root, true, CancellationToken.None);
+				await store.WriteTreeAsync(new RefName("test"), root, true, CancellationToken.None);
 
 				await CheckDirectoryTreeAsync(root);
 			}
@@ -122,7 +122,7 @@ namespace EpicGames.Horde.Tests
 			// Check we can read it back in
 			using (BundleStore store = new BundleStore(blobStore, new BundleOptions(), cache))
 			{
-				DirectoryNode root = await store.ReadTreeAsync<DirectoryNode>(new RefId("test"));
+				DirectoryNode root = await store.ReadTreeAsync<DirectoryNode>(new RefName("test"));
 				await CheckDirectoryTreeAsync(root);
 			}
 		}
@@ -155,7 +155,7 @@ namespace EpicGames.Horde.Tests
 				FileEntry world = hello.AddFile("world", FileEntryFlags.None);
 				await world.AppendAsync(Encoding.UTF8.GetBytes("world"), new ChunkingOptions(), CancellationToken.None);
 				
-				await store.WriteTreeAsync(new RefId("test"), root, true, CancellationToken.None);
+				await store.WriteTreeAsync(new RefName("test"), root, true, CancellationToken.None);
 
 				await CheckFileTreeAsync(root);
 			}
@@ -163,7 +163,7 @@ namespace EpicGames.Horde.Tests
 			// Check we can read it back in
 			using (BundleStore store = new BundleStore(blobStore, new BundleOptions(), cache))
 			{
-				DirectoryNode root = await store.ReadTreeAsync<DirectoryNode>(new RefId("test"));
+				DirectoryNode root = await store.ReadTreeAsync<DirectoryNode>(new RefName("test"));
 				await CheckFileTreeAsync(root);
 			}
 		}
@@ -250,7 +250,7 @@ namespace EpicGames.Horde.Tests
 				FileEntry file = root.AddFile("test", FileEntryFlags.None);
 				await file.AppendAsync(data, options, CancellationToken.None);
 
-				await store.WriteTreeAsync(new RefId("test"), root, true, CancellationToken.None);
+				await store.WriteTreeAsync(new RefName("test"), root, true, CancellationToken.None);
 
 				await CheckLargeFileTreeAsync(root, data);
 			}
@@ -258,7 +258,7 @@ namespace EpicGames.Horde.Tests
 			// Check we can read it back in
 			using (BundleStore store = new BundleStore(blobStore, new BundleOptions(), cache))
 			{
-				DirectoryNode root = await store.ReadTreeAsync<DirectoryNode>(new RefId("test"));
+				DirectoryNode root = await store.ReadTreeAsync<DirectoryNode>(new RefName("test"));
 				await CheckLargeFileTreeAsync(root, data);
 
 				TreeNodeRef<FileNode> file = root.GetFileEntry("test");
diff --git a/Engine/Source/Programs/Shared/EpicGames.Horde.Tests/BundleTests.cs b/Engine/Source/Programs/Shared/EpicGames.Horde.Tests/BundleTests.cs
--- a/Engine/Source/Programs/Shared/EpicGames.Horde.Tests/BundleTests.cs
+++ b/Engine/Source/Programs/Shared/EpicGames.Horde.Tests/BundleTests.cs
@@ -179,21 +179,21 @@ namespace EpicGames.Horde.Tests
 			DirectoryNode node = root.AddDirectory("hello");
 			DirectoryNode node2 = node.AddDirectory("world");
 
-			RefId refId = new RefId("testref");
-			await treeStore.WriteTreeAsync(refId, root, true);
+			RefName refName = new RefName("testref");
+			await treeStore.WriteTreeAsync(refName, root, true);
 
 			// Should be stored inline
 			Assert.AreEqual(1, blobStore.Refs.Count);
 			Assert.AreEqual(0, blobStore.Blobs.Count);
 
 			// Check the ref
-			IBlob blob = await blobStore.ReadRefAsync(refId);
+			IBlob blob = await blobStore.ReadRefAsync(refName);
 			Bundle bundle = new Bundle(new MemoryReader(blob.Data));
 			Assert.AreEqual(0, bundle.Header.Imports.Count);
 			Assert.AreEqual(3, bundle.Header.Exports.Count);
 
 			// Create a new bundle and read it back in again
-			DirectoryNode newRoot = await treeStore.ReadTreeAsync<DirectoryNode>(refId);
+			DirectoryNode newRoot = await treeStore.ReadTreeAsync<DirectoryNode>(refName);
 
 			Assert.AreEqual(0, newRoot.Files.Count);
 			Assert.AreEqual(1, newRoot.Directories.Count);
@@ -227,8 +227,8 @@ namespace EpicGames.Horde.Tests
 			Assert.AreEqual(0, blobStore.Refs.Count);
 			Assert.AreEqual(0, blobStore.Blobs.Count);
 
-			RefId refId = new RefId("ref");
-			await treeStore.WriteTreeAsync(refId, root, true);
+			RefName refName = new RefName("ref");
+			await treeStore.WriteTreeAsync(refName, root, true);
 
 			Assert.AreEqual(1, blobStore.Refs.Count);
 			Assert.AreEqual(1, blobStore.Blobs.Count);
@@ -243,7 +243,7 @@ namespace EpicGames.Horde.Tests
 
 			InMemoryBlobStore blobStore = new InMemoryBlobStore();
 
-			RefId refId = new RefId("ref");
+			RefName refName = new RefName("ref");
 
 			using (ITreeStore oldBundle = new BundleStore(blobStore, options, _cache))
 			{
@@ -254,7 +254,7 @@ namespace EpicGames.Horde.Tests
 				DirectoryNode node3 = node2.AddDirectory("node3");
 				DirectoryNode node4 = node3.AddDirectory("node4");
 
-				await oldBundle.WriteTreeAsync(refId, root, true);
+				await oldBundle.WriteTreeAsync(refName, root, true);
 
 				Assert.AreEqual(1, blobStore.Refs.Count);
 				Assert.AreEqual(4, blobStore.Blobs.Count);
@@ -262,7 +262,7 @@ namespace EpicGames.Horde.Tests
 
 			using (ITreeStore newBundle = new BundleStore(blobStore, options, _cache))
 			{
-				DirectoryNode root = await newBundle.ReadTreeAsync<DirectoryNode>(refId);
+				DirectoryNode root = await newBundle.ReadTreeAsync<DirectoryNode>(refName);
 
 				DirectoryNode? newNode1 = await root.FindDirectoryAsync("node1", CancellationToken.None);
 				Assert.IsNotNull(newNode1);
@@ -447,7 +447,7 @@ namespace EpicGames.Horde.Tests
 
 			using BundleStore treeStore = new BundleStore(_blobStore, options, _cache);
 
-			RefId refId = new RefId("ref");
+			RefName refName = new RefName("ref");
 			DirectoryNode root = new DirectoryNode();
 
 			long totalLength = 0;
@@ -470,7 +470,7 @@ namespace EpicGames.Horde.Tests
 					}
 
 					int oldWorkingSetSize = GetWorkingSetSize(root);
-					await treeStore.WriteTreeAsync(refId, root, false);
+					await treeStore.WriteTreeAsync(refName, root, false);
 					int newWorkingSetSize = GetWorkingSetSize(root);
 					Assert.IsTrue(newWorkingSetSize <= oldWorkingSetSize);
 					Assert.IsTrue(newWorkingSetSize <= 20);
@@ -480,7 +480,7 @@ namespace EpicGames.Horde.Tests
 			Assert.IsTrue(_blobStore.Blobs.Count > 0);
 			Assert.IsTrue(_blobStore.Refs.Count == 0);
 
-			await treeStore.WriteTreeAsync(refId, root, true, CancellationToken.None);
+			await treeStore.WriteTreeAsync(refName, root, true, CancellationToken.None);
 
 			Assert.AreEqual(totalLength, root.Length);
 
