Horde: Fix exception initializing content replication.

diff --git a/Engine/Source/Programs/Shared/EpicGames.Horde/Storage/Bundles/Bundle.cs b/Engine/Source/Programs/Shared/EpicGames.Horde/Storage/Bundles/Bundle.cs
--- a/Engine/Source/Programs/Shared/EpicGames.Horde/Storage/Bundles/Bundle.cs
+++ b/Engine/Source/Programs/Shared/EpicGames.Horde/Storage/Bundles/Bundle.cs
@@ -360,7 +360,7 @@ namespace EpicGames.Horde.Storage.Bundles
 		/// <returns>The bundle that was read</returns>
 		public static async Task<Bundle?> TryReadBundleAsync(this IBlobStore store, RefName name, CancellationToken cancellationToken)
 		{
-			IBlob? blob = await store.ReadRefAsync(name, cancellationToken);
+			IBlob? blob = await store.TryReadRefAsync(name, cancellationToken);
 			if (blob == null)
 			{
 				return null;
