Horde: Include name of the parent ref in the path for each bundle blob.

diff --git a/Engine/Source/Programs/Horde/Horde.Agent/Commands/Bundles/CreateCommand.cs b/Engine/Source/Programs/Horde/Horde.Agent/Commands/Bundles/CreateCommand.cs
--- a/Engine/Source/Programs/Horde/Horde.Agent/Commands/Bundles/CreateCommand.cs
+++ b/Engine/Source/Programs/Horde/Horde.Agent/Commands/Bundles/CreateCommand.cs
@@ -59,10 +59,11 @@ namespace Horde.Agent.Commands.Bundles
 				return BlobUtils.Deserialize(bytes);
 			}
 
-			public async Task<BlobId> WriteBlobAsync(ReadOnlySequence<byte> data, IReadOnlyList<BlobId> references, CancellationToken cancellationToken = default)
+			public async Task<BlobId> WriteBlobAsync(RefName refName, ReadOnlySequence<byte> data, IReadOnlyList<BlobId> references, CancellationToken cancellationToken = default)
 			{
-				BlobId id = new BlobId(Guid.NewGuid().ToString());
+				BlobId id = new BlobId($"{refName}/{Guid.NewGuid()}");
 				FileReference file = GetBlobFile(id);
+				DirectoryReference.CreateDirectory(file.Directory);
 				_logger.LogInformation("Writing {File}", file);
 				await WriteAsync(file, BlobUtils.Serialize(data, references), cancellationToken);
 				return id;
@@ -98,6 +99,7 @@ namespace Horde.Agent.Commands.Bundles
 			public async Task WriteRefAsync(RefName name, ReadOnlySequence<byte> data, IReadOnlyList<BlobId> references, CancellationToken cancellationToken = default)
 			{
 				FileReference file = GetRefFile(name);
+				DirectoryReference.CreateDirectory(file.Directory);
 				_logger.LogInformation("Writing {File}", file);
 				await WriteAsync(file, BlobUtils.Serialize(data, references), cancellationToken);
 			}
diff --git a/Engine/Source/Programs/Horde/Horde.Build/Storage/Blobs/BlobStore.cs b/Engine/Source/Programs/Horde/Horde.Build/Storage/Blobs/BlobStore.cs
--- a/Engine/Source/Programs/Horde/Horde.Build/Storage/Blobs/BlobStore.cs
+++ b/Engine/Source/Programs/Horde/Horde.Build/Storage/Blobs/BlobStore.cs
@@ -50,9 +50,9 @@ namespace Horde.Build.Storage
 		public Task<IBlob?> TryReadBlobAsync(BlobId id, CancellationToken cancellationToken = default) => TryReadAsync(GetBlobPath(id), cancellationToken);
 
 		/// <inheritdoc/>
-		public async Task<BlobId> WriteBlobAsync(ReadOnlySequence<byte> data, IReadOnlyList<BlobId> references, CancellationToken cancellationToken = default)
+		public async Task<BlobId> WriteBlobAsync(RefName refName, ReadOnlySequence<byte> data, IReadOnlyList<BlobId> references, CancellationToken cancellationToken = default)
 		{
-			BlobId id = new BlobId(Guid.NewGuid().ToString());
+			BlobId id = new BlobId($"{refName}/{Guid.NewGuid()}");
 			await WriteAsync(GetBlobPath(id), data, references, cancellationToken);
 			return id;
 		}
diff --git a/Engine/Source/Programs/Shared/EpicGames.Horde/Storage/Backends/InMemoryBlobStore.cs b/Engine/Source/Programs/Shared/EpicGames.Horde/Storage/Backends/InMemoryBlobStore.cs
--- a/Engine/Source/Programs/Shared/EpicGames.Horde/Storage/Backends/InMemoryBlobStore.cs
+++ b/Engine/Source/Programs/Shared/EpicGames.Horde/Storage/Backends/InMemoryBlobStore.cs
@@ -40,7 +40,7 @@ namespace EpicGames.Horde.Storage.Backends
 		}
 
 		/// <inheritdoc/>
-		public Task<BlobId> WriteBlobAsync(ReadOnlySequence<byte> data, IReadOnlyList<BlobId> references, CancellationToken cancellationToken = default)
+		public Task<BlobId> WriteBlobAsync(RefName refName, ReadOnlySequence<byte> data, IReadOnlyList<BlobId> references, CancellationToken cancellationToken = default)
 		{
 			BlobId blobId = new BlobId(Guid.NewGuid().ToString());
 			_blobs[blobId] = BlobUtils.FromMemory(data.ToArray(), references);
diff --git a/Engine/Source/Programs/Shared/EpicGames.Horde/Storage/Bundles/Bundle.cs b/Engine/Source/Programs/Shared/EpicGames.Horde/Storage/Bundles/Bundle.cs
--- a/Engine/Source/Programs/Shared/EpicGames.Horde/Storage/Bundles/Bundle.cs
+++ b/Engine/Source/Programs/Shared/EpicGames.Horde/Storage/Bundles/Bundle.cs
@@ -375,7 +375,7 @@ namespace EpicGames.Horde.Storage.Bundles
 		/// <param name="name">Identifier for the blob</param>
 		/// <param name="bundle">Bundle to write</param>
 		/// <param name="cancellationToken">Cancellation token for the operation</param>
-		public static Task WriteBundleAsync(this IBlobStore store, RefName name, Bundle bundle, CancellationToken cancellationToken)
+		public static Task WriteBundleRefAsync(this IBlobStore store, RefName name, Bundle bundle, CancellationToken cancellationToken)
 		{
 			ReadOnlySequence<byte> data = bundle.AsSequence();
 			List<BlobId> imports = bundle.Header.Imports.Select(x => x.BlobId).ToList();
@@ -386,14 +386,15 @@ namespace EpicGames.Horde.Storage.Bundles
 		/// Writes a bundle to a blob store
 		/// </summary>
 		/// <param name="store">Blob store to write to</param>
+		/// <param name="name">Name of the root ref</param>
 		/// <param name="bundle">Bundle to write</param>
 		/// <param name="cancellationToken">Cancellation token for the operation</param>
 		/// <returns>Identifier for the blob</returns>
-		public static Task<BlobId> WriteBundleAsync(this IBlobStore store, Bundle bundle, CancellationToken cancellationToken)
+		public static Task<BlobId> WriteBundleBlobAsync(this IBlobStore store, RefName name, Bundle bundle, CancellationToken cancellationToken)
 		{
 			ReadOnlySequence<byte> data = bundle.AsSequence();
 			List<BlobId> imports = bundle.Header.Imports.Select(x => x.BlobId).ToList();
-			return store.WriteBlobAsync(data, imports, cancellationToken);
+			return store.WriteBlobAsync(name, data, imports, cancellationToken);
 		}
 
 		static Bundle ReadBundle(IBlob blob)
diff --git a/Engine/Source/Programs/Shared/EpicGames.Horde/Storage/Bundles/BundleStore.cs b/Engine/Source/Programs/Shared/EpicGames.Horde/Storage/Bundles/BundleStore.cs
--- a/Engine/Source/Programs/Shared/EpicGames.Horde/Storage/Bundles/BundleStore.cs
+++ b/Engine/Source/Programs/Shared/EpicGames.Horde/Storage/Bundles/BundleStore.cs
@@ -294,7 +294,7 @@ namespace EpicGames.Horde.Storage.Bundles
 				if (bundleNodes.Count > 0 && bundleSize + nodeSize > _options.MaxBlobSize)
 				{
 					Bundle bundle = CreateBundle(bundleNodes);
-					BlobId blobId = await _blobStore.WriteBundleAsync(bundle, cancellationToken);
+					BlobId blobId = await _blobStore.WriteBundleBlobAsync(name, bundle, cancellationToken);
 
 					BlobInfo bundleInfo = new BlobInfo(blobId, bundleNodes.ToArray());
 					MountExportedBundle(bundleInfo, bundle);
@@ -313,7 +313,7 @@ namespace EpicGames.Horde.Storage.Bundles
 			{
 				// Write the main ref
 				Bundle rootBundle = CreateBundle(bundleNodes);
-				await _blobStore.WriteBundleAsync(name, rootBundle, cancellationToken);
+				await _blobStore.WriteBundleRefAsync(name, rootBundle, cancellationToken);
 
 				// Copy the stats over
 				Stats = _nextStats;
diff --git a/Engine/Source/Programs/Shared/EpicGames.Horde/Storage/IBlobStore.cs b/Engine/Source/Programs/Shared/EpicGames.Horde/Storage/IBlobStore.cs
--- a/Engine/Source/Programs/Shared/EpicGames.Horde/Storage/IBlobStore.cs
+++ b/Engine/Source/Programs/Shared/EpicGames.Horde/Storage/IBlobStore.cs
@@ -66,11 +66,12 @@ namespace EpicGames.Horde.Storage
 		/// <summary>
 		/// Writes a new object to the store
 		/// </summary>
+		/// <param name="refName">Name of a ref that this blob is being written for. While the returned BlobId is guaranteed to be unique, this name can be used as a prefix to aid debugging.</param>
 		/// <param name="data">Payload for the object</param>
 		/// <param name="references">Object references</param>
 		/// <param name="cancellationToken">Cancellation token for the operation</param>
 		/// <returns>Unique identifier for the bundle</returns>
-		Task<BlobId> WriteBlobAsync(ReadOnlySequence<byte> data, IReadOnlyList<BlobId> references, CancellationToken cancellationToken = default);
+		Task<BlobId> WriteBlobAsync(RefName refName, ReadOnlySequence<byte> data, IReadOnlyList<BlobId> references, CancellationToken cancellationToken = default);
 
 		#endregion
 
@@ -136,7 +137,7 @@ namespace EpicGames.Horde.Storage
 			public Task<IBlob?> TryReadBlobAsync(BlobId id, CancellationToken cancellationToken = default) => _inner.TryReadBlobAsync(id, cancellationToken);
 
 			/// <inheritdoc/>
-			public Task<BlobId> WriteBlobAsync(ReadOnlySequence<byte> data, IReadOnlyList<BlobId> references, CancellationToken cancellationToken = default) => _inner.WriteBlobAsync(data, references, cancellationToken);
+			public Task<BlobId> WriteBlobAsync(RefName refName, ReadOnlySequence<byte> data, IReadOnlyList<BlobId> references, CancellationToken cancellationToken = default) => _inner.WriteBlobAsync(refName, data, references, cancellationToken);
 
 			#endregion
 
