PR #9391: Fix deadlock due to NativePerforceConnection.Response not being disposed (Contributed by [REDACTED_PERSON])

diff --git a/Engine/Source/Programs/Shared/EpicGames.Perforce/PerforceConnection.cs b/Engine/Source/Programs/Shared/EpicGames.Perforce/PerforceConnection.cs
--- a/Engine/Source/Programs/Shared/EpicGames.Perforce/PerforceConnection.cs
+++ b/Engine/Source/Programs/Shared/EpicGames.Perforce/PerforceConnection.cs
@@ -2320,6 +2320,7 @@ namespace EpicGames.Perforce
 		/// <returns></returns>
 		public static async Task<PerforceResponse<LoginRecord>> TryLoginAsync(this IPerforceConnection connection, string password, CancellationToken cancellationToken = default)
 		{
+			List<PerforceResponse> parsedResponses;
 			await using (IPerforceOutput response = await connection.LoginCommandAsync(password, cancellationToken))
 			{
 				for (; ; )
@@ -2332,17 +2333,21 @@ namespace EpicGames.Perforce
 
 				DiscardPasswordPrompt(response);
 
-				List<PerforceResponse> parsedResponses = await response.ReadResponsesAsync(typeof(LoginRecord), cancellationToken);
+				parsedResponses = await response.ReadResponsesAsync(typeof(LoginRecord), cancellationToken);
 
-				PerforceResponse firstResponse = parsedResponses.First();
-				if (firstResponse.Info != null)
-				{
-					// Older versions of P4.EXE do not return a login record for succesful login, instread just returning a string. Call p4 login -s to get the login state instead.
-					return await TryGetLoginStateAsync(connection, cancellationToken);
-				}
+				// end lifetime of `response` here
+				// this prevents a deadlock in case `connection` is a `NativePerforceConnection` and `response` is a `NativePerforceConnection.Response`
+				// not DisposeAsync()ing here will cause a deadlock when calling `TryGetLoginStateAsync()` below
+			}
 
-				return new PerforceResponse<LoginRecord>(firstResponse);
+			PerforceResponse firstResponse = parsedResponses.First();
+			if (firstResponse.Info != null)
+			{
+				// Older versions of P4.EXE do not return a login record for succesful login, instread just returning a string. Call p4 login -s to get the login state instead.
+				return await TryGetLoginStateAsync(connection, cancellationToken);
 			}
+
+			return new PerforceResponse<LoginRecord>(firstResponse);
 		}
 
 		static void DiscardPasswordPrompt(IPerforceOutput output)
