Temporary work around to having TEnumAsBytes passed into property Setter function when it expects the actual enum value.

#rnx

diff --git a/Engine/Source/Programs/Shared/EpicGames.UHT/Exporters/CodeGen/UhtHeaderCodeGeneratorCppFile.cs b/Engine/Source/Programs/Shared/EpicGames.UHT/Exporters/CodeGen/UhtHeaderCodeGeneratorCppFile.cs
--- a/Engine/Source/Programs/Shared/EpicGames.UHT/Exporters/CodeGen/UhtHeaderCodeGeneratorCppFile.cs
+++ b/Engine/Source/Programs/Shared/EpicGames.UHT/Exporters/CodeGen/UhtHeaderCodeGeneratorCppFile.cs
@@ -1090,6 +1090,18 @@ namespace EpicGames.UHT.Exporters.CodeGen
 								.AppendPropertyText(property, UhtPropertyTextType.GetterSetterArg)
 								.Append("*)InValue;\r\n");
 						}
+						else if (property is UhtByteProperty byteProperty && byteProperty.Enum != null)
+						{
+							// If someone passed in a TByteAsEnum instead of the actual enum value, the cast in the else clause would cause an issue.
+							// Since this is known to be a TByteAsEnum, we just fetch the first byte.  *HOWEVER* on MSB machines where 
+							// the actual enum value is passed in this will fail and return zero if the native size of the enum > 1 byte.
+							builder
+								.Append("\t\t")
+								.AppendPropertyText(property, UhtPropertyTextType.GetterSetterArg)
+								.Append(" Value = (")
+								.AppendPropertyText(property, UhtPropertyTextType.GetterSetterArg)
+								.Append(")*(uint8*)InValue;\r\n");
+						}
 						else
 						{
 							builder
diff --git a/Engine/Source/Programs/UnrealHeaderTool/Private/CodeGenerator.cpp b/Engine/Source/Programs/UnrealHeaderTool/Private/CodeGenerator.cpp
--- a/Engine/Source/Programs/UnrealHeaderTool/Private/CodeGenerator.cpp
+++ b/Engine/Source/Programs/UnrealHeaderTool/Private/CodeGenerator.cpp
@@ -5097,6 +5097,13 @@ void FNativeClassHeaderGenerator::ExportNativeFunctions(FOutputDevice& OutGenera
 				{
 					RuntimeStringBuilders.AccessorWrapperImplementations.Logf(TEXT("\t\t%s* Value = (%s*)InValue;\r\n"), *PropertyType, *PropertyType);
 				}
+				else if (Prop->IsByteEnumOrByteEnumStaticArray() && Prop->GetPropertyBase().ArrayType == EArrayType::None)
+				{
+					// If someone passed in a TByteAsEnum instead of the actual enum value, the cast in the else clause would cause an issue.
+					// Since this is known to be a TByteAsEnum, we just fetch the first byte.  *HOWEVER* on MSB machines where 
+					// the actual enum value is passed in this will fail and return zero if the native size of the enum > 1 byte.
+					RuntimeStringBuilders.AccessorWrapperImplementations.Logf(TEXT("\t\t%s Value = (%s)*(uint8*)InValue;\r\n"), *PropertyType, *PropertyType);
+				}
 				else
 				{
 					RuntimeStringBuilders.AccessorWrapperImplementations.Logf(TEXT("\t\t%s& Value = *(%s*)InValue;\r\n"), *PropertyType, *PropertyType);
