Horde: Make BundleStore writes thread safe, improving replication performance.

diff --git a/Engine/Source/Programs/Horde/Horde.Build/Perforce/ReplicationService.cs b/Engine/Source/Programs/Horde/Horde.Build/Perforce/ReplicationService.cs
--- a/Engine/Source/Programs/Horde/Horde.Build/Perforce/ReplicationService.cs
+++ b/Engine/Source/Programs/Horde/Horde.Build/Perforce/ReplicationService.cs
@@ -690,6 +690,7 @@ namespace Horde.Build.Perforce
 					Stopwatch processTimer = new Stopwatch();
 					Stopwatch trimTimer = new Stopwatch();
 					Stopwatch gcTimer = new Stopwatch();
+					Task trimTask = Task.CompletedTask;
 
 					Dictionary<int, FileEntry> handles = new Dictionary<int, FileEntry>();
 					await foreach (PerforceResponse response in perforce.StreamCommandAsync("sync", Array.Empty<string>(), new string[] { syncPath }, null, typeof(SyncRecord), true, default))
@@ -734,12 +735,20 @@ namespace Horde.Build.Perforce
 							if (dataSize > trimDataSize)
 							{
 								_logger.LogInformation("Trimming working set after receiving {Size:n1}mb...", dataSize / (1024 * 1024));
-								trimTimer.Start();
-								await _treeStore.WriteTreeAsync(refName, syncNode, false, cancellationToken);
-								trimTimer.Stop();
-								gcTimer.Start();
-								GC.Collect();
-								gcTimer.Stop();
+								if (!trimTask.IsCompleted)
+								{
+									_logger.LogInformation("Waiting for previous trim to complete...");
+								}
+								await trimTask;
+								trimTask = Task.Run(async () =>
+								{
+									trimTimer.Start();
+									await _treeStore.WriteTreeAsync(refName, syncNode, false, cancellationToken);
+									trimTimer.Stop();
+									gcTimer.Start();
+									GC.Collect();
+									gcTimer.Stop();
+								}, cancellationToken);
 								trimDataSize = dataSize + TrimInterval;
 								_logger.LogInformation("Trimming complete. Next trim at {Size:n1}mb.", trimDataSize / (1024 * 1024));
 							}
@@ -747,6 +756,7 @@ namespace Horde.Build.Perforce
 
 						processTimer.Stop();
 					}
+					await trimTask;
 					_logger.LogInformation("Completed batch in {TimeSeconds:n1}s ({ProcessTimeSeconds:n1}s processing, {TrimTimeSeconds:n1}s trimming, {GcTimeSeconds:n1}s gc)", syncTimer.Elapsed.TotalSeconds, processTimer.Elapsed.TotalSeconds, trimTimer.Elapsed.TotalSeconds, gcTimer.Elapsed.TotalSeconds);
 
 					// Update the root sync node
diff --git a/Engine/Source/Programs/Shared/EpicGames.Horde/Storage/Bundles/BundleStore.cs b/Engine/Source/Programs/Shared/EpicGames.Horde/Storage/Bundles/BundleStore.cs
--- a/Engine/Source/Programs/Shared/EpicGames.Horde/Storage/Bundles/BundleStore.cs
+++ b/Engine/Source/Programs/Shared/EpicGames.Horde/Storage/Bundles/BundleStore.cs
@@ -95,13 +95,20 @@ namespace EpicGames.Horde.Storage.Bundles
 		{
 			public readonly BundleStore Owner;
 			public IoHash Hash { get; }
-			public NodeState State { get; set; }
+			public NodeState State => _state;
+
+			NodeState _state;
 
 			public NodeInfo(BundleStore owner, IoHash hash, NodeState state)
 			{
 				Owner = owner;
 				Hash = hash;
-				State = state;
+				_state = state;
+			}
+
+			public bool TrySetState(NodeState prevState, NodeState nextState)
+			{
+				return Interlocked.CompareExchange(ref _state, nextState, prevState) == prevState;
 			}
 
 			/// <inheritdoc/>
@@ -125,18 +132,12 @@ namespace EpicGames.Horde.Storage.Bundles
 			}
 		}
 
-		// Type of node state, ordered from low to high in terms of preference when resolving duplicates.
-		enum NodeStatePriority
-		{
-			Imported = 0,
-			InMemory = 1,
-			Exported = 2,
-		}
-
 		abstract class NodeState
 		{
 		}
 
+		// State for data imported from a bundle, but whose location within it is not yet known because the bundle's header has not been read yet.
+		// Can be transitioned to ExportedNodeState.
 		class ImportedNodeState : NodeState
 		{
 			public readonly BlobInfo BundleInfo;
@@ -149,6 +150,8 @@ namespace EpicGames.Horde.Storage.Bundles
 			}
 		}
 
+		// State for data persisted to a bundle. Decoded data is cached outside this structure in the store's MemoryCache.
+		// Cannot be transitioned to any other state.
 		class ExportedNodeState : ImportedNodeState
 		{
 			public readonly PacketInfo PacketInfo;
@@ -166,6 +169,8 @@ namespace EpicGames.Horde.Storage.Bundles
 			}
 		}
 
+		// Data for a node which exists in memory, but which does NOT exist in storage yet. Nodes which 
+		// are read in from storage are represented by ExportedNodeState and a MemoryCache. Can transition to an ExportedNodeState only.
 		class InMemoryNodeState : NodeState
 		{
 			public readonly ReadOnlySequence<byte> Data;
@@ -213,6 +218,8 @@ namespace EpicGames.Horde.Storage.Bundles
 		// so that other threads can also await it.
 		readonly Dictionary<string, Task<Bundle>> _readTasks = new Dictionary<string, Task<Bundle>>();
 
+		Task _writeTask = Task.CompletedTask;
+
 		readonly BundleOptions _options;
 
 		/// <summary>
@@ -278,6 +285,31 @@ namespace EpicGames.Horde.Storage.Bundles
 
 		/// <inheritdoc/>
 		public async Task WriteTreeAsync(RefName name, ITreeBlob root, bool flush = true, CancellationToken cancellationToken = default)
+		{
+			await SequenceWrite(() => WriteTreeInternalAsync(name, root, flush, cancellationToken), cancellationToken);
+		}
+
+		async Task SequenceWrite(Func<Task> writeFunc, CancellationToken cancellationToken)
+		{
+			Task task;
+			lock (_lockObject)
+			{
+				Task prevTask = _writeTask;
+
+				// Wait for the previous write to complete first, ignoring any exceptions
+				Func<Task> wrappedWriteFunc = async () =>
+				{
+					await prevTask.ContinueWith(x => { }, cancellationToken, TaskContinuationOptions.None, TaskScheduler.Default);
+					await writeFunc();
+				};
+				task = Task.Run(wrappedWriteFunc, cancellationToken);
+
+				_writeTask = task;
+			}
+			await task;
+		}
+
+		async Task WriteTreeInternalAsync(RefName name, ITreeBlob root, bool flush = true, CancellationToken cancellationToken = default)
 		{
 			// Find all the referenced in-memory nodes from the root
 			List<NodeInfo> nodes = new List<NodeInfo>();
@@ -409,6 +441,40 @@ namespace EpicGames.Horde.Storage.Bundles
 			{
 				node = _hashToNode.GetOrAdd(hash, new NodeInfo(this, hash, state));
 			}
+
+			if (state is ExportedNodeState)
+			{
+				// Transition from any non-exported state
+				for (; ; )
+				{
+					NodeState prevState = node.State;
+					if (prevState is ExportedNodeState)
+					{
+						break;
+					}
+					else if (node.TrySetState(prevState, state))
+					{
+						break;
+					}
+				}
+			}
+			else if (state is ImportedNodeState)
+			{
+				// Transition from in memory state
+				for (; ; )
+				{
+					NodeState prevState = node.State;
+					if (!(prevState is InMemoryNodeState))
+					{
+						break;
+					}
+					else if (node.TrySetState(prevState, state))
+					{
+						break;
+					}
+				}
+			}
+
 			return node;
 		}
 
@@ -539,10 +605,9 @@ namespace EpicGames.Horde.Storage.Bundles
 					BundleExport export = header.Exports[exportIdx];
 
 					List<NodeInfo> nodeRefs = export.References.ConvertAll(x => nodes[x]);
-					ExportedNodeState nodeState = new ExportedNodeState(bundleInfo, exportIdx, packetInfo, nodeOffset, export.Length, nodeRefs);
+					ExportedNodeState exportedState = new ExportedNodeState(bundleInfo, exportIdx, packetInfo, nodeOffset, export.Length, nodeRefs);
 
-					NodeInfo node = FindOrAddNode(export.Hash, nodeState);
-					node.State = nodeState;
+					NodeInfo node = FindOrAddNode(export.Hash, exportedState);
 					bundleInfo.Exports[exportIdx] = node;
 					nodes.Add(node);
 
@@ -664,74 +729,6 @@ namespace EpicGames.Horde.Storage.Bundles
 			}
 		}
 
-		async Task<bool> MarkDependencies(NodeInfo node, CancellationToken cancellationToken)
-		{
-			bool rewrite = false;
-			foreach (NodeInfo reference in await GetReferencesAsync(node, cancellationToken))
-			{
-				rewrite |= await MarkDependencies(reference, cancellationToken);
-			}
-			if (rewrite)
-			{
-				await MakeInMemory(node, cancellationToken);
-			}
-			return node.State is InMemoryNodeState;
-		}
-
-		async Task MakeInMemory(NodeInfo node, CancellationToken cancellationToken)
-		{
-			if (!(node.State is InMemoryNodeState))
-			{
-				ReadOnlySequence<byte> data = await GetDataAsync(node, cancellationToken);
-				IReadOnlyList<NodeInfo> references = await GetReferencesAsync(node, cancellationToken);
-				node.State = new InMemoryNodeState(data, references);
-			}
-		}
-
-#if false
-		/// <summary>
-		/// Splits apart any bundles that have a low utilization
-		/// </summary>
-		/// <returns></returns>
-		public async Task OptimizeAsync(double repackRatio)
-		{
-			// Find all the live nodes in the tree
-			Dictionary<BundleInfo, HashSet<NodeInfo>> bundleToLiveNodes = new Dictionary<BundleInfo, HashSet<NodeInfo>>();
-			await FindLiveNodes(_root, bundleToLiveNodes, new HashSet<NodeInfo>());
-
-			// Figure out which blobs need to be repacked
-			foreach ((BundleInfo bundleInfo, HashSet<NodeInfo> liveNodes) in bundleToLiveNodes)
-			{
-				Bundle bundle = await ReadBundleAsync(bundleInfo);
-				BundleHeader header = bundle.Header;
-
-				// Get the total size of this bundle
-				long totalSize = header.Exports.Sum(x => x.Length + x.References.Count * IoHash.NumBytes);
-
-				// Figure out how much of it we're actually using
-				long usedSize = 0;
-				foreach(NodeInfo liveNode in liveNodes)
-				{
-					if(liveNode.State is ExportedNodeState exportedState)
-					{
-						usedSize += exportedState.Length + exportedState.References.Count * IoHash.NumBytes;
-					}
-				}
-
-				int minUsedSize = (int)(totalSize * repackRatio);
-				if (usedSize < minUsedSize)
-				{
-					foreach (NodeInfo liveNode in liveNodes)
-					{
-						await MakeInMemory(liveNode);
-					}
-				}
-			}
-
-			// Walk backwards through the list of live nodes and make anything that references a repacked node also standalone
-			await MarkDependencies(_root);
-		}
-#endif
 		/// <summary>
 		/// Creates a Bundle containing a set of nodes. 
 		///
