Removed the use of "this." where possible in EpicGames.UHT.

diff --git a/Engine/Source/Programs/Shared/EpicGames.UHT/.editorconfig b/Engine/Source/Programs/Shared/EpicGames.UHT/.editorconfig
--- a/Engine/Source/Programs/Shared/EpicGames.UHT/.editorconfig
+++ b/Engine/Source/Programs/Shared/EpicGames.UHT/.editorconfig
@@ -267,8 +267,5 @@ dotnet_diagnostic.CA1508.severity = none
 # CA1501: 'UhtClassParser' has an object hierarchy '6' levels deep within the defining module. If possible, eliminate base classes within the hierarchy to decrease its hierarchy level below '6': 'UhtClassBaseParser, UhtClass, UhtStruct, UhtField, UhtObject, UhtType'.
 dotnet_diagnostic.CA1501.severity = none
 
-# IDE0003: Name can be simplified
-dotnet_diagnostic.IDE0003.severity = none
-
 # IDE0079: Remove unnecessary supression - Different tool versions will complain about unused stuff.
 dotnet_diagnostic.IDE0079.severity = none
diff --git a/Engine/Source/Programs/Shared/EpicGames.UHT/Exporters/CodeGen/UhtCodeGenerator.cs b/Engine/Source/Programs/Shared/EpicGames.UHT/Exporters/CodeGen/UhtCodeGenerator.cs
--- a/Engine/Source/Programs/Shared/EpicGames.UHT/Exporters/CodeGen/UhtCodeGenerator.cs
+++ b/Engine/Source/Programs/Shared/EpicGames.UHT/Exporters/CodeGen/UhtCodeGenerator.cs
@@ -62,36 +62,36 @@ namespace EpicGames.UHT.Exporters.CodeGen
 
 		private UhtCodeGenerator(IUhtExportFactory factory)
 		{
-			this.Factory = factory;
-			this.HeaderInfos = new HeaderInfo[this.Factory.Session.HeaderFileTypeCount];
-			this.ObjectInfos = new ObjectInfo[this.Factory.Session.ObjectTypeCount];
-			this.PackageInfos = new PackageInfo[this.Factory.Session.PackageTypeCount];
+			Factory = factory;
+			HeaderInfos = new HeaderInfo[Factory.Session.HeaderFileTypeCount];
+			ObjectInfos = new ObjectInfo[Factory.Session.ObjectTypeCount];
+			PackageInfos = new PackageInfo[Factory.Session.PackageTypeCount];
 		}
 
 		private void Generate()
 		{
 			List<Task?> prereqs = new();
 
-			this.FastArraySerializer = Session.FindType(null, UhtFindOptions.SourceName | UhtFindOptions.ScriptStruct, "FFastArraySerializer") as UhtScriptStruct;
+			FastArraySerializer = Session.FindType(null, UhtFindOptions.SourceName | UhtFindOptions.ScriptStruct, "FFastArraySerializer") as UhtScriptStruct;
 
 			// Perform some startup initialization to compute things we need over and over again
-			if (this.Session.GoWide)
+			if (Session.GoWide)
 			{
-				Parallel.ForEach(this.Factory.Session.Packages, package =>
+				Parallel.ForEach(Factory.Session.Packages, package =>
 				{
 					InitPackageInfo(package);
 				});
 			}
 			else
 			{
-				foreach (UhtPackage package in this.Factory.Session.Packages)
+				foreach (UhtPackage package in Factory.Session.Packages)
 				{
 					InitPackageInfo(package);
 				}
 			}
 
 			// Generate the files for the header files
-			foreach (UhtHeaderFile headerFile in this.Session.SortedHeaderFiles)
+			foreach (UhtHeaderFile headerFile in Session.SortedHeaderFiles)
 			{
 				if (headerFile.ShouldExport)
 				{
@@ -103,11 +103,11 @@ namespace EpicGames.UHT.Exporters.CodeGen
 					{
 						if (headerFile != referenced)
 						{
-							prereqs.Add(this.HeaderInfos[referenced.HeaderFileTypeIndex].Task);
+							prereqs.Add(HeaderInfos[referenced.HeaderFileTypeIndex].Task);
 						}
 					}
 
-					this.HeaderInfos[headerFile.HeaderFileTypeIndex].Task = Factory.CreateTask(prereqs,
+					HeaderInfos[headerFile.HeaderFileTypeIndex].Task = Factory.CreateTask(prereqs,
 						(IUhtExportFactory factory) =>
 						{
 							new UhtHeaderCodeGeneratorHFile(this, package, headerFile).Generate(factory);
@@ -117,8 +117,8 @@ namespace EpicGames.UHT.Exporters.CodeGen
 			}
 
 			// Generate the files for the packages
-			List<Task?> generatedPackages = new(this.Session.PackageTypeCount);
-			foreach (UhtPackage package in this.Session.Packages)
+			List<Task?> generatedPackages = new(Session.PackageTypeCount);
+			foreach (UhtPackage package in Session.Packages)
 			{
 				UHTManifest.Module module = package.Module;
 
@@ -126,7 +126,7 @@ namespace EpicGames.UHT.Exporters.CodeGen
 				prereqs.Clear();
 				foreach (UhtHeaderFile headerFile in package.Children)
 				{
-					prereqs.Add(this.HeaderInfos[headerFile.HeaderFileTypeIndex].Task);
+					prereqs.Add(HeaderInfos[headerFile.HeaderFileTypeIndex].Task);
 					if (!writeHeader)
 					{
 						foreach (UhtType type in headerFile.Children)
@@ -158,7 +158,7 @@ namespace EpicGames.UHT.Exporters.CodeGen
 			}
 
 			// Wait for all the packages to complete
-			List<Task> packageTasks = new(this.Session.PackageTypeCount);
+			List<Task> packageTasks = new(Session.PackageTypeCount);
 			foreach (Task? output in generatedPackages)
 			{
 				if (output != null)
@@ -182,7 +182,7 @@ namespace EpicGames.UHT.Exporters.CodeGen
 			{
 				return "nullptr";
 			}
-			return registered ? this.ObjectInfos[obj.ObjectTypeIndex].RegisteredSingletonName : this.ObjectInfos[obj.ObjectTypeIndex].UnregisteredSingletonName;
+			return registered ? ObjectInfos[obj.ObjectTypeIndex].RegisteredSingletonName : ObjectInfos[obj.ObjectTypeIndex].UnregisteredSingletonName;
 		}
 
 		/// <summary>
@@ -204,7 +204,7 @@ namespace EpicGames.UHT.Exporters.CodeGen
 		/// <returns>External declaration</returns>
 		public string GetExternalDecl(int objectIndex, bool registered)
 		{
-			return registered ? this.ObjectInfos[objectIndex].RegsiteredExternalDecl : this.ObjectInfos[objectIndex].UnregisteredExternalDecl;
+			return registered ? ObjectInfos[objectIndex].RegsiteredExternalDecl : ObjectInfos[objectIndex].UnregisteredExternalDecl;
 		}
 
 		/// <summary>
@@ -226,7 +226,7 @@ namespace EpicGames.UHT.Exporters.CodeGen
 		/// <returns>Cross reference</returns>
 		public string GetCrossReference(int objectIndex, bool registered)
 		{
-			return registered ? this.ObjectInfos[objectIndex].RegsiteredCrossReference : this.ObjectInfos[objectIndex].UnregsiteredCrossReference;
+			return registered ? ObjectInfos[objectIndex].RegsiteredCrossReference : ObjectInfos[objectIndex].UnregsiteredCrossReference;
 		}
 		#endregion
 
@@ -235,12 +235,12 @@ namespace EpicGames.UHT.Exporters.CodeGen
 		{
 			StringBuilder builder = new();
 
-			ref PackageInfo packageInfo = ref this.PackageInfos[package.PackageTypeIndex];
+			ref PackageInfo packageInfo = ref PackageInfos[package.PackageTypeIndex];
 			packageInfo.StrippedName = package.SourceName.Replace('/', '_');
 			packageInfo.Api = $"{package.ShortName.ToString().ToUpper()}_API ";
 
 			// Construct the names used commonly during export
-			ref ObjectInfo objectInfo = ref this.ObjectInfos[package.ObjectTypeIndex];
+			ref ObjectInfo objectInfo = ref ObjectInfos[package.ObjectTypeIndex];
 			builder.Append("Z_Construct_UPackage_");
 			builder.Append(packageInfo.StrippedName);
 			objectInfo.UnregisteredSingletonName = objectInfo.RegisteredSingletonName = builder.ToString();
@@ -255,7 +255,7 @@ namespace EpicGames.UHT.Exporters.CodeGen
 
 		private void InitHeaderInfo(StringBuilder builder, UhtPackage package, ref PackageInfo packageInfo, UhtHeaderFile headerFile)
 		{
-			ref HeaderInfo headerInfo = ref this.HeaderInfos[headerFile.HeaderFileTypeIndex];
+			ref HeaderInfo headerInfo = ref HeaderInfos[headerFile.HeaderFileTypeIndex];
 
 			headerInfo.IncludePath = Path.GetRelativePath(package.Module.IncludeBase, headerFile.FilePath).Replace('\\', '/');
 
@@ -308,7 +308,7 @@ namespace EpicGames.UHT.Exporters.CodeGen
 
 		private void InitObjectInfo(StringBuilder builder, UhtPackage package, ref PackageInfo packageInfo, ref HeaderInfo headerInfo, UhtObject obj)
 		{
-			ref ObjectInfo objectInfo = ref this.ObjectInfos[obj.ObjectTypeIndex];
+			ref ObjectInfo objectInfo = ref ObjectInfos[obj.ObjectTypeIndex];
 
 			builder.Clear();
 
@@ -331,7 +331,7 @@ namespace EpicGames.UHT.Exporters.CodeGen
 				{
 					if (classObj.AlternateObject != null)
 					{
-						this.ObjectInfos[classObj.AlternateObject.ObjectTypeIndex].NativeInterface = classObj;
+						ObjectInfos[classObj.AlternateObject.ObjectTypeIndex].NativeInterface = classObj;
 					}
 				}
 			}
diff --git a/Engine/Source/Programs/Shared/EpicGames.UHT/Exporters/CodeGen/UhtHeaderCodeGenerator.cs b/Engine/Source/Programs/Shared/EpicGames.UHT/Exporters/CodeGen/UhtHeaderCodeGenerator.cs
--- a/Engine/Source/Programs/Shared/EpicGames.UHT/Exporters/CodeGen/UhtHeaderCodeGenerator.cs
+++ b/Engine/Source/Programs/Shared/EpicGames.UHT/Exporters/CodeGen/UhtHeaderCodeGenerator.cs
@@ -35,12 +35,12 @@ namespace EpicGames.UHT.Exporters.CodeGen
 		#endregion
 
 		public readonly UhtHeaderFile HeaderFile;
-		public string FileId => this.HeaderInfos[this.HeaderFile.HeaderFileTypeIndex].FileId;
+		public string FileId => HeaderInfos[HeaderFile.HeaderFileTypeIndex].FileId;
 
 		public UhtHeaderCodeGenerator(UhtCodeGenerator codeGenerator, UhtPackage package, UhtHeaderFile headerFile)
 			: base(codeGenerator, package)
 		{
-			this.HeaderFile = headerFile;
+			HeaderFile = headerFile;
 		}
 
 		#region Event parameter
@@ -162,7 +162,7 @@ namespace EpicGames.UHT.Exporters.CodeGen
 					!(outerClass != null && outerClass.ClassFlags.HasAnyFlags(EClassFlags.RequiredAPI)) &&
 					exportFlags.HasAnyFlags(UhtFunctionExportFlags.RequiredAPI))
 				{
-					builder.Append(this.PackageApi);
+					builder.Append(PackageApi);
 				}
 
 				if (textType == UhtPropertyTextType.InterfaceFunctionArgOrRetVal)
diff --git a/Engine/Source/Programs/Shared/EpicGames.UHT/Exporters/CodeGen/UhtHeaderCodeGeneratorCppFile.cs b/Engine/Source/Programs/Shared/EpicGames.UHT/Exporters/CodeGen/UhtHeaderCodeGeneratorCppFile.cs
--- a/Engine/Source/Programs/Shared/EpicGames.UHT/Exporters/CodeGen/UhtHeaderCodeGeneratorCppFile.cs
+++ b/Engine/Source/Programs/Shared/EpicGames.UHT/Exporters/CodeGen/UhtHeaderCodeGeneratorCppFile.cs
@@ -30,7 +30,7 @@ namespace EpicGames.UHT.Exporters.CodeGen
 		/// <param name="factory">Requesting factory</param>
 		public void Generate(IUhtExportFactory factory)
 		{
-			ref UhtCodeGenerator.HeaderInfo headerInfo = ref this.HeaderInfos[this.HeaderFile.HeaderFileTypeIndex];
+			ref UhtCodeGenerator.HeaderInfo headerInfo = ref HeaderInfos[HeaderFile.HeaderFileTypeIndex];
 			{
 				using BorrowStringBuilder borrower = new(StringBuilderCache.Big);
 				StringBuilder builder = borrower.StringBuilder;
@@ -46,16 +46,16 @@ namespace EpicGames.UHT.Exporters.CodeGen
 				bool addedStructuredArchiveFromArchiveHeader = false;
 				bool addedArchiveUObjectFromStructuredArchiveHeader = false;
 				HashSet<UhtHeaderFile> addedIncludes = new();
-				addedIncludes.Add(this.HeaderFile);
-				foreach (UhtType type in this.HeaderFile.Children)
+				addedIncludes.Add(HeaderFile);
+				foreach (UhtType type in HeaderFile.Children)
 				{
 					if (type is UhtClass classObj)
 					{
-						if (classObj.ClassWithin != this.Session.UObject && !classObj.ClassWithin.HeaderFile.IsNoExportTypes)
+						if (classObj.ClassWithin != Session.UObject && !classObj.ClassWithin.HeaderFile.IsNoExportTypes)
 						{
 							if (addedIncludes.Add(classObj.ClassWithin.HeaderFile))
 							{
-								builder.Append("#include \"").Append(this.HeaderInfos[classObj.ClassWithin.HeaderFile.HeaderFileTypeIndex].IncludePath).Append("\"\r\n");
+								builder.Append("#include \"").Append(HeaderInfos[classObj.ClassWithin.HeaderFile.HeaderFileTypeIndex].IncludePath).Append("\"\r\n");
 							}
 						}
 
@@ -84,12 +84,12 @@ namespace EpicGames.UHT.Exporters.CodeGen
 				}
 
 				builder.Append(DisableDeprecationWarnings).Append("\r\n");
-				string cleanFileName = this.HeaderFile.FileNameWithoutExtension.Replace('.', '_');
+				string cleanFileName = HeaderFile.FileNameWithoutExtension.Replace('.', '_');
 				builder.Append("void EmptyLinkFunctionForGeneratedCode").Append(cleanFileName).Append("() {}\r\n");
 
-				if (this.HeaderFile.References.CrossModule.References.Count > 0)
+				if (HeaderFile.References.CrossModule.References.Count > 0)
 				{
-					ReadOnlyMemory<string> sorted = this.HeaderFile.References.CrossModule.GetSortedReferences(
+					ReadOnlyMemory<string> sorted = HeaderFile.References.CrossModule.GetSortedReferences(
 						(int objectIndex, bool registered) => GetCrossReference(objectIndex, registered));
 					builder.Append("// Cross Module References\r\n");
 					foreach (string crossReference in sorted.Span)
@@ -105,7 +105,7 @@ namespace EpicGames.UHT.Exporters.CodeGen
 				bool hasRegisteredScriptStructs = false;
 				bool hasRegisteredClasses = false;
 				bool allEnumsEditorOnly = true;
-				foreach (UhtField field in this.HeaderFile.References.ExportTypes)
+				foreach (UhtField field in HeaderFile.References.ExportTypes)
 				{
 					if (field is UhtEnum enumObj)
 					{
@@ -175,7 +175,7 @@ namespace EpicGames.UHT.Exporters.CodeGen
 							builder.Append("#if WITH_EDITORONLY_DATA\r\n");
 						}
 						builder.Append("\tconst FEnumRegisterCompiledInInfo ").Append(staticsName).Append("::EnumInfo[] = {\r\n");
-						foreach (UhtObject obj in this.HeaderFile.References.ExportTypes)
+						foreach (UhtObject obj in HeaderFile.References.ExportTypes)
 						{
 							if (obj is UhtEnum enumObj)
 							{
@@ -183,7 +183,7 @@ namespace EpicGames.UHT.Exporters.CodeGen
 								{
 									builder.Append("#if WITH_EDITORONLY_DATA\r\n");
 								}
-								uint hash = this.ObjectInfos[enumObj.ObjectTypeIndex].Hash;
+								uint hash = ObjectInfos[enumObj.ObjectTypeIndex].Hash;
 								builder
 									.Append("\t\t{ ")
 									.Append(enumObj.SourceName)
@@ -209,13 +209,13 @@ namespace EpicGames.UHT.Exporters.CodeGen
 					if (hasRegisteredScriptStructs)
 					{
 						builder.Append("\tconst FStructRegisterCompiledInInfo ").Append(staticsName).Append("::ScriptStructInfo[] = {\r\n");
-						foreach (UhtObject obj in this.HeaderFile.References.ExportTypes)
+						foreach (UhtObject obj in HeaderFile.References.ExportTypes)
 						{
 							if (obj is UhtScriptStruct scriptStruct)
 							{
 								if (scriptStruct.ScriptStructFlags.HasAnyFlags(EStructFlags.Native))
 								{
-									uint hash = this.ObjectInfos[scriptStruct.ObjectTypeIndex].Hash;
+									uint hash = ObjectInfos[scriptStruct.ObjectTypeIndex].Hash;
 									builder
 										.Append("\t\t{ ")
 										.Append(scriptStruct.SourceName)
@@ -238,13 +238,13 @@ namespace EpicGames.UHT.Exporters.CodeGen
 					if (hasRegisteredClasses)
 					{
 						builder.Append("\tconst FClassRegisterCompiledInInfo ").Append(staticsName).Append("::ClassInfo[] = {\r\n");
-						foreach (UhtObject obj in this.HeaderFile.References.ExportTypes)
+						foreach (UhtObject obj in HeaderFile.References.ExportTypes)
 						{
 							if (obj is UhtClass classObj)
 							{
 								if (!classObj.ClassFlags.HasAnyFlags(EClassFlags.Intrinsic))
 								{
-									uint hash = this.ObjectInfos[classObj.ObjectTypeIndex].Hash;
+									uint hash = ObjectInfos[classObj.ObjectTypeIndex].Hash;
 									builder
 										.Append("\t\t{ Z_Construct_UClass_")
 										.Append(classObj.SourceName)
@@ -276,10 +276,10 @@ namespace EpicGames.UHT.Exporters.CodeGen
 					builder.Append("\t\t").AppendArray(!hasRegisteredEnums, allEnumsEditorOnly, staticsName, "EnumInfo").Append(");\r\n");
 				}
 
-				if (this.Session.IncludeDebugOutput)
+				if (Session.IncludeDebugOutput)
 				{
 					builder.Append("#if 0\r\n");
-					ReadOnlyMemory<string> sorted = this.HeaderFile.References.Declaration.GetSortedReferences(
+					ReadOnlyMemory<string> sorted = HeaderFile.References.Declaration.GetSortedReferences(
 						(int objectIndex, bool registered) => GetExternalDecl(objectIndex, registered));
 					foreach (string declaration in sorted.Span)
 					{
@@ -296,13 +296,13 @@ namespace EpicGames.UHT.Exporters.CodeGen
 					using UhtBorrowBuffer borrowBuffer = new(builder);
 					string cppFilePath = factory.MakePath(HeaderFile, ".gen.cpp");
 					StringView generatedBody = new(borrowBuffer.Buffer.Memory);
-					if (this.SaveExportedHeaders)
+					if (SaveExportedHeaders)
 					{
 						factory.CommitOutput(cppFilePath, generatedBody);
 					}
 
 					// Save the hash of the generated body 
-					this.HeaderInfos[this.HeaderFile.HeaderFileTypeIndex].BodyHash = UhtHash.GenenerateTextHash(generatedBody.Span[generatedBodyStart..generatedBodyEnd]);
+					HeaderInfos[HeaderFile.HeaderFileTypeIndex].BodyHash = UhtHash.GenenerateTextHash(generatedBody.Span[generatedBodyStart..generatedBodyEnd]);
 				}
 			}
 		}
@@ -355,13 +355,13 @@ namespace EpicGames.UHT.Exporters.CodeGen
 				builder.Append("\t\tif (!").Append(registrationName).Append(".OuterSingleton)\r\n");
 				builder.Append("\t\t{\r\n");
 				builder.Append("\t\t\t").Append(registrationName).Append(".OuterSingleton = GetStaticEnum(").Append(singletonName).Append(", ")
-					.Append(this.PackageSingletonName).Append("(), TEXT(\"").Append(enumObj.SourceName).Append("\"));\r\n");
+					.Append(PackageSingletonName).Append("(), TEXT(\"").Append(enumObj.SourceName).Append("\"));\r\n");
 				builder.Append("\t\t}\r\n");
 				builder.Append("\t\treturn ").Append(registrationName).Append(".OuterSingleton;\r\n");
 
 				builder.Append("\t}\r\n");
 
-				builder.Append("\ttemplate<> ").Append(this.PackageApi).Append("UEnum* StaticEnum<").Append(enumObj.CppType).Append(">()\r\n");
+				builder.Append("\ttemplate<> ").Append(PackageApi).Append("UEnum* StaticEnum<").Append(enumObj.CppType).Append(">()\r\n");
 				builder.Append("\t{\r\n");
 				builder.Append("\t\treturn ").Append(enumObj.SourceName).Append("_StaticEnum();\r\n");
 				builder.Append("\t}\r\n");
@@ -396,7 +396,7 @@ namespace EpicGames.UHT.Exporters.CodeGen
 
 				// Singleton parameters
 				builder.Append("\tconst UECodeGen_Private::FEnumParams ").Append(staticsName).Append("::EnumParams = {\r\n");
-				builder.Append("\t\t(UObject*(*)())").Append(this.PackageSingletonName).Append(",\r\n");
+				builder.Append("\t\t(UObject*(*)())").Append(PackageSingletonName).Append(",\r\n");
 				builder.Append("\t\t").Append(enumDisplayNameFn).Append(",\r\n");
 				builder.Append("\t\t").AppendUTF8LiteralString(enumObj.SourceName).Append(",\r\n");
 				builder.Append("\t\t").AppendUTF8LiteralString(enumObj.CppType).Append(",\r\n");
@@ -420,7 +420,7 @@ namespace EpicGames.UHT.Exporters.CodeGen
 
 				{
 					using UhtBorrowBuffer borrowBuffer = new(builder, hashCodeBlockStart, builder.Length - hashCodeBlockStart);
-					this.ObjectInfos[enumObj.ObjectTypeIndex].Hash = UhtHash.GenenerateTextHash(borrowBuffer.Buffer.Memory.Span);
+					ObjectInfos[enumObj.ObjectTypeIndex].Hash = UhtHash.GenenerateTextHash(borrowBuffer.Buffer.Memory.Span);
 				}
 			}
 			return builder;
@@ -463,7 +463,7 @@ namespace EpicGames.UHT.Exporters.CodeGen
 					.Append(".OuterSingleton = GetStaticStruct(")
 					.Append(singletonName)
 					.Append(", ")
-					.Append(this.PackageSingletonName)
+					.Append(PackageSingletonName)
 					.Append("(), TEXT(\"")
 					.Append(scriptStruct.EngineName)
 					.Append("\"));\r\n");
@@ -519,13 +519,13 @@ namespace EpicGames.UHT.Exporters.CodeGen
 				builder.Append("}\r\n");
 
 				// Generate the StaticStruct specialization
-				builder.Append("template<> ").Append(this.PackageApi).Append("UScriptStruct* StaticStruct<").Append(scriptStruct.SourceName).Append(">()\r\n");
+				builder.Append("template<> ").Append(PackageApi).Append("UScriptStruct* StaticStruct<").Append(scriptStruct.SourceName).Append(">()\r\n");
 				builder.Append("{\r\n");
 				builder.Append("\treturn ").Append(scriptStruct.SourceName).Append("::StaticStruct();\r\n");
 				builder.Append("}\r\n");
 
 				// Inject implementation needed to support auto bindings of fast arrays
-				if (this.ObjectInfos[scriptStruct.ObjectTypeIndex].FastArrayProperty != null)
+				if (ObjectInfos[scriptStruct.ObjectTypeIndex].FastArrayProperty != null)
 				{
 					builder.Append("UE_NET_IMPLEMENT_FASTARRAY(").Append(scriptStruct.SourceName).Append(");\r\n");
 				}
@@ -574,7 +574,7 @@ namespace EpicGames.UHT.Exporters.CodeGen
 				AppendPropertiesDefs(builder, scriptStruct, scriptStruct.SourceName, staticsName, 1);
 
 				builder.Append("\tconst UECodeGen_Private::FStructParams ").Append(staticsName).Append("::ReturnStructParams = {\r\n");
-				builder.Append("\t\t(UObject* (*)())").Append(this.PackageSingletonName).Append(",\r\n");
+				builder.Append("\t\t(UObject* (*)())").Append(PackageSingletonName).Append(",\r\n");
 				builder.Append("\t\t").Append(GetSingletonName(scriptStruct.SuperScriptStruct, true)).Append(",\r\n");
 				if (scriptStruct.ScriptStructFlags.HasAnyFlags(EStructFlags.Native))
 				{
@@ -618,7 +618,7 @@ namespace EpicGames.UHT.Exporters.CodeGen
 
 			using (UhtBorrowBuffer borrowBuffer = new(builder, hashCodeBlockStart, builder.Length - hashCodeBlockStart))
 			{
-				this.ObjectInfos[scriptStruct.ObjectTypeIndex].Hash = UhtHash.GenenerateTextHash(borrowBuffer.Buffer.Memory.Span);
+				ObjectInfos[scriptStruct.ObjectTypeIndex].Hash = UhtHash.GenenerateTextHash(borrowBuffer.Buffer.Memory.Span);
 			}
 
 			// if this struct has RigVM methods we need to implement both the 
@@ -783,7 +783,7 @@ namespace EpicGames.UHT.Exporters.CodeGen
 
 			using (UhtBorrowBuffer borrowBuffer = new(builder, hashCodeBlockStart, builder.Length - hashCodeBlockStart))
 			{
-				this.ObjectInfos[function.ObjectTypeIndex].Hash = UhtHash.GenenerateTextHash(borrowBuffer.Buffer.Memory.Span);
+				ObjectInfos[function.ObjectTypeIndex].Hash = UhtHash.GenenerateTextHash(borrowBuffer.Buffer.Memory.Span);
 			}
 			return builder;
 		}
@@ -842,7 +842,7 @@ namespace EpicGames.UHT.Exporters.CodeGen
 				return builder;
 			}
 
-			PropertyMemberContextImpl context = new(this.CodeGenerator, structObj, structSourceName, staticsName);
+			PropertyMemberContextImpl context = new(CodeGenerator, structObj, structSourceName, staticsName);
 
 			using (UhtMacroBlockEmitter emitter = new(builder, "WITH_EDITORONLY_DATA"))
 			{
@@ -881,7 +881,7 @@ namespace EpicGames.UHT.Exporters.CodeGen
 				return builder;
 			}
 
-			PropertyMemberContextImpl context = new(this.CodeGenerator, structObj, structSourceName, staticsName);
+			PropertyMemberContextImpl context = new(CodeGenerator, structObj, structSourceName, staticsName);
 
 			using (UhtMacroBlockEmitter emitter = new(builder, "WITH_EDITORONLY_DATA"))
 			{
@@ -1228,7 +1228,7 @@ namespace EpicGames.UHT.Exporters.CodeGen
 			string registrationName = $"Z_Registration_Info_UClass_{classObj.SourceName}";
 			string[]? sparseDataTypes = classObj.MetaData.GetStringArray(UhtNames.SparseClassDataTypes);
 
-			PropertyMemberContextImpl context = new(this.CodeGenerator, classObj, classObj.SourceName, staticsName);
+			PropertyMemberContextImpl context = new(CodeGenerator, classObj, classObj.SourceName, staticsName);
 
 			bool hasInterfaces = classObj.Bases.Any(x => x is UhtClass baseClass && baseClass.ClassFlags.HasAnyFlags(EClassFlags.Interface));
 
@@ -1303,7 +1303,7 @@ namespace EpicGames.UHT.Exporters.CodeGen
 				{
 					builder.Append("\t\t(UObject* (*)())").Append(GetSingletonName(classObj.SuperClass, true)).Append(",\r\n");
 				}
-				builder.Append("\t\t(UObject* (*)())").Append(GetSingletonName(this.Package, true)).Append(",\r\n");
+				builder.Append("\t\t(UObject* (*)())").Append(GetSingletonName(Package, true)).Append(",\r\n");
 				builder.Append("\t};\r\n");
 
 				if (sortedFunctions.Count > 0)
@@ -1452,7 +1452,7 @@ namespace EpicGames.UHT.Exporters.CodeGen
 				uint baseClassHash = 0;
 				if (classObj.SuperClass != null && !classObj.SuperClass.ClassFlags.HasAnyFlags(EClassFlags.Intrinsic))
 				{
-					baseClassHash = this.ObjectInfos[classObj.SuperClass.ObjectTypeIndex].Hash;
+					baseClassHash = ObjectInfos[classObj.SuperClass.ObjectTypeIndex].Hash;
 				}
 				hashBuilder.Append($"\r\n// {baseClassHash}\r\n");
 
@@ -1461,7 +1461,7 @@ namespace EpicGames.UHT.Exporters.CodeGen
 				{
 					foreach (string sparseDataType in sparseDataTypes)
 					{
-						UhtType? type = this.Session.FindType(classObj, UhtFindOptions.ScriptStruct | UhtFindOptions.EngineName, sparseDataType);
+						UhtType? type = Session.FindType(classObj, UhtFindOptions.ScriptStruct | UhtFindOptions.EngineName, sparseDataType);
 						if (type != null)
 						{
 							hashBuilder.Append(type.EngineName).Append("\r\n");
@@ -1483,7 +1483,7 @@ namespace EpicGames.UHT.Exporters.CodeGen
 					builder.Append("\t*/\r\n");
 				}
 
-				if (this.Session.IncludeDebugOutput)
+				if (Session.IncludeDebugOutput)
 				{
 					using UhtBorrowBuffer borrowBuffer = new(hashBuilder, saveLength, hashBuilder.Length - saveLength);
 					builder.Append("#if 0\r\n");
@@ -1494,11 +1494,11 @@ namespace EpicGames.UHT.Exporters.CodeGen
 				// Calculate generated class initialization code hash so that we know when it changes after hot-reload
 				{
 					using UhtBorrowBuffer borrowBuffer = new(hashBuilder);
-					this.ObjectInfos[classObj.ObjectTypeIndex].Hash = UhtHash.GenenerateTextHash(borrowBuffer.Buffer.Memory.Span);
+					ObjectInfos[classObj.ObjectTypeIndex].Hash = UhtHash.GenenerateTextHash(borrowBuffer.Buffer.Memory.Span);
 				}
 			}
 
-			builder.Append("\ttemplate<> ").Append(this.PackageApi).Append("UClass* StaticClass<").Append(classObj.SourceName).Append(">()\r\n");
+			builder.Append("\ttemplate<> ").Append(PackageApi).Append("UClass* StaticClass<").Append(classObj.SourceName).Append(">()\r\n");
 			builder.Append("\t{\r\n");
 			builder.Append("\t\treturn ").Append(classObj.SourceName).Append("::StaticClass();\r\n");
 			builder.Append("\t}\r\n");
@@ -1771,26 +1771,26 @@ namespace EpicGames.UHT.Exporters.CodeGen
 
 			public PropertyMemberContextImpl(UhtCodeGenerator codeGenerator, UhtStruct outerStruct, string outerStructSourceName, string staticsName)
 			{
-				this._codeGenerator = codeGenerator;
-				this._outerStruct = outerStruct;
-				this._staticsName = staticsName;
-				this._outerStructSourceName = outerStructSourceName.Length == 0 ? outerStruct.SourceName : outerStructSourceName;
+				_codeGenerator = codeGenerator;
+				_outerStruct = outerStruct;
+				_staticsName = staticsName;
+				_outerStructSourceName = outerStructSourceName.Length == 0 ? outerStruct.SourceName : outerStructSourceName;
 			}
 
-			public UhtStruct OuterStruct => this._outerStruct;
-			public string OuterStructSourceName => this._outerStructSourceName;
-			public string StaticsName => this._staticsName;
+			public UhtStruct OuterStruct => _outerStruct;
+			public string OuterStructSourceName => _outerStructSourceName;
+			public string StaticsName => _staticsName;
 			public string NamePrefix => "NewProp_";
 			public string MetaDataSuffix => "_MetaData";
 
 			public string GetSingletonName(UhtObject? obj, bool registered)
 			{
-				return this._codeGenerator.GetSingletonName(obj, registered);
+				return _codeGenerator.GetSingletonName(obj, registered);
 			}
 
 			public uint GetTypeHash(UhtObject obj)
 			{
-				return this._codeGenerator.ObjectInfos[obj.ObjectTypeIndex].Hash;
+				return _codeGenerator.ObjectInfos[obj.ObjectTypeIndex].Hash;
 			}
 		}
 	}
diff --git a/Engine/Source/Programs/Shared/EpicGames.UHT/Exporters/CodeGen/UhtHeaderCodeGeneratorHFile.cs b/Engine/Source/Programs/Shared/EpicGames.UHT/Exporters/CodeGen/UhtHeaderCodeGeneratorHFile.cs
--- a/Engine/Source/Programs/Shared/EpicGames.UHT/Exporters/CodeGen/UhtHeaderCodeGeneratorHFile.cs
+++ b/Engine/Source/Programs/Shared/EpicGames.UHT/Exporters/CodeGen/UhtHeaderCodeGeneratorHFile.cs
@@ -34,7 +34,7 @@ namespace EpicGames.UHT.Exporters.CodeGen
 		/// <param name="factory">Requesting factory</param>
 		public void Generate(IUhtExportFactory factory)
 		{
-			ref UhtCodeGenerator.HeaderInfo headerInfo = ref this.HeaderInfos[this.HeaderFile.HeaderFileTypeIndex];
+			ref UhtCodeGenerator.HeaderInfo headerInfo = ref HeaderInfos[HeaderFile.HeaderFileTypeIndex];
 			{
 				using BorrowStringBuilder borrower = new(StringBuilderCache.Big);
 				StringBuilder builder = borrower.StringBuilder;
@@ -49,14 +49,14 @@ namespace EpicGames.UHT.Exporters.CodeGen
 				builder.Append("\r\n");
 				builder.Append(DisableDeprecationWarnings).Append("\r\n");
 
-				string strippedName = Path.GetFileNameWithoutExtension(this.HeaderFile.FilePath);
-				string defineName = $"{this.Package.ShortName.ToString().ToUpper()}_{strippedName}_generated_h";
+				string strippedName = Path.GetFileNameWithoutExtension(HeaderFile.FilePath);
+				string defineName = $"{Package.ShortName.ToString().ToUpper()}_{strippedName}_generated_h";
 
-				if (this.HeaderFile.References.ForwardDeclarations.Count > 0)
+				if (HeaderFile.References.ForwardDeclarations.Count > 0)
 				{
-					string[] sorted = new string[this.HeaderFile.References.ForwardDeclarations.Count];
+					string[] sorted = new string[HeaderFile.References.ForwardDeclarations.Count];
 					int index = 0;
-					foreach (string forwardDeclaration in this.HeaderFile.References.ForwardDeclarations)
+					foreach (string forwardDeclaration in HeaderFile.References.ForwardDeclarations)
 					{
 						sorted[index++] = forwardDeclaration;
 					}
@@ -73,7 +73,7 @@ namespace EpicGames.UHT.Exporters.CodeGen
 				builder.Append("#define ").Append(defineName).Append("\r\n");
 				builder.Append("\r\n");
 
-				foreach (UhtField field in this.HeaderFile.References.ExportTypes)
+				foreach (UhtField field in HeaderFile.References.ExportTypes)
 				{
 					if (field is UhtEnum enumObj)
 					{
@@ -101,7 +101,7 @@ namespace EpicGames.UHT.Exporters.CodeGen
 				builder.Append("#undef CURRENT_FILE_ID\r\n");
 				builder.Append("#define CURRENT_FILE_ID ").Append(headerInfo.FileId).Append("\r\n\r\n\r\n");
 
-				foreach (UhtField field in this.HeaderFile.References.ExportTypes)
+				foreach (UhtField field in HeaderFile.References.ExportTypes)
 				{
 					if (field is UhtEnum enumObject)
 					{
@@ -111,9 +111,9 @@ namespace EpicGames.UHT.Exporters.CodeGen
 
 				builder.Append(EnableDeprecationWarnings).Append("\r\n");
 
-				if (this.SaveExportedHeaders)
+				if (SaveExportedHeaders)
 				{
-					string headerFilePath = factory.MakePath(this.HeaderFile, ".generated.h");
+					string headerFilePath = factory.MakePath(HeaderFile, ".generated.h");
 					factory.CommitOutput(headerFilePath, builder);
 				}
 			}
@@ -145,7 +145,7 @@ namespace EpicGames.UHT.Exporters.CodeGen
 					builder.Append('\t');
 					if (!scriptStruct.ScriptStructFlags.HasAnyFlags(EStructFlags.RequiredAPI))
 					{
-						builder.Append(this.PackageApi);
+						builder.Append(PackageApi);
 					}
 					builder.Append("static class UScriptStruct* StaticStruct(); \\\r\n");
 
@@ -247,7 +247,7 @@ namespace EpicGames.UHT.Exporters.CodeGen
 					{
 						builder.Append("\ttypedef ").Append(scriptStruct.SuperScriptStruct.SourceName).Append(" Super; \\\r\n");
 					}
-					UhtProperty? FastArrayProperty = this.ObjectInfos[scriptStruct.ObjectTypeIndex].FastArrayProperty;
+					UhtProperty? FastArrayProperty = ObjectInfos[scriptStruct.ObjectTypeIndex].FastArrayProperty;
 					if (FastArrayProperty != null)
 					{
 						builder
@@ -258,14 +258,14 @@ namespace EpicGames.UHT.Exporters.CodeGen
 							.Append(", ");
 						if (!scriptStruct.ScriptStructFlags.HasAnyFlags(EStructFlags.RequiredAPI))
 						{
-							builder.Append(this.PackageApi);
+							builder.Append(PackageApi);
 						}
 						builder.Append("); \\\r\n");
 					}
 				}
 
 				// Forward declare the StaticStruct specialization in the header
-				builder.Append("template<> ").Append(this.PackageApi).Append("UScriptStruct* StaticStruct<struct ").Append(scriptStruct.SourceName).Append(">();\r\n");
+				builder.Append("template<> ").Append(PackageApi).Append("UScriptStruct* StaticStruct<struct ").Append(scriptStruct.SourceName).Append(">();\r\n");
 				builder.Append("\r\n");
 			}
 			return builder;
@@ -300,7 +300,7 @@ namespace EpicGames.UHT.Exporters.CodeGen
 				builder.Append("template<> struct TIsUEnumClass<").Append(enumObj.CppType).Append("> { enum { Value = true }; };\r\n");
 
 				// Forward declare the StaticEnum<> specialization for enum classes
-				builder.Append("template<> ").Append(this.PackageApi).Append("UEnum* StaticEnum<").Append(enumObj.CppType).Append(">();\r\n");
+				builder.Append("template<> ").Append(PackageApi).Append("UEnum* StaticEnum<").Append(enumObj.CppType).Append(">();\r\n");
 				builder.Append("\r\n");
 			}
 
@@ -335,8 +335,8 @@ namespace EpicGames.UHT.Exporters.CodeGen
 
 		private StringBuilder AppendClass(StringBuilder builder, UhtClass classObj)
 		{
-			string api = classObj.ClassFlags.HasAnyFlags(EClassFlags.MinimalAPI) ? this.PackageApi : "NO_API ";
-			UhtClass? nativeInterface = this.ObjectInfos[classObj.ObjectTypeIndex].NativeInterface;
+			string api = classObj.ClassFlags.HasAnyFlags(EClassFlags.MinimalAPI) ? PackageApi : "NO_API ";
+			UhtClass? nativeInterface = ObjectInfos[classObj.ObjectTypeIndex].NativeInterface;
 
 			// Write the spare declarations
 			AppendSparseDeclarations(builder, classObj);
@@ -456,7 +456,7 @@ namespace EpicGames.UHT.Exporters.CodeGen
 			}
 
 			// Forward declare the StaticClass specialization in the header
-			builder.Append("template<> ").Append(this.PackageApi).Append("UClass* StaticClass<class ").Append(classObj.SourceName).Append(">();\r\n");
+			builder.Append("template<> ").Append(PackageApi).Append("UClass* StaticClass<class ").Append(classObj.SourceName).Append(">();\r\n");
 			builder.Append("\r\n");
 			return builder;
 		}
@@ -498,7 +498,7 @@ namespace EpicGames.UHT.Exporters.CodeGen
 			bool hasProperties, bool hasFunctions, bool hasEditorFields, bool allEditorFields,
 			bool includeEditorOnlyFields)
 		{
-			builder.Append("\tUE_FIELD_NOTIFICATION_DECLARE_CLASS_DESCRIPTOR_BEGIN(").Append(this.PackageApi).Append(") \\\r\n");
+			builder.Append("\tUE_FIELD_NOTIFICATION_DECLARE_CLASS_DESCRIPTOR_BEGIN(").Append(PackageApi).Append(") \\\r\n");
 
 			//UE_FIELD_NOTIFICATION_DECLARE_FIELD
 			AppendFieldNotify(builder, classObj, hasProperties, hasFunctions, hasEditorFields, allEditorFields,
@@ -898,8 +898,8 @@ namespace EpicGames.UHT.Exporters.CodeGen
 			{
 				builder.Append("\t/** Standard constructor, called after all reflected properties have been initialized */ \\\r\n");
 
-				// Assume super class has OI constructor, this may not always be true but we should always be able to check this.
-				// In any case, it will default to old behavior before we even checked this.
+				// Assume super class has OI constructor, this may not always be true but we should always be able to check 
+				// In any case, it will default to old behavior before we even checked 
 				bool superClassObjectInitializerConstructorDeclared = true;
 				UhtClass? superClass = classObj.SuperClass;
 				if (superClass != null)
@@ -1022,7 +1022,7 @@ namespace EpicGames.UHT.Exporters.CodeGen
 			builder.Append("\tstatic void StaticRegisterNatives").Append(classObj.SourceName).Append("(); \\\r\n");
 			if (!classObj.ClassExportFlags.HasAnyFlags(UhtClassExportFlags.NoExport))
 			{
-				builder.Append("\tfriend struct ").Append(this.ObjectInfos[classObj.ObjectTypeIndex].RegisteredSingletonName).Append("_Statics; \\\r\n");
+				builder.Append("\tfriend struct ").Append(ObjectInfos[classObj.ObjectTypeIndex].RegisteredSingletonName).Append("_Statics; \\\r\n");
 			}
 			builder.Append("public: \\\r\n");
 
@@ -1049,7 +1049,7 @@ namespace EpicGames.UHT.Exporters.CodeGen
 			{
 				builder.Append("CASTCLASS_None");
 			}
-			builder.Append(", TEXT(\"").Append(this.Package.SourceName).Append("\"), ").Append(api[0..^1]).Append(") \\\r\n");
+			builder.Append(", TEXT(\"").Append(Package.SourceName).Append("\"), ").Append(api[0..^1]).Append(") \\\r\n");
 
 			builder.Append("\tDECLARE_SERIALIZER(").Append(classObj.SourceName).Append(") \\\r\n");
 
@@ -1158,7 +1158,7 @@ namespace EpicGames.UHT.Exporters.CodeGen
 
 			// we'll need a way to get to the UObject portion of a native interface, so that we can safely pass native interfaces
 			// to script VM functions
-			if (classObj.SuperClass != null && classObj.SuperClass.IsChildOf(this.Session.UInterface))
+			if (classObj.SuperClass != null && classObj.SuperClass.IsChildOf(Session.UInterface))
 			{
 				// Note: This used to be declared as a pure virtual function, but it was changed here in order to allow the Blueprint nativization process
 				// to detect C++ interface classes that explicitly declare pure virtual functions via type traits. This code will no longer trigger that check.
diff --git a/Engine/Source/Programs/Shared/EpicGames.UHT/Exporters/CodeGen/UhtMacroBlockEmitter.cs b/Engine/Source/Programs/Shared/EpicGames.UHT/Exporters/CodeGen/UhtMacroBlockEmitter.cs
--- a/Engine/Source/Programs/Shared/EpicGames.UHT/Exporters/CodeGen/UhtMacroBlockEmitter.cs
+++ b/Engine/Source/Programs/Shared/EpicGames.UHT/Exporters/CodeGen/UhtMacroBlockEmitter.cs
@@ -13,26 +13,26 @@ namespace EpicGames.UHT.Exporters.CodeGen
 
 		public UhtMacroBlockEmitter(StringBuilder builder, string macroName, bool initialState = false)
 		{
-			this._builder = builder;
-			this._macroName = macroName;
-			this._bEmitted = false;
+			_builder = builder;
+			_macroName = macroName;
+			_bEmitted = false;
 			Set(initialState);
 		}
 
 		public void Set(bool emit)
 		{
-			if (this._bEmitted == emit)
+			if (_bEmitted == emit)
 			{
 				return;
 			}
-			this._bEmitted = emit;
-			if (this._bEmitted)
+			_bEmitted = emit;
+			if (_bEmitted)
 			{
-				_builder.Append("#if ").Append(this._macroName).Append("\r\n");
+				_builder.Append("#if ").Append(_macroName).Append("\r\n");
 			}
 			else
 			{
-				_builder.Append("#endif // ").Append(this._macroName).Append("\r\n");
+				_builder.Append("#endif // ").Append(_macroName).Append("\r\n");
 			}
 		}
 
diff --git a/Engine/Source/Programs/Shared/EpicGames.UHT/Exporters/CodeGen/UhtMacroCreator.cs b/Engine/Source/Programs/Shared/EpicGames.UHT/Exporters/CodeGen/UhtMacroCreator.cs
--- a/Engine/Source/Programs/Shared/EpicGames.UHT/Exporters/CodeGen/UhtMacroCreator.cs
+++ b/Engine/Source/Programs/Shared/EpicGames.UHT/Exporters/CodeGen/UhtMacroCreator.cs
@@ -14,29 +14,29 @@ namespace EpicGames.UHT.Exporters.CodeGen
 		public UhtMacroCreator(StringBuilder builder, UhtHeaderCodeGenerator generator, int lineNumber, string macroSuffix)
 		{
 			builder.Append("#define ").AppendMacroName(generator, lineNumber, macroSuffix).Append(" \\\r\n");
-			this._builder = builder;
-			this._startingLength = builder.Length;
+			_builder = builder;
+			_startingLength = builder.Length;
 		}
 
 		public UhtMacroCreator(StringBuilder builder, UhtHeaderCodeGenerator generator, UhtClass classObj, string macroSuffix)
 		{
 			builder.Append("#define ").AppendMacroName(generator, classObj, macroSuffix).Append(" \\\r\n");
-			this._builder = builder;
-			this._startingLength = builder.Length;
+			_builder = builder;
+			_startingLength = builder.Length;
 		}
 
 		public UhtMacroCreator(StringBuilder builder, UhtHeaderCodeGenerator generator, UhtScriptStruct scriptStruct, string macroSuffix)
 		{
 			builder.Append("#define ").AppendMacroName(generator, scriptStruct, macroSuffix).Append(" \\\r\n");
-			this._builder = builder;
-			this._startingLength = builder.Length;
+			_builder = builder;
+			_startingLength = builder.Length;
 		}
 
 		public UhtMacroCreator(StringBuilder builder, UhtHeaderCodeGenerator generator, UhtFunction function, string macroSuffix)
 		{
 			builder.Append("#define ").AppendMacroName(generator, function, macroSuffix).Append(" \\\r\n");
-			this._builder = builder;
-			this._startingLength = builder.Length;
+			_builder = builder;
+			_startingLength = builder.Length;
 		}
 
 		public void Dispose()
diff --git a/Engine/Source/Programs/Shared/EpicGames.UHT/Exporters/CodeGen/UhtPackageCodeGenerator.cs b/Engine/Source/Programs/Shared/EpicGames.UHT/Exporters/CodeGen/UhtPackageCodeGenerator.cs
--- a/Engine/Source/Programs/Shared/EpicGames.UHT/Exporters/CodeGen/UhtPackageCodeGenerator.cs
+++ b/Engine/Source/Programs/Shared/EpicGames.UHT/Exporters/CodeGen/UhtPackageCodeGenerator.cs
@@ -25,19 +25,19 @@ namespace EpicGames.UHT.Exporters.CodeGen
 
 		public readonly UhtCodeGenerator CodeGenerator;
 		public readonly UhtPackage Package;
-		public bool SaveExportedHeaders => this.Package.Module.SaveExportedHeaders;
+		public bool SaveExportedHeaders => Package.Module.SaveExportedHeaders;
 
-		public Utils.UhtSession Session => this.CodeGenerator.Session;
-		public UhtCodeGenerator.PackageInfo[] PackageInfos => this.CodeGenerator.PackageInfos;
-		public UhtCodeGenerator.HeaderInfo[] HeaderInfos => this.CodeGenerator.HeaderInfos;
-		public UhtCodeGenerator.ObjectInfo[] ObjectInfos => this.CodeGenerator.ObjectInfos;
-		public string PackageApi => this.PackageInfos[this.Package.PackageTypeIndex].Api;
-		public string PackageSingletonName => this.ObjectInfos[this.Package.ObjectTypeIndex].RegisteredSingletonName;
+		public Utils.UhtSession Session => CodeGenerator.Session;
+		public UhtCodeGenerator.PackageInfo[] PackageInfos => CodeGenerator.PackageInfos;
+		public UhtCodeGenerator.HeaderInfo[] HeaderInfos => CodeGenerator.HeaderInfos;
+		public UhtCodeGenerator.ObjectInfo[] ObjectInfos => CodeGenerator.ObjectInfos;
+		public string PackageApi => PackageInfos[Package.PackageTypeIndex].Api;
+		public string PackageSingletonName => ObjectInfos[Package.ObjectTypeIndex].RegisteredSingletonName;
 
 		public UhtPackageCodeGenerator(UhtCodeGenerator codeGenerator, UhtPackage package)
 		{
-			this.CodeGenerator = codeGenerator;
-			this.Package = package;
+			CodeGenerator = codeGenerator;
+			Package = package;
 		}
 
 		#region Utility functions
@@ -50,7 +50,7 @@ namespace EpicGames.UHT.Exporters.CodeGen
 		/// <returns>Singleton name of "nullptr" if Object is null</returns>
 		public string GetSingletonName(UhtObject? obj, bool registered)
 		{
-			return this.CodeGenerator.GetSingletonName(obj, registered);
+			return CodeGenerator.GetSingletonName(obj, registered);
 		}
 
 		/// <summary>
@@ -61,7 +61,7 @@ namespace EpicGames.UHT.Exporters.CodeGen
 		/// <returns>External declaration</returns>
 		public string GetExternalDecl(UhtObject obj, bool registered)
 		{
-			return this.CodeGenerator.GetExternalDecl(obj, registered);
+			return CodeGenerator.GetExternalDecl(obj, registered);
 		}
 
 		/// <summary>
@@ -72,7 +72,7 @@ namespace EpicGames.UHT.Exporters.CodeGen
 		/// <returns>External declaration</returns>
 		public string GetExternalDecl(int objectIndex, bool registered)
 		{
-			return this.CodeGenerator.GetExternalDecl(objectIndex, registered);
+			return CodeGenerator.GetExternalDecl(objectIndex, registered);
 		}
 
 		/// <summary>
@@ -83,7 +83,7 @@ namespace EpicGames.UHT.Exporters.CodeGen
 		/// <returns>Cross reference</returns>
 		public string GetCrossReference(UhtObject obj, bool registered)
 		{
-			return this.CodeGenerator.GetCrossReference(obj, registered);
+			return CodeGenerator.GetCrossReference(obj, registered);
 		}
 
 		/// <summary>
@@ -94,7 +94,7 @@ namespace EpicGames.UHT.Exporters.CodeGen
 		/// <returns>Cross reference</returns>
 		public string GetCrossReference(int objectIndex, bool registered)
 		{
-			return this.CodeGenerator.GetCrossReference(objectIndex, registered);
+			return CodeGenerator.GetCrossReference(objectIndex, registered);
 		}
 
 		/// <summary>
diff --git a/Engine/Source/Programs/Shared/EpicGames.UHT/Exporters/CodeGen/UhtPackageCodeGeneratorCppFile.cs b/Engine/Source/Programs/Shared/EpicGames.UHT/Exporters/CodeGen/UhtPackageCodeGeneratorCppFile.cs
--- a/Engine/Source/Programs/Shared/EpicGames.UHT/Exporters/CodeGen/UhtPackageCodeGeneratorCppFile.cs
+++ b/Engine/Source/Programs/Shared/EpicGames.UHT/Exporters/CodeGen/UhtPackageCodeGeneratorCppFile.cs
@@ -39,7 +39,7 @@ namespace EpicGames.UHT.Exporters.CodeGen
 				uint bodiesHash = 0;
 				foreach (UhtHeaderFile headerFile in packageSortedHeaders)
 				{
-					ref UhtCodeGenerator.HeaderInfo headerInfo = ref this.HeaderInfos[headerFile.HeaderFileTypeIndex];
+					ref UhtCodeGenerator.HeaderInfo headerInfo = ref HeaderInfos[headerFile.HeaderFileTypeIndex];
 					ReadOnlyMemory<string> sorted = headerFile.References.Declaration.GetSortedReferences(
 						(int objectIndex, bool registered) => GetExternalDecl(objectIndex, registered));
 					foreach (string declaration in sorted.Span)
@@ -49,7 +49,7 @@ namespace EpicGames.UHT.Exporters.CodeGen
 
 					singletons.AddRange(headerFile.References.Singletons);
 
-					uint bodyHash = this.HeaderInfos[headerFile.HeaderFileTypeIndex].BodyHash;
+					uint bodyHash = HeaderInfos[headerFile.HeaderFileTypeIndex].BodyHash;
 					if (bodiesHash == 0)
 					{
 						// Don't combine in the first case because it keeps GUID backwards compatibility
@@ -68,8 +68,8 @@ namespace EpicGames.UHT.Exporters.CodeGen
 				}
 				uint declarationsHash = UhtHash.GenenerateTextHash(declarations.ToString());
 
-				string strippedName = this.PackageInfos[this.Package.PackageTypeIndex].StrippedName;
-				string singletonName = this.PackageSingletonName;
+				string strippedName = PackageInfos[Package.PackageTypeIndex].StrippedName;
+				string singletonName = PackageSingletonName;
 				singletons.Sort((UhtField lhs, UhtField rhs) =>
 				{
 					bool lhsIsDel = IsDelegateFunction(lhs);
@@ -79,16 +79,16 @@ namespace EpicGames.UHT.Exporters.CodeGen
 						return !lhsIsDel ? -1 : 1;
 					}
 					return StringComparerUE.OrdinalIgnoreCase.Compare(
-						this.ObjectInfos[lhs.ObjectTypeIndex].RegisteredSingletonName,
-						this.ObjectInfos[rhs.ObjectTypeIndex].RegisteredSingletonName);
+						ObjectInfos[lhs.ObjectTypeIndex].RegisteredSingletonName,
+						ObjectInfos[rhs.ObjectTypeIndex].RegisteredSingletonName);
 				});
 
 				builder.Append(HeaderCopyright);
 				builder.Append(RequiredCPPIncludes);
 				builder.Append(DisableDeprecationWarnings).Append("\r\n");
-				builder.Append("void EmptyLinkFunctionForGeneratedCode").Append(this.Package.ShortName).Append("_init() {}\r\n");
+				builder.Append("void EmptyLinkFunctionForGeneratedCode").Append(Package.ShortName).Append("_init() {}\r\n");
 
-				if (this.Session.IncludeDebugOutput)
+				if (Session.IncludeDebugOutput)
 				{
 					builder.Append("#if 0\r\n");
 					foreach (UhtHeaderFile headerFile in packageSortedHeaders)
@@ -101,11 +101,11 @@ namespace EpicGames.UHT.Exporters.CodeGen
 
 				foreach (UhtObject obj in singletons)
 				{
-					ref UhtCodeGenerator.ObjectInfo info = ref this.ObjectInfos[obj.ObjectTypeIndex];
+					ref UhtCodeGenerator.ObjectInfo info = ref ObjectInfos[obj.ObjectTypeIndex];
 					builder.Append(info.RegsiteredExternalDecl);
 				}
 
-				builder.AppendMetaDataDef(this.Package, null, MetaDataParamsName, 3);
+				builder.AppendMetaDataDef(Package, null, MetaDataParamsName, 3);
 
 				builder.Append("\tstatic FPackageRegistrationInfo Z_Registration_Info_UPackage_").Append(strippedName).Append(";\r\n");
 				builder.Append("\tFORCENOINLINE UPackage* ").Append(singletonName).Append("()\r\n");
@@ -118,20 +118,20 @@ namespace EpicGames.UHT.Exporters.CodeGen
 					builder.Append("\t\t\tstatic UObject* (*const SingletonFuncArray[])() = {\r\n");
 					foreach (UhtField field in singletons)
 					{
-						builder.Append("\t\t\t\t(UObject* (*)())").Append(this.ObjectInfos[field.ObjectTypeIndex].RegisteredSingletonName).Append(",\r\n");
+						builder.Append("\t\t\t\t(UObject* (*)())").Append(ObjectInfos[field.ObjectTypeIndex].RegisteredSingletonName).Append(",\r\n");
 					}
 					builder.Append("\t\t\t};\r\n");
 				}
 
-				EPackageFlags flags = this.Package.PackageFlags & (EPackageFlags.ClientOptional | EPackageFlags.ServerSideOnly | EPackageFlags.EditorOnly | EPackageFlags.Developer | EPackageFlags.UncookedOnly);
+				EPackageFlags flags = Package.PackageFlags & (EPackageFlags.ClientOptional | EPackageFlags.ServerSideOnly | EPackageFlags.EditorOnly | EPackageFlags.Developer | EPackageFlags.UncookedOnly);
 				builder.Append("\t\t\tstatic const UECodeGen_Private::FPackageParams PackageParams = {\r\n");
-				builder.Append("\t\t\t\t").AppendUTF8LiteralString(this.Package.SourceName).Append(",\r\n");
+				builder.Append("\t\t\t\t").AppendUTF8LiteralString(Package.SourceName).Append(",\r\n");
 				builder.Append("\t\t\t\t").Append(singletons.Count != 0 ? "SingletonFuncArray" : "nullptr").Append(",\r\n");
 				builder.Append("\t\t\t\t").Append(singletons.Count != 0 ? "UE_ARRAY_COUNT(SingletonFuncArray)" : "0").Append(",\r\n");
 				builder.Append("\t\t\t\tPKG_CompiledIn | ").Append($"0x{(uint)flags:X8}").Append(",\r\n");
 				builder.Append("\t\t\t\t").Append($"0x{bodiesHash:X8}").Append(",\r\n");
 				builder.Append("\t\t\t\t").Append($"0x{declarationsHash:X8}").Append(",\r\n");
-				builder.Append("\t\t\t\t").AppendMetaDataParams(this.Package, null, MetaDataParamsName).Append("\r\n");
+				builder.Append("\t\t\t\t").AppendMetaDataParams(Package, null, MetaDataParamsName).Append("\r\n");
 				builder.Append("\t\t\t};\r\n");
 				builder.Append("\t\t\tUECodeGen_Private::ConstructUPackage(Z_Registration_Info_UPackage_").Append(strippedName).Append(".OuterSingleton, PackageParams);\r\n");
 				builder.Append("\t\t}\r\n");
@@ -140,14 +140,14 @@ namespace EpicGames.UHT.Exporters.CodeGen
 
 				// Do not change the Z_CompiledInDeferPackage_UPackage_ without changing LC_SymbolPatterns
 				builder.Append("\tstatic FRegisterCompiledInInfo Z_CompiledInDeferPackage_UPackage_").Append(strippedName).Append('(').Append(singletonName)
-					.Append(", TEXT(\"").Append(this.Package.SourceName).Append("\"), Z_Registration_Info_UPackage_").Append(strippedName).Append(", CONSTRUCT_RELOAD_VERSION_INFO(FPackageReloadVersionInfo, ")
+					.Append(", TEXT(\"").Append(Package.SourceName).Append("\"), Z_Registration_Info_UPackage_").Append(strippedName).Append(", CONSTRUCT_RELOAD_VERSION_INFO(FPackageReloadVersionInfo, ")
 					.Append($"0x{bodiesHash:X8}, 0x{declarationsHash:X8}").Append("));\r\n");
 
 				builder.Append(EnableDeprecationWarnings).Append("\r\n");
 
-				if (this.SaveExportedHeaders)
+				if (SaveExportedHeaders)
 				{
-					string cppFilePath = factory.MakePath(this.Package, ".init.gen.cpp");
+					string cppFilePath = factory.MakePath(Package, ".init.gen.cpp");
 					factory.CommitOutput(cppFilePath, builder);
 				}
 			}
diff --git a/Engine/Source/Programs/Shared/EpicGames.UHT/Exporters/CodeGen/UhtPackageCodeGeneratorHFile.cs b/Engine/Source/Programs/Shared/EpicGames.UHT/Exporters/CodeGen/UhtPackageCodeGeneratorHFile.cs
--- a/Engine/Source/Programs/Shared/EpicGames.UHT/Exporters/CodeGen/UhtPackageCodeGeneratorHFile.cs
+++ b/Engine/Source/Programs/Shared/EpicGames.UHT/Exporters/CodeGen/UhtPackageCodeGeneratorHFile.cs
@@ -37,10 +37,10 @@ namespace EpicGames.UHT.Exporters.CodeGen
 				builder.Append("\r\n");
 				builder.Append("\r\n");
 
-				List<UhtHeaderFile> headerFiles = new(this.Package.Children.Count * 2);
+				List<UhtHeaderFile> headerFiles = new(Package.Children.Count * 2);
 				headerFiles.AddRange(packageSortedHeaders);
 
-				foreach (UhtHeaderFile headerFile in this.Package.Children)
+				foreach (UhtHeaderFile headerFile in Package.Children)
 				{
 					if (headerFile.HeaderFileType == UhtHeaderFileType.Classes)
 					{
@@ -55,15 +55,15 @@ namespace EpicGames.UHT.Exporters.CodeGen
 				{
 					if (headerFile.HeaderFileType == UhtHeaderFileType.Classes)
 					{
-						builder.Append("#include \"").Append(this.HeaderInfos[headerFile.HeaderFileTypeIndex].IncludePath).Append("\"\r\n");
+						builder.Append("#include \"").Append(HeaderInfos[headerFile.HeaderFileTypeIndex].IncludePath).Append("\"\r\n");
 					}
 				}
 
 				builder.Append("\r\n");
 
-				if (this.SaveExportedHeaders)
+				if (SaveExportedHeaders)
 				{
-					string headerFilePath = factory.MakePath(this.Package, "Classes.h");
+					string headerFilePath = factory.MakePath(Package, "Classes.h");
 					factory.CommitOutput(headerFilePath, builder);
 				}
 			}
diff --git a/Engine/Source/Programs/Shared/EpicGames.UHT/Exporters/Json/UhtJsonExporter.cs b/Engine/Source/Programs/Shared/EpicGames.UHT/Exporters/Json/UhtJsonExporter.cs
--- a/Engine/Source/Programs/Shared/EpicGames.UHT/Exporters/Json/UhtJsonExporter.cs
+++ b/Engine/Source/Programs/Shared/EpicGames.UHT/Exporters/Json/UhtJsonExporter.cs
@@ -26,14 +26,14 @@ namespace EpicGames.UHT.Exporters.Json
 
 		private UhtJsonExporter(IUhtExportFactory factory)
 		{
-			this.Factory = factory;
+			Factory = factory;
 		}
 
 		private void Export()
 		{
 			// Generate the files for the packages
-			List<Task?> generatedPackages = new(this.Session.PackageTypeCount);
-			foreach (UhtPackage package in this.Session.Packages)
+			List<Task?> generatedPackages = new(Session.PackageTypeCount);
+			foreach (UhtPackage package in Session.Packages)
 			{
 				UHTManifest.Module module = package.Module;
 				generatedPackages.Add(Factory.CreateTask(
@@ -46,7 +46,7 @@ namespace EpicGames.UHT.Exporters.Json
 			}
 
 			// Wait for all the packages to complete
-			List<Task> packageTasks = new(this.Session.PackageTypeCount);
+			List<Task> packageTasks = new(Session.PackageTypeCount);
 			foreach (Task? output in generatedPackages)
 			{
 				if (output != null)
diff --git a/Engine/Source/Programs/Shared/EpicGames.UHT/Parsers/UhtFunctionParser.cs b/Engine/Source/Programs/Shared/EpicGames.UHT/Parsers/UhtFunctionParser.cs
--- a/Engine/Source/Programs/Shared/EpicGames.UHT/Parsers/UhtFunctionParser.cs
+++ b/Engine/Source/Programs/Shared/EpicGames.UHT/Parsers/UhtFunctionParser.cs
@@ -27,22 +27,22 @@ namespace EpicGames.UHT.Parsers
 
 		public UhtAdvancedDisplayParameterHandler(UhtMetaData metaData)
 		{
-			this._metaData = metaData;
-			this._parameterNames = null;
-			this._numberLeaveUnmarked = -1;
-			this._alreadyLeft = 0;
-			this._bUseNumber = false;
+			_metaData = metaData;
+			_parameterNames = null;
+			_numberLeaveUnmarked = -1;
+			_alreadyLeft = 0;
+			_bUseNumber = false;
 
-			if (this._metaData.TryGetValue(UhtNames.AdvancedDisplay, out string? foundString))
+			if (_metaData.TryGetValue(UhtNames.AdvancedDisplay, out string? foundString))
 			{
-				this._parameterNames = foundString.ToString().Split(',', StringSplitOptions.RemoveEmptyEntries);
-				for (int index = 0, endIndex = this._parameterNames.Length; index < endIndex; ++index)
+				_parameterNames = foundString.ToString().Split(',', StringSplitOptions.RemoveEmptyEntries);
+				for (int index = 0, endIndex = _parameterNames.Length; index < endIndex; ++index)
 				{
-					this._parameterNames[index] = this._parameterNames[index].Trim();
+					_parameterNames[index] = _parameterNames[index].Trim();
 				}
-				if (this._parameterNames.Length == 1)
+				if (_parameterNames.Length == 1)
 				{
-					_bUseNumber = Int32.TryParse(this._parameterNames[0], out _numberLeaveUnmarked);
+					_bUseNumber = Int32.TryParse(_parameterNames[0], out _numberLeaveUnmarked);
 				}
 			}
 		}
@@ -53,26 +53,26 @@ namespace EpicGames.UHT.Parsers
 		 */
 		public bool ShouldMarkParameter(StringView parameterName)
 		{
-			if (this._bUseNumber)
+			if (_bUseNumber)
 			{
-				if (this._numberLeaveUnmarked < 0)
+				if (_numberLeaveUnmarked < 0)
 				{
 					return false;
 				}
-				if (this._alreadyLeft < this._numberLeaveUnmarked)
+				if (_alreadyLeft < _numberLeaveUnmarked)
 				{
-					++this._alreadyLeft;
+					++_alreadyLeft;
 					return false;
 				}
 				return true;
 			}
 
-			if (this._parameterNames == null)
+			if (_parameterNames == null)
 			{
 				return false;
 			}
 
-			foreach (string element in this._parameterNames)
+			foreach (string element in _parameterNames)
 			{
 				if (parameterName.Span.Equals(element, StringComparison.OrdinalIgnoreCase))
 				{
@@ -85,7 +85,7 @@ namespace EpicGames.UHT.Parsers
 		/** return if more parameters can be marked */
 		public bool CanMarkMore()
 		{
-			return this._bUseNumber ? this._numberLeaveUnmarked > 0 : (this._parameterNames != null && this._parameterNames.Length > 0);
+			return _bUseNumber ? _numberLeaveUnmarked > 0 : (_parameterNames != null && _parameterNames.Length > 0);
 		}
 	}
 
diff --git a/Engine/Source/Programs/Shared/EpicGames.UHT/Parsers/UhtHeaderFileParser.cs b/Engine/Source/Programs/Shared/EpicGames.UHT/Parsers/UhtHeaderFileParser.cs
--- a/Engine/Source/Programs/Shared/EpicGames.UHT/Parsers/UhtHeaderFileParser.cs
+++ b/Engine/Source/Programs/Shared/EpicGames.UHT/Parsers/UhtHeaderFileParser.cs
@@ -196,7 +196,7 @@ namespace EpicGames.UHT.Parsers
 		/// <summary>
 		/// If true, this header file belongs to the engine
 		/// </summary>
-		public bool IsPartOfEngine => this.HeaderFile.Package.IsPartOfEngine;
+		public bool IsPartOfEngine => HeaderFile.Package.IsPartOfEngine;
 
 		/// <summary>
 		/// If true, the inclusion of the generated header file was seen
@@ -259,9 +259,9 @@ namespace EpicGames.UHT.Parsers
 
 		private UhtHeaderFileParser(UhtHeaderFile headerFile)
 		{
-			this.TokenReader = new UhtTokenBufferReader(headerFile, headerFile.Data.Memory);
-			this.HeaderFile = headerFile;
-			this.TokenReader.TokenPreprocessor = this;
+			TokenReader = new UhtTokenBufferReader(headerFile, headerFile.Data.Memory);
+			HeaderFile = headerFile;
+			TokenReader.TokenPreprocessor = this;
 		}
 
 		/// <summary>
@@ -271,11 +271,11 @@ namespace EpicGames.UHT.Parsers
 		/// <exception cref="UhtIceException">Throw if the new scope isn't parented by the current scope</exception>
 		public void PushScope(UhtParsingScope scope)
 		{
-			if (scope.ParentScope != this._topScope)
+			if (scope.ParentScope != _topScope)
 			{
 				throw new UhtIceException("Pushing a new scope whose parent isn't the current top scope.");
 			}
-			this._topScope = scope;
+			_topScope = scope;
 		}
 
 		/// <summary>
@@ -285,11 +285,11 @@ namespace EpicGames.UHT.Parsers
 		/// <exception cref="UhtIceException">Thrown if the given scope isn't the top scope</exception>
 		public void PopScope(UhtParsingScope scope)
 		{
-			if (scope != this._topScope)
+			if (scope != _topScope)
 			{
 				throw new UhtIceException("Attempt to pop a scope that isn't the top scope");
 			}
-			this._topScope = scope.ParentScope;
+			_topScope = scope.ParentScope;
 		}
 
 		/// <summary>
@@ -298,11 +298,11 @@ namespace EpicGames.UHT.Parsers
 		/// <returns>Property parser</returns>
 		public UhtPropertyParser GetCachedPropertyParser()
 		{
-			if (this.PropertyParser == null)
+			if (PropertyParser == null)
 			{
-				this.PropertyParser = new();
+				PropertyParser = new();
 			}
-			return this.PropertyParser;
+			return PropertyParser;
 		}
 
 		/// <summary>
@@ -311,7 +311,7 @@ namespace EpicGames.UHT.Parsers
 		/// <returns>Enumeration flags for all active compiler directives</returns>
 		public UhtCompilerDirective GetCurrentCompositeCompilerDirective()
 		{
-			return this._compilerDirectives.Count > 0 ? this._compilerDirectives[^1]._composite : UhtCompilerDirective.None;
+			return _compilerDirectives.Count > 0 ? _compilerDirectives[^1]._composite : UhtCompilerDirective.None;
 		}
 
 		/// <summary>
@@ -320,7 +320,7 @@ namespace EpicGames.UHT.Parsers
 		/// <returns>Current compiler directive</returns>
 		public UhtCompilerDirective GetCurrentNonCompositeCompilerDirective()
 		{
-			return this._compilerDirectives.Count > 0 ? this._compilerDirectives[^1]._element : UhtCompilerDirective.None;
+			return _compilerDirectives.Count > 0 ? _compilerDirectives[^1]._element : UhtCompilerDirective.None;
 		}
 
 		#region ITokenPreprocessor implementation
@@ -339,15 +339,15 @@ namespace EpicGames.UHT.Parsers
 		/// <inheritdoc/>
 		public void SaveState()
 		{
-			this._savedCompilerDirectives.Clear();
-			this._savedCompilerDirectives.AddRange(this._compilerDirectives);
+			_savedCompilerDirectives.Clear();
+			_savedCompilerDirectives.AddRange(_compilerDirectives);
 		}
 
 		/// <inheritdoc/>
 		public void RestoreState()
 		{
-			this._compilerDirectives.Clear();
-			this._compilerDirectives.AddRange(this._savedCompilerDirectives);
+			_compilerDirectives.Clear();
+			_compilerDirectives.AddRange(_savedCompilerDirectives);
 		}
 		#endregion
 
@@ -369,21 +369,21 @@ namespace EpicGames.UHT.Parsers
 		/// <param name="logUnhandledKeywords">If true, log any unhandled keywords</param>
 		public void ParseStatements(char initiator, char terminator, bool logUnhandledKeywords)
 		{
-			if (this._topScope == null)
+			if (_topScope == null)
 			{
 				return;
 			}
 
 			if (initiator != 0)
 			{
-				this.TokenReader.Require(initiator);
+				TokenReader.Require(initiator);
 			}
 			while (true)
 			{
-				UhtToken token = this.TokenReader.GetToken();
+				UhtToken token = TokenReader.GetToken();
 				if (token.TokenType.IsEndType())
 				{
-					if (this._topScope != null && this._topScope.ParentScope == null)
+					if (_topScope != null && _topScope.ParentScope == null)
 					{
 						CheckEof(ref token);
 					}
@@ -394,10 +394,10 @@ namespace EpicGames.UHT.Parsers
 					return;
 				}
 
-				if (this._topScope != null)
+				if (_topScope != null)
 				{
-					ParseStatement(this._topScope, ref token, logUnhandledKeywords);
-					this.TokenReader.ClearComments();
+					ParseStatement(_topScope, ref token, logUnhandledKeywords);
+					TokenReader.ClearComments();
 				}
 			}
 		}
@@ -569,9 +569,9 @@ namespace EpicGames.UHT.Parsers
 
 		private void CheckEof(ref UhtToken token)
 		{
-			if (this._compilerDirectives.Count > 0)
+			if (_compilerDirectives.Count > 0)
 			{
-				throw new UhtException(this.TokenReader, token.InputLine, "Missing #endif");
+				throw new UhtException(TokenReader, token.InputLine, "Missing #endif");
 			}
 		}
 		#endregion
@@ -589,11 +589,11 @@ namespace EpicGames.UHT.Parsers
 			// Collect all the lines of the preprocessor statement including any continuations.
 			// We assume that the vast majority of lines will not be continuations.  So avoid using the
 			// string builder as much as possible.
-			int startingLine = this.TokenReader.InputLine;
+			int startingLine = TokenReader.InputLine;
 			StringViewBuilder builder = new();
 			while (true)
 			{
-				UhtToken lineToken = this.TokenReader.GetLine();
+				UhtToken lineToken = TokenReader.GetLine();
 				if (lineToken.TokenType != UhtTokenType.Line)
 				{
 					break;
@@ -611,14 +611,14 @@ namespace EpicGames.UHT.Parsers
 			StringView line = builder.ToStringView();
 
 			// Create a token reader we will use to decode
-			UhtTokenBufferReader lineTokenReader = new(this.HeaderFile, line.Memory);
+			UhtTokenBufferReader lineTokenReader = new(HeaderFile, line.Memory);
 			lineTokenReader.InputLine = startingLine;
 
 			if (!lineTokenReader.TryOptionalIdentifier(out UhtToken directive))
 			{
 				if (isBeingIncluded)
 				{
-					throw new UhtException(this.TokenReader, directive.InputLine, "Missing compiler directive after '#'");
+					throw new UhtException(TokenReader, directive.InputLine, "Missing compiler directive after '#'");
 				}
 				return checkClearComments;
 			}
@@ -627,7 +627,7 @@ namespace EpicGames.UHT.Parsers
 			{
 				if (isBeingIncluded)
 				{
-					throw new UhtException(this.TokenReader, directive.InputLine, "#error directive encountered");
+					throw new UhtException(TokenReader, directive.InputLine, "#error directive encountered");
 				}
 			}
 			else if (directive.IsValue("pragma"))
@@ -638,9 +638,9 @@ namespace EpicGames.UHT.Parsers
 			{
 				if (!lineTokenReader.TryOptionalConstInt(out int newInputLine))
 				{
-					throw new UhtException(this.TokenReader, directive.InputLine, "Missing line number in line number directive");
+					throw new UhtException(TokenReader, directive.InputLine, "Missing line number in line number directive");
 				}
-				this.TokenReader.InputLine = newInputLine;
+				TokenReader.InputLine = newInputLine;
 			}
 			else if (directive.IsValue("include"))
 			{
@@ -649,23 +649,23 @@ namespace EpicGames.UHT.Parsers
 					UhtToken includeName = lineTokenReader.GetToken();
 					if (includeName.IsConstString())
 					{
-						StringView includeNameString = includeName.GetUnescapedString(this.HeaderFile);
-						if (this.SpottedAutogeneratedHeaderInclude)
+						StringView includeNameString = includeName.GetUnescapedString(HeaderFile);
+						if (SpottedAutogeneratedHeaderInclude)
 						{
-							this.HeaderFile.LogError("#include found after .generated.h file - the .generated.h file should always be the last #include in a header");
+							HeaderFile.LogError("#include found after .generated.h file - the .generated.h file should always be the last #include in a header");
 						}
-						if (this.HeaderFile.GeneratedHeaderFileName.AsSpan().Equals(includeNameString.Span, StringComparison.OrdinalIgnoreCase))
+						if (HeaderFile.GeneratedHeaderFileName.AsSpan().Equals(includeNameString.Span, StringComparison.OrdinalIgnoreCase))
 						{
-							this.SpottedAutogeneratedHeaderInclude = true;
+							SpottedAutogeneratedHeaderInclude = true;
 						}
 						if (includeNameString.Span.Contains(".generated.h", StringComparison.Ordinal))
 						{
 							string refName = Path.GetFileName(includeNameString.ToString()).Replace(".generated.h", ".h", StringComparison.Ordinal);
-							this.HeaderFile.AddReferencedHeader(refName, true);
+							HeaderFile.AddReferencedHeader(refName, true);
 						}
 						else
 						{
-							this.HeaderFile.AddReferencedHeader(includeNameString.ToString(), true);
+							HeaderFile.AddReferencedHeader(includeNameString.ToString(), true);
 						}
 					}
 				}
@@ -687,7 +687,7 @@ namespace EpicGames.UHT.Parsers
 				UhtCompilerDirective newCompilerDirective = ParserConditional(lineTokenReader);
 				if (SupportsElif(oldCompilerDirective) != SupportsElif(newCompilerDirective))
 				{
-					throw new UhtException(this.TokenReader, directive.InputLine,
+					throw new UhtException(TokenReader, directive.InputLine,
 						$"Mixing {GetCompilerDirectiveText(oldCompilerDirective)} with {GetCompilerDirectiveText(newCompilerDirective)} in an #elif preprocessor block is not supported");
 				}
 				PushCompilerDirective(newCompilerDirective);
@@ -711,7 +711,7 @@ namespace EpicGames.UHT.Parsers
 						PushCompilerDirective(UhtCompilerDirective.NotCPPBlock);
 						break;
 					case UhtCompilerDirective.WithHotReload:
-						throw new UhtException(this.TokenReader, directive.InputLine, "Can not use WITH_HOT_RELOAD with an #else clause");
+						throw new UhtException(TokenReader, directive.InputLine, "Can not use WITH_HOT_RELOAD with an #else clause");
 					default:
 						PushCompilerDirective(oldCompilerDirective);
 						break;
@@ -731,7 +731,7 @@ namespace EpicGames.UHT.Parsers
 			{
 				if (isBeingIncluded)
 				{
-					throw new UhtException(this.TokenReader, directive.InputLine, $"Unrecognized compiler directive {directive.Value}");
+					throw new UhtException(TokenReader, directive.InputLine, $"Unrecognized compiler directive {directive.Value}");
 				}
 			}
 			return checkClearComments;
@@ -795,7 +795,7 @@ namespace EpicGames.UHT.Parsers
 				case UhtTokenType.EndOfDefault:
 				case UhtTokenType.EndOfType:
 				case UhtTokenType.EndOfDeclaration:
-					throw new UhtException(this.TokenReader, define.InputLine, "#if with no expression");
+					throw new UhtException(TokenReader, define.InputLine, "#if with no expression");
 			}
 
 			return UhtCompilerDirective.Unrecognized;
@@ -810,7 +810,7 @@ namespace EpicGames.UHT.Parsers
 			CompilerDirective newCompileDirective = new();
 			newCompileDirective._element = compilerDirective;
 			newCompileDirective._composite = GetCurrentCompositeCompilerDirective() | compilerDirective;
-			this._compilerDirectives.Add(newCompileDirective);
+			_compilerDirectives.Add(newCompileDirective);
 		}
 
 		/// <summary>
@@ -818,18 +818,18 @@ namespace EpicGames.UHT.Parsers
 		/// </summary>
 		private UhtCompilerDirective PopCompilerDirective(UhtToken token)
 		{
-			if (this._compilerDirectives.Count == 0)
+			if (_compilerDirectives.Count == 0)
 			{
-				throw new UhtException(this.TokenReader, token.InputLine, $"Unmatched '#{token.Value}'");
+				throw new UhtException(TokenReader, token.InputLine, $"Unmatched '#{token.Value}'");
 			}
-			UhtCompilerDirective compilerDirective = this._compilerDirectives[^1]._element;
-			this._compilerDirectives.RemoveAt(this._compilerDirectives.Count - 1);
+			UhtCompilerDirective compilerDirective = _compilerDirectives[^1]._element;
+			_compilerDirectives.RemoveAt(_compilerDirectives.Count - 1);
 			return compilerDirective;
 		}
 
 		private bool IncludeCurrentCompilerDirective()
 		{
-			if (this._compilerDirectives.Count == 0)
+			if (_compilerDirectives.Count == 0)
 			{
 				return true;
 			}
@@ -846,11 +846,11 @@ namespace EpicGames.UHT.Parsers
 		/// <exception cref="UhtIceException"></exception>
 		private bool ClearCommentsCompilerDirective()
 		{
-			if (this._compilerDirectives.Count == 0)
+			if (_compilerDirectives.Count == 0)
 			{
 				return true;
 			}
-			UhtCompilerDirective compilerDirective = this._compilerDirectives[^1]._element;
+			UhtCompilerDirective compilerDirective = _compilerDirectives[^1]._element;
 			switch (compilerDirective)
 			{
 				case UhtCompilerDirective.CPPBlock:
@@ -908,7 +908,7 @@ namespace EpicGames.UHT.Parsers
 
 		private bool CheckForConstructor(UhtClass classObj, UhtDeclaration declaration)
 		{
-			IUhtTokenReader replayReader = UhtTokenReplayReader.GetThreadInstance(this.HeaderFile, this.HeaderFile.Data.Memory, new ReadOnlyMemory<UhtToken>(declaration.Tokens), UhtTokenType.EndOfDeclaration);
+			IUhtTokenReader replayReader = UhtTokenReplayReader.GetThreadInstance(HeaderFile, HeaderFile.Data.Memory, new ReadOnlyMemory<UhtToken>(declaration.Tokens), UhtTokenType.EndOfDeclaration);
 
 			// Allow explicit constructors
 			{
@@ -1021,7 +1021,7 @@ namespace EpicGames.UHT.Parsers
 
 		bool CheckForSerialize(UhtClass classObj, UhtDeclaration declaration)
 		{
-			IUhtTokenReader replayReader = UhtTokenReplayReader.GetThreadInstance(this.HeaderFile, this.HeaderFile.Data.Memory, new ReadOnlyMemory<UhtToken>(declaration.Tokens), UhtTokenType.EndOfDeclaration);
+			IUhtTokenReader replayReader = UhtTokenReplayReader.GetThreadInstance(HeaderFile, HeaderFile.Data.Memory, new ReadOnlyMemory<UhtToken>(declaration.Tokens), UhtTokenType.EndOfDeclaration);
 
 			SkipVirtualAndAPI(replayReader);
 
diff --git a/Engine/Source/Programs/Shared/EpicGames.UHT/Parsers/UhtParsingScope.cs b/Engine/Source/Programs/Shared/EpicGames.UHT/Parsers/UhtParsingScope.cs
--- a/Engine/Source/Programs/Shared/EpicGames.UHT/Parsers/UhtParsingScope.cs
+++ b/Engine/Source/Programs/Shared/EpicGames.UHT/Parsers/UhtParsingScope.cs
@@ -51,7 +51,7 @@ namespace EpicGames.UHT.Parsers
 		/// <summary>
 		/// Current session
 		/// </summary>
-		public UhtSession Session => this.ScopeType.Session;
+		public UhtSession Session => ScopeType.Session;
 
 		/// <summary>
 		/// Return the current class scope being compiled
@@ -76,7 +76,7 @@ namespace EpicGames.UHT.Parsers
 		/// <summary>
 		/// Return the current class being compiled
 		/// </summary>
-		public UhtClass CurrentClass => (UhtClass)this.CurrentClassScope.ScopeType;
+		public UhtClass CurrentClass => (UhtClass)CurrentClassScope.ScopeType;
 
 		/// <summary>
 		/// Construct a root/global scope
@@ -86,12 +86,12 @@ namespace EpicGames.UHT.Parsers
 		/// <param name="keywordTable">Keyword table</param>
 		public UhtParsingScope(UhtHeaderFileParser headerParser, UhtType scopeType, UhtKeywordTable keywordTable)
 		{
-			this.HeaderParser = headerParser;
-			this.TokenReader = headerParser.TokenReader;
-			this.ParentScope = null;
-			this.ScopeType = scopeType;
-			this.ScopeKeywordTable = keywordTable;
-			this.HeaderParser.PushScope(this);
+			HeaderParser = headerParser;
+			TokenReader = headerParser.TokenReader;
+			ParentScope = null;
+			ScopeType = scopeType;
+			ScopeKeywordTable = keywordTable;
+			HeaderParser.PushScope(this);
 		}
 
 		/// <summary>
@@ -103,13 +103,13 @@ namespace EpicGames.UHT.Parsers
 		/// <param name="accessSpecifier">Current access specifier</param>
 		public UhtParsingScope(UhtParsingScope parentScope, UhtType scopeType, UhtKeywordTable keywordTable, UhtAccessSpecifier accessSpecifier)
 		{
-			this.HeaderParser = parentScope.HeaderParser;
-			this.TokenReader = parentScope.TokenReader;
-			this.ParentScope = parentScope;
-			this.ScopeType = scopeType;
-			this.ScopeKeywordTable = keywordTable;
-			this.AccessSpecifier = accessSpecifier;
-			this.HeaderParser.PushScope(this);
+			HeaderParser = parentScope.HeaderParser;
+			TokenReader = parentScope.TokenReader;
+			ParentScope = parentScope;
+			ScopeType = scopeType;
+			ScopeKeywordTable = keywordTable;
+			AccessSpecifier = accessSpecifier;
+			HeaderParser.PushScope(this);
 		}
 
 		/// <summary>
@@ -129,7 +129,7 @@ namespace EpicGames.UHT.Parsers
 		{
 			if (disposing)
 			{
-				this.HeaderParser.PopScope(this);
+				HeaderParser.PopScope(this);
 			}
 		}
 
@@ -138,7 +138,7 @@ namespace EpicGames.UHT.Parsers
 		/// </summary>
 		public void AddModuleRelativePathToMetaData()
 		{
-			AddModuleRelativePathToMetaData(this.ScopeType.MetaData, this.ScopeType.HeaderFile);
+			AddModuleRelativePathToMetaData(ScopeType.MetaData, ScopeType.HeaderFile);
 		}
 
 		/// <summary>
@@ -157,7 +157,7 @@ namespace EpicGames.UHT.Parsers
 		/// <param name="metaNameIndex">Index for the meta data key.  This is used for enum values</param>
 		public void AddFormattedCommentsAsTooltipMetaData(int metaNameIndex = UhtMetaData.IndexNone)
 		{
-			AddFormattedCommentsAsTooltipMetaData(this.ScopeType, metaNameIndex);
+			AddFormattedCommentsAsTooltipMetaData(ScopeType, metaNameIndex);
 		}
 
 		/// <summary>
@@ -175,7 +175,7 @@ namespace EpicGames.UHT.Parsers
 			}
 
 			// Fetch the comments
-			ReadOnlySpan<StringView> comments = this.TokenReader.Comments;
+			ReadOnlySpan<StringView> comments = TokenReader.Comments;
 
 			// If we don't have any comments, just return
 			if (comments.Length == 0)
@@ -209,12 +209,12 @@ namespace EpicGames.UHT.Parsers
 				type.MetaData.Add(UhtNames.ToolTip, metaNameIndex, toolTip);
 
 				//COMPATIBILITY-TODO - Old UHT would only clear the comments if there was some form of a tooltip
-				this.TokenReader.ClearComments();
+				TokenReader.ClearComments();
 			}
 
 			//COMPATIBILITY-TODO
 			// Clear the comments since they have been consumed
-			//this.TokenReader.ClearComments();
+			//TokenReader.ClearComments();
 		}
 
 		// We consider any alpha/digit or code point > 0xFF as a valid comment char
@@ -608,14 +608,14 @@ namespace EpicGames.UHT.Parsers
 		/// <param name="initialToken">Initial toke nto add to the recorder</param>
 		public UhtTokenRecorder(UhtParsingScope scope, ref UhtToken initialToken)
 		{
-			this._scope = scope;
-			this._compilerDirective = this._scope.HeaderParser.GetCurrentCompositeCompilerDirective();
-			this._function = null;
-			this._flushed = false;
-			if (this._scope.ScopeType is UhtClass)
+			_scope = scope;
+			_compilerDirective = _scope.HeaderParser.GetCurrentCompositeCompilerDirective();
+			_function = null;
+			_flushed = false;
+			if (_scope.ScopeType is UhtClass)
 			{
-				this._scope.TokenReader.EnableRecording();
-				this._scope.TokenReader.RecordToken(ref initialToken);
+				_scope.TokenReader.EnableRecording();
+				_scope.TokenReader.RecordToken(ref initialToken);
 			}
 		}
 
@@ -626,13 +626,13 @@ namespace EpicGames.UHT.Parsers
 		/// <param name="function">Function associated with the recorder</param>
 		public UhtTokenRecorder(UhtParsingScope scope, UhtFunction function)
 		{
-			this._scope = scope;
-			this._compilerDirective = this._scope.HeaderParser.GetCurrentCompositeCompilerDirective();
-			this._function = function;
-			this._flushed = false;
-			if (this._scope.ScopeType is UhtClass)
+			_scope = scope;
+			_compilerDirective = _scope.HeaderParser.GetCurrentCompositeCompilerDirective();
+			_function = function;
+			_flushed = false;
+			if (_scope.ScopeType is UhtClass)
 			{
-				this._scope.TokenReader.EnableRecording();
+				_scope.TokenReader.EnableRecording();
 			}
 		}
 
@@ -653,10 +653,10 @@ namespace EpicGames.UHT.Parsers
 			if (!_flushed)
 			{
 				_flushed = true;
-				if (this._scope.ScopeType is UhtClass classObj)
+				if (_scope.ScopeType is UhtClass classObj)
 				{
-					classObj.AddDeclaration(this._compilerDirective, this._scope.TokenReader.RecordedTokens, this._function);
-					this._scope.TokenReader.DisableRecording();
+					classObj.AddDeclaration(_compilerDirective, _scope.TokenReader.RecordedTokens, _function);
+					_scope.TokenReader.DisableRecording();
 					return true;
 				}
 			}
diff --git a/Engine/Source/Programs/Shared/EpicGames.UHT/Parsers/UhtPropertyParser.cs b/Engine/Source/Programs/Shared/EpicGames.UHT/Parsers/UhtPropertyParser.cs
--- a/Engine/Source/Programs/Shared/EpicGames.UHT/Parsers/UhtPropertyParser.cs
+++ b/Engine/Source/Programs/Shared/EpicGames.UHT/Parsers/UhtPropertyParser.cs
@@ -393,8 +393,8 @@ namespace EpicGames.UHT.Parsers
 		/// <param name="typeTokens">Type tokens</param>
 		public UhtPreResolveProperty(UhtPropertySettings propertySettings, ReadOnlyMemory<UhtToken> typeTokens) : base(propertySettings)
 		{
-			this.TypeTokens = typeTokens;
-			this.PropertySettings = propertySettings;
+			TypeTokens = typeTokens;
+			PropertySettings = propertySettings;
 		}
 
 		/// <inheritdoc/>
@@ -472,7 +472,7 @@ namespace EpicGames.UHT.Parsers
 		/// </summary>
 		public UhtPropertyParser()
 		{
-			this._gatherTypeTokensDelegate = GatherTypeTokens;
+			_gatherTypeTokensDelegate = GatherTypeTokens;
 		}
 
 		/// <summary>
@@ -502,11 +502,11 @@ namespace EpicGames.UHT.Parsers
 			specifierContext.SeenBlueprintGetterSpecifier = false;
 
 			// Initialize the settings
-			this._options = options;
-			this._category = category;
-			this._currentTokenReader = topScope.TokenReader;
-			this._currentTypeTokens = new List<UhtToken>();
-			this._currentTemplateDepth = 0;
+			_options = options;
+			_category = category;
+			_currentTokenReader = topScope.TokenReader;
+			_currentTypeTokens = new List<UhtToken>();
+			_currentTemplateDepth = 0;
 
 			using UhtMessageContext tokenContext = new(this);
 			ParseInternal(topScope, specifierContext, propertyDelegate);
@@ -520,7 +520,7 @@ namespace EpicGames.UHT.Parsers
 			get
 			{
 				Stack<object?> extraContext = new(1);
-				extraContext.Push(this._category.GetHintText());
+				extraContext.Push(_category.GetHintText());
 				return extraContext;
 			}
 		}
@@ -820,7 +820,7 @@ namespace EpicGames.UHT.Parsers
 			// Handle MemoryLayout.h macros
 			bool hasWrapperBrackets = false;
 			UhtLayoutMacroType layoutMacroType = UhtLayoutMacroType.None;
-			if (this._options.HasAnyFlags(UhtPropertyParseOptions.ParseLayoutMacro))
+			if (_options.HasAnyFlags(UhtPropertyParseOptions.ParseLayoutMacro))
 			{
 				ref UhtToken layoutToken = ref tokenReader.PeekToken();
 				if (layoutToken.IsIdentifier())
@@ -854,7 +854,7 @@ namespace EpicGames.UHT.Parsers
 			tokenReader.While(_gatherTypeTokensDelegate);
 
 			// Verify we at least have one type
-			if (this._currentTypeTokens.Count < 1)
+			if (_currentTypeTokens.Count < 1)
 			{
 				throw new UhtException(tokenReader, $"{propertySettings.PropertyCategory.GetHintText()}: Missing variable type or name");
 			}
@@ -871,7 +871,7 @@ namespace EpicGames.UHT.Parsers
 				tokenReader.LogError("Specified type modifiers not allowed here");
 			}
 
-			if (this._options.HasAnyFlags(UhtPropertyParseOptions.AddModuleRelativePath))
+			if (_options.HasAnyFlags(UhtPropertyParseOptions.AddModuleRelativePath))
 			{
 				UhtParsingScope.AddModuleRelativePathToMetaData(propertySettings.MetaData, topScope.ScopeType.HeaderFile);
 			}
@@ -897,20 +897,20 @@ namespace EpicGames.UHT.Parsers
 				}
 				tokenReader.Require(')');
 
-				Finalize(topScope, specifierContext, ref nameToken, new ReadOnlyMemory<UhtToken>(this._currentTypeTokens.ToArray()), layoutMacroType, propertyDelegate);
+				Finalize(topScope, specifierContext, ref nameToken, new ReadOnlyMemory<UhtToken>(_currentTypeTokens.ToArray()), layoutMacroType, propertyDelegate);
 			}
-			else if (this._options.HasAnyFlags(UhtPropertyParseOptions.List))
+			else if (_options.HasAnyFlags(UhtPropertyParseOptions.List))
 			{
 				// Extract the property name from the types
-				if (this._currentTypeTokens.Count < 2 || !this._currentTypeTokens[^1].IsIdentifier())
+				if (_currentTypeTokens.Count < 2 || !_currentTypeTokens[^1].IsIdentifier())
 				{
 					throw new UhtException(tokenReader, $"{propertySettings.PropertyCategory.GetHintText()}: Expected name");
 				}
-				UhtToken nameToken = this._currentTypeTokens[^1];
-				this._currentTypeTokens.RemoveAt(this._currentTypeTokens.Count - 1);
+				UhtToken nameToken = _currentTypeTokens[^1];
+				_currentTypeTokens.RemoveAt(_currentTypeTokens.Count - 1);
 				CheckForOptionalParts(tokenReader, propertySettings, ref nameToken);
 
-				ReadOnlyMemory<UhtToken> typeTokens = new(this._currentTypeTokens.ToArray());
+				ReadOnlyMemory<UhtToken> typeTokens = new(_currentTypeTokens.ToArray());
 
 				while (true)
 				{
@@ -926,24 +926,24 @@ namespace EpicGames.UHT.Parsers
 					throw new UhtException(tokenReader, $"Comma delimited properties cannot be converted");
 
 					// we'll need any metadata tags we parsed later on when we call ConvertEOLCommentToTooltip() so the tags aren't clobbered
-					//this.SpecifierContext.PropertySettings.MetaData = Property.MetaData.Clone();
-					//this.SpecifierContext.PropertySettings.ResetTrailingSettings();
+					//SpecifierContext.PropertySettings.MetaData = Property.MetaData.Clone();
+					//SpecifierContext.PropertySettings.ResetTrailingSettings();
 
 					// Get the next property name
 					//NameToken = tokenReader.GetIdentifier();
-					//CheckForOptionalParts(this.SpecifierContext.PropertySettings, ref NameToken);
+					//CheckForOptionalParts(SpecifierContext.PropertySettings, ref NameToken);
 				}
 			}
-			else if (this._options.HasAnyFlags(UhtPropertyParseOptions.CommaSeparatedName))
+			else if (_options.HasAnyFlags(UhtPropertyParseOptions.CommaSeparatedName))
 			{
 				tokenReader.Require(',');
 				UhtToken nameToken = tokenReader.GetIdentifier();
 				CheckForOptionalParts(tokenReader, propertySettings, ref nameToken);
-				Finalize(topScope, specifierContext, ref nameToken, new ReadOnlyMemory<UhtToken>(this._currentTypeTokens.ToArray()), layoutMacroType, propertyDelegate);
+				Finalize(topScope, specifierContext, ref nameToken, new ReadOnlyMemory<UhtToken>(_currentTypeTokens.ToArray()), layoutMacroType, propertyDelegate);
 			}
-			else if (this._options.HasAnyFlags(UhtPropertyParseOptions.FunctionNameIncluded))
+			else if (_options.HasAnyFlags(UhtPropertyParseOptions.FunctionNameIncluded))
 			{
-				UhtToken nameToken = this._currentTypeTokens[^1];
+				UhtToken nameToken = _currentTypeTokens[^1];
 				nameToken.Value = new StringView("Function");
 				if (CheckForOptionalParts(tokenReader, propertySettings, ref nameToken))
 				{
@@ -951,31 +951,31 @@ namespace EpicGames.UHT.Parsers
 				}
 				else
 				{
-					if (this._currentTypeTokens.Count < 2 || !this._currentTypeTokens[^1].IsIdentifier())
+					if (_currentTypeTokens.Count < 2 || !_currentTypeTokens[^1].IsIdentifier())
 					{
 						throw new UhtException(tokenReader, $"{propertySettings.PropertyCategory.GetHintText()}: Expected name");
 					}
-					nameToken = this._currentTypeTokens[^1];
-					this._currentTypeTokens.RemoveAt(this._currentTypeTokens.Count - 1);
+					nameToken = _currentTypeTokens[^1];
+					_currentTypeTokens.RemoveAt(_currentTypeTokens.Count - 1);
 				}
-				Finalize(topScope, specifierContext, ref nameToken, new ReadOnlyMemory<UhtToken>(this._currentTypeTokens.ToArray()), layoutMacroType, propertyDelegate);
+				Finalize(topScope, specifierContext, ref nameToken, new ReadOnlyMemory<UhtToken>(_currentTypeTokens.ToArray()), layoutMacroType, propertyDelegate);
 			}
-			else if (this._options.HasAnyFlags(UhtPropertyParseOptions.NameIncluded))
+			else if (_options.HasAnyFlags(UhtPropertyParseOptions.NameIncluded))
 			{
-				if (this._currentTypeTokens.Count < 2 || !this._currentTypeTokens[^1].IsIdentifier())
+				if (_currentTypeTokens.Count < 2 || !_currentTypeTokens[^1].IsIdentifier())
 				{
 					throw new UhtException(tokenReader, $"{propertySettings.PropertyCategory.GetHintText()}: Expected name");
 				}
-				UhtToken nameToken = this._currentTypeTokens[^1];
-				this._currentTypeTokens.RemoveAt(this._currentTypeTokens.Count - 1);
+				UhtToken nameToken = _currentTypeTokens[^1];
+				_currentTypeTokens.RemoveAt(_currentTypeTokens.Count - 1);
 				CheckForOptionalParts(tokenReader, propertySettings, ref nameToken);
-				Finalize(topScope, specifierContext, ref nameToken, new ReadOnlyMemory<UhtToken>(this._currentTypeTokens.ToArray()), layoutMacroType, propertyDelegate);
+				Finalize(topScope, specifierContext, ref nameToken, new ReadOnlyMemory<UhtToken>(_currentTypeTokens.ToArray()), layoutMacroType, propertyDelegate);
 			}
 			else
 			{
 				UhtToken nameToken = new();
 				CheckForOptionalParts(tokenReader, propertySettings, ref nameToken);
-				Finalize(topScope, specifierContext, ref nameToken, new ReadOnlyMemory<UhtToken>(this._currentTypeTokens.ToArray()), layoutMacroType, propertyDelegate);
+				Finalize(topScope, specifierContext, ref nameToken, new ReadOnlyMemory<UhtToken>(_currentTypeTokens.ToArray()), layoutMacroType, propertyDelegate);
 			}
 		}
 
@@ -1159,7 +1159,7 @@ namespace EpicGames.UHT.Parsers
 			propertyDelegate(topScope, newProperty, ref nameToken, layoutMacroType);
 
 			// Void properties don't get added when they are the return value
-			if (newProperty.PropertyCategory != UhtPropertyCategory.Return || !this._options.HasAnyFlags(UhtPropertyParseOptions.DontAddReturn))
+			if (newProperty.PropertyCategory != UhtPropertyCategory.Return || !_options.HasAnyFlags(UhtPropertyParseOptions.DontAddReturn))
 			{
 				topScope.ScopeType.AddChild(newProperty);
 			}
@@ -1206,24 +1206,24 @@ namespace EpicGames.UHT.Parsers
 
 		private bool GatherTypeTokens(ref UhtToken token)
 		{
-			if (this._currentTemplateDepth == 0 && token.IsSymbol() && (token.IsValue(',') || token.IsValue('(') || token.IsValue(')') || 
+			if (_currentTemplateDepth == 0 && token.IsSymbol() && (token.IsValue(',') || token.IsValue('(') || token.IsValue(')') || 
 				token.IsValue(';') || token.IsValue('[') || token.IsValue(':') || token.IsValue('=') || token.IsValue('{')))
 			{
 				return false;
 			}
 
-			this._currentTypeTokens.Add(token);
+			_currentTypeTokens.Add(token);
 			if (token.IsSymbol('<'))
 			{
-				++this._currentTemplateDepth;
+				++_currentTemplateDepth;
 			}
 			else if (token.IsSymbol('>'))
 			{
-				if (this._currentTemplateDepth == 0)
+				if (_currentTemplateDepth == 0)
 				{
-					throw new UhtTokenException(this._currentTokenReader!, token, "',' or ')'");
+					throw new UhtTokenException(_currentTokenReader!, token, "',' or ')'");
 				}
-				--this._currentTemplateDepth;
+				--_currentTemplateDepth;
 			}
 			return true;
 		}
