Removed the use of "this." where possible in EpicGames.UHT.

diff --git a/Engine/Source/Programs/Shared/EpicGames.UHT/Parsers/UhtSpecifierParser.cs b/Engine/Source/Programs/Shared/EpicGames.UHT/Parsers/UhtSpecifierParser.cs
--- a/Engine/Source/Programs/Shared/EpicGames.UHT/Parsers/UhtSpecifierParser.cs
+++ b/Engine/Source/Programs/Shared/EpicGames.UHT/Parsers/UhtSpecifierParser.cs
@@ -76,15 +76,15 @@ namespace EpicGames.UHT.Parsers
 		/// <param name="table">Specifier table</param>
 		public UhtSpecifierParser(UhtSpecifierContext specifierContext, StringView context, UhtSpecifierTable table)
 		{
-			this._specifierContext = specifierContext;
-			this._tokenReader = specifierContext.TokenReader;
-			this._context = context;
-			this._table = table;
-
-			this._parseAction = ParseInternal;
-			this._parseFieldMetaDataAction = ParseFieldMetaDataInternal;
-			this._parseKVPValueAction = ParseKVPValueInternal;
-			this._parseStringViewListAction = ParseStringViewListInternal;
+			_specifierContext = specifierContext;
+			_tokenReader = specifierContext.TokenReader;
+			_context = context;
+			_table = table;
+
+			_parseAction = ParseInternal;
+			_parseFieldMetaDataAction = ParseFieldMetaDataInternal;
+			_parseKVPValueAction = ParseKVPValueInternal;
+			_parseStringViewListAction = ParseStringViewListInternal;
 		}
 
 		/// <summary>
@@ -95,11 +95,11 @@ namespace EpicGames.UHT.Parsers
 		/// <param name="table">Specifier table</param>
 		public void Reset(UhtSpecifierContext specifierContext, StringView context, UhtSpecifierTable table)
 		{
-			this._specifierContext = specifierContext;
-			this._tokenReader = specifierContext.TokenReader;
-			this._context = context;
-			this._table = table;
-			this._deferredSpecifiers.Clear();
+			_specifierContext = specifierContext;
+			_tokenReader = specifierContext.TokenReader;
+			_context = context;
+			_table = table;
+			_deferredSpecifiers.Clear();
 		}
 
 		/// <summary>
@@ -108,11 +108,11 @@ namespace EpicGames.UHT.Parsers
 		/// <returns>The parser</returns>
 		public UhtSpecifierParser ParseSpecifiers()
 		{
-			this._isParsingFieldMetaData = false;
-			this._specifierContext.MetaData.LineNumber = _tokenReader.InputLine;
+			_isParsingFieldMetaData = false;
+			_specifierContext.MetaData.LineNumber = _tokenReader.InputLine;
 
 			using UhtMessageContext tokenContext = new(this);
-			this._tokenReader.RequireList('(', ')', ',', false, this._parseAction);
+			_tokenReader.RequireList('(', ')', ',', false, _parseAction);
 			return this;
 		}
 
@@ -122,17 +122,17 @@ namespace EpicGames.UHT.Parsers
 		/// <returns>Specifier parser</returns>
 		public UhtSpecifierParser ParseFieldMetaData()
 		{
-			this._tokenReader = this._specifierContext.TokenReader;
-			this._isParsingFieldMetaData = true;
+			_tokenReader = _specifierContext.TokenReader;
+			_isParsingFieldMetaData = true;
 
 			using UhtMessageContext tokenContext = new(this);
 			if (_tokenReader.TryOptional("UMETA"))
 			{
-				this._umetaElementsParsed = 0;
-				_tokenReader.RequireList('(', ')', ',', false, this._parseFieldMetaDataAction);
-				if (this._umetaElementsParsed == 0)
+				_umetaElementsParsed = 0;
+				_tokenReader.RequireList('(', ')', ',', false, _parseFieldMetaDataAction);
+				if (_umetaElementsParsed == 0)
 				{
-					this._tokenReader.LogError($"No metadata specified while parsing {UhtMessage.FormatContext(this)}");
+					_tokenReader.LogError($"No metadata specified while parsing {UhtMessage.FormatContext(this)}");
 				}
 			}
 			return this;
@@ -143,11 +143,11 @@ namespace EpicGames.UHT.Parsers
 		/// </summary>
 		public void ParseDeferred()
 		{
-			foreach (DeferredSpecifier deferred in this._deferredSpecifiers)
+			foreach (DeferredSpecifier deferred in _deferredSpecifiers)
 			{
 				Dispatch(deferred._specifier, deferred._value);
 			}
-			this._deferredSpecifiers.Clear();
+			_deferredSpecifiers.Clear();
 		}
 
 		#region IMessageExtraContext implementation
@@ -158,10 +158,10 @@ namespace EpicGames.UHT.Parsers
 			get
 			{
 				Stack<object?> extraContext = new(1);
-				string what = this._isParsingFieldMetaData ? "metadata" : "specifiers";
-				if (this._context.Span.Length > 0)
+				string what = _isParsingFieldMetaData ? "metadata" : "specifiers";
+				if (_context.Span.Length > 0)
 				{
-					extraContext.Push($"{this._context} {what}");
+					extraContext.Push($"{_context} {what}");
 				}
 				else
 				{
@@ -174,20 +174,20 @@ namespace EpicGames.UHT.Parsers
 
 		private void ParseInternal()
 		{
-			UhtToken identifier = this._tokenReader.GetIdentifier();
+			UhtToken identifier = _tokenReader.GetIdentifier();
 
-			this._currentSpecifier = identifier.Value;
-			if (this._table.TryGetValue(_currentSpecifier, out UhtSpecifier? specifier))
+			_currentSpecifier = identifier.Value;
+			if (_table.TryGetValue(_currentSpecifier, out UhtSpecifier? specifier))
 			{
 				if (TryParseValue(specifier.ValueType, out object? value))
 				{
 					if (specifier.When == UhtSpecifierWhen.Deferred)
 					{
-						if (this._deferredSpecifiers == null)
+						if (_deferredSpecifiers == null)
 						{
-							this._deferredSpecifiers = new List<DeferredSpecifier>();
+							_deferredSpecifiers = new List<DeferredSpecifier>();
 						}
-						this._deferredSpecifiers.Add(new DeferredSpecifier { _specifier = specifier, _value = value });
+						_deferredSpecifiers.Add(new DeferredSpecifier { _specifier = specifier, _value = value });
 					}
 					else
 					{
@@ -197,7 +197,7 @@ namespace EpicGames.UHT.Parsers
 			}
 			else
 			{
-				this._tokenReader.LogError($"Unknown specifier '{_currentSpecifier}' found while parsing {UhtMessage.FormatContext(this)}");
+				_tokenReader.LogError($"Unknown specifier '{_currentSpecifier}' found while parsing {UhtMessage.FormatContext(this)}");
 			}
 		}
 
@@ -205,7 +205,7 @@ namespace EpicGames.UHT.Parsers
 		{
 			if (!_tokenReader.TryOptionalIdentifier(out UhtToken key))
 			{
-				throw new UhtException(this._tokenReader, $"UMETA expects a key and optional value", this);
+				throw new UhtException(_tokenReader, $"UMETA expects a key and optional value", this);
 			}
 
 			StringViewBuilder builder = new();
@@ -213,38 +213,38 @@ namespace EpicGames.UHT.Parsers
 			{
 				if (!ReadValue(_tokenReader, builder, true))
 				{
-					throw new UhtException(this._tokenReader, $"UMETA key '{key.Value}' expects a value", this);
+					throw new UhtException(_tokenReader, $"UMETA key '{key.Value}' expects a value", this);
 				}
 			}
 
-			++this._umetaElementsParsed;
-			this._specifierContext.MetaData.Add(key.Value.ToString(), this._specifierContext.MetaNameIndex, builder.ToString());
+			++_umetaElementsParsed;
+			_specifierContext.MetaData.Add(key.Value.ToString(), _specifierContext.MetaNameIndex, builder.ToString());
 		}
 
 		private void ParseKVPValueInternal()
 		{
-			if (this._currentKVPValues == null)
+			if (_currentKVPValues == null)
 			{
-				this._currentKVPValues = new List<KeyValuePair<StringView, StringView>>();
+				_currentKVPValues = new List<KeyValuePair<StringView, StringView>>();
 			}
-			this._currentKVPValues.Add(ReadKVP());
+			_currentKVPValues.Add(ReadKVP());
 		}
 
 		private void ParseStringViewListInternal()
 		{
-			if (this._currentStringValues == null)
+			if (_currentStringValues == null)
 			{
-				this._currentStringValues = new List<StringView>();
+				_currentStringValues = new List<StringView>();
 			}
-			this._currentStringValues.Add(ReadValue());
+			_currentStringValues.Add(ReadValue());
 		}
 
 		private void Dispatch(UhtSpecifier specifier, object? value)
 		{
-			UhtSpecifierDispatchResults results = specifier.Dispatch(this._specifierContext, value);
+			UhtSpecifierDispatchResults results = specifier.Dispatch(_specifierContext, value);
 			if (results == UhtSpecifierDispatchResults.Unknown)
 			{
-				this._tokenReader.LogError($"Unknown specifier '{specifier.Name}' found while parsing {UhtMessage.FormatContext(this)}");
+				_tokenReader.LogError($"Unknown specifier '{specifier.Name}' found while parsing {UhtMessage.FormatContext(this)}");
 			}
 		}
 
@@ -258,18 +258,18 @@ namespace EpicGames.UHT.Parsers
 					throw new UhtIceException("NotSet is an invalid value for value types");
 
 				case UhtSpecifierValueType.None:
-					if (this._tokenReader.TryOptional('='))
+					if (_tokenReader.TryOptional('='))
 					{
 						ReadValue(); // consume the value;
-						this._tokenReader.LogError($"The specifier '{this._currentSpecifier}' found a value when none was expected", this);
+						_tokenReader.LogError($"The specifier '{_currentSpecifier}' found a value when none was expected", this);
 						return false;
 					}
 					return true;
 
 				case UhtSpecifierValueType.String:
-					if (!this._tokenReader.TryOptional('='))
+					if (!_tokenReader.TryOptional('='))
 					{
-						this._tokenReader.LogError($"The specifier '{this._currentSpecifier}' expects a value", this);
+						_tokenReader.LogError($"The specifier '{_currentSpecifier}' expects a value", this);
 						return false;
 					}
 					value = ReadValue();
@@ -290,7 +290,7 @@ namespace EpicGames.UHT.Parsers
 						List<StringView>? stringList = ReadValueList();
 						if (stringList == null || stringList.Count != 1)
 						{
-							this._tokenReader.LogError($"The specifier '{this._currentSpecifier}' expects a single value", this);
+							_tokenReader.LogError($"The specifier '{_currentSpecifier}' expects a single value", this);
 							return false;
 						}
 						value = stringList[0];
@@ -310,7 +310,7 @@ namespace EpicGames.UHT.Parsers
 						List<StringView>? stringList = ReadValueList();
 						if (stringList == null || stringList.Count == 0)
 						{
-							this._tokenReader.LogError($"The specifier '{this._currentSpecifier}' expects at list one value", this);
+							_tokenReader.LogError($"The specifier '{_currentSpecifier}' expects at list one value", this);
 							return false;
 						}
 						value = stringList;
@@ -319,33 +319,33 @@ namespace EpicGames.UHT.Parsers
 
 				case UhtSpecifierValueType.KeyValuePairList:
 					{
-						this._currentKVPValues = null;
-						this._tokenReader
+						_currentKVPValues = null;
+						_tokenReader
 							.Require('=')
-							.RequireList('(', ')', ',', false, this._parseKVPValueAction);
-						List<KeyValuePair<StringView, StringView>> kvps = this._currentKVPValues ?? s_emptyKVPValues;
-						this._currentKVPValues = null;
+							.RequireList('(', ')', ',', false, _parseKVPValueAction);
+						List<KeyValuePair<StringView, StringView>> kvps = _currentKVPValues ?? s_emptyKVPValues;
+						_currentKVPValues = null;
 						value = kvps;
 						return true;
 					}
 
 				case UhtSpecifierValueType.OptionalEqualsKeyValuePairList:
 					{
-						this._currentKVPValues = null;
+						_currentKVPValues = null;
 						// This parser isn't as strict as the other parsers...
-						if (this._tokenReader.TryOptional('='))
+						if (_tokenReader.TryOptional('='))
 						{
-							if (!this._tokenReader.TryOptionalList('(', ')', ',', false, this._parseKVPValueAction))
+							if (!_tokenReader.TryOptionalList('(', ')', ',', false, _parseKVPValueAction))
 							{
-								this._parseKVPValueAction();
+								_parseKVPValueAction();
 							}
 						}
 						else
 						{
-							this._tokenReader.TryOptionalList('(', ')', ',', false, this._parseKVPValueAction);
+							_tokenReader.TryOptionalList('(', ')', ',', false, _parseKVPValueAction);
 						}
-						List<KeyValuePair<StringView, StringView>> kvps = this._currentKVPValues ?? s_emptyKVPValues;
-						this._currentKVPValues = null;
+						List<KeyValuePair<StringView, StringView>> kvps = _currentKVPValues ?? s_emptyKVPValues;
+						_currentKVPValues = null;
 						value = kvps;
 						return true;
 					}
@@ -357,13 +357,13 @@ namespace EpicGames.UHT.Parsers
 
 		private KeyValuePair<StringView, StringView> ReadKVP()
 		{
-			if (!this._tokenReader.TryOptionalIdentifier(out UhtToken key))
+			if (!_tokenReader.TryOptionalIdentifier(out UhtToken key))
 			{
-				throw new UhtException(this._tokenReader, $"The specifier '{this._currentSpecifier}' expects a key and optional value", this);
+				throw new UhtException(_tokenReader, $"The specifier '{_currentSpecifier}' expects a key and optional value", this);
 			}
 
 			StringView value = "";
-			if (this._tokenReader.TryOptional('='))
+			if (_tokenReader.TryOptional('='))
 			{
 				value = ReadValue();
 			}
@@ -372,31 +372,31 @@ namespace EpicGames.UHT.Parsers
 
 		private List<StringView>? ReadValueList()
 		{
-			this._currentStringValues = null;
+			_currentStringValues = null;
 
 			// This parser isn't as strict as the other parsers...
-			if (this._tokenReader.TryOptional('='))
+			if (_tokenReader.TryOptional('='))
 			{
-				if (!this._tokenReader.TryOptionalList('(', ')', ',', false, this._parseStringViewListAction))
+				if (!_tokenReader.TryOptionalList('(', ')', ',', false, _parseStringViewListAction))
 				{
-					this._parseStringViewListAction();
+					_parseStringViewListAction();
 				}
 			}
 			else
 			{
-				this._tokenReader.TryOptionalList('(', ')', ',', false, this._parseStringViewListAction);
+				_tokenReader.TryOptionalList('(', ')', ',', false, _parseStringViewListAction);
 			}
-			List<StringView>? stringValues = this._currentStringValues;
-			this._currentStringValues = null;
+			List<StringView>? stringValues = _currentStringValues;
+			_currentStringValues = null;
 			return stringValues;
 		}
 
 		private StringView ReadValue()
 		{
 			StringViewBuilder builder = new();
-			if (!ReadValue(this._tokenReader, builder, false))
+			if (!ReadValue(_tokenReader, builder, false))
 			{
-				throw new UhtException(this._tokenReader, $"The specifier '{this._currentSpecifier}' expects a value", this);
+				throw new UhtException(_tokenReader, $"The specifier '{_currentSpecifier}' expects a value", this);
 			}
 			return builder.ToStringView();
 		}
diff --git a/Engine/Source/Programs/Shared/EpicGames.UHT/Tables/UhtEngineClassTable.cs b/Engine/Source/Programs/Shared/EpicGames.UHT/Tables/UhtEngineClassTable.cs
--- a/Engine/Source/Programs/Shared/EpicGames.UHT/Tables/UhtEngineClassTable.cs
+++ b/Engine/Source/Programs/Shared/EpicGames.UHT/Tables/UhtEngineClassTable.cs
@@ -59,7 +59,7 @@ namespace EpicGames.UHT.Tables
 		/// <returns>True if the class name is a property.  False if the class name isn't a property or isn't an engine class.</returns>
 		public bool IsValidPropertyTypeName(StringView name)
 		{
-			if (this._engineClasses.TryGetValue(name, out UhtEngineClass engineClass))
+			if (_engineClasses.TryGetValue(name, out UhtEngineClass engineClass))
 			{
 				return engineClass.IsProperty;
 			}
@@ -76,7 +76,7 @@ namespace EpicGames.UHT.Tables
 			{
 				throw new UhtIceException("EngineClassNames must have a name specified");
 			}
-			this._engineClasses.Add(engineClassAttribute.Name, new UhtEngineClass { Name = engineClassAttribute.Name, IsProperty = engineClassAttribute.IsProperty });
+			_engineClasses.Add(engineClassAttribute.Name, new UhtEngineClass { Name = engineClassAttribute.Name, IsProperty = engineClassAttribute.IsProperty });
 		}
 	}
 }
diff --git a/Engine/Source/Programs/Shared/EpicGames.UHT/Tables/UhtExporterTable.cs b/Engine/Source/Programs/Shared/EpicGames.UHT/Tables/UhtExporterTable.cs
--- a/Engine/Source/Programs/Shared/EpicGames.UHT/Tables/UhtExporterTable.cs
+++ b/Engine/Source/Programs/Shared/EpicGames.UHT/Tables/UhtExporterTable.cs
@@ -171,14 +171,14 @@ namespace EpicGames.UHT.Tables
 		/// <param name="exporterDelegate">Delegate to invoke</param>
 		public UhtExporter(UhtExporterAttribute attribute, UhtExporterDelegate exporterDelegate)
 		{
-			this.Name = attribute.Name;
-			this.Description = attribute.Description;
-			this.ModuleName = attribute.ModuleName;
-			this.Options = attribute.Options;
-			this.Delegate = exporterDelegate;
-			this.CppFilters = attribute.CppFilters != null ? new List<string>(attribute.CppFilters) : new List<string>();
-			this.HeaderFilters = attribute.HeaderFilters != null ? new List<string>(attribute.HeaderFilters) : new List<string>();
-			this.OtherFilters = attribute.OtherFilters != null ? new List<string>(attribute.OtherFilters) : new List<string>();
+			Name = attribute.Name;
+			Description = attribute.Description;
+			ModuleName = attribute.ModuleName;
+			Options = attribute.Options;
+			Delegate = exporterDelegate;
+			CppFilters = attribute.CppFilters != null ? new List<string>(attribute.CppFilters) : new List<string>();
+			HeaderFilters = attribute.HeaderFilters != null ? new List<string>(attribute.HeaderFilters) : new List<string>();
+			OtherFilters = attribute.OtherFilters != null ? new List<string>(attribute.OtherFilters) : new List<string>();
 		}
 	}
 
@@ -198,7 +198,7 @@ namespace EpicGames.UHT.Tables
 		/// <returns></returns>
 		public bool TryGet(string name, out UhtExporter value)
 		{
-			return this._exporterValues.TryGetValue(name, out value);
+			return _exporterValues.TryGetValue(name, out value);
 		}
 
 		/// <summary>
@@ -224,7 +224,7 @@ namespace EpicGames.UHT.Tables
 			}
 
 			UhtExporter exporterValue = new(exporterAttribute, (UhtExporterDelegate)Delegate.CreateDelegate(typeof(UhtExporterDelegate), methodInfo));
-			this._exporterValues.Add(exporterAttribute.Name, exporterValue);
+			_exporterValues.Add(exporterAttribute.Name, exporterValue);
 		}
 
 		/// <summary>
@@ -233,7 +233,7 @@ namespace EpicGames.UHT.Tables
 		/// <returns>Enumerator</returns>
 		public IEnumerator<UhtExporter> GetEnumerator()
 		{
-			foreach (KeyValuePair<string, UhtExporter> kvp in this._exporterValues)
+			foreach (KeyValuePair<string, UhtExporter> kvp in _exporterValues)
 			{
 				yield return kvp.Value;
 			}
@@ -245,7 +245,7 @@ namespace EpicGames.UHT.Tables
 		/// <returns>Enumerator</returns>
 		IEnumerator IEnumerable.GetEnumerator()
 		{
-			foreach (KeyValuePair<string, UhtExporter> kvp in this._exporterValues)
+			foreach (KeyValuePair<string, UhtExporter> kvp in _exporterValues)
 			{
 				yield return kvp.Value;
 			}
diff --git a/Engine/Source/Programs/Shared/EpicGames.UHT/Tables/UhtKeywordTable.cs b/Engine/Source/Programs/Shared/EpicGames.UHT/Tables/UhtKeywordTable.cs
--- a/Engine/Source/Programs/Shared/EpicGames.UHT/Tables/UhtKeywordTable.cs
+++ b/Engine/Source/Programs/Shared/EpicGames.UHT/Tables/UhtKeywordTable.cs
@@ -113,19 +113,19 @@ namespace EpicGames.UHT.Tables
 		/// <param name="attribute">Defining attribute</param>
 		public UhtKeyword(string name, UhtKeywordDelegate keywordDelegate, UhtKeywordAttribute? attribute)
 		{
-			this.Name = name;
-			this.Delegate = keywordDelegate;
+			Name = name;
+			Delegate = keywordDelegate;
 			if (attribute != null)
 			{
-				this.AllowText = attribute.AllowText;
-				this.AllScopes = attribute.AllScopes;
-				this.DisableUsageError = attribute.DisableUsageError;
+				AllowText = attribute.AllowText;
+				AllScopes = attribute.AllScopes;
+				DisableUsageError = attribute.DisableUsageError;
 			}
 			else
 			{
-				this.AllowText = null;
-				this.AllScopes = false;
-				this.DisableUsageError = false;
+				AllowText = null;
+				AllScopes = false;
+				DisableUsageError = false;
 			}
 		}
 	}
@@ -164,7 +164,7 @@ namespace EpicGames.UHT.Tables
 		/// <param name="catchAll">The catch-all to be added</param>
 		public UhtKeywordTable AddCatchAll(UhtKeywordCatchAllDelegate catchAll)
 		{
-			this.CatchAlls.Add(catchAll);
+			CatchAlls.Add(catchAll);
 			return this;
 		}
 
@@ -175,7 +175,7 @@ namespace EpicGames.UHT.Tables
 		public override void Merge(UhtLookupTableBase baseTable)
 		{
 			base.Merge(baseTable);
-			this.CatchAlls.AddRange(((UhtKeywordTable)baseTable).CatchAlls);
+			CatchAlls.AddRange(((UhtKeywordTable)baseTable).CatchAlls);
 		}
 	}
 
@@ -238,7 +238,7 @@ namespace EpicGames.UHT.Tables
 		public void LogUnhandledError(IUhtMessageSite messageSite, UhtToken token)
 		{
 			List<string>? tables = null;
-			foreach (KeyValuePair<string, UhtKeywordTable> kvp in this.Tables)
+			foreach (KeyValuePair<string, UhtKeywordTable> kvp in Tables)
 			{
 				UhtKeywordTable keywordTable = kvp.Value;
 				if (keywordTable.Internal)
diff --git a/Engine/Source/Programs/Shared/EpicGames.UHT/Tables/UhtLocTextDefaultValueTable.cs b/Engine/Source/Programs/Shared/EpicGames.UHT/Tables/UhtLocTextDefaultValueTable.cs
--- a/Engine/Source/Programs/Shared/EpicGames.UHT/Tables/UhtLocTextDefaultValueTable.cs
+++ b/Engine/Source/Programs/Shared/EpicGames.UHT/Tables/UhtLocTextDefaultValueTable.cs
@@ -62,7 +62,7 @@ namespace EpicGames.UHT.Tables
 		/// <returns></returns>
 		public bool TryGet(StringView name, out UhtLocTextDefaultValue locTextDefaultValue)
 		{
-			return this._locTextDefaultValues.TryGetValue(name, out locTextDefaultValue);
+			return _locTextDefaultValues.TryGetValue(name, out locTextDefaultValue);
 		}
 
 		/// <summary>
@@ -83,7 +83,7 @@ namespace EpicGames.UHT.Tables
 				Delegate = (UhtLocTextDefaultValueDelegate)Delegate.CreateDelegate(typeof(UhtLocTextDefaultValueDelegate), methodInfo)
 			};
 
-			this._locTextDefaultValues.Add(new StringView(locTextDefaultValueAttribute.Name), locTextDefaultValue);
+			_locTextDefaultValues.Add(new StringView(locTextDefaultValueAttribute.Name), locTextDefaultValue);
 		}
 	}
 }
diff --git a/Engine/Source/Programs/Shared/EpicGames.UHT/Tables/UhtPropertyTypeTable.cs b/Engine/Source/Programs/Shared/EpicGames.UHT/Tables/UhtPropertyTypeTable.cs
--- a/Engine/Source/Programs/Shared/EpicGames.UHT/Tables/UhtPropertyTypeTable.cs
+++ b/Engine/Source/Programs/Shared/EpicGames.UHT/Tables/UhtPropertyTypeTable.cs
@@ -165,11 +165,11 @@ namespace EpicGames.UHT.Tables
 		{
 			get
 			{
-				if (this._default == null)
+				if (_default == null)
 				{
 					throw new UhtIceException("No property type has been marked as default");
 				}
-				return (UhtPropertyType)this._default;
+				return (UhtPropertyType)_default;
 			}
 		}
 
@@ -182,8 +182,8 @@ namespace EpicGames.UHT.Tables
 		public bool TryGet(StringView name, out UhtPropertyType propertyType)
 		{
 			return
-				this._caseSensitive.TryGetValue(name, out propertyType) ||
-				this._caseInsensitive.TryGetValue(name, out propertyType);
+				_caseSensitive.TryGetValue(name, out propertyType) ||
+				_caseInsensitive.TryGetValue(name, out propertyType);
 		}
 
 		/// <summary>
@@ -207,21 +207,21 @@ namespace EpicGames.UHT.Tables
 
 			if (propertyTypeAttribute.Options.HasAnyFlags(UhtPropertyTypeOptions.Default))
 			{
-				if (this._default != null)
+				if (_default != null)
 				{
 					throw new UhtIceException("Only one property type dispatcher can be marked as default");
 				}
-				this._default = propertyType;
+				_default = propertyType;
 			}
 			else if (!String.IsNullOrEmpty(propertyTypeAttribute.Keyword))
 			{
 				if (propertyTypeAttribute.Options.HasAnyFlags(UhtPropertyTypeOptions.CaseInsensitive))
 				{
-					this._caseInsensitive.Add(propertyTypeAttribute.Keyword, propertyType);
+					_caseInsensitive.Add(propertyTypeAttribute.Keyword, propertyType);
 				}
 				else
 				{
-					this._caseSensitive.Add(propertyTypeAttribute.Keyword, propertyType);
+					_caseSensitive.Add(propertyTypeAttribute.Keyword, propertyType);
 				}
 			}
 		}
diff --git a/Engine/Source/Programs/Shared/EpicGames.UHT/Tables/UhtSpecifierTable.cs b/Engine/Source/Programs/Shared/EpicGames.UHT/Tables/UhtSpecifierTable.cs
--- a/Engine/Source/Programs/Shared/EpicGames.UHT/Tables/UhtSpecifierTable.cs
+++ b/Engine/Source/Programs/Shared/EpicGames.UHT/Tables/UhtSpecifierTable.cs
@@ -135,12 +135,12 @@ namespace EpicGames.UHT.Tables
 		/// <param name="metaNameIndex"></param>
 		public UhtSpecifierContext(UhtParsingScope scope, IUhtMessageSite messageSite, UhtMetaData metaData, int metaNameIndex = UhtMetaData.IndexNone)
 		{
-			this.Type = scope.ScopeType;
-			this.TokenReader = scope.TokenReader;
-			this.AccessSpecifier = scope.AccessSpecifier;
-			this.MessageSite = messageSite;
-			this.MetaData = metaData;
-			this.MetaNameIndex = metaNameIndex;
+			Type = scope.ScopeType;
+			TokenReader = scope.TokenReader;
+			AccessSpecifier = scope.AccessSpecifier;
+			MessageSite = messageSite;
+			MetaData = metaData;
+			MetaNameIndex = metaNameIndex;
 		}
 
 		/// <summary>
@@ -219,16 +219,16 @@ namespace EpicGames.UHT.Tables
 		/// <param name="specifierDelegate">Delegate to invoke</param>
 		public UhtSpecifierNone(string name, UhtSpecifierWhen when, UhtSpecifierNoneDelegate specifierDelegate)
 		{
-			this.Name = name;
-			this.ValueType = UhtSpecifierValueType.None;
-			this.When = when;
-			this._delegate = specifierDelegate;
+			Name = name;
+			ValueType = UhtSpecifierValueType.None;
+			When = when;
+			_delegate = specifierDelegate;
 		}
 
 		/// <inheritdoc/>
 		public override UhtSpecifierDispatchResults Dispatch(UhtSpecifierContext specifierContext, object? value)
 		{
-			this._delegate(specifierContext);
+			_delegate(specifierContext);
 			return UhtSpecifierDispatchResults.Known;
 		}
 	}
@@ -255,10 +255,10 @@ namespace EpicGames.UHT.Tables
 		/// <param name="specifierDelegate">Delegate to invoke</param>
 		public UhtSpecifierString(string name, UhtSpecifierWhen when, UhtSpecifierStringDelegate specifierDelegate)
 		{
-			this.Name = name;
-			this.ValueType = UhtSpecifierValueType.String;
-			this.When = when;
-			this._delegate = specifierDelegate;
+			Name = name;
+			ValueType = UhtSpecifierValueType.String;
+			When = when;
+			_delegate = specifierDelegate;
 		}
 
 		/// <inheritdoc/>
@@ -268,7 +268,7 @@ namespace EpicGames.UHT.Tables
 			{
 				throw new UhtIceException("Required value is null");
 			}
-			this._delegate(specifierContext, (StringView)value);
+			_delegate(specifierContext, (StringView)value);
 			return UhtSpecifierDispatchResults.Known;
 		}
 	}
@@ -295,16 +295,16 @@ namespace EpicGames.UHT.Tables
 		/// <param name="specifierDelegate">Delegate to invoke</param>
 		public UhtSpecifierOptionalString(string name, UhtSpecifierWhen when, UhtSpecifierOptionalStringDelegate specifierDelegate)
 		{
-			this.Name = name;
-			this.ValueType = UhtSpecifierValueType.OptionalString;
-			this.When = when;
-			this._delegate = specifierDelegate;
+			Name = name;
+			ValueType = UhtSpecifierValueType.OptionalString;
+			When = when;
+			_delegate = specifierDelegate;
 		}
 
 		/// <inheritdoc/>
 		public override UhtSpecifierDispatchResults Dispatch(UhtSpecifierContext specifierContext, object? value)
 		{
-			this._delegate(specifierContext, (StringView?)value);
+			_delegate(specifierContext, (StringView?)value);
 			return UhtSpecifierDispatchResults.Known;
 		}
 	}
@@ -331,10 +331,10 @@ namespace EpicGames.UHT.Tables
 		/// <param name="specifierDelegate">Delegate to invoke</param>
 		public UhtSpecifierSingleString(string name, UhtSpecifierWhen when, UhtSpecifierSingleStringDelegate specifierDelegate)
 		{
-			this.Name = name;
-			this.ValueType = UhtSpecifierValueType.SingleString;
-			this.When = when;
-			this._delegate = specifierDelegate;
+			Name = name;
+			ValueType = UhtSpecifierValueType.SingleString;
+			When = when;
+			_delegate = specifierDelegate;
 		}
 
 		/// <inheritdoc/>
@@ -344,7 +344,7 @@ namespace EpicGames.UHT.Tables
 			{
 				throw new UhtIceException("Required value is null");
 			}
-			this._delegate(specifierContext, (StringView)value);
+			_delegate(specifierContext, (StringView)value);
 			return UhtSpecifierDispatchResults.Known;
 		}
 	}
@@ -372,10 +372,10 @@ namespace EpicGames.UHT.Tables
 		/// <param name="specifierDelegate">Delegate to invoke</param>
 		public UhtSpecifierKeyValuePairList(string name, UhtSpecifierWhen when, bool equalsOptional, UhtSpecifierKeyValuePairListDelegate specifierDelegate)
 		{
-			this.Name = name;
-			this.ValueType = equalsOptional ? UhtSpecifierValueType.OptionalEqualsKeyValuePairList : UhtSpecifierValueType.KeyValuePairList;
-			this.When = when;
-			this._delegate = specifierDelegate;
+			Name = name;
+			ValueType = equalsOptional ? UhtSpecifierValueType.OptionalEqualsKeyValuePairList : UhtSpecifierValueType.KeyValuePairList;
+			When = when;
+			_delegate = specifierDelegate;
 		}
 
 		/// <inheritdoc/>
@@ -385,7 +385,7 @@ namespace EpicGames.UHT.Tables
 			{
 				throw new UhtIceException("Required value is null");
 			}
-			this._delegate(specifierContext, (List<KeyValuePair<StringView, StringView>>)value);
+			_delegate(specifierContext, (List<KeyValuePair<StringView, StringView>>)value);
 			return UhtSpecifierDispatchResults.Known;
 		}
 	}
@@ -411,10 +411,10 @@ namespace EpicGames.UHT.Tables
 		/// <param name="specifierDelegate">Delegate to invoke</param>
 		public UhtSpecifierLegacy(string name, UhtSpecifierLegacyDelegate specifierDelegate)
 		{
-			this.Name = name;
-			this.ValueType = UhtSpecifierValueType.StringList;
-			this.When = UhtSpecifierWhen.Deferred;
-			this._delegate = specifierDelegate;
+			Name = name;
+			ValueType = UhtSpecifierValueType.StringList;
+			When = UhtSpecifierWhen.Deferred;
+			_delegate = specifierDelegate;
 		}
 
 		/// <inheritdoc/>
@@ -422,9 +422,9 @@ namespace EpicGames.UHT.Tables
 		{
 			if (value != null)
 			{
-				specifierContext.TokenReader.LogDeprecation($"Specifier '{this.Name}' has a value which is unused, future versions of UnrealHeaderTool will flag this as an error.");
+				specifierContext.TokenReader.LogDeprecation($"Specifier '{Name}' has a value which is unused, future versions of UnrealHeaderTool will flag this as an error.");
 			}
-			this._delegate(specifierContext);
+			_delegate(specifierContext);
 			return UhtSpecifierDispatchResults.Known;
 		}
 	}
@@ -451,16 +451,16 @@ namespace EpicGames.UHT.Tables
 		/// <param name="specifierDelegate">Delegate to invoke</param>
 		public UhtSpecifierStringList(string name, UhtSpecifierWhen when, UhtSpecifierStringListDelegate specifierDelegate)
 		{
-			this.Name = name;
-			this.ValueType = UhtSpecifierValueType.StringList;
-			this.When = when;
-			this._delegate = specifierDelegate;
+			Name = name;
+			ValueType = UhtSpecifierValueType.StringList;
+			When = when;
+			_delegate = specifierDelegate;
 		}
 
 		/// <inheritdoc/>
 		public override UhtSpecifierDispatchResults Dispatch(UhtSpecifierContext specifierContext, object? value)
 		{
-			this._delegate(specifierContext, (List<StringView>?)value);
+			_delegate(specifierContext, (List<StringView>?)value);
 			return UhtSpecifierDispatchResults.Known;
 		}
 	}
@@ -487,10 +487,10 @@ namespace EpicGames.UHT.Tables
 		/// <param name="specifierDelegate">Delegate to invoke</param>
 		public UhtSpecifierNonEmptyStringList(string name, UhtSpecifierWhen when, UhtSpecifierNonEmptyStringListDelegate specifierDelegate)
 		{
-			this.Name = name;
-			this.ValueType = UhtSpecifierValueType.NonEmptyStringList;
-			this.When = when;
-			this._delegate = specifierDelegate;
+			Name = name;
+			ValueType = UhtSpecifierValueType.NonEmptyStringList;
+			When = when;
+			_delegate = specifierDelegate;
 		}
 
 		/// <inheritdoc/>
@@ -500,7 +500,7 @@ namespace EpicGames.UHT.Tables
 			{
 				throw new UhtIceException("Required value is null");
 			}
-			this._delegate(specifierContext, (List<StringView>)value);
+			_delegate(specifierContext, (List<StringView>)value);
 			return UhtSpecifierDispatchResults.Known;
 		}
 	}
diff --git a/Engine/Source/Programs/Shared/EpicGames.UHT/Tables/UhtSpecifierValidatorTable.cs b/Engine/Source/Programs/Shared/EpicGames.UHT/Tables/UhtSpecifierValidatorTable.cs
--- a/Engine/Source/Programs/Shared/EpicGames.UHT/Tables/UhtSpecifierValidatorTable.cs
+++ b/Engine/Source/Programs/Shared/EpicGames.UHT/Tables/UhtSpecifierValidatorTable.cs
@@ -41,8 +41,8 @@ namespace EpicGames.UHT.Tables
 		/// <param name="specifierValidatorDelegate">Delegate of the validator</param>
 		public UhtSpecifierValidator(string name, UhtSpecifierValidatorDelegate specifierValidatorDelegate)
 		{
-			this.Name = name;
-			this.Delegate = specifierValidatorDelegate;
+			Name = name;
+			Delegate = specifierValidatorDelegate;
 		}
 	}
 
diff --git a/Engine/Source/Programs/Shared/EpicGames.UHT/Tables/UhtStructDefaultValueTable.cs b/Engine/Source/Programs/Shared/EpicGames.UHT/Tables/UhtStructDefaultValueTable.cs
--- a/Engine/Source/Programs/Shared/EpicGames.UHT/Tables/UhtStructDefaultValueTable.cs
+++ b/Engine/Source/Programs/Shared/EpicGames.UHT/Tables/UhtStructDefaultValueTable.cs
@@ -126,11 +126,11 @@ namespace EpicGames.UHT.Tables
 		{
 			get
 			{
-				if (this._default == null)
+				if (_default == null)
 				{
 					throw new UhtIceException("No struct default value has been marked as default");
 				}
-				return (UhtStructDefaultValue)this._default;
+				return (UhtStructDefaultValue)_default;
 			}
 		}
 
@@ -142,7 +142,7 @@ namespace EpicGames.UHT.Tables
 		/// <returns></returns>
 		public bool TryGet(StringView name, out UhtStructDefaultValue structDefaultValue)
 		{
-			return this._structDefaultValues.TryGetValue(name, out structDefaultValue);
+			return _structDefaultValues.TryGetValue(name, out structDefaultValue);
 		}
 
 		/// <summary>
@@ -165,15 +165,15 @@ namespace EpicGames.UHT.Tables
 
 			if (structDefaultValueAttribute.Options.HasAnyFlags(UhtStructDefaultValueOptions.Default))
 			{
-				if (this._default != null)
+				if (_default != null)
 				{
 					throw new UhtIceException("Only one struct default value attribute dispatcher can be marked as default");
 				}
-				this._default = structDefaultValue;
+				_default = structDefaultValue;
 			}
 			else if (!String.IsNullOrEmpty(structDefaultValueAttribute.Name))
 			{
-				this._structDefaultValues.Add(new StringView(structDefaultValueAttribute.Name), structDefaultValue);
+				_structDefaultValues.Add(new StringView(structDefaultValueAttribute.Name), structDefaultValue);
 			}
 		}
 	}
diff --git a/Engine/Source/Programs/Shared/EpicGames.UHT/Tokenizer/UhtToken.cs b/Engine/Source/Programs/Shared/EpicGames.UHT/Tokenizer/UhtToken.cs
--- a/Engine/Source/Programs/Shared/EpicGames.UHT/Tokenizer/UhtToken.cs
+++ b/Engine/Source/Programs/Shared/EpicGames.UHT/Tokenizer/UhtToken.cs
@@ -142,7 +142,7 @@ namespace EpicGames.UHT.Tokenizer
 		/// <summary>
 		/// End position of the token value
 		/// </summary>
-		public int InputEndPos => this.InputStartPos + this.Value.Span.Length;
+		public int InputEndPos => InputStartPos + Value.Span.Length;
 
 		/// <summary>
 		/// Line containing the token
@@ -160,12 +160,12 @@ namespace EpicGames.UHT.Tokenizer
 		/// <param name="tokenType">Type of the token</param>
 		public UhtToken(UhtTokenType tokenType)
 		{
-			this.TokenType = tokenType;
-			this.UngetPos = 0;
-			this.UngetLine = 0;
-			this.InputStartPos = 0;
-			this.InputLine = 0;
-			this.Value = new StringView();
+			TokenType = tokenType;
+			UngetPos = 0;
+			UngetLine = 0;
+			InputStartPos = 0;
+			InputLine = 0;
+			Value = new StringView();
 		}
 
 		/// <summary>
@@ -179,12 +179,12 @@ namespace EpicGames.UHT.Tokenizer
 		/// <param name="value">Token value</param>
 		public UhtToken(UhtTokenType tokenType, int ungetPos, int ungetLine, int inputStartPos, int inputLine, StringView value)
 		{
-			this.TokenType = tokenType;
-			this.UngetPos = ungetPos;
-			this.UngetLine = ungetLine;
-			this.InputStartPos = inputStartPos;
-			this.InputLine = inputLine;
-			this.Value = value;
+			TokenType = tokenType;
+			UngetPos = ungetPos;
+			UngetLine = ungetLine;
+			InputStartPos = inputStartPos;
+			InputLine = inputLine;
+			Value = value;
 		}
 
 		/// <summary>
@@ -202,7 +202,7 @@ namespace EpicGames.UHT.Tokenizer
 		/// <returns>True if the token is an end token</returns>
 		public bool IsEndType()
 		{
-			return this.TokenType.IsEndType();
+			return TokenType.IsEndType();
 		}
 
 		/// <summary>
@@ -212,7 +212,7 @@ namespace EpicGames.UHT.Tokenizer
 		/// <returns>True if the token value matches the given value</returns>
 		public bool IsValue(char value)
 		{
-			return this.Value.Span.Length == 1 && this.Value.Span[0] == value;
+			return Value.Span.Length == 1 && Value.Span[0] == value;
 		}
 
 		/// <summary>
@@ -223,7 +223,7 @@ namespace EpicGames.UHT.Tokenizer
 		/// <returns>True if the value matches</returns>
 		public bool IsValue(string value, bool ignoreCase = false)
 		{
-			return this.Value.Span.Equals(value, ignoreCase ? StringComparison.OrdinalIgnoreCase : StringComparison.Ordinal);
+			return Value.Span.Equals(value, ignoreCase ? StringComparison.OrdinalIgnoreCase : StringComparison.Ordinal);
 		}
 
 		/// <summary>
@@ -234,7 +234,7 @@ namespace EpicGames.UHT.Tokenizer
 		/// <returns>True if the value matches</returns>
 		public bool IsValue(StringView value, bool ignoreCase = false)
 		{
-			return this.Value.Span.Equals(value.Span, ignoreCase ? StringComparison.OrdinalIgnoreCase : StringComparison.Ordinal);
+			return Value.Span.Equals(value.Span, ignoreCase ? StringComparison.OrdinalIgnoreCase : StringComparison.Ordinal);
 		}
 
 		/// <summary>
@@ -245,7 +245,7 @@ namespace EpicGames.UHT.Tokenizer
 		/// <returns>True is the value starts with the given string</returns>
 		public bool ValueStartsWith(string value, bool ignoreCase = false)
 		{
-			return this.Value.Span.StartsWith(value, ignoreCase ? StringComparison.OrdinalIgnoreCase : StringComparison.Ordinal);
+			return Value.Span.StartsWith(value, ignoreCase ? StringComparison.OrdinalIgnoreCase : StringComparison.Ordinal);
 		}
 
 		/// <summary>
@@ -254,7 +254,7 @@ namespace EpicGames.UHT.Tokenizer
 		/// <returns>True if the token is an identifier</returns>
 		public bool IsIdentifier()
 		{
-			return this.TokenType == UhtTokenType.Identifier;
+			return TokenType == UhtTokenType.Identifier;
 		}
 
 		/// <summary>
@@ -285,7 +285,7 @@ namespace EpicGames.UHT.Tokenizer
 		/// <returns>True if the token is a symbol</returns>
 		public bool IsSymbol()
 		{
-			return this.TokenType == UhtTokenType.Symbol;
+			return TokenType == UhtTokenType.Symbol;
 		}
 
 		/// <summary>
@@ -324,7 +324,7 @@ namespace EpicGames.UHT.Tokenizer
 		/// <returns>True if constant integer</returns>
 		public bool IsConstInt()
 		{
-			return this.TokenType == UhtTokenType.DecimalConst || this.TokenType == UhtTokenType.HexConst;
+			return TokenType == UhtTokenType.DecimalConst || TokenType == UhtTokenType.HexConst;
 		}
 
 		/// <summary>
@@ -333,7 +333,7 @@ namespace EpicGames.UHT.Tokenizer
 		/// <returns>True if constant float</returns>
 		public bool IsConstFloat()
 		{
-			return this.TokenType == UhtTokenType.FloatConst;
+			return TokenType == UhtTokenType.FloatConst;
 		}
 
 		/// <summary>
@@ -444,7 +444,7 @@ namespace EpicGames.UHT.Tokenizer
 		/// <returns>True if the token is a string or character constant</returns>
 		public bool IsConstString()
 		{
-			return this.TokenType == UhtTokenType.StringConst || this.TokenType == UhtTokenType.CharConst;
+			return TokenType == UhtTokenType.StringConst || TokenType == UhtTokenType.CharConst;
 		}
 
 		// Return the token value for string constants
@@ -457,14 +457,14 @@ namespace EpicGames.UHT.Tokenizer
 		/// <exception cref="UhtException">Thrown if the token type is not a string or character constant</exception>
 		public StringView GetUnescapedString(IUhtMessageSite messageSite)
 		{
-			switch (this.TokenType)
+			switch (TokenType)
 			{
 				case UhtTokenType.StringConst:
-					ReadOnlySpan<char> span = this.Value.Span[1..^1];
+					ReadOnlySpan<char> span = Value.Span[1..^1];
 					int index = span.IndexOf('\\');
 					if (index == -1)
 					{
-						return new StringView(this.Value, 1, span.Length);
+						return new StringView(Value, 1, span.Length);
 					}
 					else
 					{
@@ -484,9 +484,9 @@ namespace EpicGames.UHT.Tokenizer
 					}
 
 				case UhtTokenType.CharConst:
-					if (this.Value.Span[1] == '\\')
+					if (Value.Span[1] == '\\')
 					{
-						switch (this.Value.Span[2])
+						switch (Value.Span[2])
 						{
 							case 't':
 								return new StringView("\t");
@@ -495,16 +495,16 @@ namespace EpicGames.UHT.Tokenizer
 							case 'r':
 								return new StringView("\r");
 							default:
-								return new StringView(this.Value, 2, 1);
+								return new StringView(Value, 2, 1);
 						}
 					}
 					else
 					{
-						return new StringView(this.Value, 1, 1);
+						return new StringView(Value, 1, 1);
 					}
 			}
 
-			throw new UhtException(messageSite, this.InputLine, "Call to GetUnescapedString on token that isn't a string or char constant");
+			throw new UhtException(messageSite, InputLine, "Call to GetUnescapedString on token that isn't a string or char constant");
 		}
 
 		/// <summary>
@@ -514,7 +514,7 @@ namespace EpicGames.UHT.Tokenizer
 		/// <returns>Resulting string</returns>
 		public StringView GetConstantValue(bool respectQuotes = false)
 		{
-			switch (this.TokenType)
+			switch (TokenType)
 			{
 				case UhtTokenType.DecimalConst:
 					return new StringView(GetDecimalValue().ToString(NumberFormatInfo.InvariantInfo));
@@ -538,10 +538,10 @@ namespace EpicGames.UHT.Tokenizer
 		public StringView GetTokenString(bool respectQuotes = false)
 		{
 			StringViewBuilder builder = new();
-			switch (this.TokenType)
+			switch (TokenType)
 			{
 				case UhtTokenType.StringConst:
-					StringView subValue = new(this.Value, 1, this.Value.Span.Length - 2);
+					StringView subValue = new(Value, 1, Value.Span.Length - 2);
 					while (subValue.Span.Length > 0)
 					{
 						int slashIndex = subValue.Span.IndexOf('\\');
@@ -577,10 +577,10 @@ namespace EpicGames.UHT.Tokenizer
 					break;
 
 				case UhtTokenType.CharConst:
-					char charConst = this.Value.Span[1];
+					char charConst = Value.Span[1];
 					if (charConst == '\\')
 					{
-						charConst = this.Value.Span[2];
+						charConst = Value.Span[2];
 						switch (charConst)
 						{
 							case 't':
@@ -713,7 +713,7 @@ namespace EpicGames.UHT.Tokenizer
 
 		long GetDecimalValue()
 		{
-			ReadOnlySpan<char> span = this.Value.Span;
+			ReadOnlySpan<char> span = Value.Span;
 			bool isUnsigned = false;
 			while (span.Length > 0)
 			{
@@ -737,7 +737,7 @@ namespace EpicGames.UHT.Tokenizer
 
 		long GetHexValue()
 		{
-			return Convert.ToInt64(this.Value.ToString(), 16);
+			return Convert.ToInt64(Value.ToString(), 16);
 		}
 	}
 }
diff --git a/Engine/Source/Programs/Shared/EpicGames.UHT/Tokenizer/UhtTokenBufferReader.cs b/Engine/Source/Programs/Shared/EpicGames.UHT/Tokenizer/UhtTokenBufferReader.cs
--- a/Engine/Source/Programs/Shared/EpicGames.UHT/Tokenizer/UhtTokenBufferReader.cs
+++ b/Engine/Source/Programs/Shared/EpicGames.UHT/Tokenizer/UhtTokenBufferReader.cs
@@ -44,18 +44,18 @@ namespace EpicGames.UHT.Tokenizer
 		/// <param name="input">Input source</param>
 		public UhtTokenBufferReader(IUhtMessageSite messageSite, ReadOnlyMemory<char> input)
 		{
-			this._messageSite = messageSite;
-			this._data = input;
+			_messageSite = messageSite;
+			_data = input;
 		}
 
 		#region IUHTMessageSite Implementation
-		IUhtMessageSession IUhtMessageSite.MessageSession => this._messageSite.MessageSession;
-		IUhtMessageSource? IUhtMessageSite.MessageSource => this._messageSite.MessageSource;
+		IUhtMessageSession IUhtMessageSite.MessageSession => _messageSite.MessageSession;
+		IUhtMessageSource? IUhtMessageSite.MessageSource => _messageSite.MessageSource;
 		IUhtMessageLineNumber? IUhtMessageSite.MessageLineNumber => this;
 		#endregion
 
 		#region IUHTMessageLinenumber Implementation
-		int IUhtMessageLineNumber.MessageLineNumber => this.InputLine;
+		int IUhtMessageLineNumber.MessageLineNumber => InputLine;
 		#endregion
 
 		#region ITokenReader Implementation
@@ -64,28 +64,28 @@ namespace EpicGames.UHT.Tokenizer
 		{
 			get
 			{
-				if (this._hasToken)
+				if (_hasToken)
 				{
-					return this._currentToken.TokenType.IsEndType();
+					return _currentToken.TokenType.IsEndType();
 				}
 				else
 				{
-					return this.InputPos == this._data.Length;
+					return InputPos == _data.Length;
 				}
 			}
 		}
 
 		/// <inheritdoc/>
-		public int InputPos => this._hasToken ? this._preCurrentTokenInputPos : this._inputPos;
+		public int InputPos => _hasToken ? _preCurrentTokenInputPos : _inputPos;
 
 		/// <inheritdoc/>
 		public int InputLine
 		{
-			get => this._hasToken ? this._preCurrentTokenInputLine : this._inputLine;
+			get => _hasToken ? _preCurrentTokenInputLine : _inputLine;
 			set
 			{
 				ClearToken();
-				this._inputLine = value;
+				_inputLine = value;
 			}
 		}
 
@@ -94,9 +94,9 @@ namespace EpicGames.UHT.Tokenizer
 		{
 			get
 			{
-				if (this._comments != null && this._committedComments != 0)
+				if (_comments != null && _committedComments != 0)
 				{
-					return new ReadOnlySpan<StringView>(this._comments.ToArray(), 0, this._committedComments);
+					return new ReadOnlySpan<StringView>(_comments.ToArray(), 0, _committedComments);
 				}
 				else
 				{
@@ -106,17 +106,17 @@ namespace EpicGames.UHT.Tokenizer
 		}
 
 		/// <inheritdoc/>
-		public IUhtTokenPreprocessor? TokenPreprocessor { get => this._tokenPreprocessor; set => this._tokenPreprocessor = value; }
+		public IUhtTokenPreprocessor? TokenPreprocessor { get => _tokenPreprocessor; set => _tokenPreprocessor = value; }
 
 		/// <inheritdoc/>
 		public ref UhtToken PeekToken()
 		{
-			if (!this._hasToken)
+			if (!_hasToken)
 			{
-				this._currentToken = GetTokenInternal(true);
-				this._hasToken = true;
+				_currentToken = GetTokenInternal(true);
+				_hasToken = true;
 			}
-			return ref this._currentToken;
+			return ref _currentToken;
 		}
 
 		/// <inheritdoc/>
@@ -129,18 +129,18 @@ namespace EpicGames.UHT.Tokenizer
 		/// <inheritdoc/>
 		public void ConsumeToken()
 		{
-			if (this._recordTokens && !this._currentToken.IsEndType())
+			if (_recordTokens && !_currentToken.IsEndType())
 			{
-				this._recordedTokens.Add(this._currentToken);
+				_recordedTokens.Add(_currentToken);
 			}
-			this._hasToken = false;
+			_hasToken = false;
 
 			// When comments are disabled, we are still collecting comments, but aren't committing them
-			if (this._commentsDisableCount == 0)
+			if (_commentsDisableCount == 0)
 			{
-				if (this._comments != null)
+				if (_comments != null)
 				{
-					this._committedComments = this._comments.Count;
+					_committedComments = _comments.Count;
 				}
 			}
 			else
@@ -160,12 +160,12 @@ namespace EpicGames.UHT.Tokenizer
 		/// <inheritdoc/>
 		public StringView GetRawString(char terminator, UhtRawStringOptions options)
 		{
-			ReadOnlySpan<char> span = this._data.Span;
+			ReadOnlySpan<char> span = _data.Span;
 
 			ClearToken();
 			SkipWhitespaceAndComments();
 
-			int startPos = this.InputPos;
+			int startPos = InputPos;
 			bool inQuotes = false;
 			while (true)
 			{
@@ -180,7 +180,7 @@ namespace EpicGames.UHT.Tokenizer
 				// Check for end of line
 				if (c == '\r' || c == '\n')
 				{
-					--this._inputPos;
+					--_inputPos;
 					break;
 				}
 
@@ -189,7 +189,7 @@ namespace EpicGames.UHT.Tokenizer
 				{
 					if (options.HasAnyFlags(UhtRawStringOptions.DontConsumeTerminator))
 					{
-						--this._inputPos;
+						--_inputPos;
 					}
 					break;
 				}
@@ -200,7 +200,7 @@ namespace EpicGames.UHT.Tokenizer
 					char p = InternalPeekChar(span);
 					if (p == '*' || p == '/')
 					{
-						--this._inputPos;
+						--_inputPos;
 						break;
 					}
 				}
@@ -219,7 +219,7 @@ namespace EpicGames.UHT.Tokenizer
 			}
 
 			// Remove trailing whitespace
-			int endPos = this.InputPos;
+			int endPos = InputPos;
 			for (; endPos > startPos; --endPos)
 			{
 				char c = span[endPos - 1];
@@ -234,24 +234,24 @@ namespace EpicGames.UHT.Tokenizer
 			{
 				throw new UhtException(this, $"String exceeds maximum of {UhtToken.MaxStringLength} characters");
 			}
-			return new StringView(this._data, startPos, endPos - startPos);
+			return new StringView(_data, startPos, endPos - startPos);
 		}
 
 		/// <inheritdoc/>
 		public UhtToken GetLine()
 		{
-			ReadOnlySpan<char> span = this._data.Span;
+			ReadOnlySpan<char> span = _data.Span;
 
 			ClearToken();
-			this._prevPos = this._inputPos;
-			this._prevLine = this._inputLine;
+			_prevPos = _inputPos;
+			_prevLine = _inputLine;
 
-			if (this._prevPos == span.Length)
+			if (_prevPos == span.Length)
 			{
-				return new UhtToken(UhtTokenType.EndOfFile, this._prevPos, this._prevLine, this._prevPos, this._prevLine, new StringView(this._data.Slice(this._prevPos, 0)));
+				return new UhtToken(UhtTokenType.EndOfFile, _prevPos, _prevLine, _prevPos, _prevLine, new StringView(_data.Slice(_prevPos, 0)));
 			}
 
-			int lastPos = this._inputPos;
+			int lastPos = _inputPos;
 			while (true)
 			{
 				char c = InternalGetChar(span);
@@ -264,150 +264,150 @@ namespace EpicGames.UHT.Tokenizer
 				}
 				else if (c == '\n')
 				{
-					++this._inputLine;
+					++_inputLine;
 					break;
 				}
 				else
 				{
-					lastPos = this._inputPos;
+					lastPos = _inputPos;
 				}
 			}
 
-			return new UhtToken(UhtTokenType.Line, this._prevPos, this._prevLine, this._prevPos, this._prevLine, new StringView(this._data[this._prevPos..lastPos]));
+			return new UhtToken(UhtTokenType.Line, _prevPos, _prevLine, _prevPos, _prevLine, new StringView(_data[_prevPos..lastPos]));
 		}
 
 		/// <inheritdoc/>
 		public StringView GetStringView(int startPos, int count)
 		{
-			return new StringView(this._data, startPos, count);
+			return new StringView(_data, startPos, count);
 		}
 
 		/// <inheritdoc/>
 		public void ClearComments()
 		{
-			if (this._comments != null)
+			if (_comments != null)
 			{
 
 				// Clearing comments does not remove any uncommitted comments
-				this._comments.RemoveRange(0, this._committedComments);
-				this._committedComments = 0;
+				_comments.RemoveRange(0, _committedComments);
+				_committedComments = 0;
 			}
 		}
 
 		/// <inheritdoc/>
 		public void DisableComments()
 		{
-			++this._commentsDisableCount;
+			++_commentsDisableCount;
 		}
 
 		/// <inheritdoc/>
 		public void EnableComments()
 		{
-			--this._commentsDisableCount;
+			--_commentsDisableCount;
 		}
 
 		/// <inheritdoc/>
 		public void CommitPendingComments()
 		{
-			if (this._comments != null)
+			if (_comments != null)
 			{
-				this._committedComments = this._comments.Count;
+				_committedComments = _comments.Count;
 			}
 		}
 
 		/// <inheritdoc/>
 		public bool IsFirstTokenInLine(ref UhtToken token)
 		{
-			return IsFirstTokenInLine(this._data.Span, token.InputStartPos);
+			return IsFirstTokenInLine(_data.Span, token.InputStartPos);
 		}
 
 		/// <inheritdoc/>
 		public void SaveState()
 		{
-			if (this._hasSavedState)
+			if (_hasSavedState)
 			{
 				throw new UhtIceException("Can not save more than one state");
 			}
-			this._hasSavedState = true;
-			this._savedInputLine = this.InputLine;
-			this._savedInputPos = this.InputPos;
-			if (this._comments != null)
+			_hasSavedState = true;
+			_savedInputLine = InputLine;
+			_savedInputPos = InputPos;
+			if (_comments != null)
 			{
-				if (this._savedComments == null)
+				if (_savedComments == null)
 				{
-					this._savedComments = new List<StringView>();
+					_savedComments = new List<StringView>();
 				}
-				this._savedComments.Clear();
-				this._savedComments.AddRange(this._comments);
+				_savedComments.Clear();
+				_savedComments.AddRange(_comments);
 			}
-			if (this._tokenPreprocessor != null)
+			if (_tokenPreprocessor != null)
 			{
-				this._tokenPreprocessor.SaveState();
+				_tokenPreprocessor.SaveState();
 			}
 		}
 
 		/// <inheritdoc/>
 		public void RestoreState()
 		{
-			if (!this._hasSavedState)
+			if (!_hasSavedState)
 			{
 				throw new UhtIceException("Can not restore state when none have been saved");
 			}
-			this._hasSavedState = false;
+			_hasSavedState = false;
 			ClearToken();
-			this._inputPos = this._savedInputPos;
-			this._inputLine = this._savedInputLine;
-			if (this._savedComments != null && this._comments != null)
+			_inputPos = _savedInputPos;
+			_inputLine = _savedInputLine;
+			if (_savedComments != null && _comments != null)
 			{
-				this._comments.Clear();
-				this._comments.AddRange(this._savedComments);
+				_comments.Clear();
+				_comments.AddRange(_savedComments);
 			}
-			if (this._tokenPreprocessor != null)
+			if (_tokenPreprocessor != null)
 			{
-				this._tokenPreprocessor.RestoreState();
+				_tokenPreprocessor.RestoreState();
 			}
 		}
 
 		/// <inheritdoc/>
 		public void AbandonState()
 		{
-			if (!this._hasSavedState)
+			if (!_hasSavedState)
 			{
 				throw new UhtIceException("Can not abandon state when none have been saved");
 			}
-			this._hasSavedState = false;
+			_hasSavedState = false;
 			ClearToken();
 		}
 
 		/// <inheritdoc/>
 		public void EnableRecording()
 		{
-			if (this._recordTokens)
+			if (_recordTokens)
 			{
 				throw new UhtIceException("Can not nest token recording");
 			}
-			this._recordTokens = true;
+			_recordTokens = true;
 		}
 
 		/// <inheritdoc/>
 		public void RecordToken(ref UhtToken token)
 		{
-			if (!this._recordTokens)
+			if (!_recordTokens)
 			{
 				throw new UhtIceException("Attempt to disable token recording when it isn't enabled");
 			}
-			this._recordedTokens.Add(token);
+			_recordedTokens.Add(token);
 		}
 
 		/// <inheritdoc/>
 		public void DisableRecording()
 		{
-			if (!this._recordTokens)
+			if (!_recordTokens)
 			{
 				throw new UhtIceException("Attempt to disable token recording when it isn't enabled");
 			}
-			this._recordedTokens.Clear();
-			this._recordTokens = false;
+			_recordedTokens.Clear();
+			_recordTokens = false;
 		}
 
 		/// <inheritdoc/>
@@ -415,11 +415,11 @@ namespace EpicGames.UHT.Tokenizer
 		{
 			get
 			{
-				if (!this._recordTokens)
+				if (!_recordTokens)
 				{
 					throw new UhtIceException("Attempt to get recorded tokens when it isn't enabled");
 				}
-				return this._recordedTokens;
+				return _recordedTokens;
 			}
 		}
 		#endregion
@@ -427,19 +427,19 @@ namespace EpicGames.UHT.Tokenizer
 		#region Internals
 		private void ClearAllComments()
 		{
-			if (this._comments != null)
+			if (_comments != null)
 			{
-				this._comments.Clear();
-				this._committedComments = 0;
+				_comments.Clear();
+				_committedComments = 0;
 			}
 		}
 
 		private void ClearPendingComments()
 		{
-			if (this._comments != null)
+			if (_comments != null)
 			{
-				int startingComment = this._committedComments + this._preprocessorPendingCommentsCount;
-				this._comments.RemoveRange(startingComment, this._comments.Count - startingComment);
+				int startingComment = _committedComments + _preprocessorPendingCommentsCount;
+				_comments.RemoveRange(startingComment, _comments.Count - startingComment);
 			}
 		}
 
@@ -468,18 +468,18 @@ namespace EpicGames.UHT.Tokenizer
 		/// <returns>Return the next token from the stream.  If the end of stream is reached, a token type of None will be returned.</returns>
 		private UhtToken GetTokenInternal(bool enablePreprocessor)
 		{
-			ReadOnlySpan<char> span = this._data.Span;
+			ReadOnlySpan<char> span = _data.Span;
 			bool gotInlineComment = false;
 Restart:
-			this._preCurrentTokenInputLine = this._inputLine;
-			this._preCurrentTokenInputPos = this._inputPos;
+			_preCurrentTokenInputLine = _inputLine;
+			_preCurrentTokenInputPos = _inputPos;
 			SkipWhitespaceAndCommentsInternal(ref gotInlineComment, enablePreprocessor);
-			this._prevPos = this._inputPos;
-			this._prevLine = this._inputLine;
+			_prevPos = _inputPos;
+			_prevLine = _inputLine;
 
 			UhtTokenType tokenType = UhtTokenType.EndOfFile;
-			int startPos = this._inputPos;
-			int startLine = this._inputLine;
+			int startPos = _inputPos;
+			int startLine = _inputLine;
 
 			char c = InternalGetChar(span);
 			if (c == 0)
@@ -489,16 +489,16 @@ Restart:
 			else if (UhtFCString.IsAlpha(c) || c == '_')
 			{
 
-				for (; this._inputPos < span.Length; ++this._inputPos)
+				for (; _inputPos < span.Length; ++_inputPos)
 				{
-					c = span[this._inputPos];
+					c = span[_inputPos];
 					if (!(UhtFCString.IsAlpha(c) || UhtFCString.IsDigit(c) || c == '_'))
 					{
 						break;
 					}
 				}
 
-				if (this._inputPos - startPos >= UhtToken.MaxNameLength)
+				if (_inputPos - startPos >= UhtToken.MaxNameLength)
 				{
 					throw new UhtException(this, $"Identifier length exceeds maximum of {UhtToken.MaxNameLength}");
 				}
@@ -533,7 +533,7 @@ Restart:
 				// If we have a hex constant
 				if (isHex)
 				{
-					for (; this._inputPos < span.Length && UhtFCString.IsHexDigit(span[this._inputPos]); ++this._inputPos)
+					for (; _inputPos < span.Length && UhtFCString.IsHexDigit(span[_inputPos]); ++_inputPos)
 					{
 					}
 				}
@@ -543,56 +543,56 @@ Restart:
 				{
 
 					// Skip all digits
-					for (; this._inputPos < span.Length && UhtFCString.IsDigit(span[this._inputPos]); ++this._inputPos)
+					for (; _inputPos < span.Length && UhtFCString.IsDigit(span[_inputPos]); ++_inputPos)
 					{
 					}
 
 					// If we have a '.'
-					if (this._inputPos < span.Length && span[this._inputPos] == '.')
+					if (_inputPos < span.Length && span[_inputPos] == '.')
 					{
 						isFloat = true;
-						++this._inputPos;
+						++_inputPos;
 
 						// Skip all digits
-						for (; this._inputPos < span.Length && UhtFCString.IsDigit(span[this._inputPos]); ++this._inputPos)
+						for (; _inputPos < span.Length && UhtFCString.IsDigit(span[_inputPos]); ++_inputPos)
 						{
 						}
 					}
 
 					// If we have a 'e'
-					if (this._inputPos < span.Length && UhtFCString.IsExponentMarker(span[this._inputPos]))
+					if (_inputPos < span.Length && UhtFCString.IsExponentMarker(span[_inputPos]))
 					{
 						isFloat = true;
-						++this._inputPos;
+						++_inputPos;
 
 						// Skip any signs
-						if (this._inputPos < span.Length && UhtFCString.IsSign(span[this._inputPos]))
+						if (_inputPos < span.Length && UhtFCString.IsSign(span[_inputPos]))
 						{
-							++this._inputPos;
+							++_inputPos;
 						}
 
 						// Skip all digits
-						for (; this._inputPos < span.Length && UhtFCString.IsDigit(span[this._inputPos]); ++this._inputPos)
+						for (; _inputPos < span.Length && UhtFCString.IsDigit(span[_inputPos]); ++_inputPos)
 						{
 						}
 					}
 
 					// If we have a 'f'
-					if (this._inputPos < span.Length && UhtFCString.IsFloatMarker(span[this._inputPos]))
+					if (_inputPos < span.Length && UhtFCString.IsFloatMarker(span[_inputPos]))
 					{
 						isFloat = true;
-						++this._inputPos;
+						++_inputPos;
 					}
 
 					// Check for u/l markers
-					while (this._inputPos < span.Length &&
-						(UhtFCString.IsUnsignedMarker(span[this._inputPos]) || UhtFCString.IsLongMarker(span[this._inputPos])))
+					while (_inputPos < span.Length &&
+						(UhtFCString.IsUnsignedMarker(span[_inputPos]) || UhtFCString.IsLongMarker(span[_inputPos])))
 					{
-						++this._inputPos;
+						++_inputPos;
 					}
 				}
 
-				if (this._inputPos - startPos >= UhtToken.MaxNameLength)
+				if (_inputPos - startPos >= UhtToken.MaxNameLength)
 				{
 					throw new UhtException(this, $"Number length exceeds maximum of {UhtToken.MaxNameLength}");
 				}
@@ -646,7 +646,7 @@ Restart:
 					}
 				}
 
-				if (this._inputPos - startPos >= UhtToken.MaxStringLength)
+				if (_inputPos - startPos >= UhtToken.MaxStringLength)
 				{
 					throw new UhtException(this, $"String constant exceeds maximum of {UhtToken.MaxStringLength} characters");
 				}
@@ -659,9 +659,9 @@ Restart:
 			else
 			{
 				{
-					if (this._inputPos < span.Length)
+					if (_inputPos < span.Length)
 					{
-						char d = span[this._inputPos];
+						char d = span[_inputPos];
 						if ((c == '<' && d == '<')
 							|| (c == '!' && d == '=')
 							|| (c == '<' && d == '=')
@@ -680,7 +680,7 @@ Restart:
 							|| (c == '~' && d == '=')
 							|| (c == ':' && d == ':'))
 						{
-							++this._inputPos;
+							++_inputPos;
 						}
 					}
 
@@ -691,14 +691,14 @@ Restart:
 					// 1) If the block is being skipped, we want to preserve any comments ahead of the token and eliminate all other comments found in the #if block
 					// 2) If the block isn't being skipped or in the case of things such as #pragma or #include, they were considered statements in the UHT
 					//	  and would result in the comments being dropped.  In the new UHT, we just eliminate all pending tokens.
-					if (c == '#' && enablePreprocessor && this._tokenPreprocessor != null && IsFirstTokenInLine(span, startPos))
+					if (c == '#' && enablePreprocessor && _tokenPreprocessor != null && IsFirstTokenInLine(span, startPos))
 					{
-						this._preprocessorPendingCommentsCount = this._comments != null ? this._comments.Count - this._committedComments : 0;
-						UhtToken temp = new(tokenType, startPos, startLine, startPos, this._inputLine, new StringView(this._data[startPos..this._inputPos]));
-						bool include = this._tokenPreprocessor.ParsePreprocessorDirective(ref temp, true, out bool clearComments, out bool illegalContentsCheck);
+						_preprocessorPendingCommentsCount = _comments != null ? _comments.Count - _committedComments : 0;
+						UhtToken temp = new(tokenType, startPos, startLine, startPos, _inputLine, new StringView(_data[startPos.._inputPos]));
+						bool include = _tokenPreprocessor.ParsePreprocessorDirective(ref temp, true, out bool clearComments, out bool illegalContentsCheck);
 						if (!include)
 						{
-							++this._commentsDisableCount;
+							++_commentsDisableCount;
 							int checkStateMachine = 0;
 							while (true)
 							{
@@ -710,7 +710,7 @@ Restart:
 								ReadOnlySpan<char> localValueSpan = localToken.Value.Span;
 								if (localToken.IsSymbol('#') && localValueSpan.Length == 1 && localValueSpan[0] == '#' && IsFirstTokenInLine(span, localToken.InputStartPos))
 								{
-									if (this._tokenPreprocessor.ParsePreprocessorDirective(ref temp, false, out bool _, out bool scratchIllegalContentsCheck))
+									if (_tokenPreprocessor.ParsePreprocessorDirective(ref temp, false, out bool _, out bool scratchIllegalContentsCheck))
 									{
 										break;
 									}
@@ -807,12 +807,12 @@ ResetStateMachineCheck:
 									}
 								}
 							}
-							--this._commentsDisableCount;
+							--_commentsDisableCount;
 
 							// Clear any extra pending comments at this time
 							ClearPendingComments();
 						}
-						this._preprocessorPendingCommentsCount = 0;
+						_preprocessorPendingCommentsCount = 0;
 
 						// Depending on the type of directive and/or the #if expression, we clear comments.
 						if (clearComments)
@@ -826,7 +826,7 @@ ResetStateMachineCheck:
 				tokenType = UhtTokenType.Symbol;
 			}
 
-			return new UhtToken(tokenType, startPos, startLine, startPos, this._inputLine, new StringView(this._data[startPos..this._inputPos]));
+			return new UhtToken(tokenType, startPos, startLine, startPos, _inputLine, new StringView(_data[startPos.._inputPos]));
 		}
 
 		/// <summary>
@@ -834,7 +834,7 @@ ResetStateMachineCheck:
 		/// </summary>
 		private void SkipWhitespaceAndCommentsInternal(ref bool persistGotInlineComment, bool enablePreprocessor)
 		{
-			ReadOnlySpan<char> span = this._data.Span;
+			ReadOnlySpan<char> span = _data.Span;
 
 			bool gotNewlineBetweenComments = false;
 			bool gotInlineComment = persistGotInlineComment;
@@ -848,7 +848,7 @@ ResetStateMachineCheck:
 				else if (c == '\n')
 				{
 					gotNewlineBetweenComments |= gotInlineComment;
-					++this._inputLine;
+					++_inputLine;
 				}
 				else if (c == '\r' || c == '\t' || c == ' ')
 				{
@@ -862,8 +862,8 @@ ResetStateMachineCheck:
 						{
 							ClearAllComments();
 						}
-						int commentStart = this._inputPos - 1;
-						++this._inputPos;
+						int commentStart = _inputPos - 1;
+						++_inputPos;
 						for (; ; )
 						{
 							char commentChar = InternalGetChar(span);
@@ -877,17 +877,17 @@ ResetStateMachineCheck:
 							}
 							else if (commentChar == '\n')
 							{
-								++this._inputLine;
+								++_inputLine;
 							}
 							else if (commentChar == '*' && InternalPeekChar(span) == '/')
 							{
-								++this._inputPos;
+								++_inputPos;
 								break;
 							}
 						}
 						if (enablePreprocessor)
 						{
-							AddComment(new StringView(this._data[commentStart..this._inputPos]));
+							AddComment(new StringView(_data[commentStart.._inputPos]));
 						}
 					}
 					else if (nextChar == '/')
@@ -901,8 +901,8 @@ ResetStateMachineCheck:
 							}
 						}
 						gotInlineComment = true;
-						int commentStart = this._inputPos - 1;
-						++this._inputPos;
+						int commentStart = _inputPos - 1;
+						++_inputPos;
 
 						// Scan to the end of the line
 						for (; ; )
@@ -924,18 +924,18 @@ ResetStateMachineCheck:
 						}
 						if (enablePreprocessor)
 						{
-							AddComment(new StringView(this._data[commentStart..this._inputPos]));
+							AddComment(new StringView(_data[commentStart.._inputPos]));
 						}
 					}
 					else
 					{
-						--this._inputPos;
+						--_inputPos;
 						break;
 					}
 				}
 				else
 				{
-					--this._inputPos;
+					--_inputPos;
 					break;
 				}
 			}
@@ -954,29 +954,29 @@ ResetStateMachineCheck:
 			// Check for [+-]...
 			if (UhtFCString.IsSign(c))
 			{
-				if (this._inputPos == span.Length)
+				if (_inputPos == span.Length)
 				{
 					return false;
 				}
 
 				// Check for [+-][0..9]...
-				if (UhtFCString.IsDigit(span[this._inputPos]))
+				if (UhtFCString.IsDigit(span[_inputPos]))
 				{
 					return true;
 				}
 
 				// Check for [+-][.][0..9]...
-				if (span[this._inputPos] != '.')
+				if (span[_inputPos] != '.')
 				{
 					return false;
 				}
 
-				if (this._inputPos + 1 == span.Length)
+				if (_inputPos + 1 == span.Length)
 				{
 					return false;
 				}
 
-				if (UhtFCString.IsDigit(span[this._inputPos + 1]))
+				if (UhtFCString.IsDigit(span[_inputPos + 1]))
 				{
 					return true;
 				}
@@ -985,13 +985,13 @@ ResetStateMachineCheck:
 
 			if (c == '.')
 			{
-				if (this._inputPos == span.Length)
+				if (_inputPos == span.Length)
 				{
 					return false;
 				}
 
 				// Check for [.][0..9]...
-				if (UhtFCString.IsDigit(span[this._inputPos]))
+				if (UhtFCString.IsDigit(span[_inputPos]))
 				{
 					return true;
 				}
@@ -1007,7 +1007,7 @@ ResetStateMachineCheck:
 		/// <returns>Next character in the stream or zero</returns>
 		private char InternalPeekChar(ReadOnlySpan<char> span)
 		{
-			return this._inputPos < span.Length ? span[this._inputPos] : '\0';
+			return _inputPos < span.Length ? span[_inputPos] : '\0';
 		}
 
 		/// <summary>
@@ -1018,7 +1018,7 @@ ResetStateMachineCheck:
 		/// <returns>Next character in the stream or zero</returns>
 		private char InternalGetChar(ReadOnlySpan<char> span)
 		{
-			return this._inputPos < span.Length ? span[this._inputPos++] : '\0';
+			return _inputPos < span.Length ? span[_inputPos++] : '\0';
 		}
 
 		/// <summary>
@@ -1026,25 +1026,25 @@ ResetStateMachineCheck:
 		/// </summary>
 		private void ClearToken()
 		{
-			if (this._hasToken)
+			if (_hasToken)
 			{
-				this._hasToken = false;
-				if (this._comments != null && this._comments.Count > this._committedComments)
+				_hasToken = false;
+				if (_comments != null && _comments.Count > _committedComments)
 				{
-					this._comments.RemoveRange(this._committedComments, this._comments.Count - this._committedComments);
+					_comments.RemoveRange(_committedComments, _comments.Count - _committedComments);
 				}
-				this._inputPos = this._currentToken.UngetPos;
-				this._inputLine = this._currentToken.UngetLine;
+				_inputPos = _currentToken.UngetPos;
+				_inputLine = _currentToken.UngetLine;
 			}
 		}
 
 		private void AddComment(StringView comment)
 		{
-			if (this._comments == null)
+			if (_comments == null)
 			{
-				this._comments = new List<StringView>(4);
+				_comments = new List<StringView>(4);
 			}
-			this._comments.Add(comment);
+			_comments.Add(comment);
 		}
 		#endregion
 	}
diff --git a/Engine/Source/Programs/Shared/EpicGames.UHT/Tokenizer/UhtTokenReader.cs b/Engine/Source/Programs/Shared/EpicGames.UHT/Tokenizer/UhtTokenReader.cs
--- a/Engine/Source/Programs/Shared/EpicGames.UHT/Tokenizer/UhtTokenReader.cs
+++ b/Engine/Source/Programs/Shared/EpicGames.UHT/Tokenizer/UhtTokenReader.cs
@@ -270,7 +270,7 @@ namespace EpicGames.UHT.Tokenizer
 		/// <returns></returns>
 		public void Join(StringBuilder builder, string seperator = "")
 		{
-			builder.Append(this.Token.Value.ToString());
+			builder.Append(Token.Value.ToString());
 			UhtTokenList list = this;
 			while (list.Next != null)
 			{
@@ -287,9 +287,9 @@ namespace EpicGames.UHT.Tokenizer
 		/// <returns></returns>
 		public string Join(string seperator = "")
 		{
-			if (this.Next == null)
+			if (Next == null)
 			{
-				return this.Token.Value.ToString();
+				return Token.Value.ToString();
 			}
 			StringBuilder builder = new();
 			Join(builder, seperator);
@@ -308,7 +308,7 @@ namespace EpicGames.UHT.Tokenizer
 				++count;
 			}
 			UhtToken[] outTokens = new UhtToken[count];
-			outTokens[0] = this.Token;
+			outTokens[0] = Token;
 			count = 1;
 			for (UhtTokenList temp = this; temp.Next != null; temp = temp.Next)
 			{
@@ -417,8 +417,8 @@ namespace EpicGames.UHT.Tokenizer
 		/// <param name="tokenReader">Token reader to disable</param>
 		public UhtTokenDisableComments(IUhtTokenReader tokenReader)
 		{
-			this._tokenReader = tokenReader;
-			this._tokenReader.DisableComments();
+			_tokenReader = tokenReader;
+			_tokenReader.DisableComments();
 		}
 
 		/// <summary>
@@ -426,7 +426,7 @@ namespace EpicGames.UHT.Tokenizer
 		/// </summary>
 		public void Dispose()
 		{
-			this._tokenReader.EnableComments();
+			_tokenReader.EnableComments();
 		}
 	}
 
@@ -444,9 +444,9 @@ namespace EpicGames.UHT.Tokenizer
 		/// <param name="tokenReader">Token reader</param>
 		public UhtTokenSaveState(IUhtTokenReader tokenReader)
 		{
-			this._tokenReader = tokenReader;
-			this._handled = false;
-			this._tokenReader.SaveState();
+			_tokenReader = tokenReader;
+			_handled = false;
+			_tokenReader.SaveState();
 		}
 
 		/// <summary>
@@ -454,7 +454,7 @@ namespace EpicGames.UHT.Tokenizer
 		/// </summary>
 		public void Dispose()
 		{
-			if (!this._handled)
+			if (!_handled)
 			{
 				RestoreState();
 			}
@@ -466,12 +466,12 @@ namespace EpicGames.UHT.Tokenizer
 		/// <exception cref="UhtIceException">Thrown if state has already been restored or aborted</exception>
 		public void RestoreState()
 		{
-			if (this._handled)
+			if (_handled)
 			{
 				throw new UhtIceException("Can not call RestoreState/AbandonState more than once");
 			}
-			this._tokenReader.RestoreState();
-			this._handled = true;
+			_tokenReader.RestoreState();
+			_handled = true;
 		}
 
 		/// <summary>
@@ -480,12 +480,12 @@ namespace EpicGames.UHT.Tokenizer
 		/// <exception cref="UhtIceException">Thrown if state has already been restored or aborted</exception>
 		public void AbandonState()
 		{
-			if (this._handled)
+			if (_handled)
 			{
 				throw new UhtIceException("Can not call RestoreState/AbandonState more than once");
 			}
-			this._tokenReader.AbandonState();
-			this._handled = true;
+			_tokenReader.AbandonState();
+			_handled = true;
 		}
 	}
 }
diff --git a/Engine/Source/Programs/Shared/EpicGames.UHT/Tokenizer/UhtTokenReplayReader.cs b/Engine/Source/Programs/Shared/EpicGames.UHT/Tokenizer/UhtTokenReplayReader.cs
--- a/Engine/Source/Programs/Shared/EpicGames.UHT/Tokenizer/UhtTokenReplayReader.cs
+++ b/Engine/Source/Programs/Shared/EpicGames.UHT/Tokenizer/UhtTokenReplayReader.cs
@@ -42,11 +42,11 @@ namespace EpicGames.UHT.Tokenizer
 		/// <param name="endTokenType">Token type to return when end of tokens reached</param>
 		public UhtTokenReplayReader(IUhtMessageSite messageSite, ReadOnlyMemory<char> data, ReadOnlyMemory<UhtToken> tokens, UhtTokenType endTokenType)
 		{
-			this._messageSite = messageSite;
-			this._tokens = tokens;
-			this._data = data;
-			this._endTokenType = endTokenType;
-			this._currentToken = new UhtToken(endTokenType);
+			_messageSite = messageSite;
+			_tokens = tokens;
+			_data = data;
+			_endTokenType = endTokenType;
+			_currentToken = new UhtToken(endTokenType);
 		}
 
 		/// <summary>
@@ -54,9 +54,9 @@ namespace EpicGames.UHT.Tokenizer
 		/// </summary>
 		public UhtTokenReplayReader()
 		{
-			this._messageSite = new UhtEmptyMessageSite();
-			this._tokens = Array.Empty<UhtToken>().AsMemory();
-			this._data = Array.Empty<char>().AsMemory();
+			_messageSite = new UhtEmptyMessageSite();
+			_tokens = Array.Empty<UhtToken>().AsMemory();
+			_data = Array.Empty<char>().AsMemory();
 		}
 
 		/// <summary>
@@ -69,14 +69,14 @@ namespace EpicGames.UHT.Tokenizer
 		/// <returns>The replay reader</returns>
 		public UhtTokenReplayReader Reset(IUhtMessageSite messageSite, ReadOnlyMemory<char> data, ReadOnlyMemory<UhtToken> tokens, UhtTokenType endTokenType)
 		{
-			this._messageSite = messageSite;
-			this._tokens = tokens;
-			this._data = data;
-			this._currentTokenIndex = -1;
-			this._hasToken = false;
-			this._currentToken = new UhtToken(endTokenType);
-			this._savedStateCount = 0;
-			this._endTokenType = endTokenType;
+			_messageSite = messageSite;
+			_tokens = tokens;
+			_data = data;
+			_currentTokenIndex = -1;
+			_hasToken = false;
+			_currentToken = new UhtToken(endTokenType);
+			_savedStateCount = 0;
+			_endTokenType = endTokenType;
 			return this;
 		}
 
@@ -101,8 +101,8 @@ namespace EpicGames.UHT.Tokenizer
 		}
 
 		#region IUHTMessageSite Implementation
-		IUhtMessageSession IUhtMessageSite.MessageSession => this._messageSite.MessageSession;
-		IUhtMessageSource? IUhtMessageSite.MessageSource => this._messageSite.MessageSource;
+		IUhtMessageSession IUhtMessageSite.MessageSession => _messageSite.MessageSession;
+		IUhtMessageSource? IUhtMessageSite.MessageSource => _messageSite.MessageSource;
 		IUhtMessageLineNumber? IUhtMessageSite.MessageLineNumber => this;
 		#endregion
 
@@ -112,11 +112,11 @@ namespace EpicGames.UHT.Tokenizer
 		{
 			get
 			{
-				if (!this._hasToken)
+				if (!_hasToken)
 				{
 					GetTokenInternal();
 				}
-				return this._currentTokenIndex == this._tokens.Span.Length;
+				return _currentTokenIndex == _tokens.Span.Length;
 			}
 		}
 
@@ -125,7 +125,7 @@ namespace EpicGames.UHT.Tokenizer
 		{
 			get
 			{
-				if (!this._hasToken)
+				if (!_hasToken)
 				{
 					GetTokenInternal();
 				}
@@ -138,7 +138,7 @@ namespace EpicGames.UHT.Tokenizer
 		{
 			get
 			{
-				if (!this._hasToken)
+				if (!_hasToken)
 				{
 					GetTokenInternal();
 				}
@@ -165,7 +165,7 @@ namespace EpicGames.UHT.Tokenizer
 		/// <inheritdoc/>
 		public void ConsumeToken()
 		{
-			this._hasToken = false;
+			_hasToken = false;
 		}
 
 		/// <inheritdoc/>
@@ -201,7 +201,7 @@ namespace EpicGames.UHT.Tokenizer
 		/// <inheritdoc/>
 		public StringView GetStringView(int startPos, int count)
 		{
-			return new StringView(this._data, startPos, count);
+			return new StringView(_data, startPos, count);
 		}
 
 		/// <inheritdoc/>
@@ -221,11 +221,11 @@ namespace EpicGames.UHT.Tokenizer
 		/// <inheritdoc/>
 		public ref UhtToken PeekToken()
 		{
-			if (!this._hasToken)
+			if (!_hasToken)
 			{
 				GetTokenInternal();
 			}
-			return ref this._currentToken;
+			return ref _currentToken;
 		}
 
 		/// <inheritdoc/>
@@ -237,40 +237,40 @@ namespace EpicGames.UHT.Tokenizer
 		/// <inheritdoc/>
 		public void SaveState()
 		{
-			if (this._savedStateCount == MaxSavedStates)
+			if (_savedStateCount == MaxSavedStates)
 			{
 				throw new UhtIceException("Token reader saved states full");
 			}
-			this._savedStates[this._savedStateCount] = new SavedState { TokenIndex = this._currentTokenIndex, HasToken = this._hasToken };
-			++this._savedStateCount;
+			_savedStates[_savedStateCount] = new SavedState { TokenIndex = _currentTokenIndex, HasToken = _hasToken };
+			++_savedStateCount;
 		}
 
 		/// <inheritdoc/>
 		public void RestoreState()
 		{
-			if (this._savedStateCount == 0)
+			if (_savedStateCount == 0)
 			{
 				throw new UhtIceException("Attempt to restore a state when none have been saved");
 			}
 
-			--this._savedStateCount;
-			this._currentTokenIndex = this._savedStates[this._savedStateCount].TokenIndex;
-			this._hasToken = this._savedStates[this._savedStateCount].HasToken;
-			if (this._hasToken)
+			--_savedStateCount;
+			_currentTokenIndex = _savedStates[_savedStateCount].TokenIndex;
+			_hasToken = _savedStates[_savedStateCount].HasToken;
+			if (_hasToken)
 			{
-				this._currentToken = this._tokens.Span[this._currentTokenIndex];
+				_currentToken = _tokens.Span[_currentTokenIndex];
 			}
 		}
 
 		/// <inheritdoc/>
 		public void AbandonState()
 		{
-			if (this._savedStateCount == 0)
+			if (_savedStateCount == 0)
 			{
 				throw new UhtIceException("Attempt to abandon a state when none have been saved");
 			}
 
-			--this._savedStateCount;
+			--_savedStateCount;
 		}
 
 		/// <inheritdoc/>
@@ -300,45 +300,45 @@ namespace EpicGames.UHT.Tokenizer
 		{
 			get
 			{
-				if (this._tokens.Length == 0)
+				if (_tokens.Length == 0)
 				{
 					return -1;
 				}
-				if (this._currentTokenIndex < 0)
+				if (_currentTokenIndex < 0)
 				{
-					return this._tokens.Span[0].InputLine;
+					return _tokens.Span[0].InputLine;
 				}
-				if (this._currentTokenIndex < this._tokens.Length)
+				if (_currentTokenIndex < _tokens.Length)
 				{
-					return this._tokens.Span[this._currentTokenIndex].InputLine;
+					return _tokens.Span[_currentTokenIndex].InputLine;
 				}
-				return this._tokens.Span[this._tokens.Length - 1].InputLine;
+				return _tokens.Span[_tokens.Length - 1].InputLine;
 			}
 		}
 		#endregion
 
 		private ref UhtToken GetTokenInternal()
 		{
-			if (this._currentTokenIndex < this._tokens.Length)
+			if (_currentTokenIndex < _tokens.Length)
 			{
-				++this._currentTokenIndex;
+				++_currentTokenIndex;
 			}
-			if (this._currentTokenIndex < this._tokens.Length)
+			if (_currentTokenIndex < _tokens.Length)
 			{
-				this._currentToken = this._tokens.Span[this._currentTokenIndex];
+				_currentToken = _tokens.Span[_currentTokenIndex];
 			}
-			else if (this._tokens.Length == 0)
+			else if (_tokens.Length == 0)
 			{
-				this._currentToken = new UhtToken();
+				_currentToken = new UhtToken();
 			}
 			else
 			{
-				UhtToken lastToken = this._tokens.Span[this._tokens.Length - 1];
+				UhtToken lastToken = _tokens.Span[_tokens.Length - 1];
 				int endPos = lastToken.InputEndPos;
-				this._currentToken = new UhtToken(this._endTokenType, endPos, lastToken.InputLine, endPos, lastToken.InputLine, new StringView());
+				_currentToken = new UhtToken(_endTokenType, endPos, lastToken.InputLine, endPos, lastToken.InputLine, new StringView());
 			}
-			this._hasToken = true;
-			return ref this._currentToken;
+			_hasToken = true;
+			return ref _currentToken;
 		}
 	}
 }
diff --git a/Engine/Source/Programs/Shared/EpicGames.UHT/Types/Properties/UhtArrayProperty.cs b/Engine/Source/Programs/Shared/EpicGames.UHT/Types/Properties/UhtArrayProperty.cs
--- a/Engine/Source/Programs/Shared/EpicGames.UHT/Types/Properties/UhtArrayProperty.cs
+++ b/Engine/Source/Programs/Shared/EpicGames.UHT/Types/Properties/UhtArrayProperty.cs
@@ -39,23 +39,23 @@ namespace EpicGames.UHT.Types
 		public UhtArrayProperty(UhtPropertySettings propertySettings, UhtProperty value) : base(propertySettings, value)
 		{
 			// If the creation of the value property set more flags, then copy those flags to ourselves
-			this.PropertyFlags |= value.PropertyFlags & EPropertyFlags.UObjectWrapper;
+			PropertyFlags |= value.PropertyFlags & EPropertyFlags.UObjectWrapper;
 
 			if (value.MetaData.ContainsKey(UhtNames.NativeConst))
 			{
-				this.MetaData.Add(UhtNames.NativeConstTemplateArg, "");
+				MetaData.Add(UhtNames.NativeConstTemplateArg, "");
 				value.MetaData.Remove(UhtNames.NativeConst);
 			}
 
-			this.PropertyCaps |= UhtPropertyCaps.PassCppArgsByRef | UhtPropertyCaps.SupportsRigVM | UhtPropertyCaps.IsRigVMArray;
-			this.PropertyCaps &= ~(UhtPropertyCaps.CanBeContainerValue | UhtPropertyCaps.CanBeContainerKey);
-			this.PropertyCaps = (this.PropertyCaps & ~(UhtPropertyCaps.IsParameterSupportedByBlueprint | UhtPropertyCaps.IsMemberSupportedByBlueprint | UhtPropertyCaps.CanExposeOnSpawn)) |
-				(this.ValueProperty.PropertyCaps & (UhtPropertyCaps.IsParameterSupportedByBlueprint | UhtPropertyCaps.IsMemberSupportedByBlueprint | UhtPropertyCaps.CanExposeOnSpawn));
-			this.ValueProperty.SourceName = this.SourceName;
-			this.ValueProperty.EngineName = this.EngineName;
-			this.ValueProperty.PropertyFlags = this.PropertyFlags & EPropertyFlags.PropagateToArrayInner;
-			this.ValueProperty.Outer = this;
-			this.ValueProperty.MetaData.Clear();
+			PropertyCaps |= UhtPropertyCaps.PassCppArgsByRef | UhtPropertyCaps.SupportsRigVM | UhtPropertyCaps.IsRigVMArray;
+			PropertyCaps &= ~(UhtPropertyCaps.CanBeContainerValue | UhtPropertyCaps.CanBeContainerKey);
+			PropertyCaps = (PropertyCaps & ~(UhtPropertyCaps.IsParameterSupportedByBlueprint | UhtPropertyCaps.IsMemberSupportedByBlueprint | UhtPropertyCaps.CanExposeOnSpawn)) |
+				(ValueProperty.PropertyCaps & (UhtPropertyCaps.IsParameterSupportedByBlueprint | UhtPropertyCaps.IsMemberSupportedByBlueprint | UhtPropertyCaps.CanExposeOnSpawn));
+			ValueProperty.SourceName = SourceName;
+			ValueProperty.EngineName = EngineName;
+			ValueProperty.PropertyFlags = PropertyFlags & EPropertyFlags.PropagateToArrayInner;
+			ValueProperty.Outer = this;
+			ValueProperty.MetaData.Clear();
 		}
 
 		/// <inheritdoc/>
@@ -65,11 +65,11 @@ namespace EpicGames.UHT.Types
 			switch (phase)
 			{
 				case UhtResolvePhase.Final:
-					this.PropertyFlags |= ResolveAndReturnNewFlags(this.ValueProperty, phase);
-					this.MetaData.Add(this.ValueProperty.MetaData);
-					this.ValueProperty.PropertyFlags = this.PropertyFlags & EPropertyFlags.PropagateToArrayInner;
-					this.ValueProperty.MetaData.Clear();
-					PropagateFlagsFromInnerAndHandlePersistentInstanceMetadata(this, this.MetaData, this.ValueProperty);
+					PropertyFlags |= ResolveAndReturnNewFlags(ValueProperty, phase);
+					MetaData.Add(ValueProperty.MetaData);
+					ValueProperty.PropertyFlags = PropertyFlags & EPropertyFlags.PropagateToArrayInner;
+					ValueProperty.MetaData.Clear();
+					PropagateFlagsFromInnerAndHandlePersistentInstanceMetadata(this, MetaData, ValueProperty);
 					break;
 			}
 			return results;
@@ -78,19 +78,19 @@ namespace EpicGames.UHT.Types
 		/// <inheritdoc/>
 		public override void CollectReferencesInternal(IUhtReferenceCollector collector, bool templateProperty)
 		{
-			this.ValueProperty.CollectReferencesInternal(collector, true);
+			ValueProperty.CollectReferencesInternal(collector, true);
 		}
 
 		/// <inheritdoc/>
 		public override string? GetForwardDeclarations()
 		{
-			return this.ValueProperty.GetForwardDeclarations();
+			return ValueProperty.GetForwardDeclarations();
 		}
 
 		/// <inheritdoc/>
 		public override IEnumerable<UhtType> EnumerateReferencedTypes()
 		{
-			foreach (UhtType type in this.ValueProperty.EnumerateReferencedTypes())
+			foreach (UhtType type in ValueProperty.EnumerateReferencedTypes())
 			{
 				yield return type;
 			}
@@ -106,11 +106,11 @@ namespace EpicGames.UHT.Types
 					break;
 
 				case UhtPropertyTextType.FunctionThunkParameterArgType:
-					builder.AppendFunctionThunkParameterArrayType(this.ValueProperty);
+					builder.AppendFunctionThunkParameterArrayType(ValueProperty);
 					break;
 
 				default:
-					builder.Append("TArray<").AppendPropertyText(this.ValueProperty, textType, true);
+					builder.Append("TArray<").AppendPropertyText(ValueProperty, textType, true);
 					if (builder[^1] == '>')
 					{
 						// if our internal property type is a template class, add a space between the closing brackets b/c VS.NET cannot parse this correctly
@@ -125,16 +125,16 @@ namespace EpicGames.UHT.Types
 		/// <inheritdoc/>
 		public override StringBuilder AppendMemberDecl(StringBuilder builder, IUhtPropertyMemberContext context, string name, string nameSuffix, int tabs)
 		{
-			builder.AppendMemberDecl(this.ValueProperty, context, name, GetNameSuffix(nameSuffix, "_Inner"), tabs);
+			builder.AppendMemberDecl(ValueProperty, context, name, GetNameSuffix(nameSuffix, "_Inner"), tabs);
 			return AppendMemberDecl(builder, context, name, nameSuffix, tabs, "FArrayPropertyParams");
 		}
 
 		/// <inheritdoc/>
 		public override StringBuilder AppendMemberDef(StringBuilder builder, IUhtPropertyMemberContext context, string name, string nameSuffix, string? offset, int tabs)
 		{
-			builder.AppendMemberDef(this.ValueProperty, context, name, GetNameSuffix(nameSuffix, "_Inner"), "0", tabs);
+			builder.AppendMemberDef(ValueProperty, context, name, GetNameSuffix(nameSuffix, "_Inner"), "0", tabs);
 			AppendMemberDefStart(builder, context, name, nameSuffix, offset, tabs, "FArrayPropertyParams", "UECodeGen_Private::EPropertyGenFlags::Array");
-			builder.Append(this.Allocator == UhtPropertyAllocator.MemoryImage ? "EArrayPropertyFlags::UsesMemoryImageAllocator" : "EArrayPropertyFlags::None").Append(", ");
+			builder.Append(Allocator == UhtPropertyAllocator.MemoryImage ? "EArrayPropertyFlags::UsesMemoryImageAllocator" : "EArrayPropertyFlags::None").Append(", ");
 			AppendMemberDefEnd(builder, context, name, nameSuffix);
 			return builder;
 		}
@@ -142,7 +142,7 @@ namespace EpicGames.UHT.Types
 		/// <inheritdoc/>
 		public override StringBuilder AppendMemberPtr(StringBuilder builder, IUhtPropertyMemberContext context, string name, string nameSuffix, int tabs)
 		{
-			builder.AppendMemberPtr(this.ValueProperty, context, name, GetNameSuffix(nameSuffix, "_Inner"), tabs);
+			builder.AppendMemberPtr(ValueProperty, context, name, GetNameSuffix(nameSuffix, "_Inner"), tabs);
 			base.AppendMemberPtr(builder, context, name, nameSuffix, tabs);
 			return builder;
 		}
@@ -150,7 +150,7 @@ namespace EpicGames.UHT.Types
 		/// <inheritdoc/>
 		public override void AppendObjectHashes(StringBuilder builder, int startingLength, IUhtPropertyMemberContext context)
 		{
-			this.ValueProperty.AppendObjectHashes(builder, startingLength, context);
+			ValueProperty.AppendObjectHashes(builder, startingLength, context);
 		}
 
 		/// <inheritdoc/>
@@ -175,7 +175,7 @@ namespace EpicGames.UHT.Types
 			{
 				if (!function.FunctionFlags.HasAnyFlags(EFunctionFlags.NetRequest))
 				{
-					if (this.RefQualifier != UhtPropertyRefQualifier.ConstRef && !this.IsStaticArray)
+					if (RefQualifier != UhtPropertyRefQualifier.ConstRef && !IsStaticArray)
 					{
 						this.LogError("Replicated TArray parameters must be passed by const reference");
 					}
@@ -188,9 +188,9 @@ namespace EpicGames.UHT.Types
 		{
 			base.Validate(outerStruct, outermostProperty, options);
 
-			if (this.Allocator != UhtPropertyAllocator.Default)
+			if (Allocator != UhtPropertyAllocator.Default)
 			{
-				if (this.PropertyFlags.HasAnyFlags(EPropertyFlags.Net))
+				if (PropertyFlags.HasAnyFlags(EPropertyFlags.Net))
 				{
 					this.LogError("Replicated arrays with MemoryImageAllocators are not yet supported");
 				}
@@ -202,7 +202,7 @@ namespace EpicGames.UHT.Types
 		{
 			if (other is UhtArrayProperty otherArray)
 			{
-				return this.ValueProperty.IsSameType(otherArray.ValueProperty);
+				return ValueProperty.IsSameType(otherArray.ValueProperty);
 			}
 			return false;
 		}
diff --git a/Engine/Source/Programs/Shared/EpicGames.UHT/Types/Properties/UhtBoolProperty.cs b/Engine/Source/Programs/Shared/EpicGames.UHT/Types/Properties/UhtBoolProperty.cs
--- a/Engine/Source/Programs/Shared/EpicGames.UHT/Types/Properties/UhtBoolProperty.cs
+++ b/Engine/Source/Programs/Shared/EpicGames.UHT/Types/Properties/UhtBoolProperty.cs
@@ -57,7 +57,7 @@ namespace EpicGames.UHT.Types
 		{
 			get
 			{
-				switch (this.BoolType)
+				switch (BoolType)
 				{
 					case UhtBoolType.Native:
 						return "bool";
@@ -80,7 +80,7 @@ namespace EpicGames.UHT.Types
 		{
 			get
 			{
-				switch (this.BoolType)
+				switch (BoolType)
 				{
 					case UhtBoolType.Native:
 						return "UBOOL";
@@ -101,7 +101,7 @@ namespace EpicGames.UHT.Types
 		/// <summary>
 		/// If true, the boolean is a native bool and not a UBOOL
 		/// </summary>
-		protected bool IsNativeBool => this.BoolType == UhtBoolType.Native;
+		protected bool IsNativeBool => BoolType == UhtBoolType.Native;
 
 		/// <summary>
 		/// Type of the boolean
@@ -115,13 +115,13 @@ namespace EpicGames.UHT.Types
 		/// <param name="boolType">Type of the boolean</param>
 		public UhtBoolProperty(UhtPropertySettings propertySettings, UhtBoolType boolType) : base(propertySettings)
 		{
-			this.PropertyCaps |= UhtPropertyCaps.RequiresNullConstructorArg | UhtPropertyCaps.IsParameterSupportedByBlueprint |
+			PropertyCaps |= UhtPropertyCaps.RequiresNullConstructorArg | UhtPropertyCaps.IsParameterSupportedByBlueprint |
 				UhtPropertyCaps.IsMemberSupportedByBlueprint | UhtPropertyCaps.SupportsRigVM;
 			if (boolType == UhtBoolType.Native || boolType == UhtBoolType.UInt8)
 			{
-				this.PropertyCaps |= UhtPropertyCaps.CanExposeOnSpawn;
+				PropertyCaps |= UhtPropertyCaps.CanExposeOnSpawn;
 			}
-			this.BoolType = boolType;
+			BoolType = boolType;
 		}
 
 		/// <inheritdoc/>
@@ -136,7 +136,7 @@ namespace EpicGames.UHT.Types
 				case UhtPropertyTextType.RigVMType:
 				case UhtPropertyTextType.ExportMember:
 				case UhtPropertyTextType.GenericFunctionArgOrRetVal:
-					builder.Append(this.CppTypeText);
+					builder.Append(CppTypeText);
 					break;
 
 				case UhtPropertyTextType.Sparse:
@@ -158,7 +158,7 @@ namespace EpicGames.UHT.Types
 		public override StringBuilder AppendMemberDecl(StringBuilder builder, IUhtPropertyMemberContext context, string name, string nameSuffix, int tabs)
 		{
 			builder.AppendMetaDataDecl(this, context, name, nameSuffix, tabs);
-			if (this.Outer is not UhtProperty)
+			if (Outer is not UhtProperty)
 			{
 				builder.AppendTabs(tabs).Append("static void ").AppendNameDecl(context, name, nameSuffix).Append("_SetBit(void* Obj);\r\n");
 			}
@@ -170,11 +170,11 @@ namespace EpicGames.UHT.Types
 		public override StringBuilder AppendMemberDef(StringBuilder builder, IUhtPropertyMemberContext context, string name, string nameSuffix, string? offset, int tabs)
 		{
 			builder.AppendMetaDataDef(this, context, name, nameSuffix, tabs);
-			if (this.Outer == context.OuterStruct)
+			if (Outer == context.OuterStruct)
 			{
 				builder.AppendTabs(tabs).Append("void ").AppendNameDef(context, name, nameSuffix).Append("_SetBit(void* Obj)\r\n");
 				builder.AppendTabs(tabs).Append("{\r\n");
-				builder.AppendTabs(tabs + 1).Append("((").Append(context.OuterStructSourceName).Append("*)Obj)->").Append(this.SourceName).Append(" = 1;\r\n");
+				builder.AppendTabs(tabs + 1).Append("((").Append(context.OuterStructSourceName).Append("*)Obj)->").Append(SourceName).Append(" = 1;\r\n");
 				builder.AppendTabs(tabs).Append("}\r\n");
 			}
 
@@ -184,9 +184,9 @@ namespace EpicGames.UHT.Types
 				"UECodeGen_Private::EPropertyGenFlags::Bool ",
 				false, false);
 
-			builder.Append("sizeof(").Append(this.CppTypeText).Append("), ");
+			builder.Append("sizeof(").Append(CppTypeText).Append("), ");
 
-			if (this.Outer == context.OuterStruct)
+			if (Outer == context.OuterStruct)
 			{
 				builder
 					.Append("sizeof(").Append(context.OuterStructSourceName).Append("), ")
@@ -207,7 +207,7 @@ namespace EpicGames.UHT.Types
 			AppendText(builder, textType);
 
 			//@todo we currently can't have out bools.. so this isn't really necessary, but eventually out bools may be supported, so leave here for now
-			if (textType.IsParameter() && this.PropertyFlags.HasAnyFlags(EPropertyFlags.OutParm))
+			if (textType.IsParameter() && PropertyFlags.HasAnyFlags(EPropertyFlags.OutParm))
 			{
 				builder.Append('&');
 			}
@@ -216,14 +216,14 @@ namespace EpicGames.UHT.Types
 
 			if (!skipParameterName)
 			{
-				builder.Append(this.SourceName);
+				builder.Append(SourceName);
 			}
 
-			if (this.ArrayDimensions != null)
+			if (ArrayDimensions != null)
 			{
-				builder.Append('[').Append(this.ArrayDimensions).Append(']');
+				builder.Append('[').Append(ArrayDimensions).Append(']');
 			}
-			else if (textType == UhtPropertyTextType.ExportMember && !this.IsNativeBool)
+			else if (textType == UhtPropertyTextType.ExportMember && !IsNativeBool)
 			{
 				builder.Append(":1");
 			}
diff --git a/Engine/Source/Programs/Shared/EpicGames.UHT/Types/Properties/UhtByteProperty.cs b/Engine/Source/Programs/Shared/EpicGames.UHT/Types/Properties/UhtByteProperty.cs
--- a/Engine/Source/Programs/Shared/EpicGames.UHT/Types/Properties/UhtByteProperty.cs
+++ b/Engine/Source/Programs/Shared/EpicGames.UHT/Types/Properties/UhtByteProperty.cs
@@ -32,10 +32,10 @@ namespace EpicGames.UHT.Types
 		protected override string CppTypeText => "uint8";
 
 		/// <inheritdoc/>
-		protected override string PGetMacroText => this.Enum == null || this.Enum.CppForm != UhtEnumCppForm.EnumClass ? "PROPERTY" : "ENUM";
+		protected override string PGetMacroText => Enum == null || Enum.CppForm != UhtEnumCppForm.EnumClass ? "PROPERTY" : "ENUM";
 
 		/// <inheritdoc/>
-		protected override UhtPGetArgumentType PGetTypeArgument => this.Enum == null || this.Enum.CppForm != UhtEnumCppForm.EnumClass ? UhtPGetArgumentType.EngineClass : UhtPGetArgumentType.TypeText;
+		protected override UhtPGetArgumentType PGetTypeArgument => Enum == null || Enum.CppForm != UhtEnumCppForm.EnumClass ? UhtPGetArgumentType.EngineClass : UhtPGetArgumentType.TypeText;
 
 		/// <summary>
 		/// Construct a new property
@@ -45,36 +45,36 @@ namespace EpicGames.UHT.Types
 		/// <param name="enumObj">Optional referenced enum</param>
 		public UhtByteProperty(UhtPropertySettings propertySettings, UhtPropertyIntType intType, UhtEnum? enumObj = null) : base(propertySettings, intType)
 		{
-			this.Enum = enumObj;
-			this.PropertyCaps |= UhtPropertyCaps.CanExposeOnSpawn | UhtPropertyCaps.IsParameterSupportedByBlueprint |
+			Enum = enumObj;
+			PropertyCaps |= UhtPropertyCaps.CanExposeOnSpawn | UhtPropertyCaps.IsParameterSupportedByBlueprint |
 				UhtPropertyCaps.IsMemberSupportedByBlueprint | UhtPropertyCaps.SupportsRigVM;
-			if (this.Enum != null)
+			if (Enum != null)
 			{
-				this.PropertyCaps |= UhtPropertyCaps.IsRigVMEnum;
+				PropertyCaps |= UhtPropertyCaps.IsRigVMEnum;
 			}
 		}
 
 		/// <inheritdoc/>
 		public override IEnumerable<UhtType> EnumerateReferencedTypes()
 		{
-			if (this.Enum != null)
+			if (Enum != null)
 			{
-				yield return this.Enum;
+				yield return Enum;
 			}
 		}
 
 		/// <inheritdoc/>
 		public override void CollectReferencesInternal(IUhtReferenceCollector collector, bool templateProperty)
 		{
-			collector.AddCrossModuleReference(this.Enum, true);
+			collector.AddCrossModuleReference(Enum, true);
 		}
 
 		/// <inheritdoc/>
 		public override StringBuilder AppendText(StringBuilder builder, UhtPropertyTextType textType, bool isTemplateArgument)
 		{
-			if (this.Enum != null)
+			if (Enum != null)
 			{
-				return UhtEnumProperty.AppendEnumText(builder, this, this.Enum, textType, isTemplateArgument);
+				return UhtEnumProperty.AppendEnumText(builder, this, Enum, textType, isTemplateArgument);
 			}
 			else
 			{
@@ -92,7 +92,7 @@ namespace EpicGames.UHT.Types
 		public override StringBuilder AppendMemberDef(StringBuilder builder, IUhtPropertyMemberContext context, string name, string nameSuffix, string? offset, int tabs)
 		{
 			AppendMemberDefStart(builder, context, name, nameSuffix, offset, tabs, "FBytePropertyParams", "UECodeGen_Private::EPropertyGenFlags::Byte");
-			AppendMemberDefRef(builder, context, this.Enum, true, true);
+			AppendMemberDefRef(builder, context, Enum, true, true);
 			AppendMemberDefEnd(builder, context, name, nameSuffix);
 			return builder;
 		}
@@ -100,15 +100,15 @@ namespace EpicGames.UHT.Types
 		/// <inheritdoc/>
 		public override void AppendObjectHashes(StringBuilder builder, int startingLength, IUhtPropertyMemberContext context)
 		{
-			builder.AppendObjectHash(startingLength, this, context, this.Enum);
+			builder.AppendObjectHash(startingLength, this, context, Enum);
 		}
 
 		/// <inheritdoc/>
 		public override StringBuilder AppendFunctionThunkParameterArg(StringBuilder builder)
 		{
-			if (this.Enum != null)
+			if (Enum != null)
 			{
-				return UhtEnumProperty.AppendEnumFunctionThunkParameterArg(builder, this, this.Enum);
+				return UhtEnumProperty.AppendEnumFunctionThunkParameterArg(builder, this, Enum);
 			}
 			else
 			{
@@ -119,9 +119,9 @@ namespace EpicGames.UHT.Types
 		/// <inheritdoc/>
 		public override bool SanitizeDefaultValue(IUhtTokenReader defaultValueReader, StringBuilder innerDefaultValue)
 		{
-			if (this.Enum != null)
+			if (Enum != null)
 			{
-				return UhtEnumProperty.SanitizeEnumDefaultValue(this, this.Enum, defaultValueReader, innerDefaultValue);
+				return UhtEnumProperty.SanitizeEnumDefaultValue(this, Enum, defaultValueReader, innerDefaultValue);
 			}
 
 			int value = defaultValueReader.GetConstIntExpression();
@@ -134,11 +134,11 @@ namespace EpicGames.UHT.Types
 		{
 			if (other is UhtByteProperty otherByte)
 			{
-				return this.Enum == otherByte.Enum;
+				return Enum == otherByte.Enum;
 			}
 			else if (other is UhtEnumProperty otherEnum)
 			{
-				return this.Enum == otherEnum.Enum;
+				return Enum == otherEnum.Enum;
 			}
 			return false;
 		}
@@ -150,7 +150,7 @@ namespace EpicGames.UHT.Types
 
 			if (function.FunctionFlags.HasAnyFlags(EFunctionFlags.BlueprintEvent | EFunctionFlags.BlueprintCallable))
 			{
-				if (this.Enum != null && this.Enum.UnderlyingType != UhtEnumUnderlyingType.uint8 && this.Enum.UnderlyingType != UhtEnumUnderlyingType.Unspecified)
+				if (Enum != null && Enum.UnderlyingType != UhtEnumUnderlyingType.uint8 && Enum.UnderlyingType != UhtEnumUnderlyingType.Unspecified)
 				{
 					this.LogError("Invalid enum param for Blueprints - currently only uint8 supported");
 				}
