Removed the use of "this." where possible in EpicGames.UHT.

diff --git a/Engine/Source/Programs/Shared/EpicGames.UHT/Types/Properties/UhtClassProperty.cs b/Engine/Source/Programs/Shared/EpicGames.UHT/Types/Properties/UhtClassProperty.cs
--- a/Engine/Source/Programs/Shared/EpicGames.UHT/Types/Properties/UhtClassProperty.cs
+++ b/Engine/Source/Programs/Shared/EpicGames.UHT/Types/Properties/UhtClassProperty.cs
@@ -20,7 +20,7 @@ namespace EpicGames.UHT.Types
 		public override string EngineClassName => "ClassProperty";
 
 		/// <inheritdoc/>
-		protected override bool PGetPassAsNoPtr => this.PropertyFlags.HasAnyFlags(EPropertyFlags.UObjectWrapper);
+		protected override bool PGetPassAsNoPtr => PropertyFlags.HasAnyFlags(EPropertyFlags.UObjectWrapper);
 
 		/// <inheritdoc/>
 		protected override UhtPGetArgumentType PGetTypeArgument => UhtPGetArgumentType.TypeText;
@@ -35,34 +35,34 @@ namespace EpicGames.UHT.Types
 		public UhtClassProperty(UhtPropertySettings propertySettings, UhtClass classObj, UhtClass metaClass, EPropertyFlags extraFlags = EPropertyFlags.None)
 			: base(propertySettings, classObj, metaClass)
 		{
-			this.PropertyFlags |= extraFlags;
-			this.PropertyCaps |= UhtPropertyCaps.CanHaveConfig;
-			this.PropertyCaps &= ~(UhtPropertyCaps.CanBeInstanced);
+			PropertyFlags |= extraFlags;
+			PropertyCaps |= UhtPropertyCaps.CanHaveConfig;
+			PropertyCaps &= ~(UhtPropertyCaps.CanBeInstanced);
 		}
 
 		/// <inheritdoc/>
 		public override string? GetForwardDeclarations()
 		{
-			return this.MetaClass != null ? $"class {this.MetaClass.SourceName};" : null;
+			return MetaClass != null ? $"class {MetaClass.SourceName};" : null;
 		}
 
 		/// <inheritdoc/>
 		private StringBuilder AppendSubClassText(StringBuilder builder)
 		{
-			builder.Append("TSubclassOf<").Append(this.MetaClass?.SourceName).Append("> "); //COMPATIBILITY-TODO - Extra space in old UHT
+			builder.Append("TSubclassOf<").Append(MetaClass?.SourceName).Append("> "); //COMPATIBILITY-TODO - Extra space in old UHT
 			return builder;
 		}
 
 		/// <inheritdoc/>
 		private StringBuilder AppendText(StringBuilder builder)
 		{
-			if (this.PropertyFlags.HasAnyFlags(EPropertyFlags.UObjectWrapper))
+			if (PropertyFlags.HasAnyFlags(EPropertyFlags.UObjectWrapper))
 			{
 				AppendSubClassText(builder);
 			}
 			else
 			{
-				builder.Append(this.Class.SourceName).Append('*');
+				builder.Append(Class.SourceName).Append('*');
 			}
 			return builder;
 		}
@@ -89,11 +89,11 @@ namespace EpicGames.UHT.Types
 					break;
 
 				case UhtPropertyTextType.FunctionThunkRetVal:
-					if (this.PropertyFlags.HasAnyFlags(EPropertyFlags.ConstParm))
+					if (PropertyFlags.HasAnyFlags(EPropertyFlags.ConstParm))
 					{
 						builder.Append("const ");
 					}
-					if (this.PropertyFlags.HasAnyFlags(EPropertyFlags.UObjectWrapper))
+					if (PropertyFlags.HasAnyFlags(EPropertyFlags.UObjectWrapper))
 					{
 						AppendText(builder);
 					}
@@ -105,7 +105,7 @@ namespace EpicGames.UHT.Types
 
 				case UhtPropertyTextType.EventParameterMember:
 				case UhtPropertyTextType.EventParameterFunctionMember:
-					if (!this.PropertyFlags.HasAnyFlags(EPropertyFlags.UObjectWrapper))
+					if (!PropertyFlags.HasAnyFlags(EPropertyFlags.UObjectWrapper))
 					{
 						builder.Append("UClass*");
 					}
@@ -116,13 +116,13 @@ namespace EpicGames.UHT.Types
 					break;
 
 				case UhtPropertyTextType.FunctionThunkParameterArgType:
-					if (this.PropertyFlags.HasAllFlags(EPropertyFlags.OutParm | EPropertyFlags.UObjectWrapper))
+					if (PropertyFlags.HasAllFlags(EPropertyFlags.OutParm | EPropertyFlags.UObjectWrapper))
 					{
 						AppendSubClassText(builder);
 					}
 					else
 					{
-						builder.Append(this.Class.SourceName);
+						builder.Append(Class.SourceName);
 					}
 					break;
 			}
@@ -139,8 +139,8 @@ namespace EpicGames.UHT.Types
 		public override StringBuilder AppendMemberDef(StringBuilder builder, IUhtPropertyMemberContext context, string name, string nameSuffix, string? offset, int tabs)
 		{
 			AppendMemberDefStart(builder, context, name, nameSuffix, offset, tabs, "FClassPropertyParams", "UECodeGen_Private::EPropertyGenFlags::Class");
-			AppendMemberDefRef(builder, context, this.Class, false);
-			AppendMemberDefRef(builder, context, this.MetaClass, false);
+			AppendMemberDefRef(builder, context, Class, false);
+			AppendMemberDefRef(builder, context, MetaClass, false);
 			AppendMemberDefEnd(builder, context, name, nameSuffix);
 			return builder;
 		}
diff --git a/Engine/Source/Programs/Shared/EpicGames.UHT/Types/Properties/UhtClassPtrProperties.cs b/Engine/Source/Programs/Shared/EpicGames.UHT/Types/Properties/UhtClassPtrProperties.cs
--- a/Engine/Source/Programs/Shared/EpicGames.UHT/Types/Properties/UhtClassPtrProperties.cs
+++ b/Engine/Source/Programs/Shared/EpicGames.UHT/Types/Properties/UhtClassPtrProperties.cs
@@ -36,7 +36,7 @@ namespace EpicGames.UHT.Types
 		public UhtClassPtrProperty(UhtPropertySettings propertySettings, UhtClass classObj, UhtClass metaClass, EPropertyFlags extraFlags = EPropertyFlags.None)
 			: base(propertySettings, classObj, metaClass)
 		{
-			this.PropertyFlags |= extraFlags | EPropertyFlags.UObjectWrapper;
+			PropertyFlags |= extraFlags | EPropertyFlags.UObjectWrapper;
 		}
 
 		/// <inheritdoc/>
@@ -45,7 +45,7 @@ namespace EpicGames.UHT.Types
 			switch (textType)
 			{
 				default:
-					builder.Append("TObjectPtr<").Append(this.Class.SourceName).Append('>');
+					builder.Append("TObjectPtr<").Append(Class.SourceName).Append('>');
 					break;
 			}
 			return builder;
@@ -62,8 +62,8 @@ namespace EpicGames.UHT.Types
 		{
 			AppendMemberDefStart(builder, context, name, nameSuffix, offset, tabs, "FClassPtrPropertyParams",
 				"UECodeGen_Private::EPropertyGenFlags::Class | UECodeGen_Private::EPropertyGenFlags::ObjectPtr");
-			AppendMemberDefRef(builder, context, this.Class, false);
-			AppendMemberDefRef(builder, context, this.MetaClass, false);
+			AppendMemberDefRef(builder, context, Class, false);
+			AppendMemberDefRef(builder, context, MetaClass, false);
 			AppendMemberDefEnd(builder, context, name, nameSuffix);
 			return builder;
 		}
@@ -75,7 +75,7 @@ namespace EpicGames.UHT.Types
 
 			// UFunctions with a smart pointer as input parameter wont compile anyway, because of missing P_GET_... macro.
 			// UFunctions with a smart pointer as return type will crash when called via blueprint, because they are not supported in VM.
-			if (this.PropertyCategory == UhtPropertyCategory.RegularParameter || this.PropertyCategory == UhtPropertyCategory.ReplicatedParameter)
+			if (PropertyCategory == UhtPropertyCategory.RegularParameter || PropertyCategory == UhtPropertyCategory.ReplicatedParameter)
 			{
 				outerStruct.LogError("UFunctions cannot take a TObjectPtr as a parameter.");
 			}
diff --git a/Engine/Source/Programs/Shared/EpicGames.UHT/Types/Properties/UhtContainerProperty.cs b/Engine/Source/Programs/Shared/EpicGames.UHT/Types/Properties/UhtContainerProperty.cs
--- a/Engine/Source/Programs/Shared/EpicGames.UHT/Types/Properties/UhtContainerProperty.cs
+++ b/Engine/Source/Programs/Shared/EpicGames.UHT/Types/Properties/UhtContainerProperty.cs
@@ -21,16 +21,16 @@ namespace EpicGames.UHT.Types
 		/// <param name="value">Value property</param>
 		protected UhtContainerBaseProperty(UhtPropertySettings propertySettings, UhtProperty value) : base(propertySettings)
 		{
-			this.ValueProperty = value;
-			this.PropertyCaps = (this.PropertyCaps & ~(UhtPropertyCaps.CanBeInstanced | UhtPropertyCaps.CanHaveConfig)) |
-				(this.ValueProperty.PropertyCaps & (UhtPropertyCaps.CanBeInstanced | UhtPropertyCaps.CanHaveConfig));
+			ValueProperty = value;
+			PropertyCaps = (PropertyCaps & ~(UhtPropertyCaps.CanBeInstanced | UhtPropertyCaps.CanHaveConfig)) |
+				(ValueProperty.PropertyCaps & (UhtPropertyCaps.CanBeInstanced | UhtPropertyCaps.CanHaveConfig));
 		}
 
 		/// <inheritdoc/>
 		public override void Validate(UhtStruct outerStruct, UhtProperty outermostProperty, UhtValidationOptions options)
 		{
 			base.Validate(outerStruct, outermostProperty, options);
-			this.ValueProperty.Validate(outerStruct, outermostProperty, options | UhtValidationOptions.IsValue);
+			ValueProperty.Validate(outerStruct, outermostProperty, options | UhtValidationOptions.IsValue);
 		}
 
 		/// <summary>
@@ -71,13 +71,13 @@ namespace EpicGames.UHT.Types
 		/// <inheritdoc/>
 		public override bool ScanForInstancedReferenced(bool deepScan)
 		{
-			return this.ValueProperty.ScanForInstancedReferenced(deepScan);
+			return ValueProperty.ScanForInstancedReferenced(deepScan);
 		}
 
 		///<inheritdoc/>
 		public override bool ContainsEditorOnlyProperties()
 		{
-			return this.ValueProperty.ContainsEditorOnlyProperties();
+			return ValueProperty.ContainsEditorOnlyProperties();
 		}
 	}
 }
diff --git a/Engine/Source/Programs/Shared/EpicGames.UHT/Types/Properties/UhtDelegateProperty.cs b/Engine/Source/Programs/Shared/EpicGames.UHT/Types/Properties/UhtDelegateProperty.cs
--- a/Engine/Source/Programs/Shared/EpicGames.UHT/Types/Properties/UhtDelegateProperty.cs
+++ b/Engine/Source/Programs/Shared/EpicGames.UHT/Types/Properties/UhtDelegateProperty.cs
@@ -40,9 +40,9 @@ namespace EpicGames.UHT.Types
 		/// <param name="function">Referenced function</param>
 		public UhtDelegateProperty(UhtPropertySettings propertySettings, UhtFunction function) : base(propertySettings)
 		{
-			this.Function = function;
-			this.HeaderFile.AddReferencedHeader(function);
-			this.PropertyCaps |= UhtPropertyCaps.PassCppArgsByRef | UhtPropertyCaps.IsParameterSupportedByBlueprint |
+			Function = function;
+			HeaderFile.AddReferencedHeader(function);
+			PropertyCaps |= UhtPropertyCaps.PassCppArgsByRef | UhtPropertyCaps.IsParameterSupportedByBlueprint |
 				UhtPropertyCaps.IsMemberSupportedByBlueprint | UhtPropertyCaps.SupportsRigVM;
 		}
 
@@ -53,7 +53,7 @@ namespace EpicGames.UHT.Types
 			switch (phase)
 			{
 				case UhtResolvePhase.Final:
-					this.PropertyFlags |= EPropertyFlags.InstancedReference & ~this.DisallowPropertyFlags;
+					PropertyFlags |= EPropertyFlags.InstancedReference & ~DisallowPropertyFlags;
 					break;
 			}
 			return results;
@@ -68,7 +68,7 @@ namespace EpicGames.UHT.Types
 		/// <inheritdoc/>
 		public override void CollectReferencesInternal(IUhtReferenceCollector collector, bool templateProperty)
 		{
-			collector.AddCrossModuleReference(this.Function, true);
+			collector.AddCrossModuleReference(Function, true);
 		}
 
 		/// <inheritdoc/>
@@ -79,11 +79,11 @@ namespace EpicGames.UHT.Types
 				case UhtPropertyTextType.ExportMember:
 				case UhtPropertyTextType.RigVMType:
 				case UhtPropertyTextType.EventParameterFunctionMember:
-					builder.Append(this.CppTypeText);
+					builder.Append(CppTypeText);
 					break;
 
 				default:
-					builder.Append(this.Function.SourceName);
+					builder.Append(Function.SourceName);
 					break;
 			}
 			return builder;
@@ -99,7 +99,7 @@ namespace EpicGames.UHT.Types
 		public override StringBuilder AppendMemberDef(StringBuilder builder, IUhtPropertyMemberContext context, string name, string nameSuffix, string? offset, int tabs)
 		{
 			AppendMemberDefStart(builder, context, name, nameSuffix, offset, tabs, "FDelegatePropertyParams", "UECodeGen_Private::EPropertyGenFlags::Delegate");
-			AppendMemberDefRef(builder, context, this.Function, true);
+			AppendMemberDefRef(builder, context, Function, true);
 			AppendMemberDefEnd(builder, context, name, nameSuffix);
 			return builder;
 		}
@@ -107,13 +107,13 @@ namespace EpicGames.UHT.Types
 		/// <inheritdoc/>
 		public override StringBuilder AppendFunctionThunkParameterArg(StringBuilder builder)
 		{
-			return builder.Append(this.Function.SourceName).Append('(').AppendFunctionThunkParameterName(this).Append(')');
+			return builder.Append(Function.SourceName).Append('(').AppendFunctionThunkParameterName(this).Append(')');
 		}
 
 		/// <inheritdoc/>
 		public override void AppendObjectHashes(StringBuilder builder, int startingLength, IUhtPropertyMemberContext context)
 		{
-			builder.AppendObjectHash(startingLength, this, context, this.Function);
+			builder.AppendObjectHash(startingLength, this, context, Function);
 		}
 
 		/// <inheritdoc/>
@@ -134,7 +134,7 @@ namespace EpicGames.UHT.Types
 		{
 			if (other is UhtDelegateProperty otherObject)
 			{
-				return this.Function == otherObject.Function;
+				return Function == otherObject.Function;
 			}
 			return false;
 		}
@@ -148,7 +148,7 @@ namespace EpicGames.UHT.Types
 			{
 				if (function.FunctionFlags.HasAnyFlags(EFunctionFlags.NetRequest))
 				{
-					if (!this.PropertyFlags.HasAnyFlags(EPropertyFlags.RepSkip))
+					if (!PropertyFlags.HasAnyFlags(EPropertyFlags.RepSkip))
 					{
 						this.LogError("Service request functions cannot contain delegate parameters, unless marked NotReplicated");
 					}
diff --git a/Engine/Source/Programs/Shared/EpicGames.UHT/Types/Properties/UhtDoubleProperty.cs b/Engine/Source/Programs/Shared/EpicGames.UHT/Types/Properties/UhtDoubleProperty.cs
--- a/Engine/Source/Programs/Shared/EpicGames.UHT/Types/Properties/UhtDoubleProperty.cs
+++ b/Engine/Source/Programs/Shared/EpicGames.UHT/Types/Properties/UhtDoubleProperty.cs
@@ -29,7 +29,7 @@ namespace EpicGames.UHT.Types
 		/// <param name="propertySettings">Property settings</param>
 		public UhtDoubleProperty(UhtPropertySettings propertySettings) : base(propertySettings, UhtPropertyIntType.None)
 		{
-			this.PropertyCaps |= UhtPropertyCaps.IsParameterSupportedByBlueprint | UhtPropertyCaps.IsMemberSupportedByBlueprint;
+			PropertyCaps |= UhtPropertyCaps.IsParameterSupportedByBlueprint | UhtPropertyCaps.IsMemberSupportedByBlueprint;
 		}
 
 		/// <inheritdoc/>
diff --git a/Engine/Source/Programs/Shared/EpicGames.UHT/Types/Properties/UhtEnumProperty.cs b/Engine/Source/Programs/Shared/EpicGames.UHT/Types/Properties/UhtEnumProperty.cs
--- a/Engine/Source/Programs/Shared/EpicGames.UHT/Types/Properties/UhtEnumProperty.cs
+++ b/Engine/Source/Programs/Shared/EpicGames.UHT/Types/Properties/UhtEnumProperty.cs
@@ -20,16 +20,16 @@ namespace EpicGames.UHT.Types
 	public class UhtEnumProperty : UhtProperty
 	{
 		/// <inheritdoc/>
-		public override string EngineClassName => this.Enum.CppForm != UhtEnumCppForm.EnumClass ? "ByteProperty" : "EnumProperty";
+		public override string EngineClassName => Enum.CppForm != UhtEnumCppForm.EnumClass ? "ByteProperty" : "EnumProperty";
 
 		/// <inheritdoc/>
 		protected override string CppTypeText => "invalid";
 
 		/// <inheritdoc/>
-		protected override string PGetMacroText => this.Enum.CppForm != UhtEnumCppForm.EnumClass ? "PROPERTY" : "ENUM";
+		protected override string PGetMacroText => Enum.CppForm != UhtEnumCppForm.EnumClass ? "PROPERTY" : "ENUM";
 
 		/// <inheritdoc/>
-		protected override UhtPGetArgumentType PGetTypeArgument => this.Enum.CppForm != UhtEnumCppForm.EnumClass ? UhtPGetArgumentType.EngineClass : UhtPGetArgumentType.TypeText;
+		protected override UhtPGetArgumentType PGetTypeArgument => Enum.CppForm != UhtEnumCppForm.EnumClass ? UhtPGetArgumentType.EngineClass : UhtPGetArgumentType.TypeText;
 
 		/// <summary>
 		/// Referenced enum
@@ -45,12 +45,12 @@ namespace EpicGames.UHT.Types
 		/// <summary>
 		/// Underlying type which defaults to int32 if the referenced enum doesn't have an underlying type
 		/// </summary>
-		public UhtEnumUnderlyingType UnderlyingType => this.Enum.UnderlyingType != UhtEnumUnderlyingType.Unspecified ? this.Enum.UnderlyingType : UhtEnumUnderlyingType.int32;
+		public UhtEnumUnderlyingType UnderlyingType => Enum.UnderlyingType != UhtEnumUnderlyingType.Unspecified ? Enum.UnderlyingType : UhtEnumUnderlyingType.int32;
 
 		/// <summary>
 		/// Underlying type size.  Defaults to unsized if the referenced enum doesn't have an underlying type
 		/// </summary>
-		public UhtPropertyIntType UnderlyingTypeSize => this.Enum.UnderlyingType != UhtEnumUnderlyingType.Unspecified ? UhtPropertyIntType.Sized : UhtPropertyIntType.Unsized;
+		public UhtPropertyIntType UnderlyingTypeSize => Enum.UnderlyingType != UhtEnumUnderlyingType.Unspecified ? UhtPropertyIntType.Sized : UhtPropertyIntType.Unsized;
 
 		/// <summary>
 		/// Construct property
@@ -59,40 +59,40 @@ namespace EpicGames.UHT.Types
 		/// <param name="enumObj">Referenced enum</param>
 		public UhtEnumProperty(UhtPropertySettings propertySettings, UhtEnum enumObj) : base(propertySettings)
 		{
-			this.PropertyCaps |= UhtPropertyCaps.RequiresNullConstructorArg | UhtPropertyCaps.CanExposeOnSpawn | UhtPropertyCaps.IsParameterSupportedByBlueprint |
+			PropertyCaps |= UhtPropertyCaps.RequiresNullConstructorArg | UhtPropertyCaps.CanExposeOnSpawn | UhtPropertyCaps.IsParameterSupportedByBlueprint |
 				UhtPropertyCaps.IsMemberSupportedByBlueprint | UhtPropertyCaps.SupportsRigVM | UhtPropertyCaps.IsRigVMEnum;
-			this.Enum = enumObj;
-			this.HeaderFile.AddReferencedHeader(enumObj);
-			if (this.Enum.CppForm == UhtEnumCppForm.EnumClass)
+			Enum = enumObj;
+			HeaderFile.AddReferencedHeader(enumObj);
+			if (Enum.CppForm == UhtEnumCppForm.EnumClass)
 			{
-				this.UnderlyingProperty = CreateUnderlyingProperty();
+				UnderlyingProperty = CreateUnderlyingProperty();
 			}
 			else
 			{
-				this.UnderlyingProperty = null;
+				UnderlyingProperty = null;
 			}
 		}
 
 		/// <inheritdoc/>
 		public override void CollectReferencesInternal(IUhtReferenceCollector collector, bool templateProperty)
 		{
-			collector.AddCrossModuleReference(this.Enum, true);
+			collector.AddCrossModuleReference(Enum, true);
 		}
 
 		/// <inheritdoc/>
 		public override string? GetForwardDeclarations()
 		{
-			if (this.Enum.CppForm != UhtEnumCppForm.EnumClass)
+			if (Enum.CppForm != UhtEnumCppForm.EnumClass)
 			{
 				return null;
 			}
-			return $"enum class {this.Enum.SourceName} : {this.UnderlyingType};";
+			return $"enum class {Enum.SourceName} : {UnderlyingType};";
 		}
 
 		/// <inheritdoc/>
 		public override IEnumerable<UhtType> EnumerateReferencedTypes()
 		{
-			yield return this.Enum;
+			yield return Enum;
 		}
 
 		/// <summary>
@@ -212,37 +212,37 @@ namespace EpicGames.UHT.Types
 		/// <inheritdoc/>
 		public override StringBuilder AppendText(StringBuilder builder, UhtPropertyTextType textType, bool isTemplateArgument)
 		{
-			return AppendEnumText(builder, this, this.Enum, textType, isTemplateArgument);
+			return AppendEnumText(builder, this, Enum, textType, isTemplateArgument);
 		}
 
 		/// <inheritdoc/>
 		public override StringBuilder AppendMemberDecl(StringBuilder builder, IUhtPropertyMemberContext context, string name, string nameSuffix, int tabs)
 		{
-			if (this.UnderlyingProperty != null)
+			if (UnderlyingProperty != null)
 			{
-				builder.AppendMemberDecl(this.UnderlyingProperty, context, name, GetNameSuffix(nameSuffix, "_Underlying"), tabs);
+				builder.AppendMemberDecl(UnderlyingProperty, context, name, GetNameSuffix(nameSuffix, "_Underlying"), tabs);
 			}
-			return AppendMemberDecl(builder, context, name, nameSuffix, tabs, this.Enum.CppForm != UhtEnumCppForm.EnumClass ? "FBytePropertyParams" : "FEnumPropertyParams");
+			return AppendMemberDecl(builder, context, name, nameSuffix, tabs, Enum.CppForm != UhtEnumCppForm.EnumClass ? "FBytePropertyParams" : "FEnumPropertyParams");
 		}
 
 		/// <inheritdoc/>
 		public override StringBuilder AppendMemberDef(StringBuilder builder, IUhtPropertyMemberContext context, string name, string nameSuffix, string? offset, int tabs)
 		{
-			if (this.Enum.CppForm == UhtEnumCppForm.EnumClass)
+			if (Enum.CppForm == UhtEnumCppForm.EnumClass)
 			{
-				if (this.UnderlyingProperty != null)
+				if (UnderlyingProperty != null)
 				{
-					builder.AppendMemberDef(this.UnderlyingProperty, context, name, GetNameSuffix(nameSuffix, "_Underlying"), "0", tabs);
+					builder.AppendMemberDef(UnderlyingProperty, context, name, GetNameSuffix(nameSuffix, "_Underlying"), "0", tabs);
 				}
 				AppendMemberDefStart(builder, context, name, nameSuffix, offset, tabs, "FEnumPropertyParams", "UECodeGen_Private::EPropertyGenFlags::Enum");
-				AppendMemberDefRef(builder, context, this.Enum, true);
+				AppendMemberDefRef(builder, context, Enum, true);
 				AppendMemberDefEnd(builder, context, name, nameSuffix);
 				return builder;
 			}
 			else
 			{
 				AppendMemberDefStart(builder, context, name, nameSuffix, offset, tabs, "FBytePropertyParams", "UECodeGen_Private::EPropertyGenFlags::Byte", true, true);
-				AppendMemberDefRef(builder, context, this.Enum, true);
+				AppendMemberDefRef(builder, context, Enum, true);
 				AppendMemberDefEnd(builder, context, name, nameSuffix);
 			}
 			return builder;
@@ -251,9 +251,9 @@ namespace EpicGames.UHT.Types
 		/// <inheritdoc/>
 		public override StringBuilder AppendMemberPtr(StringBuilder builder, IUhtPropertyMemberContext context, string name, string nameSuffix, int tabs)
 		{
-			if (this.UnderlyingProperty != null)
+			if (UnderlyingProperty != null)
 			{
-				builder.AppendMemberPtr(this.UnderlyingProperty, context, name, GetNameSuffix(nameSuffix, "_Underlying"), tabs);
+				builder.AppendMemberPtr(UnderlyingProperty, context, name, GetNameSuffix(nameSuffix, "_Underlying"), tabs);
 			}
 			base.AppendMemberPtr(builder, context, name, nameSuffix, tabs);
 			return builder;
@@ -262,25 +262,25 @@ namespace EpicGames.UHT.Types
 		/// <inheritdoc/>
 		public override StringBuilder AppendFunctionThunkParameterArg(StringBuilder builder)
 		{
-			return AppendEnumFunctionThunkParameterArg(builder, this, this.Enum);
+			return AppendEnumFunctionThunkParameterArg(builder, this, Enum);
 		}
 
 		/// <inheritdoc/>
 		public override void AppendObjectHashes(StringBuilder builder, int startingLength, IUhtPropertyMemberContext context)
 		{
-			builder.AppendObjectHash(startingLength, this, context, this.Enum);
+			builder.AppendObjectHash(startingLength, this, context, Enum);
 		}
 
 		/// <inheritdoc/>
 		public override StringBuilder AppendNullConstructorArg(StringBuilder builder, bool isInitializer)
 		{
-			if (this.Enum.CppForm != UhtEnumCppForm.EnumClass)
+			if (Enum.CppForm != UhtEnumCppForm.EnumClass)
 			{
 				builder.Append('0');
 			}
 			else
 			{
-				builder.Append('(').Append(this.Enum.CppType).Append(")0");
+				builder.Append('(').Append(Enum.CppType).Append(")0");
 			}
 			return builder;
 		}
@@ -288,7 +288,7 @@ namespace EpicGames.UHT.Types
 		/// <inheritdoc/>
 		public override bool SanitizeDefaultValue(IUhtTokenReader defaultValueReader, StringBuilder innerDefaultValue)
 		{
-			return SanitizeEnumDefaultValue(this, this.Enum, defaultValueReader, innerDefaultValue);
+			return SanitizeEnumDefaultValue(this, Enum, defaultValueReader, innerDefaultValue);
 		}
 
 		/// <inheritdoc/>
@@ -296,11 +296,11 @@ namespace EpicGames.UHT.Types
 		{
 			if (other is UhtEnumProperty otherEnum)
 			{
-				return this.Enum == otherEnum.Enum;
+				return Enum == otherEnum.Enum;
 			}
 			else if (other is UhtByteProperty otherByte)
 			{
-				return this.Enum == otherByte.Enum;
+				return Enum == otherByte.Enum;
 			}
 			return false;
 		}
@@ -308,26 +308,26 @@ namespace EpicGames.UHT.Types
 		private UhtProperty CreateUnderlyingProperty()
 		{
 			UhtPropertySettings propertySettings = new();
-			propertySettings.Reset(this, 0, this.PropertyCategory, 0);
+			propertySettings.Reset(this, 0, PropertyCategory, 0);
 			propertySettings.SourceName = "UnderlyingType";
-			switch (this.UnderlyingType)
+			switch (UnderlyingType)
 			{
 				case UhtEnumUnderlyingType.int8:
-					return new UhtInt8Property(propertySettings, this.UnderlyingTypeSize);
+					return new UhtInt8Property(propertySettings, UnderlyingTypeSize);
 				case UhtEnumUnderlyingType.int16:
-					return new UhtInt16Property(propertySettings, this.UnderlyingTypeSize);
+					return new UhtInt16Property(propertySettings, UnderlyingTypeSize);
 				case UhtEnumUnderlyingType.int32:
-					return new UhtIntProperty(propertySettings, this.UnderlyingTypeSize);
+					return new UhtIntProperty(propertySettings, UnderlyingTypeSize);
 				case UhtEnumUnderlyingType.int64:
-					return new UhtInt64Property(propertySettings, this.UnderlyingTypeSize);
+					return new UhtInt64Property(propertySettings, UnderlyingTypeSize);
 				case UhtEnumUnderlyingType.uint8:
-					return new UhtByteProperty(propertySettings, this.UnderlyingTypeSize);
+					return new UhtByteProperty(propertySettings, UnderlyingTypeSize);
 				case UhtEnumUnderlyingType.uint16:
-					return new UhtUInt16Property(propertySettings, this.UnderlyingTypeSize);
+					return new UhtUInt16Property(propertySettings, UnderlyingTypeSize);
 				case UhtEnumUnderlyingType.uint32:
-					return new UhtUInt32Property(propertySettings, this.UnderlyingTypeSize);
+					return new UhtUInt32Property(propertySettings, UnderlyingTypeSize);
 				case UhtEnumUnderlyingType.uint64:
-					return new UhtUInt64Property(propertySettings, this.UnderlyingTypeSize);
+					return new UhtUInt64Property(propertySettings, UnderlyingTypeSize);
 				default:
 					throw new UhtIceException("Unexpected underlying enum type");
 			}
diff --git a/Engine/Source/Programs/Shared/EpicGames.UHT/Types/Properties/UhtFieldPathProperty.cs b/Engine/Source/Programs/Shared/EpicGames.UHT/Types/Properties/UhtFieldPathProperty.cs
--- a/Engine/Source/Programs/Shared/EpicGames.UHT/Types/Properties/UhtFieldPathProperty.cs
+++ b/Engine/Source/Programs/Shared/EpicGames.UHT/Types/Properties/UhtFieldPathProperty.cs
@@ -40,14 +40,14 @@ namespace EpicGames.UHT.Types
 		/// <param name="fieldClassName">Field class name</param>
 		public UhtFieldPathProperty(UhtPropertySettings propertySettings, string fieldClassName) : base(propertySettings)
 		{
-			this.FieldClassName = fieldClassName;
-			this.PropertyCaps |= UhtPropertyCaps.IsParameterSupportedByBlueprint | UhtPropertyCaps.IsMemberSupportedByBlueprint;
+			FieldClassName = fieldClassName;
+			PropertyCaps |= UhtPropertyCaps.IsParameterSupportedByBlueprint | UhtPropertyCaps.IsMemberSupportedByBlueprint;
 		}
 
 		/// <inheritdoc/>
 		public override string? GetForwardDeclarations()
 		{
-			return $"class F{this.FieldClassName};";
+			return $"class F{FieldClassName};";
 		}
 
 		/// <inheritdoc/>
@@ -60,7 +60,7 @@ namespace EpicGames.UHT.Types
 					break;
 
 				default:
-					builder.Append("TFieldPath<F").Append(this.FieldClassName).Append('>');
+					builder.Append("TFieldPath<F").Append(FieldClassName).Append('>');
 					break;
 			}
 			return builder;
@@ -76,7 +76,7 @@ namespace EpicGames.UHT.Types
 		public override StringBuilder AppendMemberDef(StringBuilder builder, IUhtPropertyMemberContext context, string name, string nameSuffix, string? offset, int tabs)
 		{
 			AppendMemberDefStart(builder, context, name, nameSuffix, offset, tabs, "FFieldPathPropertyParams", "UECodeGen_Private::EPropertyGenFlags::FieldPath");
-			builder.Append("&F").Append(this.FieldClassName).Append("::StaticClass, ");
+			builder.Append("&F").Append(FieldClassName).Append("::StaticClass, ");
 			AppendMemberDefEnd(builder, context, name, nameSuffix);
 			return builder;
 		}
@@ -99,7 +99,7 @@ namespace EpicGames.UHT.Types
 		{
 			if (other is UhtFieldPathProperty otherFieldPath)
 			{
-				return this.FieldClassName == otherFieldPath.FieldClassName;
+				return FieldClassName == otherFieldPath.FieldClassName;
 			}
 			return false;
 		}
diff --git a/Engine/Source/Programs/Shared/EpicGames.UHT/Types/Properties/UhtFloatProperty.cs b/Engine/Source/Programs/Shared/EpicGames.UHT/Types/Properties/UhtFloatProperty.cs
--- a/Engine/Source/Programs/Shared/EpicGames.UHT/Types/Properties/UhtFloatProperty.cs
+++ b/Engine/Source/Programs/Shared/EpicGames.UHT/Types/Properties/UhtFloatProperty.cs
@@ -28,7 +28,7 @@ namespace EpicGames.UHT.Types
 		/// <param name="propertySettings">Property settings</param>
 		public UhtFloatProperty(UhtPropertySettings propertySettings) : base(propertySettings, UhtPropertyIntType.None)
 		{
-			this.PropertyCaps |= UhtPropertyCaps.CanExposeOnSpawn | UhtPropertyCaps.IsParameterSupportedByBlueprint | UhtPropertyCaps.IsMemberSupportedByBlueprint;
+			PropertyCaps |= UhtPropertyCaps.CanExposeOnSpawn | UhtPropertyCaps.IsParameterSupportedByBlueprint | UhtPropertyCaps.IsMemberSupportedByBlueprint;
 		}
 
 		/// <inheritdoc/>
diff --git a/Engine/Source/Programs/Shared/EpicGames.UHT/Types/Properties/UhtInt64Property.cs b/Engine/Source/Programs/Shared/EpicGames.UHT/Types/Properties/UhtInt64Property.cs
--- a/Engine/Source/Programs/Shared/EpicGames.UHT/Types/Properties/UhtInt64Property.cs
+++ b/Engine/Source/Programs/Shared/EpicGames.UHT/Types/Properties/UhtInt64Property.cs
@@ -28,7 +28,7 @@ namespace EpicGames.UHT.Types
 		/// <param name="intType">Integer type</param>
 		public UhtInt64Property(UhtPropertySettings propertySettings, UhtPropertyIntType intType) : base(propertySettings, intType)
 		{
-			this.PropertyCaps |= UhtPropertyCaps.CanExposeOnSpawn | UhtPropertyCaps.IsParameterSupportedByBlueprint | UhtPropertyCaps.IsMemberSupportedByBlueprint;
+			PropertyCaps |= UhtPropertyCaps.CanExposeOnSpawn | UhtPropertyCaps.IsParameterSupportedByBlueprint | UhtPropertyCaps.IsMemberSupportedByBlueprint;
 		}
 
 		/// <inheritdoc/>
diff --git a/Engine/Source/Programs/Shared/EpicGames.UHT/Types/Properties/UhtIntProperty.cs b/Engine/Source/Programs/Shared/EpicGames.UHT/Types/Properties/UhtIntProperty.cs
--- a/Engine/Source/Programs/Shared/EpicGames.UHT/Types/Properties/UhtIntProperty.cs
+++ b/Engine/Source/Programs/Shared/EpicGames.UHT/Types/Properties/UhtIntProperty.cs
@@ -28,21 +28,21 @@ namespace EpicGames.UHT.Types
 		/// <param name="intType">Integer type</param>
 		public UhtIntProperty(UhtPropertySettings propertySettings, UhtPropertyIntType intType) : base(propertySettings, intType)
 		{
-			this.PropertyCaps |= UhtPropertyCaps.CanExposeOnSpawn | UhtPropertyCaps.IsParameterSupportedByBlueprint | UhtPropertyCaps.IsMemberSupportedByBlueprint;
+			PropertyCaps |= UhtPropertyCaps.CanExposeOnSpawn | UhtPropertyCaps.IsParameterSupportedByBlueprint | UhtPropertyCaps.IsMemberSupportedByBlueprint;
 		}
 
 		/// <inheritdoc/>
 		public override StringBuilder AppendMemberDecl(StringBuilder builder, IUhtPropertyMemberContext context, string name, string nameSuffix, int tabs)
 		{
 			return AppendMemberDecl(builder, context, name, nameSuffix, tabs,
-				this.IntType == UhtPropertyIntType.Unsized ? "FUnsizedIntPropertyParams" : "FIntPropertyParams");
+				IntType == UhtPropertyIntType.Unsized ? "FUnsizedIntPropertyParams" : "FIntPropertyParams");
 		}
 
 		/// <inheritdoc/>
 		public override StringBuilder AppendMemberDef(StringBuilder builder, IUhtPropertyMemberContext context, string name, string nameSuffix, string? offset, int tabs)
 		{
 			AppendMemberDefStart(builder, context, name, nameSuffix, offset, tabs,
-				this.IntType == UhtPropertyIntType.Unsized ? "FUnsizedIntPropertyParams" : "FIntPropertyParams",
+				IntType == UhtPropertyIntType.Unsized ? "FUnsizedIntPropertyParams" : "FIntPropertyParams",
 				"UECodeGen_Private::EPropertyGenFlags::Int");
 			AppendMemberDefEnd(builder, context, name, nameSuffix);
 			return builder;
diff --git a/Engine/Source/Programs/Shared/EpicGames.UHT/Types/Properties/UhtInterfaceProperty.cs b/Engine/Source/Programs/Shared/EpicGames.UHT/Types/Properties/UhtInterfaceProperty.cs
--- a/Engine/Source/Programs/Shared/EpicGames.UHT/Types/Properties/UhtInterfaceProperty.cs
+++ b/Engine/Source/Programs/Shared/EpicGames.UHT/Types/Properties/UhtInterfaceProperty.cs
@@ -43,13 +43,13 @@ namespace EpicGames.UHT.Types
 		/// <param name="interfaceClass">Referenced interface</param>
 		public UhtInterfaceProperty(UhtPropertySettings propertySettings, UhtClass interfaceClass) : base(propertySettings)
 		{
-			this.InterfaceClass = interfaceClass;
-			this.PropertyFlags |= EPropertyFlags.UObjectWrapper;
-			this.PropertyCaps |= UhtPropertyCaps.CanExposeOnSpawn | UhtPropertyCaps.IsParameterSupportedByBlueprint | UhtPropertyCaps.IsMemberSupportedByBlueprint | UhtPropertyCaps.PassCppArgsByRef;
-			this.PropertyCaps &= ~(UhtPropertyCaps.CanHaveConfig | UhtPropertyCaps.CanBeContainerKey);
-			if (this.Session.Config!.AreRigVMUInterfaceProeprtiesEnabled)
+			InterfaceClass = interfaceClass;
+			PropertyFlags |= EPropertyFlags.UObjectWrapper;
+			PropertyCaps |= UhtPropertyCaps.CanExposeOnSpawn | UhtPropertyCaps.IsParameterSupportedByBlueprint | UhtPropertyCaps.IsMemberSupportedByBlueprint | UhtPropertyCaps.PassCppArgsByRef;
+			PropertyCaps &= ~(UhtPropertyCaps.CanHaveConfig | UhtPropertyCaps.CanBeContainerKey);
+			if (Session.Config!.AreRigVMUInterfaceProeprtiesEnabled)
 			{
-				this.PropertyCaps |= UhtPropertyCaps.SupportsRigVM;
+				PropertyCaps |= UhtPropertyCaps.SupportsRigVM;
 			}
 		}
 
@@ -60,9 +60,9 @@ namespace EpicGames.UHT.Types
 			switch (phase)
 			{
 				case UhtResolvePhase.Final:
-					if (this.InterfaceClass.HierarchyHasAnyClassFlags(EClassFlags.DefaultToInstanced))
+					if (InterfaceClass.HierarchyHasAnyClassFlags(EClassFlags.DefaultToInstanced))
 					{
-						this.PropertyFlags |= (EPropertyFlags.InstancedReference | EPropertyFlags.ExportObject) & ~this.DisallowPropertyFlags;
+						PropertyFlags |= (EPropertyFlags.InstancedReference | EPropertyFlags.ExportObject) & ~DisallowPropertyFlags;
 					}
 					break;
 			}
@@ -72,19 +72,19 @@ namespace EpicGames.UHT.Types
 		/// <inheritdoc/>
 		public override bool ScanForInstancedReferenced(bool deepScan)
 		{
-			return !this.DisallowPropertyFlags.HasAnyFlags(EPropertyFlags.InstancedReference) && this.InterfaceClass.HierarchyHasAnyClassFlags(EClassFlags.DefaultToInstanced);
+			return !DisallowPropertyFlags.HasAnyFlags(EPropertyFlags.InstancedReference) && InterfaceClass.HierarchyHasAnyClassFlags(EClassFlags.DefaultToInstanced);
 		}
 
 		/// <inheritdoc/>
 		public override void CollectReferencesInternal(IUhtReferenceCollector collector, bool templateProperty)
 		{
-			collector.AddCrossModuleReference(this.InterfaceClass, false);
+			collector.AddCrossModuleReference(InterfaceClass, false);
 		}
 
 		/// <inheritdoc/>
 		public override string? GetForwardDeclarations()
 		{
-			UhtClass? exportClass = this.InterfaceClass;
+			UhtClass? exportClass = InterfaceClass;
 			while (exportClass != null && !exportClass.ClassFlags.HasAnyFlags(EClassFlags.Native))
 			{
 				exportClass = exportClass.SuperClass;
@@ -95,7 +95,7 @@ namespace EpicGames.UHT.Types
 		/// <inheritdoc/>
 		public override IEnumerable<UhtType> EnumerateReferencedTypes()
 		{
-			yield return this.InterfaceClass;
+			yield return InterfaceClass;
 		}
 
 		/// <inheritdoc/>
@@ -108,11 +108,11 @@ namespace EpicGames.UHT.Types
 					break;
 
 				case UhtPropertyTextType.FunctionThunkParameterArgType:
-					builder.Append(this.InterfaceClass.SourceName);
+					builder.Append(InterfaceClass.SourceName);
 					break;
 
 				default:
-					builder.Append("TScriptInterface<").Append(this.InterfaceClass.SourceName).Append('>');
+					builder.Append("TScriptInterface<").Append(InterfaceClass.SourceName).Append('>');
 					break;
 			}
 			return builder;
@@ -129,7 +129,7 @@ namespace EpicGames.UHT.Types
 		{
 			// FScriptInterface<USomeInterface> is valid so in that case we need to pass in the interface class and not the alternate object (which in the end is the same object)
 			AppendMemberDefStart(builder, context, name, nameSuffix, offset, tabs, "FInterfacePropertyParams", "UECodeGen_Private::EPropertyGenFlags::Interface");
-			AppendMemberDefRef(builder, context, this.InterfaceClass.AlternateObject ?? this.InterfaceClass, false);
+			AppendMemberDefRef(builder, context, InterfaceClass.AlternateObject ?? InterfaceClass, false);
 			AppendMemberDefEnd(builder, context, name, nameSuffix);
 			return builder;
 		}
@@ -152,9 +152,9 @@ namespace EpicGames.UHT.Types
 		{
 			base.ValidateMember(structObj, options);
 
-			if (this.PointerType == UhtPointerType.Native)
+			if (PointerType == UhtPointerType.Native)
 			{
-				this.LogError($"UPROPERTY pointers cannot be interfaces - did you mean TScriptInterface<{this.InterfaceClass.SourceName}>?");
+				this.LogError($"UPROPERTY pointers cannot be interfaces - did you mean TScriptInterface<{InterfaceClass.SourceName}>?");
 			}
 		}
 
@@ -163,7 +163,7 @@ namespace EpicGames.UHT.Types
 		{
 			if (other is UhtInterfaceProperty otherObject)
 			{
-				return this.InterfaceClass == otherObject.InterfaceClass;
+				return InterfaceClass == otherObject.InterfaceClass;
 			}
 			return false;
 		}
@@ -171,8 +171,8 @@ namespace EpicGames.UHT.Types
 		/// <inheritdoc/>
 		public override bool MustBeConstArgument([NotNullWhen(true)] out UhtType? errorType)
 		{
-			errorType = this.InterfaceClass;
-			return this.InterfaceClass.ClassFlags.HasAnyFlags(EClassFlags.Const);
+			errorType = InterfaceClass;
+			return InterfaceClass.ClassFlags.HasAnyFlags(EClassFlags.Const);
 		}
 
 		#region Keywords
diff --git a/Engine/Source/Programs/Shared/EpicGames.UHT/Types/Properties/UhtLargeWorldCoordinatesRealProperty.cs b/Engine/Source/Programs/Shared/EpicGames.UHT/Types/Properties/UhtLargeWorldCoordinatesRealProperty.cs
--- a/Engine/Source/Programs/Shared/EpicGames.UHT/Types/Properties/UhtLargeWorldCoordinatesRealProperty.cs
+++ b/Engine/Source/Programs/Shared/EpicGames.UHT/Types/Properties/UhtLargeWorldCoordinatesRealProperty.cs
@@ -28,7 +28,7 @@ namespace EpicGames.UHT.Types
 		/// <param name="propertySettings">Property settings</param>
 		public UhtLargeWorldCoordinatesRealProperty(UhtPropertySettings propertySettings) : base(propertySettings, UhtPropertyIntType.None)
 		{
-			this.PropertyCaps |= UhtPropertyCaps.IsParameterSupportedByBlueprint | UhtPropertyCaps.IsMemberSupportedByBlueprint;
+			PropertyCaps |= UhtPropertyCaps.IsParameterSupportedByBlueprint | UhtPropertyCaps.IsMemberSupportedByBlueprint;
 		}
 
 		/// <inheritdoc/>
diff --git a/Engine/Source/Programs/Shared/EpicGames.UHT/Types/Properties/UhtLazyObjectPtrProperty.cs b/Engine/Source/Programs/Shared/EpicGames.UHT/Types/Properties/UhtLazyObjectPtrProperty.cs
--- a/Engine/Source/Programs/Shared/EpicGames.UHT/Types/Properties/UhtLazyObjectPtrProperty.cs
+++ b/Engine/Source/Programs/Shared/EpicGames.UHT/Types/Properties/UhtLazyObjectPtrProperty.cs
@@ -40,8 +40,8 @@ namespace EpicGames.UHT.Types
 		public UhtLazyObjectPtrProperty(UhtPropertySettings propertySettings, UhtClass classObj)
 			: base(propertySettings, classObj, null)
 		{
-			this.PropertyFlags |= EPropertyFlags.UObjectWrapper;
-			this.PropertyCaps |= UhtPropertyCaps.PassCppArgsByRef | UhtPropertyCaps.RequiresNullConstructorArg;
+			PropertyFlags |= EPropertyFlags.UObjectWrapper;
+			PropertyCaps |= UhtPropertyCaps.PassCppArgsByRef | UhtPropertyCaps.RequiresNullConstructorArg;
 		}
 
 		/// <inheritdoc/>
@@ -50,11 +50,11 @@ namespace EpicGames.UHT.Types
 			switch (textType)
 			{
 				case UhtPropertyTextType.FunctionThunkParameterArgType:
-					builder.Append("TLazyObjectPtr<").Append(this.Class.SourceName).Append("> "); //COMPATIBILITY-TODO - Extra space
+					builder.Append("TLazyObjectPtr<").Append(Class.SourceName).Append("> "); //COMPATIBILITY-TODO - Extra space
 					break;
 
 				default:
-					builder.Append("TLazyObjectPtr<").Append(this.Class.SourceName).Append('>');
+					builder.Append("TLazyObjectPtr<").Append(Class.SourceName).Append('>');
 					break;
 			}
 			return builder;
@@ -70,7 +70,7 @@ namespace EpicGames.UHT.Types
 		public override StringBuilder AppendMemberDef(StringBuilder builder, IUhtPropertyMemberContext context, string name, string nameSuffix, string? offset, int tabs)
 		{
 			AppendMemberDefStart(builder, context, name, nameSuffix, offset, tabs, "FLazyObjectPropertyParams", "UECodeGen_Private::EPropertyGenFlags::LazyObject");
-			AppendMemberDefRef(builder, context, this.Class, false);
+			AppendMemberDefRef(builder, context, Class, false);
 			AppendMemberDefEnd(builder, context, name, nameSuffix);
 			return builder;
 		}
@@ -87,7 +87,7 @@ namespace EpicGames.UHT.Types
 		{
 			if (other is UhtLazyObjectPtrProperty otherObject)
 			{
-				return this.Class == otherObject.Class && this.MetaClass == otherObject.MetaClass;
+				return Class == otherObject.Class && MetaClass == otherObject.MetaClass;
 			}
 			return false;
 		}
@@ -99,7 +99,7 @@ namespace EpicGames.UHT.Types
 
 			// UFunctions with a smart pointer as input parameter wont compile anyway, because of missing P_GET_... macro.
 			// UFunctions with a smart pointer as return type will crash when called via blueprint, because they are not supported in VM.
-			if (this.PropertyCategory == UhtPropertyCategory.RegularParameter || this.PropertyCategory == UhtPropertyCategory.ReplicatedParameter)
+			if (PropertyCategory == UhtPropertyCategory.RegularParameter || PropertyCategory == UhtPropertyCategory.ReplicatedParameter)
 			{
 				outerStruct.LogError("UFunctions cannot take a lazy pointer as a parameter.");
 			}
diff --git a/Engine/Source/Programs/Shared/EpicGames.UHT/Types/Properties/UhtMapProperty.cs b/Engine/Source/Programs/Shared/EpicGames.UHT/Types/Properties/UhtMapProperty.cs
--- a/Engine/Source/Programs/Shared/EpicGames.UHT/Types/Properties/UhtMapProperty.cs
+++ b/Engine/Source/Programs/Shared/EpicGames.UHT/Types/Properties/UhtMapProperty.cs
@@ -45,34 +45,34 @@ namespace EpicGames.UHT.Types
 		public UhtMapProperty(UhtPropertySettings propertySettings, UhtProperty key, UhtProperty value) : base(propertySettings, value)
 		{
 			// If the creation of the value property set more flags, then copy those flags to ourselves
-			this.PropertyFlags |= value.PropertyFlags & EPropertyFlags.UObjectWrapper;
+			PropertyFlags |= value.PropertyFlags & EPropertyFlags.UObjectWrapper;
 
 			// Make sure the 'UObjectWrapper' flag is maintained so that both 'TMap<TSubclassOf<...>, ...>' and 'TMap<UClass*, TSubclassOf<...>>' works correctly
 			key.PropertyFlags = (value.PropertyFlags & ~EPropertyFlags.UObjectWrapper) | (key.PropertyFlags & EPropertyFlags.UObjectWrapper);
 			key.DisallowPropertyFlags = ~(EPropertyFlags.ContainsInstancedReference | EPropertyFlags.InstancedReference | EPropertyFlags.UObjectWrapper);
 
-			this.KeyProperty = key;
-			this.PropertyCaps |= UhtPropertyCaps.PassCppArgsByRef;
-			this.PropertyCaps &= ~(UhtPropertyCaps.CanBeContainerValue | UhtPropertyCaps.CanBeContainerKey);
-			this.PropertyCaps &= ~(UhtPropertyCaps.IsParameterSupportedByBlueprint | UhtPropertyCaps.IsMemberSupportedByBlueprint | UhtPropertyCaps.CanExposeOnSpawn);
-			this.PropertyCaps |= (this.ValueProperty.PropertyCaps & (UhtPropertyCaps.IsParameterSupportedByBlueprint | UhtPropertyCaps.IsMemberSupportedByBlueprint | UhtPropertyCaps.CanExposeOnSpawn)) &
-				(this.KeyProperty.PropertyCaps & (UhtPropertyCaps.IsParameterSupportedByBlueprint | UhtPropertyCaps.IsMemberSupportedByBlueprint));
-			this.PropertyFlags = value.PropertyFlags;
-			this.ValueProperty.SourceName = this.SourceName;
-			this.ValueProperty.EngineName = this.EngineName;
-			this.ValueProperty.PropertyFlags &= EPropertyFlags.PropagateToMapValue;
-			this.ValueProperty.Outer = this;
-			this.ValueProperty.MetaData.Clear();
-			this.KeyProperty.SourceName = $"{this.SourceName}_Key";
-			this.KeyProperty.EngineName = $"{this.EngineName}_Key";
-			this.KeyProperty.PropertyFlags &= EPropertyFlags.PropagateToMapKey;
-			this.KeyProperty.Outer = this;
-			this.KeyProperty.MetaData.Clear();
+			KeyProperty = key;
+			PropertyCaps |= UhtPropertyCaps.PassCppArgsByRef;
+			PropertyCaps &= ~(UhtPropertyCaps.CanBeContainerValue | UhtPropertyCaps.CanBeContainerKey);
+			PropertyCaps &= ~(UhtPropertyCaps.IsParameterSupportedByBlueprint | UhtPropertyCaps.IsMemberSupportedByBlueprint | UhtPropertyCaps.CanExposeOnSpawn);
+			PropertyCaps |= (ValueProperty.PropertyCaps & (UhtPropertyCaps.IsParameterSupportedByBlueprint | UhtPropertyCaps.IsMemberSupportedByBlueprint | UhtPropertyCaps.CanExposeOnSpawn)) &
+				(KeyProperty.PropertyCaps & (UhtPropertyCaps.IsParameterSupportedByBlueprint | UhtPropertyCaps.IsMemberSupportedByBlueprint));
+			PropertyFlags = value.PropertyFlags;
+			ValueProperty.SourceName = SourceName;
+			ValueProperty.EngineName = EngineName;
+			ValueProperty.PropertyFlags &= EPropertyFlags.PropagateToMapValue;
+			ValueProperty.Outer = this;
+			ValueProperty.MetaData.Clear();
+			KeyProperty.SourceName = $"{SourceName}_Key";
+			KeyProperty.EngineName = $"{EngineName}_Key";
+			KeyProperty.PropertyFlags &= EPropertyFlags.PropagateToMapKey;
+			KeyProperty.Outer = this;
+			KeyProperty.MetaData.Clear();
 
 			// With old UHT, Deprecated was applied after property create.
 			// With the new, it is applied prior to creation.  Deprecated in old 
 			// was not on the key.
-			this.KeyProperty.PropertyFlags &= ~EPropertyFlags.Deprecated;
+			KeyProperty.PropertyFlags &= ~EPropertyFlags.Deprecated;
 		}
 
 		/// <inheritdoc/>
@@ -82,18 +82,18 @@ namespace EpicGames.UHT.Types
 			switch (phase)
 			{
 				case UhtResolvePhase.Final:
-					this.KeyProperty.Resolve(phase);
-					this.KeyProperty.MetaData.Clear();
-
-					EPropertyFlags newFlags = ResolveAndReturnNewFlags(this.ValueProperty, phase);
-					this.PropertyFlags |= newFlags;
-					this.KeyProperty.PropertyFlags |= newFlags;
-					this.MetaData.Add(this.ValueProperty.MetaData);
-					this.ValueProperty.MetaData.Clear();
-					this.ValueProperty.PropertyFlags = this.PropertyFlags & EPropertyFlags.PropagateToMapValue;
-
-					PropagateFlagsFromInnerAndHandlePersistentInstanceMetadata(this, null, this.KeyProperty);
-					PropagateFlagsFromInnerAndHandlePersistentInstanceMetadata(this, this.MetaData, this.ValueProperty);
+					KeyProperty.Resolve(phase);
+					KeyProperty.MetaData.Clear();
+
+					EPropertyFlags newFlags = ResolveAndReturnNewFlags(ValueProperty, phase);
+					PropertyFlags |= newFlags;
+					KeyProperty.PropertyFlags |= newFlags;
+					MetaData.Add(ValueProperty.MetaData);
+					ValueProperty.MetaData.Clear();
+					ValueProperty.PropertyFlags = PropertyFlags & EPropertyFlags.PropagateToMapValue;
+
+					PropagateFlagsFromInnerAndHandlePersistentInstanceMetadata(this, null, KeyProperty);
+					PropagateFlagsFromInnerAndHandlePersistentInstanceMetadata(this, MetaData, ValueProperty);
 					break;
 			}
 			return results;
@@ -102,24 +102,24 @@ namespace EpicGames.UHT.Types
 		/// <inheritdoc/>
 		public override void CollectReferencesInternal(IUhtReferenceCollector collector, bool templateProperty)
 		{
-			this.ValueProperty.CollectReferencesInternal(collector, true);
-			this.KeyProperty.CollectReferencesInternal(collector, true);
+			ValueProperty.CollectReferencesInternal(collector, true);
+			KeyProperty.CollectReferencesInternal(collector, true);
 		}
 
 		/// <inheritdoc/>
 		public override string? GetForwardDeclarations()
 		{
-			return $"{this.KeyProperty.GetForwardDeclarations()} {this.ValueProperty.GetForwardDeclarations()}";
+			return $"{KeyProperty.GetForwardDeclarations()} {ValueProperty.GetForwardDeclarations()}";
 		}
 
 		/// <inheritdoc/>
 		public override IEnumerable<UhtType> EnumerateReferencedTypes()
 		{
-			foreach (UhtType type in this.ValueProperty.EnumerateReferencedTypes())
+			foreach (UhtType type in ValueProperty.EnumerateReferencedTypes())
 			{
 				yield return type;
 			}
-			foreach (UhtType type in this.KeyProperty.EnumerateReferencedTypes())
+			foreach (UhtType type in KeyProperty.EnumerateReferencedTypes())
 			{
 				yield return type;
 			}
@@ -135,17 +135,17 @@ namespace EpicGames.UHT.Types
 					break;
 
 				case UhtPropertyTextType.FunctionThunkParameterArgType:
-					builder.AppendFunctionThunkParameterArrayType(this.KeyProperty).Append(',').AppendFunctionThunkParameterArrayType(this.ValueProperty);
+					builder.AppendFunctionThunkParameterArrayType(KeyProperty).Append(',').AppendFunctionThunkParameterArrayType(ValueProperty);
 					break;
 
 				default:
-					builder.Append("TMap<").AppendPropertyText(this.KeyProperty, textType, true);
+					builder.Append("TMap<").AppendPropertyText(KeyProperty, textType, true);
 					if (builder[^1] == '>')
 					{
 						// if our internal property type is a template class, add a space between the closing brackets b/c VS.NET cannot parse this correctly
 						builder.Append(' ');
 					}
-					builder.Append(',').AppendPropertyText(this.ValueProperty, textType, true);
+					builder.Append(',').AppendPropertyText(ValueProperty, textType, true);
 					if (builder[^1] == '>')
 					{
 						// if our internal property type is a template class, add a space between the closing brackets b/c VS.NET cannot parse this correctly
@@ -160,18 +160,18 @@ namespace EpicGames.UHT.Types
 		/// <inheritdoc/>
 		public override StringBuilder AppendMemberDecl(StringBuilder builder, IUhtPropertyMemberContext context, string name, string nameSuffix, int tabs)
 		{
-			builder.AppendMemberDecl(this.ValueProperty, context, name, GetNameSuffix(nameSuffix, "_ValueProp"), tabs);
-			builder.AppendMemberDecl(this.KeyProperty, context, name, GetNameSuffix(nameSuffix, "_Key_KeyProp"), tabs);
+			builder.AppendMemberDecl(ValueProperty, context, name, GetNameSuffix(nameSuffix, "_ValueProp"), tabs);
+			builder.AppendMemberDecl(KeyProperty, context, name, GetNameSuffix(nameSuffix, "_Key_KeyProp"), tabs);
 			return AppendMemberDecl(builder, context, name, nameSuffix, tabs, "FMapPropertyParams");
 		}
 
 		/// <inheritdoc/>
 		public override StringBuilder AppendMemberDef(StringBuilder builder, IUhtPropertyMemberContext context, string name, string nameSuffix, string? offset, int tabs)
 		{
-			builder.AppendMemberDef(this.ValueProperty, context, name, GetNameSuffix(nameSuffix, "_ValueProp"), "1", tabs);
-			builder.AppendMemberDef(this.KeyProperty, context, name, GetNameSuffix(nameSuffix, "_Key_KeyProp"), "0", tabs);
+			builder.AppendMemberDef(ValueProperty, context, name, GetNameSuffix(nameSuffix, "_ValueProp"), "1", tabs);
+			builder.AppendMemberDef(KeyProperty, context, name, GetNameSuffix(nameSuffix, "_Key_KeyProp"), "0", tabs);
 			AppendMemberDefStart(builder, context, name, nameSuffix, offset, tabs, "FMapPropertyParams", "UECodeGen_Private::EPropertyGenFlags::Map");
-			builder.Append(this.Allocator == UhtPropertyAllocator.MemoryImage ? "EMapPropertyFlags::UsesMemoryImageAllocator" : "EMapPropertyFlags::None").Append(", ");
+			builder.Append(Allocator == UhtPropertyAllocator.MemoryImage ? "EMapPropertyFlags::UsesMemoryImageAllocator" : "EMapPropertyFlags::None").Append(", ");
 			AppendMemberDefEnd(builder, context, name, nameSuffix);
 			return builder;
 		}
@@ -179,8 +179,8 @@ namespace EpicGames.UHT.Types
 		/// <inheritdoc/>
 		public override StringBuilder AppendMemberPtr(StringBuilder builder, IUhtPropertyMemberContext context, string name, string nameSuffix, int tabs)
 		{
-			builder.AppendMemberPtr(this.ValueProperty, context, name, GetNameSuffix(nameSuffix, "_ValueProp"), tabs);
-			builder.AppendMemberPtr(this.KeyProperty, context, name, GetNameSuffix(nameSuffix, "_Key_KeyProp"), tabs);
+			builder.AppendMemberPtr(ValueProperty, context, name, GetNameSuffix(nameSuffix, "_ValueProp"), tabs);
+			builder.AppendMemberPtr(KeyProperty, context, name, GetNameSuffix(nameSuffix, "_Key_KeyProp"), tabs);
 			base.AppendMemberPtr(builder, context, name, nameSuffix, tabs);
 			return builder;
 		}
@@ -188,8 +188,8 @@ namespace EpicGames.UHT.Types
 		/// <inheritdoc/>
 		public override void AppendObjectHashes(StringBuilder builder, int startingLength, IUhtPropertyMemberContext context)
 		{
-			this.KeyProperty.AppendObjectHashes(builder, startingLength, context);
-			this.ValueProperty.AppendObjectHashes(builder, startingLength, context);
+			KeyProperty.AppendObjectHashes(builder, startingLength, context);
+			ValueProperty.AppendObjectHashes(builder, startingLength, context);
 		}
 
 		/// <inheritdoc/>
@@ -210,8 +210,8 @@ namespace EpicGames.UHT.Types
 		{
 			if (other is UhtMapProperty otherMap)
 			{
-				return this.ValueProperty.IsSameType(otherMap.ValueProperty) &&
-					this.KeyProperty.IsSameType(otherMap.KeyProperty);
+				return ValueProperty.IsSameType(otherMap.ValueProperty) &&
+					KeyProperty.IsSameType(otherMap.KeyProperty);
 			}
 			return false;
 		}
@@ -220,22 +220,22 @@ namespace EpicGames.UHT.Types
 		public override void ValidateDeprecated()
 		{
 			base.ValidateDeprecated();
-			this.KeyProperty.ValidateDeprecated();
+			KeyProperty.ValidateDeprecated();
 		}
 
 		/// <inheritdoc/>
 		public override void Validate(UhtStruct outerStruct, UhtProperty outermostProperty, UhtValidationOptions options)
 		{
 			base.Validate(outerStruct, outermostProperty, options);
-			this.KeyProperty.Validate(outerStruct, outermostProperty, options | UhtValidationOptions.IsKey);
+			KeyProperty.Validate(outerStruct, outermostProperty, options | UhtValidationOptions.IsKey);
 		}
 
 		///<inheritdoc/>
 		public override bool ValidateStructPropertyOkForNet(UhtProperty referencingProperty)
 		{
-			if (!this.PropertyFlags.HasAnyFlags(EPropertyFlags.RepSkip))
+			if (!PropertyFlags.HasAnyFlags(EPropertyFlags.RepSkip))
 			{
-				referencingProperty.LogError($"Maps are not supported for Replication or RPCs.  Map '{this.SourceName}' in '{this.Outer?.SourceName}'.  Origin '{referencingProperty.SourceName}'");
+				referencingProperty.LogError($"Maps are not supported for Replication or RPCs.  Map '{SourceName}' in '{Outer?.SourceName}'.  Origin '{referencingProperty.SourceName}'");
 				return false;
 			}
 			return true;
diff --git a/Engine/Source/Programs/Shared/EpicGames.UHT/Types/Properties/UhtMulticastDelegateProperty.cs b/Engine/Source/Programs/Shared/EpicGames.UHT/Types/Properties/UhtMulticastDelegateProperty.cs
--- a/Engine/Source/Programs/Shared/EpicGames.UHT/Types/Properties/UhtMulticastDelegateProperty.cs
+++ b/Engine/Source/Programs/Shared/EpicGames.UHT/Types/Properties/UhtMulticastDelegateProperty.cs
@@ -40,11 +40,11 @@ namespace EpicGames.UHT.Types
 		/// <param name="function">Referenced function</param>
 		protected UhtMulticastDelegateProperty(UhtPropertySettings propertySettings, UhtFunction function) : base(propertySettings)
 		{
-			this.Function = function;
-			this.HeaderFile.AddReferencedHeader(function);
-			this.PropertyCaps |= UhtPropertyCaps.PassCppArgsByRef | UhtPropertyCaps.CanBeBlueprintAssignable | UhtPropertyCaps.CanBeBlueprintCallable |
+			Function = function;
+			HeaderFile.AddReferencedHeader(function);
+			PropertyCaps |= UhtPropertyCaps.PassCppArgsByRef | UhtPropertyCaps.CanBeBlueprintAssignable | UhtPropertyCaps.CanBeBlueprintCallable |
 				UhtPropertyCaps.CanBeBlueprintAuthorityOnly | UhtPropertyCaps.IsMemberSupportedByBlueprint | UhtPropertyCaps.SupportsRigVM;
-			this.PropertyCaps &= ~(UhtPropertyCaps.IsParameterSupportedByBlueprint);
+			PropertyCaps &= ~(UhtPropertyCaps.IsParameterSupportedByBlueprint);
 		}
 
 		/// <inheritdoc/>
@@ -54,7 +54,7 @@ namespace EpicGames.UHT.Types
 			switch (phase)
 			{
 				case UhtResolvePhase.Final:
-					this.PropertyFlags |= EPropertyFlags.InstancedReference & ~this.DisallowPropertyFlags;
+					PropertyFlags |= EPropertyFlags.InstancedReference & ~DisallowPropertyFlags;
 					break;
 			}
 			return results;
@@ -69,7 +69,7 @@ namespace EpicGames.UHT.Types
 		/// <inheritdoc/>
 		public override void CollectReferencesInternal(IUhtReferenceCollector collector, bool templateProperty)
 		{
-			collector.AddCrossModuleReference(this.Function, true);
+			collector.AddCrossModuleReference(Function, true);
 		}
 
 		/// <inheritdoc/>
@@ -81,11 +81,11 @@ namespace EpicGames.UHT.Types
 				case UhtPropertyTextType.ExportMember:
 				case UhtPropertyTextType.RigVMType:
 				case UhtPropertyTextType.FunctionThunkParameterArgType:
-					builder.Append(this.CppTypeText);
+					builder.Append(CppTypeText);
 					break;
 
 				default:
-					builder.Append(this.Function.SourceName);
+					builder.Append(Function.SourceName);
 					break;
 			}
 			return builder;
@@ -94,13 +94,13 @@ namespace EpicGames.UHT.Types
 		/// <inheritdoc/>
 		public override void AppendObjectHashes(StringBuilder builder, int startingLength, IUhtPropertyMemberContext context)
 		{
-			builder.AppendObjectHash(startingLength, this, context, this.Function);
+			builder.AppendObjectHash(startingLength, this, context, Function);
 		}
 
 		/// <inheritdoc/>
 		public override StringBuilder AppendFunctionThunkParameterArg(StringBuilder builder)
 		{
-			return builder.Append(this.Function.SourceName).Append('(').AppendFunctionThunkParameterName(this).Append(')');
+			return builder.Append(Function.SourceName).Append('(').AppendFunctionThunkParameterName(this).Append(')');
 		}
 
 		/// <inheritdoc/>
@@ -121,7 +121,7 @@ namespace EpicGames.UHT.Types
 		{
 			if (other is UhtMulticastDelegateProperty otherObject)
 			{
-				return this.Function == otherObject.Function;
+				return Function == otherObject.Function;
 			}
 			return false;
 		}
diff --git a/Engine/Source/Programs/Shared/EpicGames.UHT/Types/Properties/UhtMulticastInlineDelegateProperty.cs b/Engine/Source/Programs/Shared/EpicGames.UHT/Types/Properties/UhtMulticastInlineDelegateProperty.cs
--- a/Engine/Source/Programs/Shared/EpicGames.UHT/Types/Properties/UhtMulticastInlineDelegateProperty.cs
+++ b/Engine/Source/Programs/Shared/EpicGames.UHT/Types/Properties/UhtMulticastInlineDelegateProperty.cs
@@ -49,7 +49,7 @@ namespace EpicGames.UHT.Types
 		public override StringBuilder AppendMemberDef(StringBuilder builder, IUhtPropertyMemberContext context, string name, string nameSuffix, string? offset, int tabs)
 		{
 			AppendMemberDefStart(builder, context, name, nameSuffix, offset, tabs, "FMulticastDelegatePropertyParams", "UECodeGen_Private::EPropertyGenFlags::InlineMulticastDelegate");
-			AppendMemberDefRef(builder, context, this.Function, true);
+			AppendMemberDefRef(builder, context, Function, true);
 			AppendMemberDefEnd(builder, context, name, nameSuffix);
 			return builder;
 		}
diff --git a/Engine/Source/Programs/Shared/EpicGames.UHT/Types/Properties/UhtMulticastSparseDelegateProperty.cs b/Engine/Source/Programs/Shared/EpicGames.UHT/Types/Properties/UhtMulticastSparseDelegateProperty.cs
--- a/Engine/Source/Programs/Shared/EpicGames.UHT/Types/Properties/UhtMulticastSparseDelegateProperty.cs
+++ b/Engine/Source/Programs/Shared/EpicGames.UHT/Types/Properties/UhtMulticastSparseDelegateProperty.cs
@@ -50,7 +50,7 @@ namespace EpicGames.UHT.Types
 		public override StringBuilder AppendMemberDef(StringBuilder builder, IUhtPropertyMemberContext context, string name, string nameSuffix, string? offset, int tabs)
 		{
 			AppendMemberDefStart(builder, context, name, nameSuffix, offset, tabs, "FMulticastDelegatePropertyParams", "UECodeGen_Private::EPropertyGenFlags::SparseMulticastDelegate");
-			AppendMemberDefRef(builder, context, this.Function, true);
+			AppendMemberDefRef(builder, context, Function, true);
 			AppendMemberDefEnd(builder, context, name, nameSuffix);
 			return builder;
 		}
diff --git a/Engine/Source/Programs/Shared/EpicGames.UHT/Types/Properties/UhtNameProperty.cs b/Engine/Source/Programs/Shared/EpicGames.UHT/Types/Properties/UhtNameProperty.cs
--- a/Engine/Source/Programs/Shared/EpicGames.UHT/Types/Properties/UhtNameProperty.cs
+++ b/Engine/Source/Programs/Shared/EpicGames.UHT/Types/Properties/UhtNameProperty.cs
@@ -35,7 +35,7 @@ namespace EpicGames.UHT.Types
 		/// <param name="propertySettings">Property settings</param>
 		public UhtNameProperty(UhtPropertySettings propertySettings) : base(propertySettings)
 		{
-			this.PropertyCaps |= UhtPropertyCaps.CanExposeOnSpawn | UhtPropertyCaps.IsParameterSupportedByBlueprint |
+			PropertyCaps |= UhtPropertyCaps.CanExposeOnSpawn | UhtPropertyCaps.IsParameterSupportedByBlueprint |
 				UhtPropertyCaps.IsMemberSupportedByBlueprint | UhtPropertyCaps.SupportsRigVM;
 		}
 
diff --git a/Engine/Source/Programs/Shared/EpicGames.UHT/Types/Properties/UhtNumericProperty.cs b/Engine/Source/Programs/Shared/EpicGames.UHT/Types/Properties/UhtNumericProperty.cs
--- a/Engine/Source/Programs/Shared/EpicGames.UHT/Types/Properties/UhtNumericProperty.cs
+++ b/Engine/Source/Programs/Shared/EpicGames.UHT/Types/Properties/UhtNumericProperty.cs
@@ -28,8 +28,8 @@ namespace EpicGames.UHT.Types
 		/// <param name="intType">Type of integer</param>
 		protected UhtNumericProperty(UhtPropertySettings propertySettings, UhtPropertyIntType intType) : base(propertySettings)
 		{
-			this.IntType = intType;
-			this.PropertyCaps |= UhtPropertyCaps.RequiresNullConstructorArg | UhtPropertyCaps.SupportsRigVM;
+			IntType = intType;
+			PropertyCaps |= UhtPropertyCaps.RequiresNullConstructorArg | UhtPropertyCaps.SupportsRigVM;
 		}
 
 		/// <inheritdoc/>
diff --git a/Engine/Source/Programs/Shared/EpicGames.UHT/Types/Properties/UhtObjectProperty.cs b/Engine/Source/Programs/Shared/EpicGames.UHT/Types/Properties/UhtObjectProperty.cs
--- a/Engine/Source/Programs/Shared/EpicGames.UHT/Types/Properties/UhtObjectProperty.cs
+++ b/Engine/Source/Programs/Shared/EpicGames.UHT/Types/Properties/UhtObjectProperty.cs
@@ -29,8 +29,8 @@ namespace EpicGames.UHT.Types
 		public UhtObjectProperty(UhtPropertySettings propertySettings, UhtClass classObj, UhtClass? metaClass = null, EPropertyFlags extraFlags = EPropertyFlags.None)
 			: base(propertySettings, classObj, metaClass)
 		{
-			this.PropertyFlags |= extraFlags;
-			this.PropertyCaps |= UhtPropertyCaps.RequiresNullConstructorArg | UhtPropertyCaps.CanBeInstanced | UhtPropertyCaps.CanExposeOnSpawn |
+			PropertyFlags |= extraFlags;
+			PropertyCaps |= UhtPropertyCaps.RequiresNullConstructorArg | UhtPropertyCaps.CanBeInstanced | UhtPropertyCaps.CanExposeOnSpawn |
 				UhtPropertyCaps.IsParameterSupportedByBlueprint | UhtPropertyCaps.IsMemberSupportedByBlueprint;
 		}
 
@@ -41,10 +41,10 @@ namespace EpicGames.UHT.Types
 			switch (phase)
 			{
 				case UhtResolvePhase.Final:
-					if (this.Class.HierarchyHasAnyClassFlags(EClassFlags.DefaultToInstanced))
+					if (Class.HierarchyHasAnyClassFlags(EClassFlags.DefaultToInstanced))
 					{
-						this.PropertyFlags |= EPropertyFlags.InstancedReference;
-						this.MetaData.Add(UhtNames.EditInline, true);
+						PropertyFlags |= EPropertyFlags.InstancedReference;
+						MetaData.Add(UhtNames.EditInline, true);
 					}
 					break;
 			}
@@ -54,7 +54,7 @@ namespace EpicGames.UHT.Types
 		/// <inheritdoc/>
 		public override bool ScanForInstancedReferenced(bool deepScan)
 		{
-			return this.Class.HierarchyHasAnyClassFlags(EClassFlags.DefaultToInstanced);
+			return Class.HierarchyHasAnyClassFlags(EClassFlags.DefaultToInstanced);
 		}
 
 		/// <inheritdoc/>
@@ -63,19 +63,19 @@ namespace EpicGames.UHT.Types
 			switch (textType)
 			{
 				case UhtPropertyTextType.FunctionThunkRetVal:
-					if (this.PropertyFlags.HasAnyFlags(EPropertyFlags.ConstParm))
+					if (PropertyFlags.HasAnyFlags(EPropertyFlags.ConstParm))
 					{
 						builder.Append("const ");
 					}
-					builder.Append(this.Class.SourceName).Append('*');
+					builder.Append(Class.SourceName).Append('*');
 					break;
 
 				case UhtPropertyTextType.FunctionThunkParameterArgType:
-					builder.Append(this.Class.SourceName);
+					builder.Append(Class.SourceName);
 					break;
 
 				default:
-					builder.Append(this.Class.SourceName).Append('*');
+					builder.Append(Class.SourceName).Append('*');
 					break;
 			}
 			return builder;
@@ -91,7 +91,7 @@ namespace EpicGames.UHT.Types
 		public override StringBuilder AppendMemberDef(StringBuilder builder, IUhtPropertyMemberContext context, string name, string nameSuffix, string? offset, int tabs)
 		{
 			AppendMemberDefStart(builder, context, name, nameSuffix, offset, tabs, "FObjectPropertyParams", "UECodeGen_Private::EPropertyGenFlags::Object");
-			AppendMemberDefRef(builder, context, this.Class, false);
+			AppendMemberDefRef(builder, context, Class, false);
 			AppendMemberDefEnd(builder, context, name, nameSuffix);
 			return builder;
 		}
@@ -108,7 +108,7 @@ namespace EpicGames.UHT.Types
 		{
 			if (other is UhtObjectProperty otherObject)
 			{
-				return this.Class == otherObject.Class && this.MetaClass == otherObject.MetaClass;
+				return Class == otherObject.Class && MetaClass == otherObject.MetaClass;
 			}
 			return false;
 		}
diff --git a/Engine/Source/Programs/Shared/EpicGames.UHT/Types/Properties/UhtObjectPropertyBase.cs b/Engine/Source/Programs/Shared/EpicGames.UHT/Types/Properties/UhtObjectPropertyBase.cs
--- a/Engine/Source/Programs/Shared/EpicGames.UHT/Types/Properties/UhtObjectPropertyBase.cs
+++ b/Engine/Source/Programs/Shared/EpicGames.UHT/Types/Properties/UhtObjectPropertyBase.cs
@@ -47,23 +47,23 @@ namespace EpicGames.UHT.Types
 		/// <param name="metaClass">Referenced UCLASS used by class properties</param>
 		protected UhtObjectPropertyBase(UhtPropertySettings propertySettings, UhtClass classObj, UhtClass? metaClass) : base(propertySettings)
 		{
-			this.Class = classObj;
-			this.MetaClass = metaClass;
+			Class = classObj;
+			MetaClass = metaClass;
 
 			// This applies to EVERYTHING including raw pointer
 			// Imply const if it's a parameter that is a pointer to a const class
 			// NOTE: We shouldn't be automatically adding const param because in some cases with functions and blueprint native event, the 
 			// generated code won't match.  For now, just disabled the auto add in that case and check for the error in the validation code.
 			// Otherwise, the user is not warned and they will get compile errors.
-			if (propertySettings.PropertyCategory != UhtPropertyCategory.Member && this.Class.ClassFlags.HasAnyFlags(EClassFlags.Const) && !propertySettings.Options.HasAnyFlags(UhtPropertyOptions.NoAutoConst))
+			if (propertySettings.PropertyCategory != UhtPropertyCategory.Member && Class.ClassFlags.HasAnyFlags(EClassFlags.Const) && !propertySettings.Options.HasAnyFlags(UhtPropertyOptions.NoAutoConst))
 			{
-				this.PropertyFlags |= EPropertyFlags.ConstParm;
+				PropertyFlags |= EPropertyFlags.ConstParm;
 			}
 
-			this.PropertyCaps &= ~(UhtPropertyCaps.CanHaveConfig);
-			if (this.Session.Config!.AreRigVMUObjectProeprtiesEnabled)
+			PropertyCaps &= ~(UhtPropertyCaps.CanHaveConfig);
+			if (Session.Config!.AreRigVMUObjectProeprtiesEnabled)
 			{
-				this.PropertyCaps |= UhtPropertyCaps.SupportsRigVM;
+				PropertyCaps |= UhtPropertyCaps.SupportsRigVM;
 			}
 		}
 
@@ -74,9 +74,9 @@ namespace EpicGames.UHT.Types
 			switch (phase)
 			{
 				case UhtResolvePhase.Final:
-					if (this.Class.HierarchyHasAnyClassFlags(EClassFlags.DefaultToInstanced))
+					if (Class.HierarchyHasAnyClassFlags(EClassFlags.DefaultToInstanced))
 					{
-						this.PropertyFlags |= (EPropertyFlags.InstancedReference | EPropertyFlags.ExportObject) & ~this.DisallowPropertyFlags;
+						PropertyFlags |= (EPropertyFlags.InstancedReference | EPropertyFlags.ExportObject) & ~DisallowPropertyFlags;
 					}
 					break;
 			}
@@ -86,29 +86,29 @@ namespace EpicGames.UHT.Types
 		/// <inheritdoc/>
 		public override bool ScanForInstancedReferenced(bool deepScan)
 		{
-			return !this.DisallowPropertyFlags.HasAnyFlags(EPropertyFlags.InstancedReference) && this.Class.HierarchyHasAnyClassFlags(EClassFlags.DefaultToInstanced);
+			return !DisallowPropertyFlags.HasAnyFlags(EPropertyFlags.InstancedReference) && Class.HierarchyHasAnyClassFlags(EClassFlags.DefaultToInstanced);
 		}
 
 		/// <inheritdoc/>
 		public override void CollectReferencesInternal(IUhtReferenceCollector collector, bool templateProperty)
 		{
-			collector.AddCrossModuleReference(this.Class, false);
-			collector.AddCrossModuleReference(this.MetaClass, false);
+			collector.AddCrossModuleReference(Class, false);
+			collector.AddCrossModuleReference(MetaClass, false);
 		}
 
 		/// <inheritdoc/>
 		public override string? GetForwardDeclarations()
 		{
-			return $"class {this.Class.SourceName};";
+			return $"class {Class.SourceName};";
 		}
 
 		/// <inheritdoc/>
 		public override IEnumerable<UhtType> EnumerateReferencedTypes()
 		{
-			yield return this.Class;
-			if (this.MetaClass != null)
+			yield return Class;
+			if (MetaClass != null)
 			{
-				yield return this.MetaClass;
+				yield return MetaClass;
 			}
 		}
 
@@ -128,15 +128,15 @@ namespace EpicGames.UHT.Types
 		/// <inheritdoc/>
 		public override void ValidateDeprecated()
 		{
-			ValidateDeprecatedClass(this.Class);
-			ValidateDeprecatedClass(this.MetaClass);
+			ValidateDeprecatedClass(Class);
+			ValidateDeprecatedClass(MetaClass);
 		}
 
 		/// <inheritdoc/>
 		public override bool MustBeConstArgument([NotNullWhen(true)] out UhtType? errorType)
 		{
-			errorType = this.Class;
-			return this.Class.ClassFlags.HasAnyFlags(EClassFlags.Const);
+			errorType = Class;
+			return Class.ClassFlags.HasAnyFlags(EClassFlags.Const);
 		}
 
 		#region Parsing support methods
diff --git a/Engine/Source/Programs/Shared/EpicGames.UHT/Types/Properties/UhtObjectPtrProperty.cs b/Engine/Source/Programs/Shared/EpicGames.UHT/Types/Properties/UhtObjectPtrProperty.cs
--- a/Engine/Source/Programs/Shared/EpicGames.UHT/Types/Properties/UhtObjectPtrProperty.cs
+++ b/Engine/Source/Programs/Shared/EpicGames.UHT/Types/Properties/UhtObjectPtrProperty.cs
@@ -37,8 +37,8 @@ namespace EpicGames.UHT.Types
 		public UhtObjectPtrProperty(UhtPropertySettings propertySettings, UhtClass classObj, EPropertyFlags extraFlags = EPropertyFlags.None)
 			: base(propertySettings, classObj)
 		{
-			this.PropertyFlags |= extraFlags | EPropertyFlags.UObjectWrapper;
-			this.PropertyCaps |= UhtPropertyCaps.PassCppArgsByRef;
+			PropertyFlags |= extraFlags | EPropertyFlags.UObjectWrapper;
+			PropertyCaps |= UhtPropertyCaps.PassCppArgsByRef;
 		}
 
 		/// <inheritdoc/>
@@ -49,24 +49,24 @@ namespace EpicGames.UHT.Types
 				case UhtPropertyTextType.GetterSetterArg:
 					if (isTemplateArgument)
 					{
-						builder.Append("TObjectPtr<").Append(this.Class.SourceName).Append('>');
+						builder.Append("TObjectPtr<").Append(Class.SourceName).Append('>');
 					}
 					else
 					{
-						builder.Append(this.Class.SourceName).Append('*');
+						builder.Append(Class.SourceName).Append('*');
 					}
 					break;
 
 				case UhtPropertyTextType.FunctionThunkRetVal:
-					if (this.PropertyFlags.HasAnyFlags(EPropertyFlags.ConstParm))
+					if (PropertyFlags.HasAnyFlags(EPropertyFlags.ConstParm))
 					{
 						builder.Append("const ");
 					}
-					builder.Append("TObjectPtr<").Append(this.Class.SourceName).Append('>');
+					builder.Append("TObjectPtr<").Append(Class.SourceName).Append('>');
 					break;
 
 				default:
-					builder.Append("TObjectPtr<").Append(this.Class.SourceName).Append('>');
+					builder.Append("TObjectPtr<").Append(Class.SourceName).Append('>');
 					break;
 			}
 			return builder;
@@ -82,7 +82,7 @@ namespace EpicGames.UHT.Types
 		public override StringBuilder AppendMemberDef(StringBuilder builder, IUhtPropertyMemberContext context, string name, string nameSuffix, string? offset, int tabs)
 		{
 			AppendMemberDefStart(builder, context, name, nameSuffix, offset, tabs, "FObjectPtrPropertyParams", "UECodeGen_Private::EPropertyGenFlags::Object | UECodeGen_Private::EPropertyGenFlags::ObjectPtr");
-			AppendMemberDefRef(builder, context, this.Class, false);
+			AppendMemberDefRef(builder, context, Class, false);
 			AppendMemberDefEnd(builder, context, name, nameSuffix);
 			return builder;
 		}
@@ -94,7 +94,7 @@ namespace EpicGames.UHT.Types
 
 			// UFunctions with a smart pointer as input parameter wont compile anyway, because of missing P_GET_... macro.
 			// UFunctions with a smart pointer as return type will crash when called via blueprint, because they are not supported in VM.
-			if (this.PropertyCategory == UhtPropertyCategory.RegularParameter || this.PropertyCategory == UhtPropertyCategory.ReplicatedParameter)
+			if (PropertyCategory == UhtPropertyCategory.RegularParameter || PropertyCategory == UhtPropertyCategory.ReplicatedParameter)
 			{
 				outerStruct.LogError("UFunctions cannot take a TObjectPtr as a parameter.");
 			}
diff --git a/Engine/Source/Programs/Shared/EpicGames.UHT/Types/Properties/UhtSetProperty.cs b/Engine/Source/Programs/Shared/EpicGames.UHT/Types/Properties/UhtSetProperty.cs
--- a/Engine/Source/Programs/Shared/EpicGames.UHT/Types/Properties/UhtSetProperty.cs
+++ b/Engine/Source/Programs/Shared/EpicGames.UHT/Types/Properties/UhtSetProperty.cs
@@ -38,17 +38,17 @@ namespace EpicGames.UHT.Types
 		public UhtSetProperty(UhtPropertySettings propertySettings, UhtProperty value) : base(propertySettings, value)
 		{
 			// If the creation of the value property set more flags, then copy those flags to ourselves
-			this.PropertyFlags |= value.PropertyFlags & EPropertyFlags.UObjectWrapper;
+			PropertyFlags |= value.PropertyFlags & EPropertyFlags.UObjectWrapper;
 
-			this.PropertyCaps |= UhtPropertyCaps.PassCppArgsByRef;
-			this.PropertyCaps &= ~(UhtPropertyCaps.CanBeContainerValue | UhtPropertyCaps.CanBeContainerKey);
-			this.PropertyCaps = (this.PropertyCaps & ~(UhtPropertyCaps.IsParameterSupportedByBlueprint | UhtPropertyCaps.IsMemberSupportedByBlueprint | UhtPropertyCaps.CanExposeOnSpawn)) |
-				(this.ValueProperty.PropertyCaps & (UhtPropertyCaps.IsParameterSupportedByBlueprint | UhtPropertyCaps.IsMemberSupportedByBlueprint | UhtPropertyCaps.CanExposeOnSpawn));
-			this.ValueProperty.SourceName = this.SourceName;
-			this.ValueProperty.EngineName = this.EngineName;
-			this.ValueProperty.PropertyFlags = this.PropertyFlags & EPropertyFlags.PropagateToSetElement;
-			this.ValueProperty.Outer = this;
-			this.ValueProperty.MetaData.Clear();
+			PropertyCaps |= UhtPropertyCaps.PassCppArgsByRef;
+			PropertyCaps &= ~(UhtPropertyCaps.CanBeContainerValue | UhtPropertyCaps.CanBeContainerKey);
+			PropertyCaps = (PropertyCaps & ~(UhtPropertyCaps.IsParameterSupportedByBlueprint | UhtPropertyCaps.IsMemberSupportedByBlueprint | UhtPropertyCaps.CanExposeOnSpawn)) |
+				(ValueProperty.PropertyCaps & (UhtPropertyCaps.IsParameterSupportedByBlueprint | UhtPropertyCaps.IsMemberSupportedByBlueprint | UhtPropertyCaps.CanExposeOnSpawn));
+			ValueProperty.SourceName = SourceName;
+			ValueProperty.EngineName = EngineName;
+			ValueProperty.PropertyFlags = PropertyFlags & EPropertyFlags.PropagateToSetElement;
+			ValueProperty.Outer = this;
+			ValueProperty.MetaData.Clear();
 		}
 
 		/// <inheritdoc/>
@@ -58,11 +58,11 @@ namespace EpicGames.UHT.Types
 			switch (phase)
 			{
 				case UhtResolvePhase.Final:
-					this.PropertyFlags |= ResolveAndReturnNewFlags(this.ValueProperty, phase);
-					this.MetaData.Add(this.ValueProperty.MetaData);
-					this.ValueProperty.MetaData.Clear();
-					this.ValueProperty.PropertyFlags = this.PropertyFlags & EPropertyFlags.PropagateToSetElement;
-					PropagateFlagsFromInnerAndHandlePersistentInstanceMetadata(this, this.MetaData, this.ValueProperty);
+					PropertyFlags |= ResolveAndReturnNewFlags(ValueProperty, phase);
+					MetaData.Add(ValueProperty.MetaData);
+					ValueProperty.MetaData.Clear();
+					ValueProperty.PropertyFlags = PropertyFlags & EPropertyFlags.PropagateToSetElement;
+					PropagateFlagsFromInnerAndHandlePersistentInstanceMetadata(this, MetaData, ValueProperty);
 					break;
 			}
 			return results;
@@ -71,19 +71,19 @@ namespace EpicGames.UHT.Types
 		/// <inheritdoc/>
 		public override void CollectReferencesInternal(IUhtReferenceCollector collector, bool templateProperty)
 		{
-			this.ValueProperty.CollectReferencesInternal(collector, true);
+			ValueProperty.CollectReferencesInternal(collector, true);
 		}
 
 		/// <inheritdoc/>
 		public override string? GetForwardDeclarations()
 		{
-			return this.ValueProperty.GetForwardDeclarations();
+			return ValueProperty.GetForwardDeclarations();
 		}
 
 		/// <inheritdoc/>
 		public override IEnumerable<UhtType> EnumerateReferencedTypes()
 		{
-			foreach (UhtType type in this.ValueProperty.EnumerateReferencedTypes())
+			foreach (UhtType type in ValueProperty.EnumerateReferencedTypes())
 			{
 				yield return type;
 			}
@@ -99,11 +99,11 @@ namespace EpicGames.UHT.Types
 					break;
 
 				case UhtPropertyTextType.FunctionThunkParameterArgType:
-					builder.AppendFunctionThunkParameterArrayType(this.ValueProperty);
+					builder.AppendFunctionThunkParameterArrayType(ValueProperty);
 					break;
 
 				default:
-					builder.Append("TSet<").AppendPropertyText(this.ValueProperty, textType, true);
+					builder.Append("TSet<").AppendPropertyText(ValueProperty, textType, true);
 					if (builder[^1] == '>')
 					{
 						// if our internal property type is a template class, add a space between the closing brackets b/c VS.NET cannot parse this correctly
@@ -118,17 +118,17 @@ namespace EpicGames.UHT.Types
 		/// <inheritdoc/>
 		public override StringBuilder AppendMemberDecl(StringBuilder builder, IUhtPropertyMemberContext context, string name, string nameSuffix, int tabs)
 		{
-			builder.AppendMemberDecl(this.ValueProperty, context, name, GetNameSuffix(nameSuffix, "_ElementProp"), tabs);
+			builder.AppendMemberDecl(ValueProperty, context, name, GetNameSuffix(nameSuffix, "_ElementProp"), tabs);
 			return AppendMemberDecl(builder, context, name, nameSuffix, tabs, "FSetPropertyParams");
 		}
 
 		/// <inheritdoc/>
 		public override StringBuilder AppendMemberDef(StringBuilder builder, IUhtPropertyMemberContext context, string name, string nameSuffix, string? offset, int tabs)
 		{
-			builder.AppendMemberDef(this.ValueProperty, context, name, GetNameSuffix(nameSuffix, "_ElementProp"), "0", tabs);
+			builder.AppendMemberDef(ValueProperty, context, name, GetNameSuffix(nameSuffix, "_ElementProp"), "0", tabs);
 			builder.AppendMetaDataDef(this, context, name, nameSuffix, tabs);
 
-			if (this.ValueProperty is UhtStructProperty structProperty)
+			if (ValueProperty is UhtStructProperty structProperty)
 			{
 				builder
 					.AppendTabs(tabs)
@@ -147,7 +147,7 @@ namespace EpicGames.UHT.Types
 		/// <inheritdoc/>
 		public override StringBuilder AppendMemberPtr(StringBuilder builder, IUhtPropertyMemberContext context, string name, string nameSuffix, int tabs)
 		{
-			builder.AppendMemberPtr(this.ValueProperty, context, name, GetNameSuffix(nameSuffix, "_ElementProp"), tabs);
+			builder.AppendMemberPtr(ValueProperty, context, name, GetNameSuffix(nameSuffix, "_ElementProp"), tabs);
 			base.AppendMemberPtr(builder, context, name, nameSuffix, tabs);
 			return builder;
 		}
@@ -155,7 +155,7 @@ namespace EpicGames.UHT.Types
 		/// <inheritdoc/>
 		public override void AppendObjectHashes(StringBuilder builder, int startingLength, IUhtPropertyMemberContext context)
 		{
-			this.ValueProperty.AppendObjectHashes(builder, startingLength, context);
+			ValueProperty.AppendObjectHashes(builder, startingLength, context);
 		}
 
 		/// <inheritdoc/>
@@ -176,7 +176,7 @@ namespace EpicGames.UHT.Types
 		{
 			if (other is UhtSetProperty otherSet)
 			{
-				return this.ValueProperty.IsSameType(otherSet.ValueProperty);
+				return ValueProperty.IsSameType(otherSet.ValueProperty);
 			}
 			return false;
 		}
@@ -184,9 +184,9 @@ namespace EpicGames.UHT.Types
 		///<inheritdoc/>
 		public override bool ValidateStructPropertyOkForNet(UhtProperty referencingProperty)
 		{
-			if (!this.PropertyFlags.HasAnyFlags(EPropertyFlags.RepSkip))
+			if (!PropertyFlags.HasAnyFlags(EPropertyFlags.RepSkip))
 			{
-				referencingProperty.LogError($"Sets are not supported for Replication or RPCs.  Set '{this.SourceName}' in '{this.Outer?.SourceName}'.  Origin '{referencingProperty.SourceName}'");
+				referencingProperty.LogError($"Sets are not supported for Replication or RPCs.  Set '{SourceName}' in '{Outer?.SourceName}'.  Origin '{referencingProperty.SourceName}'");
 				return false;
 			}
 			return true;
diff --git a/Engine/Source/Programs/Shared/EpicGames.UHT/Types/Properties/UhtSoftClassProperty.cs b/Engine/Source/Programs/Shared/EpicGames.UHT/Types/Properties/UhtSoftClassProperty.cs
--- a/Engine/Source/Programs/Shared/EpicGames.UHT/Types/Properties/UhtSoftClassProperty.cs
+++ b/Engine/Source/Programs/Shared/EpicGames.UHT/Types/Properties/UhtSoftClassProperty.cs
@@ -42,13 +42,13 @@ namespace EpicGames.UHT.Types
 		/// <inheritdoc/>
 		public override string? GetForwardDeclarations()
 		{
-			return this.MetaClass != null ? $"class {this.MetaClass.SourceName};" : null;
+			return MetaClass != null ? $"class {MetaClass.SourceName};" : null;
 		}
 
 		/// <inheritdoc/>
 		public override void CollectReferencesInternal(IUhtReferenceCollector collector, bool templateProperty)
 		{
-			collector.AddCrossModuleReference(this.MetaClass, false);
+			collector.AddCrossModuleReference(MetaClass, false);
 		}
 
 		/// <inheritdoc/>
@@ -57,7 +57,7 @@ namespace EpicGames.UHT.Types
 			switch (textType)
 			{
 				default:
-					builder.Append("TSoftClassPtr<").Append(this.MetaClass?.SourceName).Append("> ");
+					builder.Append("TSoftClassPtr<").Append(MetaClass?.SourceName).Append("> ");
 					break;
 			}
 			return builder;
@@ -73,7 +73,7 @@ namespace EpicGames.UHT.Types
 		public override StringBuilder AppendMemberDef(StringBuilder builder, IUhtPropertyMemberContext context, string name, string nameSuffix, string? offset, int tabs)
 		{
 			AppendMemberDefStart(builder, context, name, nameSuffix, offset, tabs, "FSoftClassPropertyParams", "UECodeGen_Private::EPropertyGenFlags::SoftClass");
-			AppendMemberDefRef(builder, context, this.MetaClass, false);
+			AppendMemberDefRef(builder, context, MetaClass, false);
 			AppendMemberDefEnd(builder, context, name, nameSuffix);
 			return builder;
 		}
diff --git a/Engine/Source/Programs/Shared/EpicGames.UHT/Types/Properties/UhtSoftObjectProperty.cs b/Engine/Source/Programs/Shared/EpicGames.UHT/Types/Properties/UhtSoftObjectProperty.cs
--- a/Engine/Source/Programs/Shared/EpicGames.UHT/Types/Properties/UhtSoftObjectProperty.cs
+++ b/Engine/Source/Programs/Shared/EpicGames.UHT/Types/Properties/UhtSoftObjectProperty.cs
@@ -37,8 +37,8 @@ namespace EpicGames.UHT.Types
 		public UhtSoftObjectProperty(UhtPropertySettings propertySettings, UhtClass classObj, UhtClass? metaClass = null)
 			: base(propertySettings, classObj, metaClass)
 		{
-			this.PropertyFlags |= EPropertyFlags.UObjectWrapper;
-			this.PropertyCaps |= UhtPropertyCaps.PassCppArgsByRef | UhtPropertyCaps.RequiresNullConstructorArg | UhtPropertyCaps.CanExposeOnSpawn |
+			PropertyFlags |= EPropertyFlags.UObjectWrapper;
+			PropertyCaps |= UhtPropertyCaps.PassCppArgsByRef | UhtPropertyCaps.RequiresNullConstructorArg | UhtPropertyCaps.CanExposeOnSpawn |
 				UhtPropertyCaps.CanHaveConfig | UhtPropertyCaps.IsParameterSupportedByBlueprint | UhtPropertyCaps.IsMemberSupportedByBlueprint;
 		}
 
@@ -55,7 +55,7 @@ namespace EpicGames.UHT.Types
 			switch (textType)
 			{
 				default:
-					builder.Append("TSoftObjectPtr<").Append(this.Class.SourceName).Append('>');
+					builder.Append("TSoftObjectPtr<").Append(Class.SourceName).Append('>');
 					break;
 			}
 			return builder;
@@ -71,7 +71,7 @@ namespace EpicGames.UHT.Types
 		public override StringBuilder AppendMemberDef(StringBuilder builder, IUhtPropertyMemberContext context, string name, string nameSuffix, string? offset, int tabs)
 		{
 			AppendMemberDefStart(builder, context, name, nameSuffix, offset, tabs, "FSoftObjectPropertyParams", "UECodeGen_Private::EPropertyGenFlags::SoftObject");
-			AppendMemberDefRef(builder, context, this.Class, false);
+			AppendMemberDefRef(builder, context, Class, false);
 			AppendMemberDefEnd(builder, context, name, nameSuffix);
 			return builder;
 		}
@@ -81,7 +81,7 @@ namespace EpicGames.UHT.Types
 		{
 			if (other is UhtSoftObjectProperty otherObject)
 			{
-				return this.Class == otherObject.Class && this.MetaClass == otherObject.MetaClass;
+				return Class == otherObject.Class && MetaClass == otherObject.MetaClass;
 			}
 			return false;
 		}
diff --git a/Engine/Source/Programs/Shared/EpicGames.UHT/Types/Properties/UhtStrProperty.cs b/Engine/Source/Programs/Shared/EpicGames.UHT/Types/Properties/UhtStrProperty.cs
--- a/Engine/Source/Programs/Shared/EpicGames.UHT/Types/Properties/UhtStrProperty.cs
+++ b/Engine/Source/Programs/Shared/EpicGames.UHT/Types/Properties/UhtStrProperty.cs
@@ -35,7 +35,7 @@ namespace EpicGames.UHT.Types
 		/// <param name="propertySettings">Property settings</param>
 		public UhtStrProperty(UhtPropertySettings propertySettings) : base(propertySettings)
 		{
-			this.PropertyCaps |= UhtPropertyCaps.PassCppArgsByRef | UhtPropertyCaps.CanExposeOnSpawn | UhtPropertyCaps.IsParameterSupportedByBlueprint |
+			PropertyCaps |= UhtPropertyCaps.PassCppArgsByRef | UhtPropertyCaps.CanExposeOnSpawn | UhtPropertyCaps.IsParameterSupportedByBlueprint |
 				UhtPropertyCaps.IsMemberSupportedByBlueprint | UhtPropertyCaps.SupportsRigVM;
 		}
 
@@ -93,7 +93,7 @@ namespace EpicGames.UHT.Types
 			{
 				if (!function.FunctionFlags.HasAnyFlags(EFunctionFlags.NetRequest))
 				{
-					if (this.RefQualifier != UhtPropertyRefQualifier.ConstRef && !this.IsStaticArray)
+					if (RefQualifier != UhtPropertyRefQualifier.ConstRef && !IsStaticArray)
 					{
 						this.LogError("Replicated FString parameters must be passed by const reference");
 					}
diff --git a/Engine/Source/Programs/Shared/EpicGames.UHT/Types/Properties/UhtStructProperty.cs b/Engine/Source/Programs/Shared/EpicGames.UHT/Types/Properties/UhtStructProperty.cs
--- a/Engine/Source/Programs/Shared/EpicGames.UHT/Types/Properties/UhtStructProperty.cs
+++ b/Engine/Source/Programs/Shared/EpicGames.UHT/Types/Properties/UhtStructProperty.cs
@@ -44,20 +44,20 @@ namespace EpicGames.UHT.Types
 		/// <param name="scriptStruct">USTRUCT being referenced</param>
 		public UhtStructProperty(UhtPropertySettings propertySettings, UhtScriptStruct scriptStruct) : base(propertySettings)
 		{
-			this.ScriptStruct = scriptStruct;
-			this.HeaderFile.AddReferencedHeader(scriptStruct);
-			this.PropertyCaps |= UhtPropertyCaps.SupportsRigVM;
-			if (this.ScriptStruct.HasNoOpConstructor)
+			ScriptStruct = scriptStruct;
+			HeaderFile.AddReferencedHeader(scriptStruct);
+			PropertyCaps |= UhtPropertyCaps.SupportsRigVM;
+			if (ScriptStruct.HasNoOpConstructor)
 			{
-				this.PropertyCaps |= UhtPropertyCaps.RequiresNullConstructorArg;
+				PropertyCaps |= UhtPropertyCaps.RequiresNullConstructorArg;
 			}
-			if (this.ScriptStruct.MetaData.GetBoolean(UhtNames.BlueprintType))
+			if (ScriptStruct.MetaData.GetBoolean(UhtNames.BlueprintType))
 			{
-				this.PropertyCaps |= UhtPropertyCaps.CanExposeOnSpawn | UhtPropertyCaps.IsParameterSupportedByBlueprint | UhtPropertyCaps.IsMemberSupportedByBlueprint;
+				PropertyCaps |= UhtPropertyCaps.CanExposeOnSpawn | UhtPropertyCaps.IsParameterSupportedByBlueprint | UhtPropertyCaps.IsMemberSupportedByBlueprint;
 			}
-			else if (this.ScriptStruct.MetaData.GetBooleanHierarchical(UhtNames.BlueprintType))
+			else if (ScriptStruct.MetaData.GetBooleanHierarchical(UhtNames.BlueprintType))
 			{
-				this.PropertyCaps |= UhtPropertyCaps.IsParameterSupportedByBlueprint | UhtPropertyCaps.IsMemberSupportedByBlueprint;
+				PropertyCaps |= UhtPropertyCaps.IsParameterSupportedByBlueprint | UhtPropertyCaps.IsMemberSupportedByBlueprint;
 			}
 		}
 
@@ -70,7 +70,7 @@ namespace EpicGames.UHT.Types
 				case UhtResolvePhase.Final:
 					if (ScanForInstancedReferenced(true))
 					{
-						this.PropertyFlags |= EPropertyFlags.ContainsInstancedReference;
+						PropertyFlags |= EPropertyFlags.ContainsInstancedReference;
 					}
 					break;
 			}
@@ -80,25 +80,25 @@ namespace EpicGames.UHT.Types
 		/// <inheritdoc/>
 		public override bool ScanForInstancedReferenced(bool deepScan)
 		{
-			return this.ScriptStruct.ScanForInstancedReferenced(deepScan);
+			return ScriptStruct.ScanForInstancedReferenced(deepScan);
 		}
 
 		/// <inheritdoc/>
 		public override void CollectReferencesInternal(IUhtReferenceCollector collector, bool templateProperty)
 		{
-			collector.AddCrossModuleReference(this.ScriptStruct, true);
+			collector.AddCrossModuleReference(ScriptStruct, true);
 		}
 
 		/// <inheritdoc/>
 		public override string? GetForwardDeclarations()
 		{
-			return !this.ScriptStruct.IsCoreType ? $"struct {this.ScriptStruct.SourceName};" : null;
+			return !ScriptStruct.IsCoreType ? $"struct {ScriptStruct.SourceName};" : null;
 		}
 
 		/// <inheritdoc/>
 		public override IEnumerable<UhtType> EnumerateReferencedTypes()
 		{
-			yield return this.ScriptStruct;
+			yield return ScriptStruct;
 		}
 
 		/// <inheritdoc/>
@@ -107,7 +107,7 @@ namespace EpicGames.UHT.Types
 			switch (textType)
 			{
 				default:
-					builder.Append(this.ScriptStruct.SourceName);
+					builder.Append(ScriptStruct.SourceName);
 					break;
 			}
 			return builder;
@@ -123,7 +123,7 @@ namespace EpicGames.UHT.Types
 		public override StringBuilder AppendMemberDef(StringBuilder builder, IUhtPropertyMemberContext context, string name, string nameSuffix, string? offset, int tabs)
 		{
 			AppendMemberDefStart(builder, context, name, nameSuffix, offset, tabs, "FStructPropertyParams", "UECodeGen_Private::EPropertyGenFlags::Struct");
-			AppendMemberDefRef(builder, context, this.ScriptStruct, true);
+			AppendMemberDefRef(builder, context, ScriptStruct, true);
 			AppendMemberDefEnd(builder, context, name, nameSuffix);
 			return builder;
 		}
@@ -131,13 +131,13 @@ namespace EpicGames.UHT.Types
 		/// <inheritdoc/>
 		public override void AppendObjectHashes(StringBuilder builder, int startingLength, IUhtPropertyMemberContext context)
 		{
-			builder.AppendObjectHash(startingLength, this, context, this.ScriptStruct);
+			builder.AppendObjectHash(startingLength, this, context, ScriptStruct);
 		}
 
 		/// <inheritdoc/>
 		public override StringBuilder AppendNullConstructorArg(StringBuilder builder, bool isInitializer)
 		{
-			bool hasNoOpConstructor = this.ScriptStruct.HasNoOpConstructor;
+			bool hasNoOpConstructor = ScriptStruct.HasNoOpConstructor;
 			if (isInitializer && hasNoOpConstructor)
 			{
 				builder.Append("ForceInit");
@@ -160,9 +160,9 @@ namespace EpicGames.UHT.Types
 		/// <inheritdoc/>
 		public override bool SanitizeDefaultValue(IUhtTokenReader defaultValueReader, StringBuilder innerDefaultValue)
 		{
-			if (!this.Session.TryGetStructDefaultValue(this.ScriptStruct.SourceName, out UhtStructDefaultValue structDefaultValue))
+			if (!Session.TryGetStructDefaultValue(ScriptStruct.SourceName, out UhtStructDefaultValue structDefaultValue))
 			{
-				structDefaultValue = this.Session.DefaultStructDefaultValue;
+				structDefaultValue = Session.DefaultStructDefaultValue;
 			}
 			return structDefaultValue.Delegate(this, defaultValueReader, innerDefaultValue);
 		}
@@ -172,7 +172,7 @@ namespace EpicGames.UHT.Types
 		{
 			if (other is UhtStructProperty otherObject)
 			{
-				return this.ScriptStruct == otherObject.ScriptStruct;
+				return ScriptStruct == otherObject.ScriptStruct;
 			}
 			return false;
 		}
@@ -186,7 +186,7 @@ namespace EpicGames.UHT.Types
 			{
 				if (!function.FunctionFlags.HasAnyFlags(EFunctionFlags.NetRequest | EFunctionFlags.NetResponse))
 				{
-					this.Session.ValidateScriptStructOkForNet(this, this.ScriptStruct);
+					Session.ValidateScriptStructOkForNet(this, ScriptStruct);
 				}
 			}
 		}
@@ -194,22 +194,22 @@ namespace EpicGames.UHT.Types
 		/// <inheritdoc/>
 		protected override void ValidateMember(UhtStruct structObj, UhtValidationOptions options)
 		{
-			if (this.PropertyFlags.HasAnyFlags(EPropertyFlags.Net))
+			if (PropertyFlags.HasAnyFlags(EPropertyFlags.Net))
 			{
-				this.Session.ValidateScriptStructOkForNet(this, this.ScriptStruct);
+				Session.ValidateScriptStructOkForNet(this, ScriptStruct);
 			}
 		}
 
 		///<inheritdoc/>
 		public override bool ValidateStructPropertyOkForNet(UhtProperty referencingProperty)
 		{
-			return referencingProperty.Session.ValidateScriptStructOkForNet(referencingProperty, this.ScriptStruct);
+			return referencingProperty.Session.ValidateScriptStructOkForNet(referencingProperty, ScriptStruct);
 		}
 
 		///<inheritdoc/>
 		public override bool ContainsEditorOnlyProperties()
 		{
-			foreach (UhtType Child in this.ScriptStruct.Children)
+			foreach (UhtType Child in ScriptStruct.Children)
 			{
 				if (Child is UhtProperty Property)
 				{
diff --git a/Engine/Source/Programs/Shared/EpicGames.UHT/Types/Properties/UhtTextProperty.cs b/Engine/Source/Programs/Shared/EpicGames.UHT/Types/Properties/UhtTextProperty.cs
--- a/Engine/Source/Programs/Shared/EpicGames.UHT/Types/Properties/UhtTextProperty.cs
+++ b/Engine/Source/Programs/Shared/EpicGames.UHT/Types/Properties/UhtTextProperty.cs
@@ -36,9 +36,9 @@ namespace EpicGames.UHT.Types
 		/// <param name="propertySettings">Property settings</param>
 		public UhtTextProperty(UhtPropertySettings propertySettings) : base(propertySettings)
 		{
-			this.PropertyCaps |= UhtPropertyCaps.PassCppArgsByRef | UhtPropertyCaps.CanExposeOnSpawn | UhtPropertyCaps.IsParameterSupportedByBlueprint |
+			PropertyCaps |= UhtPropertyCaps.PassCppArgsByRef | UhtPropertyCaps.CanExposeOnSpawn | UhtPropertyCaps.IsParameterSupportedByBlueprint |
 				UhtPropertyCaps.IsMemberSupportedByBlueprint | UhtPropertyCaps.SupportsRigVM;
-			this.PropertyCaps &= ~(UhtPropertyCaps.CanBeContainerKey);
+			PropertyCaps &= ~(UhtPropertyCaps.CanBeContainerKey);
 		}
 
 		/// <inheritdoc/>
@@ -101,7 +101,7 @@ namespace EpicGames.UHT.Types
 			if (token.IsIdentifier("LOCTEXT"))
 			{
 				//ETSTODO - Add a test case for this
-				this.LogError($"LOCTEXT default parameter values are not supported; use NSLOCTEXT instead: {this.SourceName}");
+				this.LogError($"LOCTEXT default parameter values are not supported; use NSLOCTEXT instead: {SourceName}");
 				return false;
 			}
 			else
diff --git a/Engine/Source/Programs/Shared/EpicGames.UHT/Types/Properties/UhtUInt32Property.cs b/Engine/Source/Programs/Shared/EpicGames.UHT/Types/Properties/UhtUInt32Property.cs
--- a/Engine/Source/Programs/Shared/EpicGames.UHT/Types/Properties/UhtUInt32Property.cs
+++ b/Engine/Source/Programs/Shared/EpicGames.UHT/Types/Properties/UhtUInt32Property.cs
@@ -35,14 +35,14 @@ namespace EpicGames.UHT.Types
 		public override StringBuilder AppendMemberDecl(StringBuilder builder, IUhtPropertyMemberContext context, string name, string nameSuffix, int tabs)
 		{
 			return AppendMemberDecl(builder, context, name, nameSuffix, tabs,
-				this.IntType == UhtPropertyIntType.Unsized ? "FUnsizedFIntPropertyParams" : "FUInt32PropertyParams");
+				IntType == UhtPropertyIntType.Unsized ? "FUnsizedFIntPropertyParams" : "FUInt32PropertyParams");
 		}
 
 		/// <inheritdoc/>
 		public override StringBuilder AppendMemberDef(StringBuilder builder, IUhtPropertyMemberContext context, string name, string nameSuffix, string? offset, int tabs)
 		{
 			AppendMemberDefStart(builder, context, name, nameSuffix, offset, tabs,
-				this.IntType == UhtPropertyIntType.Unsized ? "FUnsizedFIntPropertyParams" : "FUInt32PropertyParams",
+				IntType == UhtPropertyIntType.Unsized ? "FUnsizedFIntPropertyParams" : "FUInt32PropertyParams",
 				"UECodeGen_Private::EPropertyGenFlags::UInt32");
 			AppendMemberDefEnd(builder, context, name, nameSuffix);
 			return builder;
diff --git a/Engine/Source/Programs/Shared/EpicGames.UHT/Types/Properties/UhtWeakObjectPtrProperty.cs b/Engine/Source/Programs/Shared/EpicGames.UHT/Types/Properties/UhtWeakObjectPtrProperty.cs
--- a/Engine/Source/Programs/Shared/EpicGames.UHT/Types/Properties/UhtWeakObjectPtrProperty.cs
+++ b/Engine/Source/Programs/Shared/EpicGames.UHT/Types/Properties/UhtWeakObjectPtrProperty.cs
@@ -21,7 +21,7 @@ namespace EpicGames.UHT.Types
 		public override string EngineClassName => "WeakObjectProperty";
 
 		/// <inheritdoc/>
-		protected override string PGetMacroText => this.PropertyFlags.HasAnyFlags(EPropertyFlags.AutoWeak) ? "AUTOWEAKOBJECT" : "WEAKOBJECT";
+		protected override string PGetMacroText => PropertyFlags.HasAnyFlags(EPropertyFlags.AutoWeak) ? "AUTOWEAKOBJECT" : "WEAKOBJECT";
 
 		/// <inheritdoc/>
 		protected override bool PGetPassAsNoPtr => true;
@@ -38,9 +38,9 @@ namespace EpicGames.UHT.Types
 		public UhtWeakObjectPtrProperty(UhtPropertySettings propertySettings, UhtClass propertyClass, EPropertyFlags extraFlags = EPropertyFlags.None)
 			: base(propertySettings, propertyClass, null)
 		{
-			this.PropertyFlags |= EPropertyFlags.UObjectWrapper | extraFlags;
-			this.PropertyCaps |= UhtPropertyCaps.PassCppArgsByRef | UhtPropertyCaps.RequiresNullConstructorArg | UhtPropertyCaps.IsMemberSupportedByBlueprint;
-			this.PropertyCaps &= ~(UhtPropertyCaps.IsParameterSupportedByBlueprint);
+			PropertyFlags |= EPropertyFlags.UObjectWrapper | extraFlags;
+			PropertyCaps |= UhtPropertyCaps.PassCppArgsByRef | UhtPropertyCaps.RequiresNullConstructorArg | UhtPropertyCaps.IsMemberSupportedByBlueprint;
+			PropertyCaps &= ~(UhtPropertyCaps.IsParameterSupportedByBlueprint);
 		}
 
 		/// <inheritdoc/>
@@ -49,13 +49,13 @@ namespace EpicGames.UHT.Types
 			switch (textType)
 			{
 				default:
-					if (this.PropertyFlags.HasAnyFlags(EPropertyFlags.AutoWeak))
+					if (PropertyFlags.HasAnyFlags(EPropertyFlags.AutoWeak))
 					{
-						builder.Append("TAutoWeakObjectPtr<").Append(this.Class.SourceName).Append('>');
+						builder.Append("TAutoWeakObjectPtr<").Append(Class.SourceName).Append('>');
 					}
 					else
 					{
-						builder.Append("TWeakObjectPtr<").Append(this.Class.SourceName).Append('>');
+						builder.Append("TWeakObjectPtr<").Append(Class.SourceName).Append('>');
 					}
 					break;
 			}
@@ -72,7 +72,7 @@ namespace EpicGames.UHT.Types
 		public override StringBuilder AppendMemberDef(StringBuilder builder, IUhtPropertyMemberContext context, string name, string nameSuffix, string? offset, int tabs)
 		{
 			AppendMemberDefStart(builder, context, name, nameSuffix, offset, tabs, "FWeakObjectPropertyParams", "UECodeGen_Private::EPropertyGenFlags::WeakObject");
-			AppendMemberDefRef(builder, context, this.Class, false);
+			AppendMemberDefRef(builder, context, Class, false);
 			AppendMemberDefEnd(builder, context, name, nameSuffix);
 			return builder;
 		}
@@ -89,7 +89,7 @@ namespace EpicGames.UHT.Types
 		{
 			if (other is UhtWeakObjectPtrProperty otherObject)
 			{
-				return this.Class == otherObject.Class && this.MetaClass == otherObject.MetaClass;
+				return Class == otherObject.Class && MetaClass == otherObject.MetaClass;
 			}
 			return false;
 		}
