Removed the use of "this." where possible in EpicGames.UHT.

diff --git a/Engine/Source/Programs/Shared/EpicGames.UHT/Types/UhtClass.cs b/Engine/Source/Programs/Shared/EpicGames.UHT/Types/UhtClass.cs
--- a/Engine/Source/Programs/Shared/EpicGames.UHT/Types/UhtClass.cs
+++ b/Engine/Source/Programs/Shared/EpicGames.UHT/Types/UhtClass.cs
@@ -333,7 +333,7 @@ namespace EpicGames.UHT.Types
 		/// Collection of functions and other declarations found in the class
 		/// </summary>
 		[JsonIgnore]
-		public IList<UhtDeclaration>? Declarations => this._declarations;
+		public IList<UhtDeclaration>? Declarations => _declarations;
 
 		/// <summary>
 		/// If this, this class is a UINTERFACE and NativeInterface is the associated native interface
@@ -428,8 +428,8 @@ namespace EpicGames.UHT.Types
 		/// <param name="flags">Flags to add</param>
 		public void AddClassFlags(EClassFlags flags)
 		{
-			this.ClassFlags |= flags;
-			this.RemovedClassFlags &= ~flags;
+			ClassFlags |= flags;
+			RemovedClassFlags &= ~flags;
 		}
 
 		/// <summary>
@@ -438,8 +438,8 @@ namespace EpicGames.UHT.Types
 		/// <param name="flags">Flags to remove</param>
 		public void RemoveClassFlags(EClassFlags flags)
 		{
-			this.RemovedClassFlags |= flags;
-			this.ClassFlags &= ~flags;
+			RemovedClassFlags |= flags;
+			ClassFlags &= ~flags;
 		}
 		#endregion
 
@@ -449,7 +449,7 @@ namespace EpicGames.UHT.Types
 		{
 			get
 			{
-				switch (this.ClassType)
+				switch (ClassType)
 				{
 					case UhtClassType.Class:
 						return UhtEngineType.Class;
@@ -468,7 +468,7 @@ namespace EpicGames.UHT.Types
 
 		///<inheritdoc/>
 		[JsonIgnore]
-		public override bool Deprecated => this.ClassFlags.HasAnyFlags(EClassFlags.Deprecated);
+		public override bool Deprecated => ClassFlags.HasAnyFlags(EClassFlags.Deprecated);
 
 		///<inheritdoc/>
 		[JsonIgnore]
@@ -476,14 +476,14 @@ namespace EpicGames.UHT.Types
 		{
 			get
 			{
-				switch (this.ClassType)
+				switch (ClassType)
 				{
 					case UhtClassType.Class:
-						return this.Session.GetSpecifierValidatorTable(UhtTableNames.Class);
+						return Session.GetSpecifierValidatorTable(UhtTableNames.Class);
 					case UhtClassType.Interface:
-						return this.Session.GetSpecifierValidatorTable(UhtTableNames.Interface);
+						return Session.GetSpecifierValidatorTable(UhtTableNames.Interface);
 					case UhtClassType.NativeInterface:
-						return this.Session.GetSpecifierValidatorTable(UhtTableNames.NativeInterface);
+						return Session.GetSpecifierValidatorTable(UhtTableNames.NativeInterface);
 					default:
 						throw new UhtIceException("Unknown class type");
 				}
@@ -494,7 +494,7 @@ namespace EpicGames.UHT.Types
 		/// The super class
 		/// </summary>
 		[JsonConverter(typeof(UhtNullableTypeSourceNameJsonConverter<UhtClass>))]
-		public UhtClass? SuperClass => (UhtClass?)this.Super;
+		public UhtClass? SuperClass => (UhtClass?)Super;
 
 		/// <summary>
 		/// Construct a new instance of the class
@@ -503,14 +503,14 @@ namespace EpicGames.UHT.Types
 		/// <param name="lineNumber">Line number where class begins</param>
 		public UhtClass(UhtType outer, int lineNumber) : base(outer, lineNumber)
 		{
-			this.ClassWithin = this;
+			ClassWithin = this;
 		}
 
 		/// <summary>
 		/// True if this class inherits from AActor
 		/// </summary>
 		[JsonIgnore]
-		public bool IsActorClass => IsChildOf(this.Session.AActor);
+		public bool IsActorClass => IsChildOf(Session.AActor);
 
 		///<inheritdoc/>
 		[JsonIgnore]
@@ -518,16 +518,16 @@ namespace EpicGames.UHT.Types
 		{
 			get
 			{
-				switch (this.ClassType)
+				switch (ClassType)
 				{
 					case UhtClassType.Class:
 						if (IsActorClass)
 						{
-							return this.ClassFlags.HasAnyFlags(EClassFlags.Deprecated) ? "ADEPRECATED_" : "A";
+							return ClassFlags.HasAnyFlags(EClassFlags.Deprecated) ? "ADEPRECATED_" : "A";
 						}
 						else
 						{
-							return this.ClassFlags.HasAnyFlags(EClassFlags.Deprecated) ? "UDEPRECATED_" : "U";
+							return ClassFlags.HasAnyFlags(EClassFlags.Deprecated) ? "UDEPRECATED_" : "U";
 						}
 
 					case UhtClassType.Interface:
@@ -550,11 +550,11 @@ namespace EpicGames.UHT.Types
 		/// <param name="function">If parsed as part of a UFUNCTION, this will reference it</param>
 		public void AddDeclaration(UhtCompilerDirective compilerDirectives, List<UhtToken> tokens, UhtFunction? function)
 		{
-			if (this._declarations == null)
+			if (_declarations == null)
 			{
-				this._declarations = new List<UhtDeclaration>();
+				_declarations = new List<UhtDeclaration>();
 			}
-			this._declarations.Add(new UhtDeclaration { CompilerDirectives = compilerDirectives, Tokens = tokens.ToArray(), Function = function });
+			_declarations.Add(new UhtDeclaration { CompilerDirectives = compilerDirectives, Tokens = tokens.ToArray(), Function = function });
 		}
 
 		/// <summary>
@@ -567,9 +567,9 @@ namespace EpicGames.UHT.Types
 		{
 			if (name.Length > 0)
 			{
-				if (this.Declarations != null)
+				if (Declarations != null)
 				{
-					foreach (UhtDeclaration declaration in this.Declarations)
+					foreach (UhtDeclaration declaration in Declarations)
 					{
 						bool isVirtual = false;
 						for (int index = 0; index < declaration.Tokens.Length; ++index)
@@ -593,9 +593,9 @@ namespace EpicGames.UHT.Types
 					}
 				}
 			}
-			if (this.NativeInterface != null)
+			if (NativeInterface != null)
 			{
-				return this.NativeInterface.TryGetDeclaration(name, out foundDeclaration);
+				return NativeInterface.TryGetDeclaration(name, out foundDeclaration);
 			}
 			foundDeclaration = new UhtFoundDeclaration();
 			return false;
@@ -632,36 +632,36 @@ namespace EpicGames.UHT.Types
 			base.ResolveSuper(resolvePhase);
 
 			// Make sure the class within is resolved.  We have to make sure we don't try to resolve ourselves.
-			if (this.ClassWithin != null && this.ClassWithin != this)
+			if (ClassWithin != null && ClassWithin != this)
 			{
-				this.ClassWithin.Resolve(resolvePhase);
+				ClassWithin.Resolve(resolvePhase);
 			}
 
 			switch (resolvePhase)
 			{
 				case UhtResolvePhase.Bases:
-					BindAndResolveSuper(this.SuperIdentifier, UhtFindOptions.Class);
-					BindAndResolveBases(this.BaseIdentifiers, UhtFindOptions.Class);
+					BindAndResolveSuper(SuperIdentifier, UhtFindOptions.Class);
+					BindAndResolveBases(BaseIdentifiers, UhtFindOptions.Class);
 
 					// Force the MatchedSerializers on for anything being exported
-					if (!this.ClassExportFlags.HasAnyFlags(UhtClassExportFlags.NoExport))
+					if (!ClassExportFlags.HasAnyFlags(UhtClassExportFlags.NoExport))
 					{
-						this.ClassFlags |= EClassFlags.MatchedSerializers;
+						ClassFlags |= EClassFlags.MatchedSerializers;
 					}
 
-					switch (this.ClassType)
+					switch (ClassType)
 					{
 						case UhtClassType.Class:
 							{
-								UhtClass? superClass = this.SuperClass;
+								UhtClass? superClass = SuperClass;
 
 								// Merge the super class flags
 								if (superClass != null)
 								{
-									this.ClassFlags |= superClass.ClassFlags & EClassFlags.ScriptInherit & ~this.RemovedClassFlags;
+									ClassFlags |= superClass.ClassFlags & EClassFlags.ScriptInherit & ~RemovedClassFlags;
 								}
 
-								foreach (UhtStruct baseStruct in this.Bases)
+								foreach (UhtStruct baseStruct in Bases)
 								{
 									if (baseStruct is UhtClass baseClass)
 									{
@@ -669,18 +669,18 @@ namespace EpicGames.UHT.Types
 										{
 											this.LogError($"Class '{baseClass.SourceName}' is not an interface; Can only inherit from non-UObjects or UInterface derived interfaces");
 										}
-										this.ClassFlags |= baseClass.ClassFlags & EClassFlags.ScriptInherit & ~this.RemovedClassFlags;
+										ClassFlags |= baseClass.ClassFlags & EClassFlags.ScriptInherit & ~RemovedClassFlags;
 									}
 								}
 
 								// These following collections only inherit from he parent if they are empty in this class
-								if (this.SparseClassDataTypes.Count == 0)
+								if (SparseClassDataTypes.Count == 0)
 								{
-									AppendStringListMetaData(this.SuperClass, UhtNames.SparseClassDataTypes, this.SparseClassDataTypes);
+									AppendStringListMetaData(SuperClass, UhtNames.SparseClassDataTypes, SparseClassDataTypes);
 								}
-								if (this.PrioritizeCategories.Count == 0)
+								if (PrioritizeCategories.Count == 0)
 								{
-									AppendStringListMetaData(this.SuperClass, UhtNames.PrioritizeCategories, this.PrioritizeCategories);
+									AppendStringListMetaData(SuperClass, UhtNames.PrioritizeCategories, PrioritizeCategories);
 								}
 
 								// Merge with categories inherited from the parent.
@@ -690,34 +690,34 @@ namespace EpicGames.UHT.Types
 								SetAndValidateConfigName();
 
 								// Copy all of the lists back to the meta data
-								this.MetaData.AddIfNotEmpty(UhtNames.ClassGroupNames, this.ClassGroupNames);
-								this.MetaData.AddIfNotEmpty(UhtNames.AutoCollapseCategories, this.AutoCollapseCategories);
-								this.MetaData.AddIfNotEmpty(UhtNames.AutoExpandCategories, this.AutoExpandCategories);
-								this.MetaData.AddIfNotEmpty(UhtNames.PrioritizeCategories, this.PrioritizeCategories);
-								this.MetaData.AddIfNotEmpty(UhtNames.HideCategories, this.HideCategories);
-								this.MetaData.AddIfNotEmpty(UhtNames.ShowCategories, this.ShowCategories);
-								this.MetaData.AddIfNotEmpty(UhtNames.SparseClassDataTypes, this.SparseClassDataTypes);
-								this.MetaData.AddIfNotEmpty(UhtNames.HideFunctions, this.HideFunctions);
-
-								this.MetaData.Add(UhtNames.IncludePath, this.HeaderFile.IncludeFilePath);
+								MetaData.AddIfNotEmpty(UhtNames.ClassGroupNames, ClassGroupNames);
+								MetaData.AddIfNotEmpty(UhtNames.AutoCollapseCategories, AutoCollapseCategories);
+								MetaData.AddIfNotEmpty(UhtNames.AutoExpandCategories, AutoExpandCategories);
+								MetaData.AddIfNotEmpty(UhtNames.PrioritizeCategories, PrioritizeCategories);
+								MetaData.AddIfNotEmpty(UhtNames.HideCategories, HideCategories);
+								MetaData.AddIfNotEmpty(UhtNames.ShowCategories, ShowCategories);
+								MetaData.AddIfNotEmpty(UhtNames.SparseClassDataTypes, SparseClassDataTypes);
+								MetaData.AddIfNotEmpty(UhtNames.HideFunctions, HideFunctions);
+
+								MetaData.Add(UhtNames.IncludePath, HeaderFile.IncludeFilePath);
 							}
 							break;
 
 						case UhtClassType.Interface:
 							{
-								UhtClass? superClass = this.SuperClass;
+								UhtClass? superClass = SuperClass;
 								if (superClass != null)
 								{
-									this.ClassFlags |= superClass.ClassFlags & EClassFlags.ScriptInherit;
+									ClassFlags |= superClass.ClassFlags & EClassFlags.ScriptInherit;
 									if (!superClass.ClassFlags.HasAnyFlags(EClassFlags.Native))
 									{
 										throw new UhtException(this, $"Native classes cannot extend non-native classes");
 									}
-									this.ClassWithin = superClass.ClassWithin;
+									ClassWithin = superClass.ClassWithin;
 								}
 								else
 								{
-									this.ClassWithin = this.Session.UObject;
+									ClassWithin = Session.UObject;
 								}
 							}
 							break;
@@ -740,51 +740,51 @@ namespace EpicGames.UHT.Types
 			switch (phase)
 			{
 				case UhtResolvePhase.InvalidCheck:
-					switch (this.ClassType)
+					switch (ClassType)
 					{
 						case UhtClassType.Class:
 							break;
 
 						case UhtClassType.Interface:
 							{
-								string nativeInterfaceName = "I" + this.EngineName;
-								UhtType? nativeInterface = this.Session.FindType(null, UhtFindOptions.SourceName | UhtFindOptions.Class, nativeInterfaceName);
+								string nativeInterfaceName = "I" + EngineName;
+								UhtType? nativeInterface = Session.FindType(null, UhtFindOptions.SourceName | UhtFindOptions.Class, nativeInterfaceName);
 								if (nativeInterface == null)
 								{
-									this.LogError($"UInterface '{this.SourceName}' parsed without a corresponding '{nativeInterfaceName}'");
+									this.LogError($"UInterface '{SourceName}' parsed without a corresponding '{nativeInterfaceName}'");
 								}
 								else
 								{
 									// Copy the children
-									this.AddChildren(nativeInterface.DetachChildren());
+									AddChildren(nativeInterface.DetachChildren());
 								}
 							}
 							break;
 
 						case UhtClassType.NativeInterface:
 							{
-								string interfaceName = "U" + this.EngineName;
-								UhtClass? interfaceObj = (UhtClass?)this.Session.FindType(null, UhtFindOptions.SourceName | UhtFindOptions.Class, interfaceName);
+								string interfaceName = "U" + EngineName;
+								UhtClass? interfaceObj = (UhtClass?)Session.FindType(null, UhtFindOptions.SourceName | UhtFindOptions.Class, interfaceName);
 								if (interfaceObj == null)
 								{
-									if (this.HasGeneratedBody || this.Children.Count != 0)
+									if (HasGeneratedBody || Children.Count != 0)
 									{
-										this.LogError($"Native interface '{this.SourceName}' parsed without a corresponding '{interfaceName}'");
+										this.LogError($"Native interface '{SourceName}' parsed without a corresponding '{interfaceName}'");
 									}
 									else
 									{
-										this.VisibleType = false;
+										VisibleType = false;
 										result = false;
 									}
 								}
 								else
 								{
-									this.AlternateObject = interfaceObj;
+									AlternateObject = interfaceObj;
 									interfaceObj.NativeInterface = this;
 									//COMPATIBILITY-TODO - Use the native interface access specifier
-									interfaceObj.GeneratedBodyAccessSpecifier = this.GeneratedBodyAccessSpecifier;
+									interfaceObj.GeneratedBodyAccessSpecifier = GeneratedBodyAccessSpecifier;
 
-									if (this.GeneratedBodyLineNumber == -1)
+									if (GeneratedBodyLineNumber == -1)
 									{
 										this.LogError("Expected a GENERATED_BODY() at the start of the native interface");
 									}
@@ -801,24 +801,24 @@ namespace EpicGames.UHT.Types
 
 					// Check to see if this class matches an entry in the ClassCastFlags.  If it does, that
 					// becomes our ClassCastFlags.  If not, then we inherit the flags from the super.
-					if (Enum.TryParse<EClassCastFlags>(this.SourceName, false, out EClassCastFlags found))
+					if (Enum.TryParse<EClassCastFlags>(SourceName, false, out EClassCastFlags found))
 					{
 						if (found != EClassCastFlags.None)
 						{
-							this.ClassCastFlags = found;
+							ClassCastFlags = found;
 						}
 					}
-					if (this.ClassCastFlags == EClassCastFlags.None && this.SuperClass != null)
+					if (ClassCastFlags == EClassCastFlags.None && SuperClass != null)
 					{
-						this.ClassCastFlags = this.SuperClass.ClassCastFlags;
+						ClassCastFlags = SuperClass.ClassCastFlags;
 					}
 
 					// force members to be 'blueprint read only' if in a const class
 					// This must be done after the Const flag has been propagated from the parent
 					// and before properties have been finalized.
-					if (this.ClassType == UhtClassType.Class && this.ClassFlags.HasAnyFlags(EClassFlags.Const))
+					if (ClassType == UhtClassType.Class && ClassFlags.HasAnyFlags(EClassFlags.Const))
 					{
-						foreach (UhtType child in this.Children)
+						foreach (UhtType child in Children)
 						{
 							if (child is UhtProperty property)
 							{
@@ -835,18 +835,18 @@ namespace EpicGames.UHT.Types
 
 				case UhtResolvePhase.Final:
 					Dictionary<string, List<GetterSetterToResolve>>? gsToResolve = null;
-					foreach (UhtType child in this.Children)
+					foreach (UhtType child in Children)
 					{
 						if (child is UhtProperty property)
 						{
 							if (property.PropertyExportFlags.HasAnyFlags(UhtPropertyExportFlags.FieldNotify))
 							{
-								this.ClassExportFlags |= UhtClassExportFlags.HasFieldNotify;
+								ClassExportFlags |= UhtClassExportFlags.HasFieldNotify;
 							}
 
 							if (property.PropertyFlags.HasAnyFlags(EPropertyFlags.Net))
 							{
-								this.ClassExportFlags |= UhtClassExportFlags.SelfHasReplicatedProperties;
+								ClassExportFlags |= UhtClassExportFlags.SelfHasReplicatedProperties;
 								break;
 							}
 
@@ -870,26 +870,26 @@ namespace EpicGames.UHT.Types
 						{
 							if (function.FunctionExportFlags.HasAnyFlags(UhtFunctionExportFlags.FieldNotify))
 							{
-								this.ClassExportFlags |= UhtClassExportFlags.HasFieldNotify;
+								ClassExportFlags |= UhtClassExportFlags.HasFieldNotify;
 							}
 						}
 					}
 
-					if (this.SuperClass != null)
+					if (SuperClass != null)
 					{
-						if (this.SuperClass.ClassExportFlags.HasAnyFlags(UhtClassExportFlags.HasReplciatedProperties))
+						if (SuperClass.ClassExportFlags.HasAnyFlags(UhtClassExportFlags.HasReplciatedProperties))
 						{
-							this.ClassExportFlags |= UhtClassExportFlags.SuperHasReplicatedProperties;
+							ClassExportFlags |= UhtClassExportFlags.SuperHasReplicatedProperties;
 						}
 					}
 
-					if (!this.ClassFlags.HasAnyFlags(EClassFlags.Interface))
+					if (!ClassFlags.HasAnyFlags(EClassFlags.Interface))
 					{
-						if (this.ClassExportFlags.HasAnyFlags(UhtClassExportFlags.SelfHasReplicatedProperties))
+						if (ClassExportFlags.HasAnyFlags(UhtClassExportFlags.SelfHasReplicatedProperties))
 						{
 							if (TryGetDeclaration("GetLifetimeReplicatedProps", out UhtFoundDeclaration _))
 							{
-								this.ClassExportFlags |= UhtClassExportFlags.HasGetLifetimeReplicatedProps;
+								ClassExportFlags |= UhtClassExportFlags.HasGetLifetimeReplicatedProps;
 							}
 						}
 					}
@@ -897,9 +897,9 @@ namespace EpicGames.UHT.Types
 					// If we have any getters/setters to resolve 
 					if (gsToResolve != null)
 					{
-						if (this.Declarations != null)
+						if (Declarations != null)
 						{
-							foreach (UhtDeclaration declaration in this.Declarations)
+							foreach (UhtDeclaration declaration in Declarations)
 							{
 
 								// Locate the index of the name
@@ -956,7 +956,7 @@ namespace EpicGames.UHT.Types
 				case UhtResolvePhase.Final:
 					if (ScanForInstancedReferenced(false))
 					{
-						this.ClassFlags |= EClassFlags.HasInstancedReference;
+						ClassFlags |= EClassFlags.HasInstancedReference;
 					}
 					break;
 			}
@@ -965,12 +965,12 @@ namespace EpicGames.UHT.Types
 		/// <inheritdoc/>
 		public override bool ScanForInstancedReferenced(bool deepScan)
 		{
-			if (this.ClassFlags.HasAnyFlags(EClassFlags.HasInstancedReference))
+			if (ClassFlags.HasAnyFlags(EClassFlags.HasInstancedReference))
 			{
 				return true;
 			}
 
-			if (this.SuperClass != null && this.SuperClass.ScanForInstancedReferenced(deepScan))
+			if (SuperClass != null && SuperClass.ScanForInstancedReferenced(deepScan))
 			{
 				return true;
 			}
@@ -983,51 +983,51 @@ namespace EpicGames.UHT.Types
 			MergeShowCategories();
 
 			// Merge ShowFunctions and HideFunctions
-			AppendStringListMetaData(this.SuperClass, UhtNames.HideFunctions, this.HideFunctions);
-			foreach (string value in this.ShowFunctions)
+			AppendStringListMetaData(SuperClass, UhtNames.HideFunctions, HideFunctions);
+			foreach (string value in ShowFunctions)
 			{
-				this.HideFunctions.RemoveSwap(value);
+				HideFunctions.RemoveSwap(value);
 			}
-			this.ShowFunctions.Clear();
+			ShowFunctions.Clear();
 
 			// Merge AutoExpandCategories and AutoCollapseCategories (we still want to keep AutoExpandCategories though!)
-			List<string> parentAutoExpandCategories = GetStringListMetaData(this.SuperClass, UhtNames.AutoExpandCategories);
-			List<string> parentAutoCollapseCategories = GetStringListMetaData(this.SuperClass, UhtNames.AutoCollapseCategories);
+			List<string> parentAutoExpandCategories = GetStringListMetaData(SuperClass, UhtNames.AutoExpandCategories);
+			List<string> parentAutoCollapseCategories = GetStringListMetaData(SuperClass, UhtNames.AutoCollapseCategories);
 
-			foreach (string value in this.AutoExpandCategories)
+			foreach (string value in AutoExpandCategories)
 			{
-				this.AutoCollapseCategories.RemoveSwap(value);
+				AutoCollapseCategories.RemoveSwap(value);
 				parentAutoCollapseCategories.RemoveSwap(value);
 			}
 
 			// Do the same as above but the other way around
-			foreach (string value in this.AutoCollapseCategories)
+			foreach (string value in AutoCollapseCategories)
 			{
-				this.AutoExpandCategories.RemoveSwap(value);
+				AutoExpandCategories.RemoveSwap(value);
 				parentAutoExpandCategories.RemoveSwap(value);
 			}
 
 			// Once AutoExpandCategories and AutoCollapseCategories for THIS class have been parsed, add the parent inherited categories
-			this.AutoCollapseCategories.AddRange(parentAutoCollapseCategories);
-			this.AutoExpandCategories.AddRange(parentAutoExpandCategories);
+			AutoCollapseCategories.AddRange(parentAutoCollapseCategories);
+			AutoExpandCategories.AddRange(parentAutoExpandCategories);
 		}
 
 		private void MergeShowCategories()
 		{
 
 			// Add the super class hide categories and prime the output show categories with the parent
-			List<string> outShowCategories = GetStringListMetaData(this.SuperClass, UhtNames.ShowCategories);
-			AppendStringListMetaData(this.SuperClass, UhtNames.HideCategories, this.HideCategories);
+			List<string> outShowCategories = GetStringListMetaData(SuperClass, UhtNames.ShowCategories);
+			AppendStringListMetaData(SuperClass, UhtNames.HideCategories, HideCategories);
 
 			// If this class has new show categories, then merge them
-			if (this.ShowCategories.Count != 0)
+			if (ShowCategories.Count != 0)
 			{
 				StringBuilder subCategoryPath = new();
-				foreach (string value in this.ShowCategories)
+				foreach (string value in ShowCategories)
 				{
 
 					// if we didn't find this specific category path in the HideCategories metadata
-					if (!this.HideCategories.RemoveSwap(value))
+					if (!HideCategories.RemoveSwap(value))
 					{
 						string[] subCategories = value.ToString().Split('|', StringSplitOptions.RemoveEmptyEntries);
 
@@ -1049,22 +1049,22 @@ namespace EpicGames.UHT.Types
 			}
 
 			// Replace the show categories
-			this.ShowCategories.Clear();
-			this.ShowCategories.AddRange(outShowCategories);
+			ShowCategories.Clear();
+			ShowCategories.AddRange(outShowCategories);
 		}
 
 		private void SetAndValidateWithinClass(UhtResolvePhase resolvePhase)
 		{
 			// The class within must be a child of any super class within
-			UhtClass expectedClassWithin = this.SuperClass != null ? this.SuperClass.ClassWithin : this.Session.UObject;
+			UhtClass expectedClassWithin = SuperClass != null ? SuperClass.ClassWithin : Session.UObject;
 
 			// Process all of the class specifiers
-			if (!String.IsNullOrEmpty(this.ClassWithinIdentifier))
+			if (!String.IsNullOrEmpty(ClassWithinIdentifier))
 			{
-				UhtClass? specifiedClassWithin = (UhtClass?)this.Session.FindType(null, UhtFindOptions.EngineName | UhtFindOptions.Class, this.ClassWithinIdentifier);
+				UhtClass? specifiedClassWithin = (UhtClass?)Session.FindType(null, UhtFindOptions.EngineName | UhtFindOptions.Class, ClassWithinIdentifier);
 				if (specifiedClassWithin == null)
 				{
-					this.LogError($"Within class '{this.ClassWithinIdentifier}' not found");
+					this.LogError($"Within class '{ClassWithinIdentifier}' not found");
 				}
 				else
 				{
@@ -1076,7 +1076,7 @@ namespace EpicGames.UHT.Types
 						specifiedClassWithin.Resolve(resolvePhase);
 					}
 
-					if (specifiedClassWithin.IsChildOf(this.Session.UInterface))
+					if (specifiedClassWithin.IsChildOf(Session.UInterface))
 					{
 						this.LogError("Classes cannot be 'within' interfaces");
 					}
@@ -1086,7 +1086,7 @@ namespace EpicGames.UHT.Types
 					}
 					else
 					{
-						this.ClassWithin = specifiedClassWithin;
+						ClassWithin = specifiedClassWithin;
 					}
 				}
 			}
@@ -1094,7 +1094,7 @@ namespace EpicGames.UHT.Types
 			// If we don't have a class within set, then just use the expected within
 			else
 			{
-				this.ClassWithin = expectedClassWithin;
+				ClassWithin = expectedClassWithin;
 			}
 		}
 
@@ -1102,53 +1102,53 @@ namespace EpicGames.UHT.Types
 		{
 			// Since this flag is computed in this method, we have to re-propagate the flag from the super
 			// just in case they were defined in this source file.
-			if (this.SuperClass != null)
+			if (SuperClass != null)
 			{
-				this.ClassFlags |= this.SuperClass.ClassFlags & EClassFlags.Config;
+				ClassFlags |= SuperClass.ClassFlags & EClassFlags.Config;
 			}
 
 			// Set the class config flag if any properties have config
-			if (!this.ClassFlags.HasAnyFlags(EClassFlags.Config))
+			if (!ClassFlags.HasAnyFlags(EClassFlags.Config))
 			{
-				foreach (UhtProperty property in this.Properties)
+				foreach (UhtProperty property in Properties)
 				{
 					if (property.PropertyFlags.HasAnyFlags(EPropertyFlags.Config))
 					{
-						this.ClassFlags |= EClassFlags.Config;
+						ClassFlags |= EClassFlags.Config;
 						break;
 					}
 				}
 			}
 
-			if (this.Config.Length > 0)
+			if (Config.Length > 0)
 			{
 				// if the user specified "inherit", we're just going to use the parent class's config filename
 				// this is not actually necessary but it can be useful for explicitly communicating config-ness
-				if (this.Config.Equals("inherit", StringComparison.OrdinalIgnoreCase))
+				if (Config.Equals("inherit", StringComparison.OrdinalIgnoreCase))
 				{
-					if (this.SuperClass == null)
+					if (SuperClass == null)
 					{
-						this.LogError($"Cannot inherit config filename for class '{this.SourceName}' when it has no super class");
+						this.LogError($"Cannot inherit config filename for class '{SourceName}' when it has no super class");
 					}
-					else if (this.SuperClass.Config.Length == 0)
+					else if (SuperClass.Config.Length == 0)
 					{
-						this.LogError($"Cannot inherit config filename for class '{this.SourceName}' when parent class '{this.SuperClass.SourceName}' has no config filename");
+						this.LogError($"Cannot inherit config filename for class '{SourceName}' when parent class '{SuperClass.SourceName}' has no config filename");
 					}
 					else
 					{
-						this.Config = this.SuperClass.Config;
+						Config = SuperClass.Config;
 					}
 				}
 			}
-			else if (this.ClassFlags.HasAnyFlags(EClassFlags.Config) && this.SuperClass != null)
+			else if (ClassFlags.HasAnyFlags(EClassFlags.Config) && SuperClass != null)
 			{
-				this.Config = this.SuperClass.Config;
+				Config = SuperClass.Config;
 			}
 
-			if (this.ClassFlags.HasAnyFlags(EClassFlags.Config) && this.Config.Length == 0)
+			if (ClassFlags.HasAnyFlags(EClassFlags.Config) && Config.Length == 0)
 			{
 				this.LogError("Classes with config / globalconfig member variables need to specify config file.");
-				this.Config = "Engine";
+				Config = "Engine";
 			}
 		}
 
@@ -1430,27 +1430,27 @@ namespace EpicGames.UHT.Types
 			options = base.Validate(options);
 
 			// Classes must start with a valid prefix
-			string expectedClassName = this.EngineNamePrefix + this.EngineName;
-			if (expectedClassName != this.SourceName)
+			string expectedClassName = EngineNamePrefix + EngineName;
+			if (expectedClassName != SourceName)
 			{
-				this.LogError($"Class '{this.SourceName}' has an invalid Unreal prefix, expecting '{expectedClassName}'");
+				this.LogError($"Class '{SourceName}' has an invalid Unreal prefix, expecting '{expectedClassName}'");
 			}
 
 			// If we have a super class
-			if (this.SuperClass != null)
+			if (SuperClass != null)
 			{
-				if (!this.ClassFlags.HasAnyFlags(EClassFlags.NotPlaceable) && this.SuperClass.ClassFlags.HasAnyFlags(EClassFlags.NotPlaceable))
+				if (!ClassFlags.HasAnyFlags(EClassFlags.NotPlaceable) && SuperClass.ClassFlags.HasAnyFlags(EClassFlags.NotPlaceable))
 				{
 					this.LogError("The 'placeable' specifier cannot override a 'nonplaceable' base class. Classes are assumed to be placeable by default. "
 						+ "Consider whether using the 'abstract' specifier on the base class would work.");
 				}
 
 				// Native interfaces don't require a super class, but UHT technically allows it.
-				if (this.ClassType == UhtClassType.Interface || this.ClassType == UhtClassType.NativeInterface)
+				if (ClassType == UhtClassType.Interface || ClassType == UhtClassType.NativeInterface)
 				{
-					if (this != this.Session.UInterface && !this.SuperClass.ClassFlags.HasAnyFlags(EClassFlags.Interface))
+					if (this != Session.UInterface && !SuperClass.ClassFlags.HasAnyFlags(EClassFlags.Interface))
 					{
-						this.LogError($"Interface class '{this.SourceName}' cannot inherit from non-interface class '{this.SuperClass.SourceName}'");
+						this.LogError($"Interface class '{SourceName}' cannot inherit from non-interface class '{SuperClass.SourceName}'");
 					}
 				}
 			}
@@ -1458,55 +1458,55 @@ namespace EpicGames.UHT.Types
 			// If we don't have a super class
 			else
 			{
-				if (this.ClassType == UhtClassType.Interface)
+				if (ClassType == UhtClassType.Interface)
 				{
-					this.LogError($"Interface '{this.SourceName}' must derive from an interface");
+					this.LogError($"Interface '{SourceName}' must derive from an interface");
 				}
 			}
 
-			if (this.ClassFlags.HasAnyFlags(EClassFlags.NeedsDeferredDependencyLoading) && !this.IsChildOf(this.Session.UClass))
+			if (ClassFlags.HasAnyFlags(EClassFlags.NeedsDeferredDependencyLoading) && !IsChildOf(Session.UClass))
 			{
 				// CLASS_NeedsDeferredDependencyLoading can only be set on classes derived from UClass
-				this.LogError($"'NeedsDeferredDependencyLoading' is set on '{this.SourceName}' but the flag can only be used with classes derived from UClass.");
+				this.LogError($"'NeedsDeferredDependencyLoading' is set on '{SourceName}' but the flag can only be used with classes derived from UClass.");
 			}
 
-			if (this.ClassExportFlags.HasAnyFlags(UhtClassExportFlags.HasFieldNotify))
+			if (ClassExportFlags.HasAnyFlags(UhtClassExportFlags.HasFieldNotify))
 			{
-				if (this.Session.INotifyFieldValueChanged == null)
+				if (Session.INotifyFieldValueChanged == null)
 				{
-					this.LogError($"UClass '{this.SourceName}' has FieldNotify elements but the interface 'INotifyFieldValueChanged' is not defined.");
+					this.LogError($"UClass '{SourceName}' has FieldNotify elements but the interface 'INotifyFieldValueChanged' is not defined.");
 				}
-				else if (!ImplementsInterface(this.Session.INotifyFieldValueChanged))
+				else if (!ImplementsInterface(Session.INotifyFieldValueChanged))
 				{
-					this.LogError($"UClass '{this.SourceName}' need to implement the interface INotifyFieldValueChanged' to support FieldNotify.");
+					this.LogError($"UClass '{SourceName}' need to implement the interface INotifyFieldValueChanged' to support FieldNotify.");
 				}
 			}
 
 			ValidateProperties();
 
 			// Perform any validation specific to the class type
-			switch (this.ClassType)
+			switch (ClassType)
 			{
 				case UhtClassType.Class:
 					{
-						if (this.ClassFlags.HasAnyFlags(EClassFlags.EditInlineNew))
+						if (ClassFlags.HasAnyFlags(EClassFlags.EditInlineNew))
 						{
 							// don't allow actor classes to be declared EditInlineNew
-							if (this.IsChildOf(this.Session.AActor))
+							if (IsChildOf(Session.AActor))
 							{
 								this.LogError("Invalid class attribute: Creating actor instances via the property window is not allowed");
 							}
 						}
 
 						// Make sure both RequiredAPI and MinimalAPI aren't specified
-						if (this.ClassFlags.HasAllFlags(EClassFlags.MinimalAPI | EClassFlags.RequiredAPI))
+						if (ClassFlags.HasAllFlags(EClassFlags.MinimalAPI | EClassFlags.RequiredAPI))
 						{
 							this.LogError("MinimalAPI cannot be specified when the class is fully exported using a MODULENAME_API macro");
 						}
 
 						// Make sure that all interface functions are implemented property
 						// Iterate over all the bases looking for any interfaces
-						foreach (UhtStruct baseStruct in this.Bases)
+						foreach (UhtStruct baseStruct in Bases)
 						{
 
 							// Skip children that aren't interfaces or if a common ancestor
@@ -1532,7 +1532,7 @@ namespace EpicGames.UHT.Types
 								bool implemented = baseFunction.FunctionFlags.HasAnyFlags(EFunctionFlags.Delegate);
 
 								// Try to find an existing function
-								foreach (UhtType child in this.Children)
+								foreach (UhtType child in Children)
 								{
 									UhtFunction? function = child as UhtFunction;
 									if (function == null || function.SourceName != baseFunction.SourceName)
@@ -1588,13 +1588,13 @@ namespace EpicGames.UHT.Types
 						}
 					}
 
-					if (!this.ClassFlags.HasAnyFlags(EClassFlags.Interface))
+					if (!ClassFlags.HasAnyFlags(EClassFlags.Interface))
 					{
-						if (this.ClassExportFlags.HasAnyFlags(UhtClassExportFlags.SelfHasReplicatedProperties) &&
-							!this.ClassExportFlags.HasAnyFlags(UhtClassExportFlags.HasGetLifetimeReplicatedProps) &&
-							this.GeneratedCodeVersion != EGeneratedCodeVersion.V1)
+						if (ClassExportFlags.HasAnyFlags(UhtClassExportFlags.SelfHasReplicatedProperties) &&
+							!ClassExportFlags.HasAnyFlags(UhtClassExportFlags.HasGetLifetimeReplicatedProps) &&
+							GeneratedCodeVersion != EGeneratedCodeVersion.V1)
 						{
-							this.LogError($"Class {this.SourceName} has Net flagged properties and should declare member function: void GetLifetimeReplicatedProps(TArray<FLifetimeProperty>& OutLifetimeProps) const override");
+							this.LogError($"Class {SourceName} has Net flagged properties and should declare member function: void GetLifetimeReplicatedProps(TArray<FLifetimeProperty>& OutLifetimeProps) const override");
 						}
 					}
 					break;
@@ -1605,7 +1605,7 @@ namespace EpicGames.UHT.Types
 						// In the backward compatible case where they declare neither, both of these bools are false
 						bool canImplementInBlueprints = CanImplementInBlueprints();
 						bool cannotImplementInBlueprints = CannotImplementInBlueprints(canImplementInBlueprints);
-						foreach (UhtType child in this.Children)
+						foreach (UhtType child in Children)
 						{
 							if (child is UhtFunction function)
 							{
@@ -1764,7 +1764,7 @@ namespace EpicGames.UHT.Types
 
 		void ValidateProperties()
 		{
-			foreach (UhtType child in this.Children)
+			foreach (UhtType child in Children)
 			{
 				if (child is UhtProperty property)
 				{
@@ -1791,7 +1791,7 @@ namespace EpicGames.UHT.Types
 					}
 
 					// Validate if we are using editor only data in a class or struct definition
-					if (this.ClassFlags.HasAnyFlags(EClassFlags.Optional))
+					if (ClassFlags.HasAnyFlags(EClassFlags.Optional))
 					{
 						if (property.PropertyFlags.HasAnyFlags(EPropertyFlags.EditorOnly))
 						{
@@ -1843,7 +1843,7 @@ namespace EpicGames.UHT.Types
 		/// <returns>True if the interface is marked as Blueprintable</returns>
 		private bool CanImplementInBlueprints()
 		{
-			return this.MetaData.GetBoolean(UhtNames.IsBlueprintBase);
+			return MetaData.GetBoolean(UhtNames.IsBlueprintBase);
 		}
 
 		/// <summary>
@@ -1864,8 +1864,8 @@ namespace EpicGames.UHT.Types
 		/// <returns>True if the interface is marked as NotBlueprintable</returns>
 		private bool CannotImplementInBlueprints(bool canImplementInBlueprints)
 		{
-			return (!canImplementInBlueprints && this.MetaData.ContainsKey(UhtNames.IsBlueprintBase))
-				|| this.MetaData.ContainsKey(UhtNames.CannotImplementInterfaceInBlueprint);
+			return (!canImplementInBlueprints && MetaData.ContainsKey(UhtNames.IsBlueprintBase))
+				|| MetaData.ContainsKey(UhtNames.CannotImplementInterfaceInBlueprint);
 		}
 		#endregion
 
@@ -1874,13 +1874,13 @@ namespace EpicGames.UHT.Types
 		{
 
 			// Ignore any intrinsics
-			if (this.ClassFlags.HasAnyFlags(EClassFlags.Intrinsic))
+			if (ClassFlags.HasAnyFlags(EClassFlags.Intrinsic))
 			{
 				return;
 			}
 
 			// In the original UHT, we don't export IInterface (it wasn't even in the header file)
-			if (this.HeaderFile.IsNoExportTypes && this.SourceName == "IInterface")
+			if (HeaderFile.IsNoExportTypes && SourceName == "IInterface")
 			{
 				return;
 			}
@@ -1892,14 +1892,14 @@ namespace EpicGames.UHT.Types
 			collector.AddCrossModuleReference(this, true);
 
 			// Add the super class
-			if (this.SuperClass != null)
+			if (SuperClass != null)
 			{
-				collector.AddDeclaration(this.SuperClass, true);
-				collector.AddCrossModuleReference(this.SuperClass, true);
+				collector.AddDeclaration(SuperClass, true);
+				collector.AddCrossModuleReference(SuperClass, true);
 			}
 
 			// Add any other base classes
-			foreach (UhtStruct @base in this.Bases)
+			foreach (UhtStruct @base in Bases)
 			{
 				if (@base is UhtClass baseClass)
 				{
@@ -1908,11 +1908,11 @@ namespace EpicGames.UHT.Types
 			}
 
 			// Add the package
-			collector.AddDeclaration(this.Package, true);
-			collector.AddCrossModuleReference(this.Package, true);
+			collector.AddDeclaration(Package, true);
+			collector.AddCrossModuleReference(Package, true);
 
 			// Collect children references
-			foreach (UhtType child in this.Children)
+			foreach (UhtType child in Children)
 			{
 				child.CollectReferences(collector);
 			}
@@ -1920,10 +1920,10 @@ namespace EpicGames.UHT.Types
 			// Done at the end so any contained delegate functions are added first
 			// We also add interfaces in the export type order where the native interface
 			// is located.  But we add the interface...
-			switch (this.ClassType)
+			switch (ClassType)
 			{
 				case UhtClassType.NativeInterface:
-					if (this.AlternateObject != null && this.AlternateObject is UhtField field)
+					if (AlternateObject != null && AlternateObject is UhtField field)
 					{
 						collector.AddExportType(field);
 					}
diff --git a/Engine/Source/Programs/Shared/EpicGames.UHT/Types/UhtEnum.cs b/Engine/Source/Programs/Shared/EpicGames.UHT/Types/UhtEnum.cs
--- a/Engine/Source/Programs/Shared/EpicGames.UHT/Types/UhtEnum.cs
+++ b/Engine/Source/Programs/Shared/EpicGames.UHT/Types/UhtEnum.cs
@@ -147,7 +147,7 @@ namespace EpicGames.UHT.Types
 
 		///<inheritdoc/>
 		[JsonIgnore]
-		protected override UhtSpecifierValidatorTable? SpecifierValidatorTable => this.Session.GetSpecifierValidatorTable(UhtTableNames.Enum);
+		protected override UhtSpecifierValidatorTable? SpecifierValidatorTable => Session.GetSpecifierValidatorTable(UhtTableNames.Enum);
 
 		/// <summary>
 		/// Construct a new enumeration
@@ -156,8 +156,8 @@ namespace EpicGames.UHT.Types
 		/// <param name="lineNumber">Line number of declaration</param>
 		public UhtEnum(UhtType outer, int lineNumber) : base(outer, lineNumber)
 		{
-			this.MetaData.KeyConversion = this;
-			this.EnumValues = new List<UhtEnumValue>();
+			MetaData.KeyConversion = this;
+			EnumValues = new List<UhtEnumValue>();
 		}
 
 		/// <summary>
@@ -187,7 +187,7 @@ namespace EpicGames.UHT.Types
 			name = CleanEnumValueName(name);
 			for (int index = 0; index < EnumValues.Count; ++index)
 			{
-				if (this.EnumValues[index].Name == name)
+				if (EnumValues[index].Name == name)
 				{
 					return index;
 				}
@@ -203,8 +203,8 @@ namespace EpicGames.UHT.Types
 		/// <returns>Meta data name with the enum value name</returns>
 		public string GetMetaDataKey(string name, int nameIndex)
 		{
-			string enumName = this.EnumValues[nameIndex].Name;
-			if (this.CppForm != UhtEnumCppForm.Regular)
+			string enumName = EnumValues[nameIndex].Name;
+			if (CppForm != UhtEnumCppForm.Regular)
 			{
 				int scopeIndex = enumName.IndexOf("::", StringComparison.Ordinal);
 				if (scopeIndex >= 0)
@@ -223,11 +223,11 @@ namespace EpicGames.UHT.Types
 		/// <exception cref="UhtIceException">Unexpected enum form</exception>
 		public string GetFullEnumName(string shortEnumName)
 		{
-			switch (this.CppForm)
+			switch (CppForm)
 			{
 				case UhtEnumCppForm.Namespaced:
 				case UhtEnumCppForm.EnumClass:
-					return $"{this.SourceName}::{shortEnumName}";
+					return $"{SourceName}::{shortEnumName}";
 
 				case UhtEnumCppForm.Regular:
 					return shortEnumName;
@@ -245,8 +245,8 @@ namespace EpicGames.UHT.Types
 		/// <returns></returns>
 		public int AddEnumValue(string shortEnumName, long value)
 		{
-			int enumIndex = this.EnumValues.Count;
-			this.EnumValues.Add(new UhtEnumValue { Name = GetFullEnumName(shortEnumName), Value = value });
+			int enumIndex = EnumValues.Count;
+			EnumValues.Add(new UhtEnumValue { Name = GetFullEnumName(shortEnumName), Value = value });
 			return enumIndex;
 		}
 
@@ -268,29 +268,29 @@ namespace EpicGames.UHT.Types
 		{
 			if (policy.ClassOrStructCommentRequired)
 			{
-				if (!this.MetaData.ContainsKey(UhtNames.ToolTip))
+				if (!MetaData.ContainsKey(UhtNames.ToolTip))
 				{
-					this.LogError(this.MetaData.LineNumber, $"Enum '{this.SourceName}' does not provide a tooltip / comment (DocumentationPolicy)");
+					this.LogError(MetaData.LineNumber, $"Enum '{SourceName}' does not provide a tooltip / comment (DocumentationPolicy)");
 				}
 			}
 
 			Dictionary<string, string> toolTips = new(StringComparer.OrdinalIgnoreCase);
-			for (int enumIndex = 0; enumIndex < this.EnumValues.Count; ++enumIndex)
+			for (int enumIndex = 0; enumIndex < EnumValues.Count; ++enumIndex)
 			{
-				if (!this.MetaData.TryGetValue(UhtNames.Name, enumIndex, out string? entryName))
+				if (!MetaData.TryGetValue(UhtNames.Name, enumIndex, out string? entryName))
 				{
 					continue;
 				}
 
-				if (!this.MetaData.TryGetValue(UhtNames.ToolTip, enumIndex, out string? toolTip))
+				if (!MetaData.TryGetValue(UhtNames.ToolTip, enumIndex, out string? toolTip))
 				{
-					this.LogError(this.MetaData.LineNumber, $"Enum entry '{this.SourceName}::{entryName}' does not provide a tooltip / comment (DocumentationPolicy)");
+					this.LogError(MetaData.LineNumber, $"Enum entry '{SourceName}::{entryName}' does not provide a tooltip / comment (DocumentationPolicy)");
 					continue;
 				}
 
 				if (toolTips.TryGetValue(toolTip, out string? dupName))
 				{
-					this.LogError(this.MetaData.LineNumber, $"Enum entries '{this.SourceName}::{entryName}' and '{this.SourceName}::{dupName}' have identical tooltips / comments (DocumentationPolicy)");
+					this.LogError(MetaData.LineNumber, $"Enum entries '{SourceName}::{entryName}' and '{SourceName}::{dupName}' have identical tooltips / comments (DocumentationPolicy)");
 				}
 				else
 				{
@@ -306,7 +306,7 @@ namespace EpicGames.UHT.Types
 			collector.AddExportType(this);
 			collector.AddDeclaration(this, true);
 			collector.AddCrossModuleReference(this, true);
-			collector.AddCrossModuleReference(this.Package, true);
+			collector.AddCrossModuleReference(Package, true);
 		}
 	}
 }
diff --git a/Engine/Source/Programs/Shared/EpicGames.UHT/Types/UhtFunction.cs b/Engine/Source/Programs/Shared/EpicGames.UHT/Types/UhtFunction.cs
--- a/Engine/Source/Programs/Shared/EpicGames.UHT/Types/UhtFunction.cs
+++ b/Engine/Source/Programs/Shared/EpicGames.UHT/Types/UhtFunction.cs
@@ -247,18 +247,18 @@ namespace EpicGames.UHT.Types
 		/// The super function.  Currently unused
 		/// </summary>
 		[JsonConverter(typeof(UhtNullableTypeSourceNameJsonConverter<UhtFunction>))]
-		public UhtFunction? SuperFunction => (UhtFunction?)this.Super;
+		public UhtFunction? SuperFunction => (UhtFunction?)Super;
 
 		/// <inheritdoc/>
 		[JsonIgnore]
-		public override UhtEngineType EngineType => this.FunctionFlags.HasAnyFlags(EFunctionFlags.Delegate) ? UhtEngineType.Delegate : UhtEngineType.Function;
+		public override UhtEngineType EngineType => FunctionFlags.HasAnyFlags(EFunctionFlags.Delegate) ? UhtEngineType.Delegate : UhtEngineType.Function;
 
 		/// <inheritdoc/>
 		public override string EngineClassName
 		{
 			get
 			{
-				switch (this.FunctionType)
+				switch (FunctionType)
 				{
 					case UhtFunctionType.Function:
 						return "Function";
@@ -278,17 +278,17 @@ namespace EpicGames.UHT.Types
 		[JsonIgnore]
 		public string StrippedFunctionName
 		{
-			get => this._strippedFunctionName ?? this.EngineName;
-			set => this._strippedFunctionName = value;
+			get => _strippedFunctionName ?? EngineName;
+			set => _strippedFunctionName = value;
 		}
 
 		///<inheritdoc/>
 		[JsonIgnore]
-		public override bool Deprecated => this.MetaData.ContainsKey(UhtNames.DeprecatedFunction);
+		public override bool Deprecated => MetaData.ContainsKey(UhtNames.DeprecatedFunction);
 
 		///<inheritdoc/>
 		[JsonIgnore]
-		protected override UhtSpecifierValidatorTable? SpecifierValidatorTable => this.Session.GetSpecifierValidatorTable(UhtTableNames.Function);
+		protected override UhtSpecifierValidatorTable? SpecifierValidatorTable => Session.GetSpecifierValidatorTable(UhtTableNames.Function);
 
 		/// <summary>
 		/// Identifier for an RPC call to a platform service
@@ -329,25 +329,25 @@ namespace EpicGames.UHT.Types
 		/// True if the function has a return value.
 		/// </summary>
 		[JsonIgnore]
-		public bool HasReturnProperty => this.Children.Count > 0 && ((UhtProperty)this.Children[^1]).PropertyFlags.HasAnyFlags(EPropertyFlags.ReturnParm);
+		public bool HasReturnProperty => Children.Count > 0 && ((UhtProperty)Children[^1]).PropertyFlags.HasAnyFlags(EPropertyFlags.ReturnParm);
 
 		/// <summary>
 		/// The return value property or null if the function doesn't have a return value
 		/// </summary>
 		[JsonIgnore]
-		public UhtProperty? ReturnProperty => this.HasReturnProperty ? (UhtProperty)this.Children[^1] : null;
+		public UhtProperty? ReturnProperty => HasReturnProperty ? (UhtProperty)Children[^1] : null;
 
 		/// <summary>
 		/// True if the function has parameters.
 		/// </summary>
 		[JsonIgnore]
-		public bool HasParameters => this.Children.Count > 0 && !((UhtProperty)this.Children[0]).PropertyFlags.HasAnyFlags(EPropertyFlags.ReturnParm);
+		public bool HasParameters => Children.Count > 0 && !((UhtProperty)Children[0]).PropertyFlags.HasAnyFlags(EPropertyFlags.ReturnParm);
 
 		/// <summary>
 		/// Return read only memory of all the function parameters
 		/// </summary>
 		[JsonIgnore]
-		public ReadOnlyMemory<UhtType> ParameterProperties => new(this.Children.ToArray(), 0, this.HasReturnProperty ? this.Children.Count - 1 : this.Children.Count);
+		public ReadOnlyMemory<UhtType> ParameterProperties => new(Children.ToArray(), 0, HasReturnProperty ? Children.Count - 1 : Children.Count);
 
 		/// <summary>
 		/// True if the function has any outputs including a return value
@@ -357,7 +357,7 @@ namespace EpicGames.UHT.Types
 		{
 			get
 			{
-				foreach (UhtType type in this.Children)
+				foreach (UhtType type in Children)
 				{
 					if (type is UhtProperty property)
 					{
@@ -389,7 +389,7 @@ namespace EpicGames.UHT.Types
 				{
 					throw new UhtIceException("Only properties marked as parameters can be added to functions");
 				}
-				if (property.PropertyFlags.HasAnyFlags(EPropertyFlags.ReturnParm) && this.HasReturnProperty)
+				if (property.PropertyFlags.HasAnyFlags(EPropertyFlags.ReturnParm) && HasReturnProperty)
 				{
 					throw new UhtIceException("Attempt to set multiple return properties on function");
 				}
@@ -410,35 +410,35 @@ namespace EpicGames.UHT.Types
 			switch (phase)
 			{
 				case UhtResolvePhase.Bases:
-					if (this.FunctionType == UhtFunctionType.Function && this.Outer is UhtClass outerClass)
+					if (FunctionType == UhtFunctionType.Function && Outer is UhtClass outerClass)
 					{
 						// non-static functions in a const class must be const themselves
 						if (outerClass.ClassFlags.HasAnyFlags(EClassFlags.Const))
 						{
-							this.FunctionFlags |= EFunctionFlags.Const;
+							FunctionFlags |= EFunctionFlags.Const;
 						}
 					}
 					break;
 
 				case UhtResolvePhase.Properties:
 					UhtPropertyParser.ResolveChildren(this, GetPropertyParseOptions(false));
-					foreach (UhtProperty property in this.Properties)
+					foreach (UhtProperty property in Properties)
 					{
 						if (property.DefaultValueTokens != null)
 						{
 							string key = "CPP_Default_" + property.EngineName;
-							if (!this.MetaData.ContainsKey(key))
+							if (!MetaData.ContainsKey(key))
 							{
 								bool parsed = false;
 								try
 								{
 									// All tokens MUST be consumed from the reader
 									StringBuilder builder = new();
-									IUhtTokenReader defaultValueReader = UhtTokenReplayReader.GetThreadInstance(property, this.HeaderFile.Data.Memory, property.DefaultValueTokens.ToArray(), UhtTokenType.EndOfDefault);
+									IUhtTokenReader defaultValueReader = UhtTokenReplayReader.GetThreadInstance(property, HeaderFile.Data.Memory, property.DefaultValueTokens.ToArray(), UhtTokenType.EndOfDefault);
 									parsed = property.SanitizeDefaultValue(defaultValueReader, builder) && defaultValueReader.IsEOF;
 									if (parsed)
 									{
-										this.MetaData.Add(key, builder.ToString());
+										MetaData.Add(key, builder.ToString());
 									}
 								}
 								catch (Exception)
@@ -448,7 +448,7 @@ namespace EpicGames.UHT.Types
 
 								if (!parsed)
 								{
-									StringView defaultValueText = new(this.HeaderFile.Data, property.DefaultValueTokens.First().InputStartPos,
+									StringView defaultValueText = new(HeaderFile.Data, property.DefaultValueTokens.First().InputStartPos,
 										property.DefaultValueTokens.Last().InputEndPos - property.DefaultValueTokens.First().InputStartPos);
 									property.LogError($"C++ Default parameter not parsed: {property.SourceName} '{defaultValueText}'");
 								}
@@ -458,31 +458,31 @@ namespace EpicGames.UHT.Types
 					break;
 
 				case UhtResolvePhase.Final:
-					if (this.Outer is UhtClass classObj)
+					if (Outer is UhtClass classObj)
 					{
-						if (this.FunctionFlags.HasAnyFlags(EFunctionFlags.Native) &&
-							!this.FunctionExportFlags.HasAnyFlags(UhtFunctionExportFlags.CustomThunk) &&
-							this.CppImplName != this.SourceName)
+						if (FunctionFlags.HasAnyFlags(EFunctionFlags.Native) &&
+							!FunctionExportFlags.HasAnyFlags(UhtFunctionExportFlags.CustomThunk) &&
+							CppImplName != SourceName)
 						{
-							if (classObj.TryGetDeclaration(this.CppImplName, out UhtFoundDeclaration declaration))
+							if (classObj.TryGetDeclaration(CppImplName, out UhtFoundDeclaration declaration))
 							{
-								this.FunctionExportFlags |= UhtFunctionExportFlags.ImplFound;
+								FunctionExportFlags |= UhtFunctionExportFlags.ImplFound;
 								if (declaration.IsVirtual)
 								{
-									this.FunctionExportFlags |= UhtFunctionExportFlags.ImplVirtual;
+									FunctionExportFlags |= UhtFunctionExportFlags.ImplVirtual;
 								}
 							}
-							if (classObj.TryGetDeclaration(this.CppValidationImplName, out declaration))
+							if (classObj.TryGetDeclaration(CppValidationImplName, out declaration))
 							{
-								this.FunctionExportFlags |= UhtFunctionExportFlags.ValidationImplFound;
+								FunctionExportFlags |= UhtFunctionExportFlags.ValidationImplFound;
 								if (declaration.IsVirtual)
 								{
-									this.FunctionExportFlags |= UhtFunctionExportFlags.ValidationImplVirtual;
+									FunctionExportFlags |= UhtFunctionExportFlags.ValidationImplVirtual;
 								}
 							}
 						}
 
-						if (this.FunctionType == UhtFunctionType.Function && this.FunctionExportFlags.HasAnyFlags(UhtFunctionExportFlags.DeclaredConst))
+						if (FunctionType == UhtFunctionType.Function && FunctionExportFlags.HasAnyFlags(UhtFunctionExportFlags.DeclaredConst))
 						{
 							// @todo: the presence of const and one or more outputs does not guarantee that there are
 							// no side effects. On GCC and clang we could use __attribure__((pure)) or __attribute__((const))
@@ -490,31 +490,31 @@ namespace EpicGames.UHT.Types
 							// const identifier to determine purity is not desirable. We should remove the following logic:
 
 							// If its a const BlueprintCallable function with some sort of output and is not being marked as an BlueprintPure=false function, mark it as BlueprintPure as well
-							if (this.HasAnyOutputs && this.FunctionFlags.HasAnyFlags(EFunctionFlags.BlueprintCallable) &&
-								!this.FunctionExportFlags.HasAnyFlags(UhtFunctionExportFlags.ForceBlueprintImpure))
+							if (HasAnyOutputs && FunctionFlags.HasAnyFlags(EFunctionFlags.BlueprintCallable) &&
+								!FunctionExportFlags.HasAnyFlags(UhtFunctionExportFlags.ForceBlueprintImpure))
 							{
-								this.FunctionFlags |= EFunctionFlags.BlueprintPure;
-								this.FunctionExportFlags |= UhtFunctionExportFlags.AutoBlueprintPure; // Disable error for pure being set
+								FunctionFlags |= EFunctionFlags.BlueprintPure;
+								FunctionExportFlags |= UhtFunctionExportFlags.AutoBlueprintPure; // Disable error for pure being set
 							}
 						}
 					}
 
 					// The following code is only performed on functions in a class.
-					if (this.Outer is UhtClass)
+					if (Outer is UhtClass)
 					{
-						foreach (UhtType type in this.Children)
+						foreach (UhtType type in Children)
 						{
 							if (type is UhtProperty property)
 							{
 								if (property.PropertyFlags.HasExactFlags(EPropertyFlags.OutParm | EPropertyFlags.ReturnParm, EPropertyFlags.OutParm))
 								{
-									this.FunctionFlags |= EFunctionFlags.HasOutParms;
+									FunctionFlags |= EFunctionFlags.HasOutParms;
 								}
 								if (property is UhtStructProperty structProperty)
 								{
 									if (structProperty.ScriptStruct.HasDefaults)
 									{
-										this.FunctionFlags |= EFunctionFlags.HasDefaults;
+										FunctionFlags |= EFunctionFlags.HasDefaults;
 									}
 								}
 							}
@@ -538,29 +538,29 @@ namespace EpicGames.UHT.Types
 			bool checkForBlueprint = false;
 
 			// Function type specific validation
-			switch (this.FunctionType)
+			switch (FunctionType)
 			{
 				case UhtFunctionType.Function:
-					UhtClass? outerClass = (UhtClass?)this.Outer;
+					UhtClass? outerClass = (UhtClass?)Outer;
 
-					if (this.FunctionFlags.HasAnyFlags(EFunctionFlags.BlueprintPure) && outerClass != null &&
+					if (FunctionFlags.HasAnyFlags(EFunctionFlags.BlueprintPure) && outerClass != null &&
 						outerClass.ClassFlags.HasAnyFlags(EClassFlags.Interface) &&
-						!this.FunctionExportFlags.HasAnyFlags(UhtFunctionExportFlags.AutoBlueprintPure))
+						!FunctionExportFlags.HasAnyFlags(UhtFunctionExportFlags.AutoBlueprintPure))
 					{
 						this.LogError("BlueprintPure specifier is not allowed for interface functions");
 					}
 
 					// If this function is blueprint callable or blueprint pure, require a category 
-					if (this.FunctionFlags.HasAnyFlags(EFunctionFlags.BlueprintCallable | EFunctionFlags.BlueprintPure))
+					if (FunctionFlags.HasAnyFlags(EFunctionFlags.BlueprintCallable | EFunctionFlags.BlueprintPure))
 					{
-						bool blueprintAccessor = this.MetaData.ContainsKey(UhtNames.BlueprintSetter) || this.MetaData.ContainsKey(UhtNames.BlueprintGetter);
-						bool hasMenuCategory = this.MetaData.ContainsKey(UhtNames.Category);
-						bool internalOnly = this.MetaData.GetBoolean(UhtNames.BlueprintInternalUseOnly);
+						bool blueprintAccessor = MetaData.ContainsKey(UhtNames.BlueprintSetter) || MetaData.ContainsKey(UhtNames.BlueprintGetter);
+						bool hasMenuCategory = MetaData.ContainsKey(UhtNames.Category);
+						bool internalOnly = MetaData.GetBoolean(UhtNames.BlueprintInternalUseOnly);
 
-						if (!hasMenuCategory && !internalOnly && !this.Deprecated && !blueprintAccessor)
+						if (!hasMenuCategory && !internalOnly && !Deprecated && !blueprintAccessor)
 						{
 							// To allow for quick iteration, don't enforce the requirement that game functions have to be categorized
-							if (this.HeaderFile.Package.IsPartOfEngine)
+							if (HeaderFile.Package.IsPartOfEngine)
 							{
 								this.LogError("An explicit Category specifier is required for Blueprint accessible functions in an Engine module.");
 							}
@@ -568,11 +568,11 @@ namespace EpicGames.UHT.Types
 					}
 
 					// Process the virtual-ness
-					if (this.FunctionExportFlags.HasAnyFlags(UhtFunctionExportFlags.Virtual))
+					if (FunctionExportFlags.HasAnyFlags(UhtFunctionExportFlags.Virtual))
 					{
 
 						// if this is a BlueprintNativeEvent or BlueprintImplementableEvent in an interface, make sure it's not "virtual"
-						if (this.FunctionFlags.HasAnyFlags(EFunctionFlags.BlueprintEvent))
+						if (FunctionFlags.HasAnyFlags(EFunctionFlags.BlueprintEvent))
 						{
 							if (outerClass != null && outerClass.ClassFlags.HasAnyFlags(EClassFlags.Interface))
 							{
@@ -580,7 +580,7 @@ namespace EpicGames.UHT.Types
 							}
 
 							// if this is a BlueprintNativeEvent, make sure it's not "virtual"
-							else if (this.FunctionFlags.HasAnyFlags(EFunctionFlags.Native))
+							else if (FunctionFlags.HasAnyFlags(EFunctionFlags.Native))
 							{
 								this.LogError("BlueprintNativeEvent functions must be non-virtual.");
 							}
@@ -594,147 +594,147 @@ namespace EpicGames.UHT.Types
 					else
 					{
 						// if this is a function in an Interface, it must be marked 'virtual' unless it's an event
-						if (outerClass != null && outerClass.ClassFlags.HasAnyFlags(EClassFlags.Interface) && !this.FunctionFlags.HasAnyFlags(EFunctionFlags.BlueprintEvent))
+						if (outerClass != null && outerClass.ClassFlags.HasAnyFlags(EClassFlags.Interface) && !FunctionFlags.HasAnyFlags(EFunctionFlags.BlueprintEvent))
 						{
 							this.LogError("Interface functions that are not BlueprintImplementableEvents must be declared 'virtual'");
 						}
 					}
 
 					// This is a final (prebinding, non-overridable) function
-					if (this.FunctionFlags.HasAnyFlags(EFunctionFlags.Final) || this.FunctionExportFlags.HasAnyFlags(UhtFunctionExportFlags.Final))
+					if (FunctionFlags.HasAnyFlags(EFunctionFlags.Final) || FunctionExportFlags.HasAnyFlags(UhtFunctionExportFlags.Final))
 					{
 						if (outerClass != null && outerClass.ClassFlags.HasAnyFlags(EClassFlags.Interface))
 						{
 							this.LogError("Interface functions cannot be declared 'final'");
 						}
-						else if (this.FunctionFlags.HasAnyFlags(EFunctionFlags.BlueprintEvent) && !this.FunctionExportFlags.HasAnyFlags(UhtFunctionExportFlags.AutoFinal))
+						else if (FunctionFlags.HasAnyFlags(EFunctionFlags.BlueprintEvent) && !FunctionExportFlags.HasAnyFlags(UhtFunctionExportFlags.AutoFinal))
 						{
 							this.LogError("Blueprint events cannot be declared 'final'");
 						}
 					}
 
-					if (this.FunctionFlags.HasAnyFlags(EFunctionFlags.RequiredAPI) &&
+					if (FunctionFlags.HasAnyFlags(EFunctionFlags.RequiredAPI) &&
 						outerClass != null && outerClass.ClassFlags.HasAnyFlags(EClassFlags.RequiredAPI))
 					{
 						this.LogError($"API must not be used on methods of a class that is marked with an API itself.");
 					}
 
 					// Make sure that blueprint pure functions have some type of output
-					if (this.FunctionFlags.HasAnyFlags(EFunctionFlags.BlueprintPure) && !this.HasAnyOutputs)
+					if (FunctionFlags.HasAnyFlags(EFunctionFlags.BlueprintPure) && !HasAnyOutputs)
 					{
 						this.LogError("BlueprintPure specifier is not allowed for functions with no return value and no output parameters.");
 					}
 
-					if (this.FunctionFlags.HasExactFlags(EFunctionFlags.Net | EFunctionFlags.NetRequest | EFunctionFlags.NetResponse, EFunctionFlags.Net) && this.ReturnProperty != null)
+					if (FunctionFlags.HasExactFlags(EFunctionFlags.Net | EFunctionFlags.NetRequest | EFunctionFlags.NetResponse, EFunctionFlags.Net) && ReturnProperty != null)
 					{
 						this.LogError("Replicated functions can't have return values");
 					}
 
-					checkForBlueprint = this.FunctionFlags.HasAnyFlags(EFunctionFlags.BlueprintCallable | EFunctionFlags.BlueprintEvent);
+					checkForBlueprint = FunctionFlags.HasAnyFlags(EFunctionFlags.BlueprintCallable | EFunctionFlags.BlueprintEvent);
 
-					if (this.Outer is UhtStruct outerStruct)
+					if (Outer is UhtStruct outerStruct)
 					{
 						UhtStruct? outerSuperStruct = outerStruct.SuperStruct;
 						if (outerSuperStruct != null)
 						{
-							UhtType? overriddenFunction = outerSuperStruct.FindType(UhtFindOptions.SourceName | UhtFindOptions.Function, this.SourceName);
+							UhtType? overriddenFunction = outerSuperStruct.FindType(UhtFindOptions.SourceName | UhtFindOptions.Function, SourceName);
 							if (overriddenFunction != null)
 							{
 								// Native function overrides should be done in CPP text, not in a UFUNCTION() declaration (you can't change flags, and it'd otherwise be a burden to keep them identical)
 								this.LogError(
-									$"Override of UFUNCTION '{this.SourceName}' in parent '{overriddenFunction.Outer?.SourceName}' cannot have a UFUNCTION() declaration above it; it will use the same parameters as the original declaration.");
+									$"Override of UFUNCTION '{SourceName}' in parent '{overriddenFunction.Outer?.SourceName}' cannot have a UFUNCTION() declaration above it; it will use the same parameters as the original declaration.");
 							}
 						}
 					}
 
 					// Make sure that the replication flags set on an overridden function match the parent function
-					UhtFunction? superFunction = this.SuperFunction;
+					UhtFunction? superFunction = SuperFunction;
 					if (superFunction != null)
 					{
-						if ((this.FunctionFlags & EFunctionFlags.NetFuncFlags) != (superFunction.FunctionFlags & EFunctionFlags.NetFuncFlags))
+						if ((FunctionFlags & EFunctionFlags.NetFuncFlags) != (superFunction.FunctionFlags & EFunctionFlags.NetFuncFlags))
 						{
-							this.LogError($"Overridden function '{this.SourceName}' cannot specify different replication flags when overriding a function.");
+							this.LogError($"Overridden function '{SourceName}' cannot specify different replication flags when overriding a function.");
 						}
 					}
 
 					// if this function is an RPC in state scope, verify that it is an override
 					// this is required because the networking code only checks the class for RPCs when initializing network data, not any states within it
-					if (this.FunctionFlags.HasAnyFlags(EFunctionFlags.Net) && superFunction == null && this.Outer is not UhtClass)
+					if (FunctionFlags.HasAnyFlags(EFunctionFlags.Net) && superFunction == null && Outer is not UhtClass)
 					{
-						this.LogError($"Function '{this.SourceName}' base implementation of RPCs cannot be in a state. Add a stub outside state scope.");
+						this.LogError($"Function '{SourceName}' base implementation of RPCs cannot be in a state. Add a stub outside state scope.");
 					}
 
 					// Check implemented RPC functions
-					if (this.CppImplName != this.EngineName &&
-						!this.FunctionExportFlags.HasAnyFlags(UhtFunctionExportFlags.CustomThunk) &&
+					if (CppImplName != EngineName &&
+						!FunctionExportFlags.HasAnyFlags(UhtFunctionExportFlags.CustomThunk) &&
 						outerClass!.GeneratedCodeVersion > EGeneratedCodeVersion.V1)
 					{
-						bool isNative = this.FunctionFlags.HasAnyFlags(EFunctionFlags.Native);
-						bool isNet = this.FunctionFlags.HasAnyFlags(EFunctionFlags.Net);
-						bool isNetValidate = this.FunctionFlags.HasAnyFlags(EFunctionFlags.NetValidate);
-						bool isNetResponse = this.FunctionFlags.HasAnyFlags(EFunctionFlags.NetResponse);
-						bool isBlueprintEvent = this.FunctionFlags.HasAnyFlags(EFunctionFlags.BlueprintEvent);
+						bool isNative = FunctionFlags.HasAnyFlags(EFunctionFlags.Native);
+						bool isNet = FunctionFlags.HasAnyFlags(EFunctionFlags.Net);
+						bool isNetValidate = FunctionFlags.HasAnyFlags(EFunctionFlags.NetValidate);
+						bool isNetResponse = FunctionFlags.HasAnyFlags(EFunctionFlags.NetResponse);
+						bool isBlueprintEvent = FunctionFlags.HasAnyFlags(EFunctionFlags.BlueprintEvent);
 
 						bool needsImplementation = (isNet && !isNetResponse) || isBlueprintEvent || isNative;
 						bool needsValidate = (isNative || isNet) && !isNetResponse && isNetValidate;
 
-						bool hasImplementation = this.FunctionExportFlags.HasAnyFlags(UhtFunctionExportFlags.ImplFound);
-						bool hasValidate = this.FunctionExportFlags.HasAnyFlags(UhtFunctionExportFlags.ValidationImplFound);
-						bool hasVirtualImplementation = this.FunctionExportFlags.HasAnyFlags(UhtFunctionExportFlags.ImplVirtual);
-						bool hasVirtualValidate = this.FunctionExportFlags.HasAnyFlags(UhtFunctionExportFlags.ValidationImplVirtual);
+						bool hasImplementation = FunctionExportFlags.HasAnyFlags(UhtFunctionExportFlags.ImplFound);
+						bool hasValidate = FunctionExportFlags.HasAnyFlags(UhtFunctionExportFlags.ValidationImplFound);
+						bool hasVirtualImplementation = FunctionExportFlags.HasAnyFlags(UhtFunctionExportFlags.ImplVirtual);
+						bool hasVirtualValidate = FunctionExportFlags.HasAnyFlags(UhtFunctionExportFlags.ValidationImplVirtual);
 
 						if (needsImplementation || needsValidate)
 						{
 							// Check if functions are missing.
 							if (needsImplementation && !hasImplementation)
 							{
-								LogRpcFunctionError(false, this.CppImplName);
+								LogRpcFunctionError(false, CppImplName);
 							}
 
 							if (needsValidate && !hasValidate)
 							{
-								LogRpcFunctionError(false, this.CppValidationImplName);
+								LogRpcFunctionError(false, CppValidationImplName);
 							}
 
 							// If all needed functions are declared, check if they have virtual specifiers.
 							if (needsImplementation && hasImplementation && !hasVirtualImplementation)
 							{
-								LogRpcFunctionError(true, this.CppImplName);
+								LogRpcFunctionError(true, CppImplName);
 							}
 
 							if (needsValidate && hasValidate && !hasVirtualValidate)
 							{
-								LogRpcFunctionError(true, this.CppValidationImplName);
+								LogRpcFunctionError(true, CppValidationImplName);
 							}
 						}
 					}
 
 					// Validate field notify
-					if (this.FunctionExportFlags.HasAnyFlags(UhtFunctionExportFlags.FieldNotify))
+					if (FunctionExportFlags.HasAnyFlags(UhtFunctionExportFlags.FieldNotify))
 					{
 						if (outerClass == null)
 						{
-							this.LogError($"FieldNofity function '{this.SourceName}' are only valid as UClass member function.");
+							this.LogError($"FieldNofity function '{SourceName}' are only valid as UClass member function.");
 						}
 
-						if (this.ParameterProperties.Length != 0)
+						if (ParameterProperties.Length != 0)
 						{
-							this.LogError($"FieldNotify function '{this.SourceName}' must not have parameters.");
+							this.LogError($"FieldNotify function '{SourceName}' must not have parameters.");
 						}
 
-						if (this.ReturnProperty == null)
+						if (ReturnProperty == null)
 						{
-							this.LogError($"FieldNotify function '{this.SourceName}' must return a value.");
+							this.LogError($"FieldNotify function '{SourceName}' must return a value.");
 						}
 
-						if (this.FunctionFlags.HasAnyFlags(EFunctionFlags.Event))
+						if (FunctionFlags.HasAnyFlags(EFunctionFlags.Event))
 						{
-							this.LogError($"FieldNotify function '{this.SourceName}' cannot be a blueprint event.");
+							this.LogError($"FieldNotify function '{SourceName}' cannot be a blueprint event.");
 						}
 
-						if (!this.FunctionFlags.HasAnyFlags(EFunctionFlags.BlueprintPure))
+						if (!FunctionFlags.HasAnyFlags(EFunctionFlags.BlueprintPure))
 						{
-							this.LogError($"FieldNotify function '{this.SourceName}' must be pure.");
+							this.LogError($"FieldNotify function '{SourceName}' must be pure.");
 						}
 					}
 					break;
@@ -745,24 +745,24 @@ namespace EpicGames.UHT.Types
 					// Check for shadowing if requested
 					if (options.HasAnyFlags(UhtValidationOptions.Shadowing))
 					{
-						UhtType? existing = this.Outer!.FindType(UhtFindOptions.DelegateFunction | UhtFindOptions.EngineName | UhtFindOptions.SelfOnly, this.EngineName);
+						UhtType? existing = Outer!.FindType(UhtFindOptions.DelegateFunction | UhtFindOptions.EngineName | UhtFindOptions.SelfOnly, EngineName);
 						if (existing != null && existing != this)
 						{
-							this.LogError($"Can't override delegate signature function '{this.SourceName}'");
+							this.LogError($"Can't override delegate signature function '{SourceName}'");
 						}
 					}
 					break;
 			}
 
 			// Make sure this isn't a duplicate type.  In some ways this has already been check above
-			UhtType? existingType = this.Outer!.FindType(UhtFindOptions.PropertyOrFunction | UhtFindOptions.EngineName | UhtFindOptions.SelfOnly, this.EngineName);
+			UhtType? existingType = Outer!.FindType(UhtFindOptions.PropertyOrFunction | UhtFindOptions.EngineName | UhtFindOptions.SelfOnly, EngineName);
 			if (existingType != null && existingType != this)
 			{
-				this.LogError($"'{this.EngineName}' conflicts with '{existingType.FullName}'");
+				this.LogError($"'{EngineName}' conflicts with '{existingType.FullName}'");
 			}
 
 			// Do a more generic validation of the arguments
-			foreach (UhtType type in this.Children)
+			foreach (UhtType type in Children)
 			{
 				if (type is UhtProperty property)
 				{
@@ -786,17 +786,17 @@ namespace EpicGames.UHT.Types
 					{
 						if (property.IsStaticArray)
 						{
-							this.LogError($"Static array cannot be exposed to blueprint. Function: {this.SourceName} Parameter {property.SourceName}");
+							this.LogError($"Static array cannot be exposed to blueprint. Function: {SourceName} Parameter {property.SourceName}");
 						}
 						if (!property.PropertyCaps.HasAnyFlags(UhtPropertyCaps.IsParameterSupportedByBlueprint))
 						{
-							this.LogError($"Type '{property.GetUserFacingDecl()}' is not supported by blueprint. Function:  {this.SourceName} Parameter {property.SourceName}");
+							this.LogError($"Type '{property.GetUserFacingDecl()}' is not supported by blueprint. Function:  {SourceName} Parameter {property.SourceName}");
 						}
 					}
 				}
 			}
 
-			if (!this.Deprecated)
+			if (!Deprecated)
 			{
 				options |= UhtValidationOptions.Deprecated;
 			}
@@ -817,20 +817,20 @@ namespace EpicGames.UHT.Types
 			}
 			else
 			{
-				builder.Append($"Function {this.SourceName} was marked as ");
-				if (this.FunctionFlags.HasAnyFlags(EFunctionFlags.Native))
+				builder.Append($"Function {SourceName} was marked as ");
+				if (FunctionFlags.HasAnyFlags(EFunctionFlags.Native))
 				{
 					builder.Append("Native, ");
 				}
-				if (this.FunctionFlags.HasAnyFlags(EFunctionFlags.Net))
+				if (FunctionFlags.HasAnyFlags(EFunctionFlags.Net))
 				{
 					builder.Append("Net, ");
 				}
-				if (this.FunctionFlags.HasAnyFlags(EFunctionFlags.BlueprintEvent))
+				if (FunctionFlags.HasAnyFlags(EFunctionFlags.BlueprintEvent))
 				{
 					builder.Append("BlueprintEvent, ");
 				}
-				if (this.FunctionFlags.HasAnyFlags(EFunctionFlags.NetValidate))
+				if (FunctionFlags.HasAnyFlags(EFunctionFlags.NetValidate))
 				{
 					builder.Append("NetValidate, ");
 				}
@@ -838,7 +838,7 @@ namespace EpicGames.UHT.Types
 				builder.Append(". Declare function \"virtual ");
 			}
 
-			UhtProperty? returnType = this.ReturnProperty;
+			UhtProperty? returnType = ReturnProperty;
 			if (returnType != null)
 			{
 				builder.AppendPropertyText(returnType, UhtPropertyTextType.EventFunctionArgOrRetVal);
@@ -849,10 +849,10 @@ namespace EpicGames.UHT.Types
 				builder.Append("void ");
 			}
 
-			builder.Append(this.Outer!.SourceName).Append("::").Append(functionName).Append('(');
+			builder.Append(Outer!.SourceName).Append("::").Append(functionName).Append('(');
 
 			bool first = true;
-			foreach (UhtProperty arg in this.ParameterProperties.Span)
+			foreach (UhtProperty arg in ParameterProperties.Span)
 			{
 				if (!first)
 				{
@@ -863,7 +863,7 @@ namespace EpicGames.UHT.Types
 			}
 
 			builder.Append(')');
-			if (this.FunctionFlags.HasAnyFlags(EFunctionFlags.Const))
+			if (FunctionFlags.HasAnyFlags(EFunctionFlags.Const))
 			{
 				builder.Append(" const");
 			}
@@ -881,51 +881,51 @@ namespace EpicGames.UHT.Types
 
 		private void ValidateSpecifierFlags()
 		{
-			if (this.FunctionFlags.HasAnyFlags(EFunctionFlags.Net))
+			if (FunctionFlags.HasAnyFlags(EFunctionFlags.Net))
 			{
-				if (this.FunctionFlags.HasAnyFlags(EFunctionFlags.BlueprintEvent | EFunctionFlags.Exec))
+				if (FunctionFlags.HasAnyFlags(EFunctionFlags.BlueprintEvent | EFunctionFlags.Exec))
 				{
 					this.LogError("Replicated functions can not be blueprint events or exec");
 				}
 
-				bool isNetService = this.FunctionFlags.HasAnyFlags(EFunctionFlags.NetRequest | EFunctionFlags.NetResponse);
-				bool isNetReliable = this.FunctionFlags.HasAnyFlags(EFunctionFlags.NetReliable);
+				bool isNetService = FunctionFlags.HasAnyFlags(EFunctionFlags.NetRequest | EFunctionFlags.NetResponse);
+				bool isNetReliable = FunctionFlags.HasAnyFlags(EFunctionFlags.NetReliable);
 
-				if (this.FunctionFlags.HasAnyFlags(EFunctionFlags.Static))
+				if (FunctionFlags.HasAnyFlags(EFunctionFlags.Static))
 				{
 					this.LogError("Static functions can't be replicated");
 				}
 
-				if (!isNetReliable && !this.FunctionExportFlags.HasAnyFlags(UhtFunctionExportFlags.Unreliable) && !isNetService)
+				if (!isNetReliable && !FunctionExportFlags.HasAnyFlags(UhtFunctionExportFlags.Unreliable) && !isNetService)
 				{
 					this.LogError("Replicated function: 'reliable' or 'unreliable' is required");
 				}
 
-				if (isNetReliable && this.FunctionExportFlags.HasAnyFlags(UhtFunctionExportFlags.Unreliable) && !isNetService)
+				if (isNetReliable && FunctionExportFlags.HasAnyFlags(UhtFunctionExportFlags.Unreliable) && !isNetService)
 				{
 					this.LogError("'reliable' and 'unreliable' are mutually exclusive");
 				}
 			}
-			else if (this.FunctionFlags.HasAnyFlags(EFunctionFlags.NetReliable))
+			else if (FunctionFlags.HasAnyFlags(EFunctionFlags.NetReliable))
 			{
 				this.LogError("'reliable' specified without 'client' or 'server'");
 			}
-			else if (this.FunctionExportFlags.HasAnyFlags(UhtFunctionExportFlags.Unreliable))
+			else if (FunctionExportFlags.HasAnyFlags(UhtFunctionExportFlags.Unreliable))
 			{
 				this.LogError("'unreliable' specified without 'client' or 'server'");
 			}
 
-			if (this.FunctionExportFlags.HasAnyFlags(UhtFunctionExportFlags.SealedEvent) && !this.FunctionFlags.HasAnyFlags(EFunctionFlags.Event))
+			if (FunctionExportFlags.HasAnyFlags(UhtFunctionExportFlags.SealedEvent) && !FunctionFlags.HasAnyFlags(EFunctionFlags.Event))
 			{
 				this.LogError("SealedEvent may only be used on events");
 			}
 
-			if (this.FunctionExportFlags.HasAnyFlags(UhtFunctionExportFlags.SealedEvent) && this.FunctionFlags.HasAnyFlags(EFunctionFlags.BlueprintEvent))
+			if (FunctionExportFlags.HasAnyFlags(UhtFunctionExportFlags.SealedEvent) && FunctionFlags.HasAnyFlags(EFunctionFlags.BlueprintEvent))
 			{
 				this.LogError("SealedEvent cannot be used on Blueprint events");
 			}
 
-			if (this.FunctionExportFlags.HasAnyFlags(UhtFunctionExportFlags.ForceBlueprintImpure) && this.FunctionFlags.HasAnyFlags(EFunctionFlags.BlueprintPure))
+			if (FunctionExportFlags.HasAnyFlags(UhtFunctionExportFlags.ForceBlueprintImpure) && FunctionFlags.HasAnyFlags(EFunctionFlags.BlueprintPure))
 			{
 				this.LogError("BlueprintPure (or BlueprintPure=true) and BlueprintPure=false should not both appear on the same function, they are mutually exclusive");
 			}
@@ -934,30 +934,29 @@ namespace EpicGames.UHT.Types
 		/// <inheritdoc/>
 		protected override void ValidateDocumentationPolicy(UhtDocumentationPolicy policy)
 		{
-			if (this.FunctionType != UhtFunctionType.Function)
+			if (FunctionType != UhtFunctionType.Function)
 			{
 				return;
 			}
 
-			UhtClass? outerClass = (UhtClass?)this.Outer;
+			UhtClass? outerClass = (UhtClass?)Outer;
 
 			if (policy.FunctionToolTipsRequired)
 			{
-				if (!this.MetaData.ContainsKey(UhtNames.ToolTip))
+				if (!MetaData.ContainsKey(UhtNames.ToolTip))
 				{
-					this.LogError($"Function '{outerClass?.SourceName}::{this.SourceName}' does not provide a tooltip / comment (DocumentationPolicy).");
+					this.LogError($"Function '{outerClass?.SourceName}::{SourceName}' does not provide a tooltip / comment (DocumentationPolicy).");
 				}
 			}
 
 			if (policy.ParameterToolTipsRequired)
 			{
-				if (!this.MetaData.ContainsKey(UhtNames.Comment))
+				if (!MetaData.ContainsKey(UhtNames.Comment))
 				{
-					this.LogError($"Function '{outerClass?.SourceName}::{this.SourceName}' does not provide a comment (DocumentationPolicy).");
-
+					this.LogError($"Function '{outerClass?.SourceName}::{SourceName}' does not provide a comment (DocumentationPolicy).");
 				}
 
-				Dictionary<StringView, StringView> paramToolTips = GetParameterToolTipsFromFunctionComment(this.MetaData.GetValueOrDefault(UhtNames.Comment));
+				Dictionary<StringView, StringView> paramToolTips = GetParameterToolTipsFromFunctionComment(MetaData.GetValueOrDefault(UhtNames.Comment));
 				bool hasAnyParamToolTips = paramToolTips.Count > 1 || (paramToolTips.Count == 1 && !paramToolTips.ContainsKey(UhtNames.ReturnValue));
 
 				// only apply the validation for parameter tooltips if a function has any @param statements at all.
@@ -965,12 +964,12 @@ namespace EpicGames.UHT.Types
 				{
 					// ensure each parameter has a tooltip
 					HashSet<StringView> existingFields = new(StringViewComparer.OrdinalIgnoreCase);
-					foreach (UhtType parameter in this.ParameterProperties.Span)
+					foreach (UhtType parameter in ParameterProperties.Span)
 					{
 						existingFields.Add(parameter.SourceName);
 						if (!paramToolTips.ContainsKey(parameter.SourceName))
 						{
-							this.LogError($"Function '{outerClass?.SourceName}::{this.SourceName}' doesn't provide a tooltip for parameter '{parameter.SourceName}' (DocumentationPolicy).");
+							this.LogError($"Function '{outerClass?.SourceName}::{SourceName}' doesn't provide a tooltip for parameter '{parameter.SourceName}' (DocumentationPolicy).");
 						}
 					}
 
@@ -984,7 +983,7 @@ namespace EpicGames.UHT.Types
 
 						if (!existingFields.Contains(kvp.Key))
 						{
-							this.LogError($"Function '{outerClass?.SourceName}::{this.SourceName}' provides a tooltip for an unknown parameter '{kvp.Key}'");
+							this.LogError($"Function '{outerClass?.SourceName}::{SourceName}' provides a tooltip for an unknown parameter '{kvp.Key}'");
 						}
 					}
 
@@ -999,7 +998,7 @@ namespace EpicGames.UHT.Types
 
 						if (toolTipToParam.TryGetValue(kvp.Value, out StringView existing))
 						{
-							this.LogError($"Function '{outerClass?.SourceName}::{this.SourceName}' uses identical tooltips for parameters '{kvp.Key}' and '{existing}' (DocumentationPolicy).");
+							this.LogError($"Function '{outerClass?.SourceName}::{SourceName}' uses identical tooltips for parameters '{kvp.Key}' and '{existing}' (DocumentationPolicy).");
 						}
 						else
 						{
@@ -1084,15 +1083,15 @@ namespace EpicGames.UHT.Types
 		/// <inheritdoc/>
 		public override void CollectReferences(IUhtReferenceCollector collector)
 		{
-			if (this.FunctionType != UhtFunctionType.Function)
+			if (FunctionType != UhtFunctionType.Function)
 			{
 				collector.AddExportType(this);
 				collector.AddSingleton(this);
 				collector.AddDeclaration(this, true);
 				collector.AddCrossModuleReference(this, true);
-				collector.AddCrossModuleReference(this.Package, true);
+				collector.AddCrossModuleReference(Package, true);
 			}
-			foreach (UhtType child in this.Children)
+			foreach (UhtType child in Children)
 			{
 				child.CollectReferences(collector);
 			}
@@ -1106,7 +1105,7 @@ namespace EpicGames.UHT.Types
 		/// <exception cref="UhtIceException"></exception>
 		protected internal UhtPropertyParseOptions GetPropertyParseOptions(bool returnValue)
 		{
-			switch (this.FunctionType)
+			switch (FunctionType)
 			{
 				case UhtFunctionType.Delegate:
 				case UhtFunctionType.SparseDelegate:
@@ -1115,7 +1114,7 @@ namespace EpicGames.UHT.Types
 				case UhtFunctionType.Function:
 					UhtPropertyParseOptions options = UhtPropertyParseOptions.DontAddReturn; // Fetch the function name
 					options |= returnValue ? UhtPropertyParseOptions.FunctionNameIncluded : UhtPropertyParseOptions.NameIncluded;
-					if (this.FunctionFlags.HasAllFlags(EFunctionFlags.BlueprintEvent | EFunctionFlags.Native))
+					if (FunctionFlags.HasAllFlags(EFunctionFlags.BlueprintEvent | EFunctionFlags.Native))
 					{
 						options |= UhtPropertyParseOptions.NoAutoConst;
 					}
diff --git a/Engine/Source/Programs/Shared/EpicGames.UHT/Types/UhtHeaderFile.cs b/Engine/Source/Programs/Shared/EpicGames.UHT/Types/UhtHeaderFile.cs
--- a/Engine/Source/Programs/Shared/EpicGames.UHT/Types/UhtHeaderFile.cs
+++ b/Engine/Source/Programs/Shared/EpicGames.UHT/Types/UhtHeaderFile.cs
@@ -111,13 +111,13 @@ namespace EpicGames.UHT.Types
 		/// Contents of the header
 		/// </summary>
 		[JsonIgnore]
-		public StringView Data => this._sourceFile.Data;
+		public StringView Data => _sourceFile.Data;
 
 		/// <summary>
 		/// Path of the header
 		/// </summary>
 		[JsonIgnore]
-		public string FilePath => this._sourceFile.FilePath;
+		public string FilePath => _sourceFile.FilePath;
 
 		/// <summary>
 		/// File name without the extension
@@ -164,7 +164,7 @@ namespace EpicGames.UHT.Types
 		/// If true, the header file should be exported
 		/// </summary>
 		[JsonIgnore]
-		public bool ShouldExport => this.HeaderFileExportFlags.HasAnyFlags(UhtHeaderFileExportFlags.Referenced) || this.Children.Count > 0;
+		public bool ShouldExport => HeaderFileExportFlags.HasAnyFlags(UhtHeaderFileExportFlags.Referenced) || Children.Count > 0;
 
 		/// <summary>
 		/// Resource collector for the header file
@@ -178,11 +178,11 @@ namespace EpicGames.UHT.Types
 		{
 			get
 			{
-				if (this.Outer == null)
+				if (Outer == null)
 				{
 					throw new UhtIceException("Attempt to fetch header file package but it has no outer");
 				}
-				return (UhtPackage)this.Outer;
+				return (UhtPackage)Outer;
 			}
 		}
 
@@ -207,11 +207,11 @@ namespace EpicGames.UHT.Types
 
 		/// <inheritdoc/>
 		[JsonIgnore]
-		public override IUhtMessageSession MessageSession => this._messageSite.MessageSession;
+		public override IUhtMessageSession MessageSession => _messageSite.MessageSession;
 
 		/// <inheritdoc/>
 		[JsonIgnore]
-		public override IUhtMessageSource? MessageSource => this._messageSite.MessageSource;
+		public override IUhtMessageSource? MessageSource => _messageSite.MessageSource;
 		#endregion
 
 		/// <summary>
@@ -221,14 +221,14 @@ namespace EpicGames.UHT.Types
 		/// <param name="path">Path to the header file</param>
 		public UhtHeaderFile(UhtPackage package, string path) : base(package, 1)
 		{
-			this.HeaderFileTypeIndex = this.Session.GetNextHeaderFileTypeIndex();
-			this._messageSite = new UhtSimpleMessageSite(this.Session);
-			this._sourceFile = new UhtSourceFile(this.Session, path);
-			this._messageSite.MessageSource = this._sourceFile;
-			this.FileNameWithoutExtension = System.IO.Path.GetFileNameWithoutExtension(this._sourceFile.FilePath);
-			this.GeneratedHeaderFileName = this.FileNameWithoutExtension + ".generated.h";
-			this.SourceName = System.IO.Path.GetFileName(this._sourceFile.FilePath);
-			this.IsNoExportTypes = String.Equals(this._sourceFile.FileName, "NoExportTypes", StringComparison.OrdinalIgnoreCase);
+			HeaderFileTypeIndex = Session.GetNextHeaderFileTypeIndex();
+			_messageSite = new UhtSimpleMessageSite(Session);
+			_sourceFile = new UhtSourceFile(Session, path);
+			_messageSite.MessageSource = _sourceFile;
+			FileNameWithoutExtension = System.IO.Path.GetFileNameWithoutExtension(_sourceFile.FilePath);
+			GeneratedHeaderFileName = FileNameWithoutExtension + ".generated.h";
+			SourceName = System.IO.Path.GetFileName(_sourceFile.FilePath);
+			IsNoExportTypes = String.Equals(_sourceFile.FileName, "NoExportTypes", StringComparison.OrdinalIgnoreCase);
 		}
 
 		/// <summary>
@@ -236,7 +236,7 @@ namespace EpicGames.UHT.Types
 		/// </summary>
 		public void Read()
 		{
-			this._sourceFile.Read();
+			_sourceFile.Read();
 		}
 
 		/// <summary>
@@ -246,7 +246,7 @@ namespace EpicGames.UHT.Types
 		/// <param name="isIncludedFile">True if this is a directly included file</param>
 		public void AddReferencedHeader(string id, bool isIncludedFile)
 		{
-			UhtHeaderFile? headerFile = this.Session.FindHeaderFile(Path.GetFileName(id));
+			UhtHeaderFile? headerFile = Session.FindHeaderFile(Path.GetFileName(id));
 			if (headerFile != null)
 			{
 				AddReferencedHeader(headerFile, isIncludedFile);
@@ -269,10 +269,10 @@ namespace EpicGames.UHT.Types
 		/// <param name="isIncludedFile">True if this is a directly included file</param>
 		public void AddReferencedHeader(UhtHeaderFile headerFile, bool isIncludedFile)
 		{
-			lock (this._referencedHeaders)
+			lock (_referencedHeaders)
 			{
 				// Check for a duplicate
-				foreach (UhtHeaderFile reference in this._referencedHeaders)
+				foreach (UhtHeaderFile reference in _referencedHeaders)
 				{
 					if (reference == headerFile)
 					{
@@ -282,14 +282,14 @@ namespace EpicGames.UHT.Types
 
 				// There is questionable compatibility hack where a source file will always be exported
 				// regardless of having types when it is being included by the SAME package.
-				if (headerFile.Package == this.Package)
+				if (headerFile.Package == Package)
 				{
 					headerFile.HeaderFileExportFlags |= UhtHeaderFileExportFlags.Referenced;
 				}
-				this._referencedHeaders.Add(headerFile);
+				_referencedHeaders.Add(headerFile);
 				if (isIncludedFile)
 				{
-					this.IncludedHeaders.Add(headerFile);
+					IncludedHeaders.Add(headerFile);
 				}
 			}
 		}
@@ -298,7 +298,7 @@ namespace EpicGames.UHT.Types
 		/// Return an enumerator without locking.  This method can only be utilized AFTER all header parsing is complete. 
 		/// </summary>
 		[JsonIgnore]
-		public IEnumerable<UhtHeaderFile> ReferencedHeadersNoLock => this._referencedHeaders;
+		public IEnumerable<UhtHeaderFile> ReferencedHeadersNoLock => _referencedHeaders;
 
 		/// <summary>
 		/// Return an enumerator of all current referenced headers under a lock.  This should be used during parsing.
@@ -308,9 +308,9 @@ namespace EpicGames.UHT.Types
 		{
 			get
 			{
-				lock (this._referencedHeaders)
+				lock (_referencedHeaders)
 				{
-					foreach (UhtHeaderFile reference in this._referencedHeaders)
+					foreach (UhtHeaderFile reference in _referencedHeaders)
 					{
 						yield return reference;
 					}
@@ -322,9 +322,9 @@ namespace EpicGames.UHT.Types
 		public override void AppendPathName(StringBuilder builder, UhtType? stopOuter = null)
 		{
 			// Headers do not contribute to path names
-			if (this != stopOuter && this.Outer != null)
+			if (this != stopOuter && Outer != null)
 			{
-				this.Outer.AppendPathName(builder, stopOuter);
+				Outer.AppendPathName(builder, stopOuter);
 			}
 		}
 
@@ -335,7 +335,7 @@ namespace EpicGames.UHT.Types
 
 			Dictionary<int, UhtFunction> usedRPCIds = new();
 			Dictionary<int, UhtFunction> rpcsNeedingHookup = new();
-			foreach (UhtType type in this.Children)
+			foreach (UhtType type in Children)
 			{
 				if (type is UhtClass classObj)
 				{
diff --git a/Engine/Source/Programs/Shared/EpicGames.UHT/Types/UhtManifestFile.cs b/Engine/Source/Programs/Shared/EpicGames.UHT/Types/UhtManifestFile.cs
--- a/Engine/Source/Programs/Shared/EpicGames.UHT/Types/UhtManifestFile.cs
+++ b/Engine/Source/Programs/Shared/EpicGames.UHT/Types/UhtManifestFile.cs
@@ -23,10 +23,10 @@ namespace EpicGames.UHT.Types
 		#region IUHTMessageSite implementation
 
 		/// <inheritdoc/>
-		public IUhtMessageSession MessageSession => this._sourceFile.MessageSession;
+		public IUhtMessageSession MessageSession => _sourceFile.MessageSession;
 
 		/// <inheritdoc/>
-		public IUhtMessageSource? MessageSource => this._sourceFile.MessageSource;
+		public IUhtMessageSource? MessageSource => _sourceFile.MessageSource;
 
 		/// <inheritdoc/>
 		public IUhtMessageLineNumber? MessageLineNumber => null;
@@ -39,7 +39,7 @@ namespace EpicGames.UHT.Types
 		/// <param name="filePath">Path of the file</param>
 		public UhtManifestFile(UhtSession session, string filePath)
 		{
-			this._sourceFile = new UhtSourceFile(session, filePath);
+			_sourceFile = new UhtSourceFile(session, filePath);
 		}
 
 		/// <summary>
@@ -47,8 +47,8 @@ namespace EpicGames.UHT.Types
 		/// </summary>
 		public void Read()
 		{
-			this._sourceFile.Read();
-			this.Manifest = JsonSerializer.Deserialize<UHTManifest>(this._sourceFile.Data.ToString());
+			_sourceFile.Read();
+			Manifest = JsonSerializer.Deserialize<UHTManifest>(_sourceFile.Data.ToString());
 		}
 	}
 }
diff --git a/Engine/Source/Programs/Shared/EpicGames.UHT/Types/UhtMetaData.cs b/Engine/Source/Programs/Shared/EpicGames.UHT/Types/UhtMetaData.cs
--- a/Engine/Source/Programs/Shared/EpicGames.UHT/Types/UhtMetaData.cs
+++ b/Engine/Source/Programs/Shared/EpicGames.UHT/Types/UhtMetaData.cs
@@ -48,8 +48,8 @@ namespace EpicGames.UHT.Types
 		/// <param name="index">Meta data index</param>
 		public UhtMetaDataKey(string name, int index = UhtMetaData.IndexNone)
 		{
-			this.Name = name;
-			this.Index = index;
+			Name = name;
+			Index = index;
 		}
 
 		/// <summary>
@@ -60,11 +60,11 @@ namespace EpicGames.UHT.Types
 		{
 			if (Index == UhtMetaData.IndexNone)
 			{
-				return this.Name.ToString();
+				return Name.ToString();
 			}
 			else
 			{
-				return $"{this.Name}:{this.Index}";
+				return $"{Name}:{Index}";
 			}
 		}
 	}
@@ -180,8 +180,8 @@ namespace EpicGames.UHT.Types
 		/// <param name="config">Configuration for redirects</param>
 		public UhtMetaData(IUhtMessageSite? messageSite, IUhtConfig? config)
 		{
-			this.MessageSite = messageSite;
-			this.Config = config;
+			MessageSite = messageSite;
+			Config = config;
 		}
 
 		/// <summary>
@@ -190,7 +190,7 @@ namespace EpicGames.UHT.Types
 		/// <returns>True if the meta data object contains no entries</returns>
 		public bool IsEmpty()
 		{
-			return this.Dictionary == null || this.Dictionary.Count == 0;
+			return Dictionary == null || Dictionary.Count == 0;
 		}
 
 		/// <summary>
@@ -198,9 +198,9 @@ namespace EpicGames.UHT.Types
 		/// </summary>
 		public void Clear()
 		{
-			if (this.Dictionary != null)
+			if (Dictionary != null)
 			{
-				this.Dictionary.Clear();
+				Dictionary.Clear();
 			}
 		}
 
@@ -221,11 +221,11 @@ namespace EpicGames.UHT.Types
 		/// <returns>True if the key is found</returns>
 		public bool ContainsKey(string name, int nameIndex = IndexNone)
 		{
-			if (this.Dictionary == null)
+			if (Dictionary == null)
 			{
 				return false;
 			}
-			return this.Dictionary.ContainsKey(new UhtMetaDataKey(name, nameIndex));
+			return Dictionary.ContainsKey(new UhtMetaDataKey(name, nameIndex));
 		}
 
 		/// <summary>
@@ -277,12 +277,12 @@ namespace EpicGames.UHT.Types
 		/// <returns>True if the key is found</returns>
 		public bool TryGetValue(string name, int nameIndex, [NotNullWhen(true)] out string? value)
 		{
-			if (this.Dictionary == null)
+			if (Dictionary == null)
 			{
 				value = String.Empty;
 				return false;
 			}
-			return this.Dictionary.TryGetValue(new UhtMetaDataKey(name, nameIndex), out value);
+			return Dictionary.TryGetValue(new UhtMetaDataKey(name, nameIndex), out value);
 		}
 
 		/// <summary>
@@ -421,13 +421,13 @@ namespace EpicGames.UHT.Types
 		/// <param name="value">Value of the meta data</param>
 		public void Add(string name, int nameIndex, string value)
 		{
-			if (this.Config != null)
+			if (Config != null)
 			{
-				if (this.Config.RedirectMetaDataKey(name, out string remappedName))
+				if (Config.RedirectMetaDataKey(name, out string remappedName))
 				{
-					if (this.MessageSite != null)
+					if (MessageSite != null)
 					{
-						this.MessageSite.LogWarning(this.LineNumber, $"Remapping old metadata key '{name}' to new key '{remappedName}', please update the declaration.");
+						MessageSite.LogWarning(LineNumber, $"Remapping old metadata key '{name}' to new key '{remappedName}', please update the declaration.");
 					}
 				}
 			}
@@ -493,9 +493,9 @@ namespace EpicGames.UHT.Types
 		/// <param name="nameIndex">Index of the meta data key</param>
 		public void Remove(string name, int nameIndex = IndexNone)
 		{
-			if (this.Dictionary != null)
+			if (Dictionary != null)
 			{
-				this.Dictionary.Remove(new UhtMetaDataKey(name, nameIndex));
+				Dictionary.Remove(new UhtMetaDataKey(name, nameIndex));
 			}
 		}
 
@@ -538,12 +538,12 @@ namespace EpicGames.UHT.Types
 		/// <returns>List of meta data key and value pairs</returns>
 		public List<KeyValuePair<string, string>> GetSorted()
 		{
-			List<KeyValuePair<string, string>> output = new(this.Dictionary != null ? this.Dictionary.Count : 0);
-			if (this.Dictionary != null && this.Dictionary.Count > 0)
+			List<KeyValuePair<string, string>> output = new(Dictionary != null ? Dictionary.Count : 0);
+			if (Dictionary != null && Dictionary.Count > 0)
 			{
-				foreach (KeyValuePair<UhtMetaDataKey, string> kvp in this.Dictionary)
+				foreach (KeyValuePair<UhtMetaDataKey, string> kvp in Dictionary)
 				{
-					output.Add(new KeyValuePair<string, string>(this.GetKeyString(kvp.Key), kvp.Value));
+					output.Add(new KeyValuePair<string, string>(GetKeyString(kvp.Key), kvp.Value));
 				}
 
 				output.Sort((KeyValuePair<string, string> lhs, KeyValuePair<string, string> rhs) =>
@@ -556,11 +556,11 @@ namespace EpicGames.UHT.Types
 
 		private SortedList<UhtMetaDataKey, string> GetDictionary()
 		{
-			if (this.Dictionary == null)
+			if (Dictionary == null)
 			{
-				this.Dictionary = new SortedList<UhtMetaDataKey, string>(s_comparer);
+				Dictionary = new SortedList<UhtMetaDataKey, string>(s_comparer);
 			}
-			return this.Dictionary;
+			return Dictionary;
 		}
 	}
 }
diff --git a/Engine/Source/Programs/Shared/EpicGames.UHT/Types/UhtObject.cs b/Engine/Source/Programs/Shared/EpicGames.UHT/Types/UhtObject.cs
--- a/Engine/Source/Programs/Shared/EpicGames.UHT/Types/UhtObject.cs
+++ b/Engine/Source/Programs/Shared/EpicGames.UHT/Types/UhtObject.cs
@@ -38,7 +38,7 @@ namespace EpicGames.UHT.Types
 		/// <param name="session">Session of the object</param>
 		protected UhtObject(UhtSession session) : base(session)
 		{
-			ObjectTypeIndex = this.Session.GetNextObjectTypeIndex();
+			ObjectTypeIndex = Session.GetNextObjectTypeIndex();
 		}
 
 		/// <summary>
@@ -48,7 +48,7 @@ namespace EpicGames.UHT.Types
 		/// <param name="lineNumber">Line number where object is defined</param>
 		protected UhtObject(UhtType outer, int lineNumber) : base(outer, lineNumber)
 		{
-			ObjectTypeIndex = this.Session.GetNextObjectTypeIndex();
+			ObjectTypeIndex = Session.GetNextObjectTypeIndex();
 		}
 	}
 }
diff --git a/Engine/Source/Programs/Shared/EpicGames.UHT/Types/UhtPackage.cs b/Engine/Source/Programs/Shared/EpicGames.UHT/Types/UhtPackage.cs
--- a/Engine/Source/Programs/Shared/EpicGames.UHT/Types/UhtPackage.cs
+++ b/Engine/Source/Programs/Shared/EpicGames.UHT/Types/UhtPackage.cs
@@ -55,10 +55,10 @@ namespace EpicGames.UHT.Types
 		{
 			get
 			{
-				switch (this.Module.ModuleType)
+				switch (Module.ModuleType)
 				{
 					case UHTModuleType.Program:
-						return this.Module.BaseDirectory.Replace('\\', '/').StartsWith(Unreal.EngineDirectory.FullName.Replace('\\', '/'), StringComparison.Ordinal);
+						return Module.BaseDirectory.Replace('\\', '/').StartsWith(Unreal.EngineDirectory.FullName.Replace('\\', '/'), StringComparison.Ordinal);
 					case UHTModuleType.EngineRuntime:
 					case UHTModuleType.EngineUncooked:
 					case UHTModuleType.EngineDeveloper:
@@ -81,7 +81,7 @@ namespace EpicGames.UHT.Types
 		/// True if the package is a plugin
 		/// </summary>
 		[JsonIgnore]
-		public bool IsPlugin => this.Module.BaseDirectory.Replace('\\', '/').Contains("/Plugins/", StringComparison.Ordinal);
+		public bool IsPlugin => Module.BaseDirectory.Replace('\\', '/').Contains("/Plugins/", StringComparison.Ordinal);
 
 		/// <summary>
 		/// Short name of the package (without the /Script/)
@@ -96,20 +96,20 @@ namespace EpicGames.UHT.Types
 		/// <param name="packageFlags">Assorted package flags</param>
 		public UhtPackage(UhtSession session, UHTManifest.Module module, EPackageFlags packageFlags) : base(session)
 		{
-			this.Module = module;
-			this.PackageFlags = packageFlags;
-			this.PackageTypeIndex = this.Session.GetNextPackageTypeIndex();
+			Module = module;
+			PackageFlags = packageFlags;
+			PackageTypeIndex = Session.GetNextPackageTypeIndex();
 
 			int lastSlashIndex = module.Name.LastIndexOf('/');
 			if (lastSlashIndex == -1)
 			{
-				this.SourceName = $"/Script/{module.Name}";
-				this.ShortName = module.Name;
+				SourceName = $"/Script/{module.Name}";
+				ShortName = module.Name;
 			}
 			else
 			{
-				this.SourceName = module.Name;
-				this.ShortName = this.SourceName[(lastSlashIndex + 1)..];
+				SourceName = module.Name;
+				ShortName = SourceName[(lastSlashIndex + 1)..];
 			}
 		}
 	}
