Removed the use of "this." where possible in EpicGames.UHT.

diff --git a/Engine/Source/Programs/Shared/EpicGames.UHT/Types/UhtProperty.cs b/Engine/Source/Programs/Shared/EpicGames.UHT/Types/UhtProperty.cs
--- a/Engine/Source/Programs/Shared/EpicGames.UHT/Types/UhtProperty.cs
+++ b/Engine/Source/Programs/Shared/EpicGames.UHT/Types/UhtProperty.cs
@@ -736,24 +736,24 @@ namespace EpicGames.UHT.Types
 		/// <param name="messageSite">Message site used to construct meta data object</param>
 		public void Reset(UhtPropertySettings parentPropertySettings, string sourceName, IUhtMessageSite messageSite)
 		{
-			this.SourceName = sourceName;
-			this.EngineName = sourceName;
-			this.MetaData = new UhtMetaData(messageSite, parentPropertySettings.Outer.Session.Config);
-			this.Outer = parentPropertySettings.Outer;
-			this.LineNumber = parentPropertySettings.LineNumber;
-			this.PropertyCategory = parentPropertySettings.PropertyCategory;
-			this.PropertyFlags = parentPropertySettings.PropertyFlags;
-			this.DisallowPropertyFlags = parentPropertySettings.DisallowPropertyFlags;
-			this.PropertyExportFlags = UhtPropertyExportFlags.Public;
-			this.RepNotifyName = null;
-			this.Allocator = UhtPropertyAllocator.Default;
-			this.Options = parentPropertySettings.Options;
-			this.PointerType = UhtPointerType.None;
-			this.ArrayDimensions = null;
-			this.DefaultValueTokens = null;
-			this.Setter = null;
-			this.Getter = null;
-			this.IsBitfield = false;
+			SourceName = sourceName;
+			EngineName = sourceName;
+			MetaData = new UhtMetaData(messageSite, parentPropertySettings.Outer.Session.Config);
+			Outer = parentPropertySettings.Outer;
+			LineNumber = parentPropertySettings.LineNumber;
+			PropertyCategory = parentPropertySettings.PropertyCategory;
+			PropertyFlags = parentPropertySettings.PropertyFlags;
+			DisallowPropertyFlags = parentPropertySettings.DisallowPropertyFlags;
+			PropertyExportFlags = UhtPropertyExportFlags.Public;
+			RepNotifyName = null;
+			Allocator = UhtPropertyAllocator.Default;
+			Options = parentPropertySettings.Options;
+			PointerType = UhtPointerType.None;
+			ArrayDimensions = null;
+			DefaultValueTokens = null;
+			Setter = null;
+			Getter = null;
+			IsBitfield = false;
 		}
 
 		/// <summary>
@@ -765,23 +765,23 @@ namespace EpicGames.UHT.Types
 		/// <param name="disallowPropertyFlags">Property flags that are not allowed</param>
 		public void Reset(UhtType outer, int lineNumber, UhtPropertyCategory propertyCategory, EPropertyFlags disallowPropertyFlags)
 		{
-			this.SourceName = String.Empty;
-			this.EngineName = String.Empty;
-			this.MetaData = new UhtMetaData(outer, outer.Session.Config);
-			this.Outer = outer;
-			this.LineNumber = lineNumber;
-			this.PropertyCategory = propertyCategory;
-			this.PropertyFlags = EPropertyFlags.None;
-			this.DisallowPropertyFlags = disallowPropertyFlags;
-			this.PropertyExportFlags = UhtPropertyExportFlags.Public;
-			this.RepNotifyName = null;
-			this.Allocator = UhtPropertyAllocator.Default;
-			this.PointerType = UhtPointerType.None;
-			this.ArrayDimensions = null;
-			this.DefaultValueTokens = null;
-			this.Setter = null;
-			this.Getter = null;
-			this.IsBitfield = false;
+			SourceName = String.Empty;
+			EngineName = String.Empty;
+			MetaData = new UhtMetaData(outer, outer.Session.Config);
+			Outer = outer;
+			LineNumber = lineNumber;
+			PropertyCategory = propertyCategory;
+			PropertyFlags = EPropertyFlags.None;
+			DisallowPropertyFlags = disallowPropertyFlags;
+			PropertyExportFlags = UhtPropertyExportFlags.Public;
+			RepNotifyName = null;
+			Allocator = UhtPropertyAllocator.Default;
+			PointerType = UhtPointerType.None;
+			ArrayDimensions = null;
+			DefaultValueTokens = null;
+			Setter = null;
+			Getter = null;
+			IsBitfield = false;
 		}
 
 		/// <summary>
@@ -796,24 +796,24 @@ namespace EpicGames.UHT.Types
 			{
 				throw new UhtIceException("Property must have an outer specified");
 			}
-			this.SourceName = property.SourceName;
-			this.EngineName = property.EngineName;
-			this.MetaData = property.MetaData;
-			this.Outer = property.Outer;
-			this.LineNumber = property.LineNumber;
-			this.PropertyCategory = property.PropertyCategory;
-			this.PropertyFlags = property.PropertyFlags;
-			this.DisallowPropertyFlags = property.DisallowPropertyFlags;
-			this.PropertyExportFlags = property.PropertyExportFlags;
-			this.Allocator = property.Allocator;
-			this.Options = options;
-			this.PointerType = property.PointerType;
-			this.RepNotifyName = property.RepNotifyName;
-			this.ArrayDimensions = property.ArrayDimensions;
-			this.DefaultValueTokens = property.DefaultValueTokens;
-			this.Setter = property.Setter;
-			this.Getter = property.Getter;
-			this.IsBitfield = property.IsBitfield;
+			SourceName = property.SourceName;
+			EngineName = property.EngineName;
+			MetaData = property.MetaData;
+			Outer = property.Outer;
+			LineNumber = property.LineNumber;
+			PropertyCategory = property.PropertyCategory;
+			PropertyFlags = property.PropertyFlags;
+			DisallowPropertyFlags = property.DisallowPropertyFlags;
+			PropertyExportFlags = property.PropertyExportFlags;
+			Allocator = property.Allocator;
+			Options = options;
+			PointerType = property.PointerType;
+			RepNotifyName = property.RepNotifyName;
+			ArrayDimensions = property.ArrayDimensions;
+			DefaultValueTokens = property.DefaultValueTokens;
+			Setter = property.Setter;
+			Getter = property.Getter;
+			IsBitfield = property.IsBitfield;
 		}
 	}
 
@@ -962,21 +962,21 @@ namespace EpicGames.UHT.Types
 
 		///<inheritdoc/>
 		[JsonIgnore]
-		public override bool Deprecated => this.PropertyFlags.HasAnyFlags(EPropertyFlags.Deprecated);
+		public override bool Deprecated => PropertyFlags.HasAnyFlags(EPropertyFlags.Deprecated);
 
 		///<inheritdoc/>
 		protected override UhtSpecifierValidatorTable? SpecifierValidatorTable
 		{
 			get
 			{
-				switch (this.PropertyCategory)
+				switch (PropertyCategory)
 				{
 					case UhtPropertyCategory.Member:
-						return this.Session.GetSpecifierValidatorTable(UhtTableNames.PropertyMember);
+						return Session.GetSpecifierValidatorTable(UhtTableNames.PropertyMember);
 					case UhtPropertyCategory.RegularParameter:
 					case UhtPropertyCategory.ReplicatedParameter:
 					case UhtPropertyCategory.Return:
-						return this.Session.GetSpecifierValidatorTable(UhtTableNames.PropertyArgument);
+						return Session.GetSpecifierValidatorTable(UhtTableNames.PropertyArgument);
 					default:
 						throw new UhtIceException("Unknown property category type");
 				}
@@ -987,13 +987,13 @@ namespace EpicGames.UHT.Types
 		/// If true, the property is a fixed/static array
 		/// </summary>
 		[JsonIgnore]
-		public bool IsStaticArray => !String.IsNullOrEmpty(this.ArrayDimensions);
+		public bool IsStaticArray => !String.IsNullOrEmpty(ArrayDimensions);
 
 		/// <summary>
 		/// If true, the property is editor only
 		/// </summary>
 		[JsonIgnore]
-		public bool IsEditorOnlyProperty => this.PropertyFlags.HasAnyFlags(EPropertyFlags.DevelopmentAssets);
+		public bool IsEditorOnlyProperty => PropertyFlags.HasAnyFlags(EPropertyFlags.DevelopmentAssets);
 
 		/// <summary>
 		/// Construct a new property
@@ -1002,8 +1002,8 @@ namespace EpicGames.UHT.Types
 		/// <param name="lineNumber">Line number where property was declared</param>
 		protected UhtProperty(UhtType outer, int lineNumber) : base(outer, lineNumber)
 		{
-			this.PropertyFlags = EPropertyFlags.None;
-			this.PropertyCaps = UhtPropertyCaps.CanBeContainerValue | UhtPropertyCaps.CanBeContainerKey | UhtPropertyCaps.CanHaveConfig;
+			PropertyFlags = EPropertyFlags.None;
+			PropertyCaps = UhtPropertyCaps.CanBeContainerValue | UhtPropertyCaps.CanBeContainerKey | UhtPropertyCaps.CanHaveConfig;
 		}
 
 		/// <summary>
@@ -1012,31 +1012,31 @@ namespace EpicGames.UHT.Types
 		/// <param name="propertySettings">Property settings from parsing</param>
 		protected UhtProperty(UhtPropertySettings propertySettings) : base(propertySettings.Outer, propertySettings.LineNumber, propertySettings.MetaData)
 		{
-			this.SourceName = propertySettings.SourceName;
+			SourceName = propertySettings.SourceName;
 			// Engine name defaults to source name.  If it doesn't match what is coming in, then set it.
 			if (propertySettings.EngineName.Length > 0 && propertySettings.SourceName != propertySettings.EngineName)
 			{
-				this.EngineName = propertySettings.EngineName;
+				EngineName = propertySettings.EngineName;
 			}
-			this.PropertyCategory = propertySettings.PropertyCategory;
-			this.PropertyFlags = propertySettings.PropertyFlags;
-			this.DisallowPropertyFlags = propertySettings.DisallowPropertyFlags;
-			this.PropertyExportFlags = propertySettings.PropertyExportFlags;
-			this.Allocator = propertySettings.Allocator;
-			this.PointerType = propertySettings.PointerType;
-			this.RepNotifyName = propertySettings.RepNotifyName;
-			this.ArrayDimensions = propertySettings.ArrayDimensions;
-			this.DefaultValueTokens = propertySettings.DefaultValueTokens;
-			this.Getter = propertySettings.Getter;
-			this.Setter = propertySettings.Setter;
-			this.IsBitfield = propertySettings.IsBitfield;
-			this.PropertyCaps = UhtPropertyCaps.CanBeContainerValue | UhtPropertyCaps.CanBeContainerKey | UhtPropertyCaps.CanHaveConfig;
+			PropertyCategory = propertySettings.PropertyCategory;
+			PropertyFlags = propertySettings.PropertyFlags;
+			DisallowPropertyFlags = propertySettings.DisallowPropertyFlags;
+			PropertyExportFlags = propertySettings.PropertyExportFlags;
+			Allocator = propertySettings.Allocator;
+			PointerType = propertySettings.PointerType;
+			RepNotifyName = propertySettings.RepNotifyName;
+			ArrayDimensions = propertySettings.ArrayDimensions;
+			DefaultValueTokens = propertySettings.DefaultValueTokens;
+			Getter = propertySettings.Getter;
+			Setter = propertySettings.Setter;
+			IsBitfield = propertySettings.IsBitfield;
+			PropertyCaps = UhtPropertyCaps.CanBeContainerValue | UhtPropertyCaps.CanBeContainerKey | UhtPropertyCaps.CanHaveConfig;
 		}
 
 		/// <inheritdoc/>
 		protected override string GetDisplayStringFromEngineName()
 		{
-			return UhtFCString.NameToDisplayString(this.EngineName, this is UhtBoolProperty);
+			return UhtFCString.NameToDisplayString(EngineName, this is UhtBoolProperty);
 		}
 
 		#region Text and code generation support
@@ -1050,7 +1050,7 @@ namespace EpicGames.UHT.Types
 		public virtual StringBuilder AppendText(StringBuilder builder, UhtPropertyTextType textType, bool isTemplateArgument = false)
 		{
 			// By default, we assume it will be just the simple text
-			return builder.Append(this.CppTypeText);
+			return builder.Append(CppTypeText);
 		}
 
 		/// <summary>
@@ -1062,7 +1062,7 @@ namespace EpicGames.UHT.Types
 		/// <returns>Destination builder</returns>
 		public virtual StringBuilder AppendFullDecl(StringBuilder builder, UhtPropertyTextType textType, bool skipParameterName = false)
 		{
-			UhtPropertyCaps caps = this.PropertyCaps;
+			UhtPropertyCaps caps = PropertyCaps;
 
 			bool isParameter = textType.IsParameter();
 			bool isInterfaceProp = this is UhtInterfaceProperty;
@@ -1074,14 +1074,14 @@ namespace EpicGames.UHT.Types
 			//		AND - if this is a const param OR (if this is an interface property and not an out param)
 			// 3) If this is a parameter without array dimensions, must be passed by reference, but not an out parameter or const param
 			bool passCppArgsByRef = caps.HasAnyFlags(UhtPropertyCaps.PassCppArgsByRef);
-			bool isConstParam = isParameter && (this.PropertyFlags.HasAnyFlags(EPropertyFlags.ConstParm) || (isInterfaceProp && !this.PropertyFlags.HasAnyFlags(EPropertyFlags.OutParm)));
-			bool isConstArgsByRef = isParameter && this.ArrayDimensions == null && passCppArgsByRef && !this.PropertyFlags.HasAnyFlags(EPropertyFlags.ConstParm | EPropertyFlags.OutParm);
+			bool isConstParam = isParameter && (PropertyFlags.HasAnyFlags(EPropertyFlags.ConstParm) || (isInterfaceProp && !PropertyFlags.HasAnyFlags(EPropertyFlags.OutParm)));
+			bool isConstArgsByRef = isParameter && ArrayDimensions == null && passCppArgsByRef && !PropertyFlags.HasAnyFlags(EPropertyFlags.ConstParm | EPropertyFlags.OutParm);
 			bool isOnConstClass = false;
 			if (this is UhtObjectProperty objectProperty)
 			{
 				isOnConstClass = objectProperty.Class.ClassFlags.HasAnyFlags(EClassFlags.Const);
 			}
-			bool shouldHaveRef = this.PropertyFlags.HasAnyFlags(EPropertyFlags.OutParm | EPropertyFlags.ReferenceParm);
+			bool shouldHaveRef = PropertyFlags.HasAnyFlags(EPropertyFlags.OutParm | EPropertyFlags.ReferenceParm);
 
 			bool constAtTheBeginning = isOnConstClass || isConstArgsByRef || (isConstParam && !shouldHaveRef);
 			if (constAtTheBeginning)
@@ -1089,12 +1089,12 @@ namespace EpicGames.UHT.Types
 				builder.Append("const ");
 			}
 
-			this.AppendText(builder, textType);
+			AppendText(builder, textType);
 
 			bool fromConstClass = false;
-			if (textType == UhtPropertyTextType.ExportMember && this.Outer != null)
+			if (textType == UhtPropertyTextType.ExportMember && Outer != null)
 			{
-				if (this.Outer is UhtClass outerClass)
+				if (Outer is UhtClass outerClass)
 				{
 					fromConstClass = outerClass.ClassFlags.HasAnyFlags(EClassFlags.Const);
 				}
@@ -1105,7 +1105,7 @@ namespace EpicGames.UHT.Types
 				builder.Append(" const");
 			}
 
-			if (isParameter && this.ArrayDimensions == null && (passCppArgsByRef || this.PropertyFlags.HasAnyFlags(EPropertyFlags.OutParm | EPropertyFlags.ReferenceParm)))
+			if (isParameter && ArrayDimensions == null && (passCppArgsByRef || PropertyFlags.HasAnyFlags(EPropertyFlags.OutParm | EPropertyFlags.ReferenceParm)))
 			{
 				builder.Append('&');
 			}
@@ -1113,12 +1113,12 @@ namespace EpicGames.UHT.Types
 			builder.Append(' ');
 			if (!skipParameterName)
 			{
-				builder.Append(this.SourceName);
+				builder.Append(SourceName);
 			}
 
-			if (this.ArrayDimensions != null)
+			if (ArrayDimensions != null)
 			{
-				builder.Append('[').Append(this.ArrayDimensions).Append(']');
+				builder.Append('[').Append(ArrayDimensions).Append(']');
 			}
 			return builder;
 		}
@@ -1192,25 +1192,25 @@ namespace EpicGames.UHT.Types
 				.AppendTabs(tabs)
 				.Append("const UECodeGen_Private::").Append(paramsStructName).Append(' ')
 				.AppendNameDef(context, name, nameSuffix).Append(" = { ")
-				.AppendUTF8LiteralString(this.EngineName).Append(", ")
+				.AppendUTF8LiteralString(EngineName).Append(", ")
 				.AppendNotifyFunc(this).Append(", ")
-				.AppendFlags(this.PropertyFlags).Append(", ")
+				.AppendFlags(PropertyFlags).Append(", ")
 				.Append(paramsGenFlags).Append(", ")
 				.Append(ObjectFlags).Append(", ")
 				.AppendArrayDim(this, context).Append(", ");
 
-			if (this.PropertyExportFlags.HasAnyFlags(UhtPropertyExportFlags.SetterFound))
+			if (PropertyExportFlags.HasAnyFlags(UhtPropertyExportFlags.SetterFound))
 			{
-				builder.Append('&').Append(this.Outer!.SourceName).Append("::").AppendPropertySetterWrapperName(this).Append(", ");
+				builder.Append('&').Append(Outer!.SourceName).Append("::").AppendPropertySetterWrapperName(this).Append(", ");
 			}
 			else
 			{
 				builder.Append("nullptr, ");
 			}
 
-			if (this.PropertyExportFlags.HasAnyFlags(UhtPropertyExportFlags.GetterFound))
+			if (PropertyExportFlags.HasAnyFlags(UhtPropertyExportFlags.GetterFound))
 			{
-				builder.Append('&').Append(this.Outer!.SourceName).Append("::").AppendPropertyGetterWrapperName(this).Append(", ");
+				builder.Append('&').Append(Outer!.SourceName).Append("::").AppendPropertyGetterWrapperName(this).Append(", ");
 			}
 			else
 			{
@@ -1225,7 +1225,7 @@ namespace EpicGames.UHT.Types
 				}
 				else
 				{
-					builder.Append("STRUCT_OFFSET(").Append(context.OuterStructSourceName).Append(", ").Append(this.SourceName).Append("), ");
+					builder.Append("STRUCT_OFFSET(").Append(context.OuterStructSourceName).Append(", ").Append(SourceName).Append("), ");
 				}
 			}
 			return builder;
@@ -1294,17 +1294,17 @@ namespace EpicGames.UHT.Types
 		public virtual StringBuilder AppendFunctionThunkParameterGet(StringBuilder builder)
 		{
 			builder.Append("P_GET_");
-			if (this.ArrayDimensions != null)
+			if (ArrayDimensions != null)
 			{
 				builder.Append("ARRAY");
 			}
 			else
 			{
-				builder.Append(this.PGetMacroText);
+				builder.Append(PGetMacroText);
 			}
-			if (this.PropertyFlags.HasAnyFlags(EPropertyFlags.OutParm))
+			if (PropertyFlags.HasAnyFlags(EPropertyFlags.OutParm))
 			{
-				if (this.PGetPassAsNoPtr)
+				if (PGetPassAsNoPtr)
 				{
 					builder.Append("_REF_NO_PTR");
 				}
@@ -1314,19 +1314,19 @@ namespace EpicGames.UHT.Types
 				}
 			}
 			builder.Append('(');
-			if (this.ArrayDimensions != null)
+			if (ArrayDimensions != null)
 			{
 				builder.AppendFunctionThunkParameterArrayType(this).Append(',');
 			}
 			else
 			{
-				switch (this.PGetTypeArgument)
+				switch (PGetTypeArgument)
 				{
 					case UhtPGetArgumentType.None:
 						break;
 
 					case UhtPGetArgumentType.EngineClass:
-						builder.Append('F').Append(this.EngineClassName).Append(',');
+						builder.Append('F').Append(EngineClassName).Append(',');
 						break;
 
 					case UhtPGetArgumentType.TypeText:
@@ -1356,11 +1356,11 @@ namespace EpicGames.UHT.Types
 		public StringBuilder AppendFunctionThunkParameterName(StringBuilder builder)
 		{
 			builder.Append(FunctionParameterThunkPrefix);
-			if (this.PropertyFlags.HasAnyFlags(EPropertyFlags.OutParm))
+			if (PropertyFlags.HasAnyFlags(EPropertyFlags.OutParm))
 			{
 				builder.Append("Out_");
 			}
-			builder.Append(this.EngineName);
+			builder.Append(EngineName);
 			return builder;
 		}
 
@@ -1426,9 +1426,9 @@ namespace EpicGames.UHT.Types
 			switch (phase)
 			{
 				case UhtResolvePhase.Final:
-					if (this.ArrayDimensions != null)
+					if (ArrayDimensions != null)
 					{
-						ReadOnlySpan<char> dim = this.ArrayDimensions.AsSpan();
+						ReadOnlySpan<char> dim = ArrayDimensions.AsSpan();
 
 						bool again;
 						do
@@ -1478,7 +1478,7 @@ namespace EpicGames.UHT.Types
 						} while (again && dim.Length > 0);
 
 						//COMPATIBILITY-TODO - This method is more robust, but causes differences.  See UhtSession for future
-						// plans on fix in this.
+						// plans on fix in 
 						//// Now try to see if this is an enum
 						//UhtEnum? Enum = null;
 						//int Sep = Dim.IndexOf("::");
@@ -1493,27 +1493,27 @@ namespace EpicGames.UHT.Types
 						//}
 						//else
 						//{
-						//	Enum = this.Session.FindRegularEnumValue(Dim.ToString());
+						//	Enum = Session.FindRegularEnumValue(Dim.ToString());
 						//}
 						//if (Enum == null)
 						//{
-						//	Enum = this.Session.FindType(this.Outer, UhtFindOptions.Enum | UhtFindOptions.SourceName, Dim.ToString()) as UhtEnum;
+						//	Enum = Session.FindType(Outer, UhtFindOptions.Enum | UhtFindOptions.SourceName, Dim.ToString()) as UhtEnum;
 						//}
 						//if (Enum != null)
 						//{
-						//	this.MetaData.Add(UhtNames.ArraySizeEnum, Enum.GetPathName());
+						//	MetaData.Add(UhtNames.ArraySizeEnum, Enum.GetPathName());
 						//}
 
 						if (dim.Length > 0 && !UhtFCString.IsDigit(dim[0]))
 						{
-							UhtEnum? enumObj = this.Session.FindRegularEnumValue(dim.ToString());
+							UhtEnum? enumObj = Session.FindRegularEnumValue(dim.ToString());
 							if (enumObj == null)
 							{
-								enumObj = this.Session.FindType(this.Outer, UhtFindOptions.Enum | UhtFindOptions.SourceName, dim.ToString()) as UhtEnum;
+								enumObj = Session.FindType(Outer, UhtFindOptions.Enum | UhtFindOptions.SourceName, dim.ToString()) as UhtEnum;
 							}
 							if (enumObj != null)
 							{
-								this.MetaData.Add(UhtNames.ArraySizeEnum, enumObj.PathName);
+								MetaData.Add(UhtNames.ArraySizeEnum, enumObj.PathName);
 							}
 						}
 					}
@@ -1541,7 +1541,7 @@ namespace EpicGames.UHT.Types
 			options = base.Validate(options);
 
 			// The outer should never be null...
-			if (this.Outer == null)
+			if (Outer == null)
 			{
 				return options;
 			}
@@ -1551,14 +1551,14 @@ namespace EpicGames.UHT.Types
 			// In the old code if you defined the property after the function with argument with the same name, UHT would
 			// not issue an error.  However, if the property was defined PRIOR to the function with the matching argument name,
 			// UHT would generate an error.
-			if (options.HasAnyFlags(UhtValidationOptions.Shadowing) && this.PropertyCategory == UhtPropertyCategory.Member)
+			if (options.HasAnyFlags(UhtValidationOptions.Shadowing) && PropertyCategory == UhtPropertyCategory.Member)
 			{
 
 				// First check for duplicate name in self and then duplicate name in parents
-				UhtType? existing = this.Outer.FindType(UhtFindOptions.PropertyOrFunction | UhtFindOptions.EngineName | UhtFindOptions.SelfOnly, this.EngineName);
+				UhtType? existing = Outer.FindType(UhtFindOptions.PropertyOrFunction | UhtFindOptions.EngineName | UhtFindOptions.SelfOnly, EngineName);
 				if (existing == this)
 				{
-					existing = this.Outer.FindType(UhtFindOptions.PropertyOrFunction | UhtFindOptions.EngineName | UhtFindOptions.ParentsOnly | UhtFindOptions.NoGlobal | UhtFindOptions.NoIncludes, this.EngineName);
+					existing = Outer.FindType(UhtFindOptions.PropertyOrFunction | UhtFindOptions.EngineName | UhtFindOptions.ParentsOnly | UhtFindOptions.NoGlobal | UhtFindOptions.NoIncludes, EngineName);
 				}
 
 				if (existing != null && existing != this)
@@ -1568,29 +1568,29 @@ namespace EpicGames.UHT.Types
 						//@TODO: This exception does not seem sound either, but there is enough existing code that it will need to be
 						// fixed up first before the exception it is removed.
 						bool existingPropDeprecated = existingProperty.PropertyFlags.HasAnyFlags(EPropertyFlags.Deprecated);
-						bool newPropDeprecated = this.PropertyCategory == UhtPropertyCategory.Member && this.PropertyFlags.HasAnyFlags(EPropertyFlags.Deprecated);
+						bool newPropDeprecated = PropertyCategory == UhtPropertyCategory.Member && PropertyFlags.HasAnyFlags(EPropertyFlags.Deprecated);
 						if (!newPropDeprecated && !existingPropDeprecated)
 						{
-							this.LogShadowingError(existingProperty);
+							LogShadowingError(existingProperty);
 						}
 					}
 					else if (existing is UhtFunction existingFunction)
 					{
-						if (this.PropertyCategory == UhtPropertyCategory.Member)
+						if (PropertyCategory == UhtPropertyCategory.Member)
 						{
-							this.LogShadowingError(existingFunction);
+							LogShadowingError(existingFunction);
 						}
 					}
 				}
 			}
 
-			Validate((UhtStruct)this.Outer, this, options);
+			Validate((UhtStruct)Outer, this, options);
 			return options;
 		}
 
 		private void LogShadowingError(UhtType shadows)
 		{
-			this.LogError($"{this.PropertyCategory.GetHintText()}: '{this.SourceName}' cannot be defined in '{this.Outer?.SourceName}' as it is already defined in scope '{shadows.Outer?.SourceName}' (shadowing is not allowed)");
+			this.LogError($"{PropertyCategory.GetHintText()}: '{SourceName}' cannot be defined in '{Outer?.SourceName}' as it is already defined in scope '{shadows.Outer?.SourceName}' (shadowing is not allowed)");
 		}
 
 		/// <summary>
@@ -1603,15 +1603,15 @@ namespace EpicGames.UHT.Types
 		public virtual void Validate(UhtStruct outerStruct, UhtProperty outermostProperty, UhtValidationOptions options)
 		{
 			// Check for deprecation
-			if (options.HasAnyFlags(UhtValidationOptions.Deprecated) && !this.Deprecated)
+			if (options.HasAnyFlags(UhtValidationOptions.Deprecated) && !Deprecated)
 			{
 				ValidateDeprecated();
 			}
 
 			// Validate the types allowed with arrays
-			if (this.ArrayDimensions != null)
+			if (ArrayDimensions != null)
 			{
-				switch (this.PropertyCategory)
+				switch (PropertyCategory)
 				{
 					case UhtPropertyCategory.Return:
 						this.LogError("Arrays aren't allowed as return types");
@@ -1634,45 +1634,45 @@ namespace EpicGames.UHT.Types
 				}
 			}
 
-			if (!options.HasAnyFlags(UhtValidationOptions.IsKey) && this.PropertyFlags.HasAnyFlags(EPropertyFlags.PersistentInstance) && !this.PropertyCaps.HasAnyFlags(UhtPropertyCaps.CanBeInstanced))
+			if (!options.HasAnyFlags(UhtValidationOptions.IsKey) && PropertyFlags.HasAnyFlags(EPropertyFlags.PersistentInstance) && !PropertyCaps.HasAnyFlags(UhtPropertyCaps.CanBeInstanced))
 			{
 				this.LogError("'Instanced' is only allowed on an object property, an array of objects, a set of objects, or a map with an object value type.");
 			}
 
-			if (this.PropertyFlags.HasAnyFlags(EPropertyFlags.Config) && !this.PropertyCaps.HasAnyFlags(UhtPropertyCaps.CanHaveConfig))
+			if (PropertyFlags.HasAnyFlags(EPropertyFlags.Config) && !PropertyCaps.HasAnyFlags(UhtPropertyCaps.CanHaveConfig))
 			{
 				this.LogError("Not allowed to use 'config' with object variables");
 			}
 
-			if (this.PropertyFlags.HasAnyFlags(EPropertyFlags.ExposeOnSpawn))
+			if (PropertyFlags.HasAnyFlags(EPropertyFlags.ExposeOnSpawn))
 			{
-				if (this.PropertyFlags.HasAnyFlags(EPropertyFlags.DisableEditOnInstance))
+				if (PropertyFlags.HasAnyFlags(EPropertyFlags.DisableEditOnInstance))
 				{
 					this.LogWarning("Property cannot have both 'DisableEditOnInstance' and 'ExposeOnSpawn' flags");
 				}
-				if (!this.PropertyFlags.HasAnyFlags(EPropertyFlags.BlueprintVisible))
+				if (!PropertyFlags.HasAnyFlags(EPropertyFlags.BlueprintVisible))
 				{
 					this.LogWarning("Property cannot have 'ExposeOnSpawn' without 'BlueprintVisible' flag.");
 				}
 			}
 
-			if (this.PropertyFlags.HasAnyFlags(EPropertyFlags.BlueprintAssignable) && !this.PropertyCaps.HasAnyFlags(UhtPropertyCaps.CanBeBlueprintAssignable))
+			if (PropertyFlags.HasAnyFlags(EPropertyFlags.BlueprintAssignable) && !PropertyCaps.HasAnyFlags(UhtPropertyCaps.CanBeBlueprintAssignable))
 			{
 				this.LogError("'BlueprintAssignable' is only allowed on multicast delegate properties");
 			}
 
-			if (this.PropertyFlags.HasAnyFlags(EPropertyFlags.BlueprintCallable) && !this.PropertyCaps.HasAnyFlags(UhtPropertyCaps.CanBeBlueprintCallable))
+			if (PropertyFlags.HasAnyFlags(EPropertyFlags.BlueprintCallable) && !PropertyCaps.HasAnyFlags(UhtPropertyCaps.CanBeBlueprintCallable))
 			{
 				this.LogError("'BlueprintCallable' is only allowed on multicast delegate properties");
 			}
 
-			if (this.PropertyFlags.HasAnyFlags(EPropertyFlags.BlueprintAuthorityOnly) && !this.PropertyCaps.HasAnyFlags(UhtPropertyCaps.CanBeBlueprintAuthorityOnly))
+			if (PropertyFlags.HasAnyFlags(EPropertyFlags.BlueprintAuthorityOnly) && !PropertyCaps.HasAnyFlags(UhtPropertyCaps.CanBeBlueprintAuthorityOnly))
 			{
 				this.LogError("'BlueprintAuthorityOnly' is only allowed on multicast delegate properties");
 			}
 
 			// Check for invalid transients
-			EPropertyFlags transients = this.PropertyFlags & (EPropertyFlags.DuplicateTransient | EPropertyFlags.TextExportTransient | EPropertyFlags.NonPIEDuplicateTransient);
+			EPropertyFlags transients = PropertyFlags & (EPropertyFlags.DuplicateTransient | EPropertyFlags.TextExportTransient | EPropertyFlags.NonPIEDuplicateTransient);
 			if (transients != 0 && outerStruct is not UhtClass)
 			{
 				this.LogError($"'{String.Join(", ", transients.ToStringList(false))}' specifier(s) are only allowed on class member variables");
@@ -1680,7 +1680,7 @@ namespace EpicGames.UHT.Types
 
 			if (!options.HasAnyFlags(UhtValidationOptions.IsKeyOrValue))
 			{
-				switch (this.PropertyCategory)
+				switch (PropertyCategory)
 				{
 					case UhtPropertyCategory.Member:
 						ValidateMember(outerStruct, options);
@@ -1712,14 +1712,14 @@ namespace EpicGames.UHT.Types
 		{
 			if (func.FunctionFlags.HasAnyFlags(EFunctionFlags.Net))
 			{
-				if (this.PropertyFlags.HasExactFlags(EPropertyFlags.ReferenceParm | EPropertyFlags.ConstParm, EPropertyFlags.ReferenceParm))
+				if (PropertyFlags.HasExactFlags(EPropertyFlags.ReferenceParm | EPropertyFlags.ConstParm, EPropertyFlags.ReferenceParm))
 				{
 					this.LogError($"Replicated parameters cannot be passed by non-const reference");
 				}
 
 				if (func.FunctionFlags.HasAnyFlags(EFunctionFlags.NetRequest))
 				{
-					if (this.PropertyFlags.HasExactFlags(EPropertyFlags.OutParm | EPropertyFlags.RepSkip, EPropertyFlags.OutParm))
+					if (PropertyFlags.HasExactFlags(EPropertyFlags.OutParm | EPropertyFlags.RepSkip, EPropertyFlags.OutParm))
 					{
 						// This is difficult to trigger since NotReplicated also sets the property category
 						this.LogError("Service request functions cannot contain out parameters, unless marked NotReplicated");
@@ -1727,12 +1727,12 @@ namespace EpicGames.UHT.Types
 				}
 				else
 				{
-					if (this.PropertyFlags.HasAnyFlags(EPropertyFlags.OutParm))
+					if (PropertyFlags.HasAnyFlags(EPropertyFlags.OutParm))
 					{
 						this.LogError("Replicated functions cannot contain out parameters");
 					}
 
-					if (this.PropertyFlags.HasAnyFlags(EPropertyFlags.RepSkip))
+					if (PropertyFlags.HasAnyFlags(EPropertyFlags.RepSkip))
 					{
 						// This is difficult to trigger since NotReplicated also sets the property category
 						this.LogError("Only service request functions cannot contain NotReplicated parameters");
@@ -1744,9 +1744,9 @@ namespace EpicGames.UHT.Types
 			if (func.FunctionFlags.HasAnyFlags(EFunctionFlags.BlueprintEvent | EFunctionFlags.BlueprintCallable))
 			{
 				// Check that the parameter name is valid and does not conflict with pre-defined types
-				if (s_invalidParamNames.Contains(this.SourceName))
+				if (s_invalidParamNames.Contains(SourceName))
 				{
-					this.LogError($"Parameter name '{this.SourceName}' in function is invalid, '{this.SourceName}' is a reserved name.");
+					this.LogError($"Parameter name '{SourceName}' in function is invalid, '{SourceName}' is a reserved name.");
 				}
 			}
 		}
@@ -1761,11 +1761,11 @@ namespace EpicGames.UHT.Types
 			if (!options.HasAnyFlags(UhtValidationOptions.IsKeyOrValue))
 			{
 				// First check if the category was specified at all and if the property was exposed to the editor.
-				if (!this.MetaData.TryGetValue(UhtNames.Category, out string? category))
+				if (!MetaData.TryGetValue(UhtNames.Category, out string? category))
 				{
-					if (this.PropertyFlags.HasAnyFlags(EPropertyFlags.Edit | EPropertyFlags.BlueprintVisible))
+					if (PropertyFlags.HasAnyFlags(EPropertyFlags.Edit | EPropertyFlags.BlueprintVisible))
 					{
-						if (this.Package.IsPartOfEngine)
+						if (Package.IsPartOfEngine)
 						{
 							this.LogError("An explicit Category specifier is required for any property exposed to the editor or Blueprints in an Engine module.");
 						}
@@ -1773,7 +1773,7 @@ namespace EpicGames.UHT.Types
 				}
 
 				// If the category was specified explicitly, it wins
-				if (!String.IsNullOrEmpty(category) && !this.PropertyFlags.HasAnyFlags(EPropertyFlags.Edit | EPropertyFlags.BlueprintVisible |
+				if (!String.IsNullOrEmpty(category) && !PropertyFlags.HasAnyFlags(EPropertyFlags.Edit | EPropertyFlags.BlueprintVisible |
 					EPropertyFlags.BlueprintAssignable | EPropertyFlags.BlueprintCallable))
 				{
 					this.LogWarning("Property has a Category set but is not exposed to the editor or Blueprints with EditAnywhere, BlueprintReadWrite, " +
@@ -1782,28 +1782,28 @@ namespace EpicGames.UHT.Types
 			}
 
 			// Make sure that editblueprint variables are editable
-			if (!this.PropertyFlags.HasAnyFlags(EPropertyFlags.Edit))
+			if (!PropertyFlags.HasAnyFlags(EPropertyFlags.Edit))
 			{
-				if (this.PropertyFlags.HasAnyFlags(EPropertyFlags.DisableEditOnInstance))
+				if (PropertyFlags.HasAnyFlags(EPropertyFlags.DisableEditOnInstance))
 				{
 					this.LogError("Property cannot have 'DisableEditOnInstance' without being editable");
 				}
 
-				if (this.PropertyFlags.HasAnyFlags(EPropertyFlags.DisableEditOnTemplate))
+				if (PropertyFlags.HasAnyFlags(EPropertyFlags.DisableEditOnTemplate))
 				{
 					this.LogError("Property cannot have 'DisableEditOnTemplate' without being editable");
 				}
 			}
 
-			string exposeOnSpawnValue = this.MetaData.GetValueOrDefault(UhtNames.ExposeOnSpawn);
-			if (exposeOnSpawnValue.Equals("true", StringComparison.OrdinalIgnoreCase) && !this.PropertyCaps.HasAnyFlags(UhtPropertyCaps.CanExposeOnSpawn))
+			string exposeOnSpawnValue = MetaData.GetValueOrDefault(UhtNames.ExposeOnSpawn);
+			if (exposeOnSpawnValue.Equals("true", StringComparison.OrdinalIgnoreCase) && !PropertyCaps.HasAnyFlags(UhtPropertyCaps.CanExposeOnSpawn))
 			{
 				this.LogError("ExposeOnSpawn - Property cannot be exposed");
 			}
 
-			if (this.PropertyExportFlags.HasAnyFlags(UhtPropertyExportFlags.FieldNotify))
+			if (PropertyExportFlags.HasAnyFlags(UhtPropertyExportFlags.FieldNotify))
 			{
-				if (this.Outer is UhtClass classObj)
+				if (Outer is UhtClass classObj)
 				{
 					if (classObj.ClassType != UhtClassType.Class)
 					{
@@ -1833,7 +1833,7 @@ namespace EpicGames.UHT.Types
 				return;
 			}
 
-			if (this.PropertyCategory == UhtPropertyCategory.Member)
+			if (PropertyCategory == UhtPropertyCategory.Member)
 			{
 				this.LogError($"Property is using a deprecated class: '{classObj.SourceName}'.  Property should be marked deprecated as well.");
 			}
@@ -1880,9 +1880,9 @@ namespace EpicGames.UHT.Types
 		public override void CollectReferences(IUhtReferenceCollector collector)
 		{
 			CollectReferencesInternal(collector, false);
-			if (this.PropertyFlags.HasAnyFlags(EPropertyFlags.ParmFlags))
+			if (PropertyFlags.HasAnyFlags(EPropertyFlags.ParmFlags))
 			{
-				collector.AddForwardDeclaration(this.GetForwardDeclarations());
+				collector.AddForwardDeclaration(GetForwardDeclarations());
 			}
 		}
 
@@ -1943,19 +1943,19 @@ namespace EpicGames.UHT.Types
 		/// <returns></returns>
 		public bool MatchesType(UhtProperty other)
 		{
-			if (this.PropertyFlags.HasAnyFlags(EPropertyFlags.OutParm))
+			if (PropertyFlags.HasAnyFlags(EPropertyFlags.OutParm))
 			{
 				if (!other.PropertyFlags.HasAnyFlags(EPropertyFlags.OutParm))
 				{
 					return false;
 				}
 
-				if (this.PropertyFlags.HasAnyFlags(EPropertyFlags.ConstParm) && !other.PropertyFlags.HasAnyFlags(EPropertyFlags.ConstParm))
+				if (PropertyFlags.HasAnyFlags(EPropertyFlags.ConstParm) && !other.PropertyFlags.HasAnyFlags(EPropertyFlags.ConstParm))
 				{
 					return false;
 				}
 			}
-			if (this.IsStaticArray != other.IsStaticArray)
+			if (IsStaticArray != other.IsStaticArray)
 			{
 				return false;
 			}
diff --git a/Engine/Source/Programs/Shared/EpicGames.UHT/Types/UhtScriptStruct.cs b/Engine/Source/Programs/Shared/EpicGames.UHT/Types/UhtScriptStruct.cs
--- a/Engine/Source/Programs/Shared/EpicGames.UHT/Types/UhtScriptStruct.cs
+++ b/Engine/Source/Programs/Shared/EpicGames.UHT/Types/UhtScriptStruct.cs
@@ -141,43 +141,43 @@ namespace EpicGames.UHT.Types
 		/// True if the parameter is marked as "Constant"
 		/// </summary>
 		[JsonIgnore]
-		public bool Constant => this.ParameterFlags.HasAnyFlags(UhtRigVMParameterFlags.Constant);
+		public bool Constant => ParameterFlags.HasAnyFlags(UhtRigVMParameterFlags.Constant);
 
 		/// <summary>
 		/// True if the parameter is marked as "Input"
 		/// </summary>
 		[JsonIgnore]
-		public bool Input => this.ParameterFlags.HasAnyFlags(UhtRigVMParameterFlags.Input);
+		public bool Input => ParameterFlags.HasAnyFlags(UhtRigVMParameterFlags.Input);
 
 		/// <summary>
 		/// True if the parameter is marked as "Output"
 		/// </summary>
 		[JsonIgnore]
-		public bool Output => this.ParameterFlags.HasAnyFlags(UhtRigVMParameterFlags.Output);
+		public bool Output => ParameterFlags.HasAnyFlags(UhtRigVMParameterFlags.Output);
 
 		/// <summary>
 		/// True if the parameter is marked as "Singleton"
 		/// </summary>
 		[JsonIgnore]
-		public bool Singleton => this.ParameterFlags.HasAnyFlags(UhtRigVMParameterFlags.Singleton);
+		public bool Singleton => ParameterFlags.HasAnyFlags(UhtRigVMParameterFlags.Singleton);
 
 		/// <summary>
 		/// True if the parameter is editor only
 		/// </summary>
 		[JsonIgnore]
-		public bool EditorOnly => this.ParameterFlags.HasAnyFlags(UhtRigVMParameterFlags.EditorOnly);
+		public bool EditorOnly => ParameterFlags.HasAnyFlags(UhtRigVMParameterFlags.EditorOnly);
 
 		/// <summary>
 		/// True if the parameter is an enum
 		/// </summary>
 		[JsonIgnore]
-		public bool IsEnum => this.ParameterFlags.HasAnyFlags(UhtRigVMParameterFlags.IsEnum);
+		public bool IsEnum => ParameterFlags.HasAnyFlags(UhtRigVMParameterFlags.IsEnum);
 
 		/// <summary>
 		/// True if the parameter is an array
 		/// </summary>
 		[JsonIgnore]
-		public bool IsArray => this.ParameterFlags.HasAnyFlags(UhtRigVMParameterFlags.IsArray);
+		public bool IsArray => ParameterFlags.HasAnyFlags(UhtRigVMParameterFlags.IsArray);
 
 		/// <summary>
 		/// Create a new RigVM parameter from a property
@@ -186,47 +186,47 @@ namespace EpicGames.UHT.Types
 		/// <param name="index">Parameter index.  Used to create a unique cast name.</param>
 		public UhtRigVMParameter(UhtProperty property, int index)
 		{
-			this.Property = property;
+			Property = property;
 
-			this.Name = property.EngineName;
-			this.ParameterFlags |= property.MetaData.ContainsKey("Constant") ? UhtRigVMParameterFlags.Constant : UhtRigVMParameterFlags.None;
-			this.ParameterFlags |= property.MetaData.ContainsKey("Input") ? UhtRigVMParameterFlags.Input : UhtRigVMParameterFlags.None;
-			this.ParameterFlags |= property.MetaData.ContainsKey("Output") ? UhtRigVMParameterFlags.Output : UhtRigVMParameterFlags.None;
-			this.ParameterFlags |= property.IsEditorOnlyProperty ? UhtRigVMParameterFlags.EditorOnly : UhtRigVMParameterFlags.None;
-			this.ParameterFlags |= property.MetaData.ContainsKey("Singleton") ? UhtRigVMParameterFlags.Singleton : UhtRigVMParameterFlags.None;
+			Name = property.EngineName;
+			ParameterFlags |= property.MetaData.ContainsKey("Constant") ? UhtRigVMParameterFlags.Constant : UhtRigVMParameterFlags.None;
+			ParameterFlags |= property.MetaData.ContainsKey("Input") ? UhtRigVMParameterFlags.Input : UhtRigVMParameterFlags.None;
+			ParameterFlags |= property.MetaData.ContainsKey("Output") ? UhtRigVMParameterFlags.Output : UhtRigVMParameterFlags.None;
+			ParameterFlags |= property.IsEditorOnlyProperty ? UhtRigVMParameterFlags.EditorOnly : UhtRigVMParameterFlags.None;
+			ParameterFlags |= property.MetaData.ContainsKey("Singleton") ? UhtRigVMParameterFlags.Singleton : UhtRigVMParameterFlags.None;
 
 			if (property.MetaData.ContainsKey("Visible"))
 			{
-				this.ParameterFlags |= UhtRigVMParameterFlags.Constant | UhtRigVMParameterFlags.Input;
-				this.ParameterFlags &= ~UhtRigVMParameterFlags.Output;
+				ParameterFlags |= UhtRigVMParameterFlags.Constant | UhtRigVMParameterFlags.Input;
+				ParameterFlags &= ~UhtRigVMParameterFlags.Output;
 			}
 
-			if (this.EditorOnly)
+			if (EditorOnly)
 			{
-				property.LogError($"RigVM Struct '{this.Property.Outer?.SourceName}' - Member '{this.Property.SourceName}' is editor only - WITH_EDITORONLY_DATA not allowed on structs with RIGVM_METHOD.");
+				property.LogError($"RigVM Struct '{Property.Outer?.SourceName}' - Member '{Property.SourceName}' is editor only - WITH_EDITORONLY_DATA not allowed on structs with RIGVM_METHOD.");
 			}
 
 			if (property.PropertyCaps.HasAnyFlags(UhtPropertyCaps.SupportsRigVM))
 			{
-				this.Type = property.GetRigVMType();
+				Type = property.GetRigVMType();
 				if (property.PropertyCaps.HasAnyFlags(UhtPropertyCaps.IsRigVMEnum))
 				{
-					this.ParameterFlags |= UhtRigVMParameterFlags.IsEnum;
+					ParameterFlags |= UhtRigVMParameterFlags.IsEnum;
 				}
 				if (property.PropertyCaps.HasAnyFlags(UhtPropertyCaps.IsRigVMArray))
 				{
-					this.ParameterFlags |= UhtRigVMParameterFlags.IsArray;
-					if (this.IsConst())
+					ParameterFlags |= UhtRigVMParameterFlags.IsArray;
+					if (IsConst())
 					{
-						string extendedType = this.ExtendedType(false);
-						this.CastName = $"{this.Name}_{index}_Array";
-						this.CastType = $"TArrayView<const {extendedType[1..^1]}>";
+						string extendedType = ExtendedType(false);
+						CastName = $"{Name}_{index}_Array";
+						CastType = $"TArrayView<const {extendedType[1..^1]}>";
 					}
 				}
 			}
 			else
 			{
-				property.LogError($"RigVM Struct '{this.Property.Outer?.SourceName}' - Member '{this.Property.SourceName}' type '{this.Property.GetUserFacingDecl()}' not supported by RigVM.");
+				property.LogError($"RigVM Struct '{Property.Outer?.SourceName}' - Member '{Property.SourceName}' type '{Property.GetUserFacingDecl()}' not supported by RigVM.");
 			}
 		}
 
@@ -237,8 +237,8 @@ namespace EpicGames.UHT.Types
 		/// <param name="type">Type of the parameter</param>
 		public UhtRigVMParameter(string name, string type)
 		{
-			this.Name = name;
-			this.Type = type;
+			Name = name;
+			Type = type;
 		}
 
 		/// <summary>
@@ -248,7 +248,7 @@ namespace EpicGames.UHT.Types
 		/// <returns>Parameter name</returns>
 		public string NameOriginal(bool castName = false)
 		{
-			return castName && this.CastName != null ? this.CastName : this.Name;
+			return castName && CastName != null ? CastName : Name;
 		}
 
 		/// <summary>
@@ -258,7 +258,7 @@ namespace EpicGames.UHT.Types
 		/// <returns>Parameter type</returns>
 		public string TypeOriginal(bool castType = false)
 		{
-			return castType && this.CastType != null ? this.CastType : this.Type;
+			return castType && CastType != null ? CastType : Type;
 		}
 
 		/// <summary>
@@ -386,7 +386,7 @@ namespace EpicGames.UHT.Types
 		/// <returns>True if the parameter is constant</returns>
 		public bool IsConst()
 		{
-			return this.Constant || (this.Input && !this.Output);
+			return Constant || (Input && !Output);
 		}
 
 		/// <summary>
@@ -395,7 +395,7 @@ namespace EpicGames.UHT.Types
 		/// <returns>True if the parameter requires a cast</returns>
 		public bool RequiresCast()
 		{
-			return this.CastType != null && this.CastName != null;
+			return CastType != null && CastName != null;
 		}
 	}
 
@@ -577,42 +577,42 @@ namespace EpicGames.UHT.Types
 		/// True if the struct has the "HasDefaults" specifier
 		/// </summary>
 		[JsonIgnore]
-		public bool HasDefaults => this.ScriptStructExportFlags.HasAnyFlags(UhtScriptStructExportFlags.HasDefaults);
+		public bool HasDefaults => ScriptStructExportFlags.HasAnyFlags(UhtScriptStructExportFlags.HasDefaults);
 
 		/// <summary>
 		/// True if the struct has the "IsAlwaysAccessible" specifier
 		/// </summary>
 		[JsonIgnore]
-		public bool IsAlwaysAccessible => this.ScriptStructExportFlags.HasAnyFlags(UhtScriptStructExportFlags.IsAlwaysAccessible);
+		public bool IsAlwaysAccessible => ScriptStructExportFlags.HasAnyFlags(UhtScriptStructExportFlags.IsAlwaysAccessible);
 
 		/// <summary>
 		/// True if the struct has the "IsCoreType" specifier
 		/// </summary>
 		[JsonIgnore]
-		public bool IsCoreType => this.ScriptStructExportFlags.HasAnyFlags(UhtScriptStructExportFlags.IsCoreType);
+		public bool IsCoreType => ScriptStructExportFlags.HasAnyFlags(UhtScriptStructExportFlags.IsCoreType);
 
 		/// <summary>
 		/// True if the struct has the "HasNoOpConstructor" specifier
 		/// </summary>
 		[JsonIgnore]
-		public bool HasNoOpConstructor => this.ScriptStructExportFlags.HasAnyFlags(UhtScriptStructExportFlags.HasNoOpConstructor);
+		public bool HasNoOpConstructor => ScriptStructExportFlags.HasAnyFlags(UhtScriptStructExportFlags.HasNoOpConstructor);
 
 		/// <inheritdoc/>
 		public override string EngineClassName => "ScriptStruct";
 
 		/// <inheritdoc/>
 		[JsonIgnore]
-		public override string EngineNamePrefix => this.Session.Config!.IsStructWithTPrefix(this.EngineName) ? "T" : "F";
+		public override string EngineNamePrefix => Session.Config!.IsStructWithTPrefix(EngineName) ? "T" : "F";
 
 		///<inheritdoc/>
 		[JsonIgnore]
-		protected override UhtSpecifierValidatorTable? SpecifierValidatorTable => this.Session.GetSpecifierValidatorTable(UhtTableNames.ScriptStruct);
+		protected override UhtSpecifierValidatorTable? SpecifierValidatorTable => Session.GetSpecifierValidatorTable(UhtTableNames.ScriptStruct);
 
 		/// <summary>
 		/// Super struct
 		/// </summary>
 		[JsonConverter(typeof(UhtNullableTypeSourceNameJsonConverter<UhtScriptStruct>))]
-		public UhtScriptStruct? SuperScriptStruct => (UhtScriptStruct?)this.Super;
+		public UhtScriptStruct? SuperScriptStruct => (UhtScriptStruct?)Super;
 
 		/// <summary>
 		/// Construct a new script struct
@@ -632,13 +632,13 @@ namespace EpicGames.UHT.Types
 			switch (resolvePhase)
 			{
 				case UhtResolvePhase.Bases:
-					BindAndResolveSuper(this.SuperIdentifier, UhtFindOptions.ScriptStruct);
+					BindAndResolveSuper(SuperIdentifier, UhtFindOptions.ScriptStruct);
 
 					// if we have a base struct, propagate inherited struct flags now
-					UhtScriptStruct? superScriptStruct = this.SuperScriptStruct;
+					UhtScriptStruct? superScriptStruct = SuperScriptStruct;
 					if (superScriptStruct != null)
 					{
-						this.ScriptStructFlags |= superScriptStruct.ScriptStructFlags & EStructFlags.Inherit;
+						ScriptStructFlags |= superScriptStruct.ScriptStructFlags & EStructFlags.Inherit;
 					}
 					break;
 			}
@@ -671,7 +671,7 @@ namespace EpicGames.UHT.Types
 			switch (phase)
 			{
 				case UhtResolvePhase.Properties:
-					foreach (UhtType child in this.Children)
+					foreach (UhtType child in Children)
 					{
 						if (child is UhtProperty property)
 						{
@@ -695,9 +695,9 @@ namespace EpicGames.UHT.Types
 							if (originalFlags != property.PropertyFlags &&
 								property.PropertyFlags.HasAnyFlags(EPropertyFlags.BlueprintVisible | EPropertyFlags.BlueprintReadOnly))
 							{
-								if (!this.MetaData.GetBooleanHierarchical(UhtNames.BlueprintType))
+								if (!MetaData.GetBooleanHierarchical(UhtNames.BlueprintType))
 								{
-									this.MetaData.Add(UhtNames.BlueprintType, true);
+									MetaData.Add(UhtNames.BlueprintType, true);
 								}
 
 								if (!property.MetaData.ContainsKey(UhtNames.Category))
@@ -717,7 +717,7 @@ namespace EpicGames.UHT.Types
 				case UhtResolvePhase.Final:
 					if (ScanForInstancedReferenced(false))
 					{
-						this.ScriptStructFlags |= EStructFlags.HasInstancedReference;
+						ScriptStructFlags |= EStructFlags.HasInstancedReference;
 					}
 					CollectRigVMMembers();
 					break;
@@ -727,12 +727,12 @@ namespace EpicGames.UHT.Types
 		/// <inheritdoc/>
 		public override bool ScanForInstancedReferenced(bool deepScan)
 		{
-			if (this.ScriptStructFlags.HasAnyFlags(EStructFlags.HasInstancedReference))
+			if (ScriptStructFlags.HasAnyFlags(EStructFlags.HasInstancedReference))
 			{
 				return true;
 			}
 
-			if (this.SuperScriptStruct != null && this.SuperScriptStruct.ScanForInstancedReferenced(deepScan))
+			if (SuperScriptStruct != null && SuperScriptStruct.ScanForInstancedReferenced(deepScan))
 			{
 				return true;
 			}
@@ -747,25 +747,25 @@ namespace EpicGames.UHT.Types
 		{
 			options = base.Validate(options);
 
-			if (this.ScriptStructFlags.HasAnyFlags(EStructFlags.Immutable))
+			if (ScriptStructFlags.HasAnyFlags(EStructFlags.Immutable))
 			{
-				if (this.HeaderFile != this.Session.UObject.HeaderFile)
+				if (HeaderFile != Session.UObject.HeaderFile)
 				{
 					this.LogError("Immutable is being phased out in favor of SerializeNative, and is only legal on the mirror structs declared in UObject");
 				}
 			}
 
 			// Validate the engine name
-			string expectedName = $"{this.EngineNamePrefix}{this.EngineName}";
-			if (this.SourceName != expectedName)
+			string expectedName = $"{EngineNamePrefix}{EngineName}";
+			if (SourceName != expectedName)
 			{
-				this.LogError($"Struct '{this.SourceName}' has an invalid Unreal prefix, expecting '{expectedName}");
+				this.LogError($"Struct '{SourceName}' has an invalid Unreal prefix, expecting '{expectedName}");
 			}
 
 			// Validate RigVM
-			if (this.RigVMStructInfo != null && this.MetaData.ContainsKey("Deprecated") && !this.RigVMStructInfo.HasGetUpgradeInfoMethod)
+			if (RigVMStructInfo != null && MetaData.ContainsKey("Deprecated") && !RigVMStructInfo.HasGetUpgradeInfoMethod)
 			{
-				this.LogError($"RigVMStruct '{this.SourceName}' is marked as deprecated but is missing 'GetUpgradeInfo method.");
+				this.LogError($"RigVMStruct '{SourceName}' is marked as deprecated but is missing 'GetUpgradeInfo method.");
 				this.LogError("Please implement a method like below:");
 				this.LogError("RIGVM_METHOD()");
 				this.LogError("virtual FRigVMStructUpgradeInfo GetUpgradeInfo() const override;");
@@ -778,19 +778,19 @@ namespace EpicGames.UHT.Types
 		/// <inheritdoc/>
 		public override void CollectReferences(IUhtReferenceCollector collector)
 		{
-			if (this.ScriptStructFlags.HasAnyFlags(EStructFlags.NoExport))
+			if (ScriptStructFlags.HasAnyFlags(EStructFlags.NoExport))
 			{
 				collector.AddSingleton(this);
 			}
 			collector.AddExportType(this);
 			collector.AddDeclaration(this, true);
 			collector.AddCrossModuleReference(this, true);
-			if (this.SuperScriptStruct != null)
+			if (SuperScriptStruct != null)
 			{
-				collector.AddCrossModuleReference(this.SuperScriptStruct, true);
+				collector.AddCrossModuleReference(SuperScriptStruct, true);
 			}
-			collector.AddCrossModuleReference(this.Package, true);
-			foreach (UhtType child in this.Children)
+			collector.AddCrossModuleReference(Package, true);
+			foreach (UhtType child in Children)
 			{
 				child.CollectReferences(collector);
 			}
@@ -798,23 +798,23 @@ namespace EpicGames.UHT.Types
 
 		private void CollectRigVMMembers()
 		{
-			if (this.RigVMStructInfo != null)
+			if (RigVMStructInfo != null)
 			{
 				for (UhtStruct? current = this; current != null; current = current.SuperStruct)
 				{
 					foreach (UhtProperty property in current.Properties)
 					{
-						this.RigVMStructInfo.Members.Add(new UhtRigVMParameter(property, this.RigVMStructInfo.Members.Count));
+						RigVMStructInfo.Members.Add(new UhtRigVMParameter(property, RigVMStructInfo.Members.Count));
 					}
 				}
 
-				if (this.RigVMStructInfo.Members.Count == 0)
+				if (RigVMStructInfo.Members.Count == 0)
 				{
-					this.LogError($"RigVM Struct '{this.SourceName}' - has zero members - invalid RIGVM_METHOD.");
+					this.LogError($"RigVM Struct '{SourceName}' - has zero members - invalid RIGVM_METHOD.");
 				}
-				else if (this.RigVMStructInfo.Members.Count > 64)
+				else if (RigVMStructInfo.Members.Count > 64)
 				{
-					this.LogError($"RigVM Struct '{this.SourceName}' - has {this.RigVMStructInfo.Members.Count} members (64 is the limit).");
+					this.LogError($"RigVM Struct '{SourceName}' - has {RigVMStructInfo.Members.Count} members (64 is the limit).");
 				}
 			}
 		}
diff --git a/Engine/Source/Programs/Shared/EpicGames.UHT/Types/UhtSourceFile.cs b/Engine/Source/Programs/Shared/EpicGames.UHT/Types/UhtSourceFile.cs
--- a/Engine/Source/Programs/Shared/EpicGames.UHT/Types/UhtSourceFile.cs
+++ b/Engine/Source/Programs/Shared/EpicGames.UHT/Types/UhtSourceFile.cs
@@ -52,10 +52,10 @@ namespace EpicGames.UHT.Types
 
 		#region IUHTMessageSite implementation
 		/// <inheritdoc/>
-		public IUhtMessageSession MessageSession => this._messageSite.MessageSession;
+		public IUhtMessageSession MessageSession => _messageSite.MessageSession;
 
 		/// <inheritdoc/>
-		public IUhtMessageSource? MessageSource => this._messageSite.MessageSource;
+		public IUhtMessageSource? MessageSource => _messageSite.MessageSource;
 
 		/// <inheritdoc/>
 		public IUhtMessageLineNumber? MessageLineNumber => null;
@@ -63,22 +63,22 @@ namespace EpicGames.UHT.Types
 
 		#region IUHTMessageSource implementation
 		/// <inheritdoc/>
-		public string MessageFilePath => this.FilePath;
+		public string MessageFilePath => FilePath;
 
 		/// <inheritdoc/>
-		public string MessageFullFilePath => this.FullFilePath;
+		public string MessageFullFilePath => FullFilePath;
 
 		/// <inheritdoc/>
-		public bool MessageIsFragment => this.FragmentSourceFile != null;
+		public bool MessageIsFragment => FragmentSourceFile != null;
 
 		/// <inheritdoc/>
-		public string MessageFragmentFilePath => this.FragmentSourceFile != null ? this.FragmentSourceFile.MessageFilePath : "";
+		public string MessageFragmentFilePath => FragmentSourceFile != null ? FragmentSourceFile.MessageFilePath : "";
 
 		/// <inheritdoc/>
-		public string MessageFragmentFullFilePath => this.FragmentSourceFile != null ? this.FragmentSourceFile.MessageFullFilePath : "";
+		public string MessageFragmentFullFilePath => FragmentSourceFile != null ? FragmentSourceFile.MessageFullFilePath : "";
 
 		/// <inheritdoc/>
-		public int MessageFragmentLineNumber => this.FragmentLineNumber;
+		public int MessageFragmentLineNumber => FragmentLineNumber;
 		#endregion
 
 		/// <summary>
@@ -88,10 +88,10 @@ namespace EpicGames.UHT.Types
 		/// <param name="filePath">The full file path</param>
 		public UhtSourceFile(UhtSession session, string filePath)
 		{
-			this.Session = session;
-			this.FilePath = filePath;
-			this.FileName = Path.GetFileNameWithoutExtension(this.FilePath);
-			this._messageSite = new UhtSimpleMessageSite(this.Session, this);
+			Session = session;
+			FilePath = filePath;
+			FileName = Path.GetFileNameWithoutExtension(FilePath);
+			_messageSite = new UhtSimpleMessageSite(Session, this);
 		}
 
 		/// <summary>
@@ -99,7 +99,7 @@ namespace EpicGames.UHT.Types
 		/// </summary>
 		public virtual void Read()
 		{
-			this._sourceFragment = Session.ReadSource(this.FilePath);
+			_sourceFragment = Session.ReadSource(FilePath);
 		}
 	}
 }
diff --git a/Engine/Source/Programs/Shared/EpicGames.UHT/Types/UhtStruct.cs b/Engine/Source/Programs/Shared/EpicGames.UHT/Types/UhtStruct.cs
--- a/Engine/Source/Programs/Shared/EpicGames.UHT/Types/UhtStruct.cs
+++ b/Engine/Source/Programs/Shared/EpicGames.UHT/Types/UhtStruct.cs
@@ -38,7 +38,7 @@ namespace EpicGames.UHT.Types
 		{
 			get
 			{
-				foreach (UhtType type in this.Children)
+				foreach (UhtType type in Children)
 				{
 					if (type is UhtProperty property)
 					{
@@ -61,7 +61,7 @@ namespace EpicGames.UHT.Types
 		{
 			get
 			{
-				foreach (UhtType type in this.Children)
+				foreach (UhtType type in Children)
 				{
 					if (type is UhtFunction function)
 					{
@@ -146,12 +146,12 @@ namespace EpicGames.UHT.Types
 		/// <inheritdoc/>
 		protected override void ResolveSuper(UhtResolvePhase resolvePhase)
 		{
-			if (this.Super != null)
+			if (Super != null)
 			{
-				this.Super.Resolve(resolvePhase);
+				Super.Resolve(resolvePhase);
 			}
 
-			foreach (UhtStruct baseStruct in this.Bases)
+			foreach (UhtStruct baseStruct in Bases)
 			{
 				baseStruct.Resolve(resolvePhase);
 			}
@@ -167,7 +167,7 @@ namespace EpicGames.UHT.Types
 		/// <returns></returns>
 		public virtual bool ScanForInstancedReferenced(bool deepScan)
 		{
-			foreach (UhtType type in this.Children)
+			foreach (UhtType type in Children)
 			{
 				if (type is UhtProperty property)
 				{
@@ -196,14 +196,14 @@ namespace EpicGames.UHT.Types
 		{
 			if (superIdentifier)
 			{
-				this.Super = (UhtStruct?)this.FindType(findOptions | UhtFindOptions.SourceName | UhtFindOptions.NoSelf, ref superIdentifier);
-				if (this.Super == null)
+				Super = (UhtStruct?)FindType(findOptions | UhtFindOptions.SourceName | UhtFindOptions.NoSelf, ref superIdentifier);
+				if (Super == null)
 				{
-					throw new UhtException(this, $"Unable to find parent {this.EngineType.ShortLowercaseText()} type for '{this.SourceName}' named '{superIdentifier.Value}'");
+					throw new UhtException(this, $"Unable to find parent {EngineType.ShortLowercaseText()} type for '{SourceName}' named '{superIdentifier.Value}'");
 				}
-				this.HeaderFile.AddReferencedHeader(this.Super);
-				this.MetaData.Parent = this.Super.MetaData;
-				this.Super.Resolve(UhtResolvePhase.Bases);
+				HeaderFile.AddReferencedHeader(Super);
+				MetaData.Parent = Super.MetaData;
+				Super.Resolve(UhtResolvePhase.Bases);
 			}
 		}
 
@@ -220,15 +220,15 @@ namespace EpicGames.UHT.Types
 				foreach (UhtToken[] baseIdentifier in baseIdentifiers)
 				{
 					// We really only case about interfaces, but we can also handle structs
-					UhtStruct? baseStruct = (UhtStruct?)this.FindType(findOptions | UhtFindOptions.Class | UhtFindOptions.ScriptStruct | UhtFindOptions.SourceName | UhtFindOptions.NoSelf, baseIdentifier);
+					UhtStruct? baseStruct = (UhtStruct?)FindType(findOptions | UhtFindOptions.Class | UhtFindOptions.ScriptStruct | UhtFindOptions.SourceName | UhtFindOptions.NoSelf, baseIdentifier);
 					if (baseStruct != null)
 					{
-						if (this._bases == null)
+						if (_bases == null)
 						{
-							this._bases = new List<UhtStruct>();
+							_bases = new List<UhtStruct>();
 						}
-						this._bases.Add(baseStruct);
-						this.HeaderFile.AddReferencedHeader(baseStruct);
+						_bases.Add(baseStruct);
+						HeaderFile.AddReferencedHeader(baseStruct);
 						baseStruct.Resolve(UhtResolvePhase.Bases);
 					}
 				}
@@ -274,28 +274,28 @@ namespace EpicGames.UHT.Types
 		{
 			if (policy.ClassOrStructCommentRequired)
 			{
-				string classTooltip = this.MetaData.GetValueOrDefault(UhtNames.ToolTip);
-				if (classTooltip.Length == 0 || classTooltip.Equals(this.EngineName, StringComparison.OrdinalIgnoreCase))
+				string classTooltip = MetaData.GetValueOrDefault(UhtNames.ToolTip);
+				if (classTooltip.Length == 0 || classTooltip.Equals(EngineName, StringComparison.OrdinalIgnoreCase))
 				{
-					this.LogError($"{this.EngineType.CapitalizedText()} '{this.SourceName}' does not provide a tooltip / comment (DocumentationPolicy).");
+					this.LogError($"{EngineType.CapitalizedText()} '{SourceName}' does not provide a tooltip / comment (DocumentationPolicy).");
 				}
 			}
 
 			if (policy.MemberToolTipsRequired)
 			{
 				Dictionary<string, UhtProperty> toolTipToType = new();
-				foreach (UhtProperty property in this.Properties)
+				foreach (UhtProperty property in Properties)
 				{
 					string toolTip = property.GetToolTipText();
 					if (toolTip.Length == 0 || toolTip == property.GetDisplayNameText())
 					{
-						property.LogError($"Property '{this.SourceName}::{property.SourceName}' does not provide a tooltip / comment (DocumentationPolicy).");
+						property.LogError($"Property '{SourceName}::{property.SourceName}' does not provide a tooltip / comment (DocumentationPolicy).");
 						continue;
 					}
 
 					if (toolTipToType.TryGetValue(toolTip, out UhtProperty? existing))
 					{
-						property.LogError($"Property '{this.SourceName}::{existing.SourceName}' and '{this.SourceName}::{property.SourceName}' are using identical tooltips (DocumentationPolicy).");
+						property.LogError($"Property '{SourceName}::{existing.SourceName}' and '{SourceName}::{property.SourceName}' are using identical tooltips (DocumentationPolicy).");
 					}
 					else
 					{
@@ -306,13 +306,13 @@ namespace EpicGames.UHT.Types
 
 			if (policy.FloatRangesRequired)
 			{
-				foreach (UhtType child in this.Children)
+				foreach (UhtType child in Children)
 				{
 					if (child is UhtDoubleProperty || child is UhtFloatProperty)
 					{
 						if (!CheckUIMinMaxRangeFromMetaData(child))
 						{
-							child.LogError($"Property '{this.SourceName}::{child.SourceName}' does not provide a valid UIMin / UIMax (DocumentationPolicy).");
+							child.LogError($"Property '{SourceName}::{child.SourceName}' does not provide a valid UIMin / UIMax (DocumentationPolicy).");
 						}
 					}
 				}
@@ -324,7 +324,7 @@ namespace EpicGames.UHT.Types
 				if (this is UhtClass)
 				{
 					Dictionary<string, UhtType> toolTipToType = new();
-					foreach (UhtType child in this.Children)
+					foreach (UhtType child in Children)
 					{
 						if (child is UhtFunction function)
 						{
@@ -332,13 +332,13 @@ namespace EpicGames.UHT.Types
 							if (toolTip.Length == 0)
 							{
 								// NOTE: This does not fire because it doesn't check to see if it matches the display name as above.
-								child.LogError($"Function '{this.SourceName}::{function.SourceName}' does not provide a tooltip / comment (DocumentationPolicy).");
+								child.LogError($"Function '{SourceName}::{function.SourceName}' does not provide a tooltip / comment (DocumentationPolicy).");
 								continue;
 							}
 
 							if (toolTipToType.TryGetValue(toolTip, out UhtType? existing))
 							{
-								child.LogError($"Functions '{this.SourceName}::{existing.SourceName}' and '{this.SourceName}::{function.SourceName}' are using identical tooltips / comments (DocumentationPolicy).");
+								child.LogError($"Functions '{SourceName}::{existing.SourceName}' and '{SourceName}::{function.SourceName}' are using identical tooltips / comments (DocumentationPolicy).");
 							}
 							else
 							{
@@ -353,7 +353,7 @@ namespace EpicGames.UHT.Types
 		private void ValidateSparseClassData()
 		{
 			// Fetch the data types
-			string[]? sparseClassDataTypes = this.MetaData.GetStringArray(UhtNames.SparseClassDataTypes);
+			string[]? sparseClassDataTypes = MetaData.GetStringArray(UhtNames.SparseClassDataTypes);
 			if (sparseClassDataTypes == null)
 			{
 				return;
@@ -363,30 +363,30 @@ namespace EpicGames.UHT.Types
 			UhtClass? classObj = this as UhtClass;
 			if (classObj == null)
 			{
-				this.LogError($"{this.EngineType.CapitalizedText()} '{this.SourceName}' contains sparse class data but is not a class.");
+				this.LogError($"{EngineType.CapitalizedText()} '{SourceName}' contains sparse class data but is not a class.");
 				return;
 			}
 
 			// for now we only support one sparse class data structure per class
 			if (sparseClassDataTypes.Length > 1)
 			{
-				this.LogError($"Class '{this.SourceName}' contains multiple sparse class data types");
+				this.LogError($"Class '{SourceName}' contains multiple sparse class data types");
 				return;
 			}
 			if (sparseClassDataTypes.Length == 0)
 			{
-				this.LogError($"Class '{this.SourceName}' has sparse class metadata but does not specify a type");
+				this.LogError($"Class '{SourceName}' has sparse class metadata but does not specify a type");
 				return;
 			}
 
 			foreach (string sparseClassDataTypeName in sparseClassDataTypes)
 			{
-				UhtScriptStruct? sparseScriptStruct = this.Session.FindType(null, UhtFindOptions.EngineName | UhtFindOptions.ScriptStruct, sparseClassDataTypeName) as UhtScriptStruct;
+				UhtScriptStruct? sparseScriptStruct = Session.FindType(null, UhtFindOptions.EngineName | UhtFindOptions.ScriptStruct, sparseClassDataTypeName) as UhtScriptStruct;
 
 				// make sure the sparse class data struct actually exists
 				if (sparseScriptStruct == null)
 				{
-					this.LogError($"Unable to find sparse data type '{sparseClassDataTypeName}' for class '{this.SourceName}'");
+					this.LogError($"Unable to find sparse data type '{sparseClassDataTypeName}' for class '{SourceName}'");
 					continue;
 				}
 
@@ -420,13 +420,13 @@ namespace EpicGames.UHT.Types
 					{
 						foreach (string superSparseClassDataTypeName in superSparseClassDataTypes)
 						{
-							UhtScriptStruct? superSparseScriptStruct = this.Session.FindType(null, UhtFindOptions.EngineName | UhtFindOptions.ScriptStruct, superSparseClassDataTypeName) as UhtScriptStruct;
+							UhtScriptStruct? superSparseScriptStruct = Session.FindType(null, UhtFindOptions.EngineName | UhtFindOptions.ScriptStruct, superSparseClassDataTypeName) as UhtScriptStruct;
 							if (superSparseScriptStruct != null)
 							{
 								if (!sparseScriptStruct.IsChildOf(superSparseScriptStruct))
 								{
 									this.LogError(
-										$"Class '{this.SourceName}' is a child of '{superClass.SourceName}' but its sparse class data struct " +
+										$"Class '{SourceName}' is a child of '{superClass.SourceName}' but its sparse class data struct " +
 										$"'{sparseScriptStruct.EngineName}', does not inherit from '{superSparseScriptStruct.EngineName}'.");
 								}
 							}
diff --git a/Engine/Source/Programs/Shared/EpicGames.UHT/Types/UhtType.cs b/Engine/Source/Programs/Shared/EpicGames.UHT/Types/UhtType.cs
--- a/Engine/Source/Programs/Shared/EpicGames.UHT/Types/UhtType.cs
+++ b/Engine/Source/Programs/Shared/EpicGames.UHT/Types/UhtType.cs
@@ -727,7 +727,7 @@ namespace EpicGames.UHT.Types
 		/// <summary>
 		/// The name of the type as used by the engines type system.  If not set, will default to the source name
 		/// </summary>
-		public string EngineName { get => this._engineName ?? this.SourceName; set => this._engineName = value; }
+		public string EngineName { get => _engineName ?? SourceName; set => _engineName = value; }
 
 		/// <summary>
 		/// Simple enumeration that can be used to detect the different types of objects
@@ -788,7 +788,7 @@ namespace EpicGames.UHT.Types
 		/// Return the name of the documentation policy to be used
 		/// </summary>
 		[JsonIgnore]
-		public string DocumentationPolicyName => this.MetaData.GetValueOrDefaultHierarchical(UhtNames.DocumentationPolicy);
+		public string DocumentationPolicyName => MetaData.GetValueOrDefaultHierarchical(UhtNames.DocumentationPolicy);
 
 		/// <summary>
 		/// Return the specifier validation table to be used to validate the meta data on this type
@@ -812,7 +812,7 @@ namespace EpicGames.UHT.Types
 		/// Helper to string method
 		/// </summary>
 		/// <returns></returns>
-		public override string ToString() { return this.SourceName; }
+		public override string ToString() { return SourceName; }
 
 		private string? _engineName = null;
 		private List<UhtType>? _children = null;
@@ -821,11 +821,11 @@ namespace EpicGames.UHT.Types
 		#region IUHTMessageSite implementation
 		/// <inheritdoc/>
 		[JsonIgnore]
-		public virtual IUhtMessageSession MessageSession => this.HeaderFile.MessageSession;
+		public virtual IUhtMessageSession MessageSession => HeaderFile.MessageSession;
 
 		/// <inheritdoc/>
 		[JsonIgnore]
-		public virtual IUhtMessageSource? MessageSource => this.HeaderFile.MessageSource;
+		public virtual IUhtMessageSource? MessageSource => HeaderFile.MessageSource;
 
 		/// <inheritdoc/>
 		[JsonIgnore]
@@ -843,12 +843,12 @@ namespace EpicGames.UHT.Types
 		/// <param name="session"></param>
 		protected UhtType(UhtSession session)
 		{
-			this.Session = session;
-			this.Outer = null;
-			this.SourceName = String.Empty;
-			this.LineNumber = 1;
-			this.TypeIndex = session.GetNextTypeIndex();
-			this.MetaData = new UhtMetaData(this, this.Session.Config);
+			Session = session;
+			Outer = null;
+			SourceName = String.Empty;
+			LineNumber = 1;
+			TypeIndex = session.GetNextTypeIndex();
+			MetaData = new UhtMetaData(this, Session.Config);
 		}
 
 		/// <summary>
@@ -859,14 +859,14 @@ namespace EpicGames.UHT.Types
 		/// <param name="metaData">Optional meta data to be associated with the type.</param>
 		protected UhtType(UhtType outer, int lineNumber, UhtMetaData? metaData = null)
 		{
-			this.Session = outer.Session;
-			this.Outer = outer;
-			this.SourceName = String.Empty;
-			this.LineNumber = lineNumber;
-			this.TypeIndex = this.Session.GetNextTypeIndex();
-			this.MetaData = metaData ?? new UhtMetaData(this, this.Session.Config);
-			this.MetaData.MessageSite = this; // Make sure the site is correct
-			this.MetaData.Config = this.Session.Config;
+			Session = outer.Session;
+			Outer = outer;
+			SourceName = String.Empty;
+			LineNumber = lineNumber;
+			TypeIndex = Session.GetNextTypeIndex();
+			MetaData = metaData ?? new UhtMetaData(this, Session.Config);
+			MetaData.MessageSite = this; // Make sure the site is correct
+			MetaData.Config = Session.Config;
 		}
 
 		/// <summary>
@@ -875,11 +875,11 @@ namespace EpicGames.UHT.Types
 		/// <param name="child">The child to be added.</param>
 		public virtual void AddChild(UhtType child)
 		{
-			if (this._children == null)
+			if (_children == null)
 			{
-				this._children = new List<UhtType>();
+				_children = new List<UhtType>();
 			}
-			this._children.Add(child);
+			_children.Add(child);
 		}
 
 		/// <summary>
@@ -889,8 +889,8 @@ namespace EpicGames.UHT.Types
 		/// <returns>The current list of children</returns>
 		public List<UhtType>? DetachChildren()
 		{
-			List<UhtType>? children = this._children;
-			this._children = null;
+			List<UhtType>? children = _children;
+			_children = null;
 			return children;
 		}
 
@@ -920,7 +920,7 @@ namespace EpicGames.UHT.Types
 		public bool Resolve(UhtResolvePhase phase)
 		{
 			int initState = (int)phase << 1;
-			int currentState = Interlocked.Add(ref this._resolveState, 0);
+			int currentState = Interlocked.Add(ref _resolveState, 0);
 
 			// If we are already passed the init phase for this state, then we are good
 			if (currentState > initState)
@@ -930,7 +930,7 @@ namespace EpicGames.UHT.Types
 
 			// Attempt to make myself the resolve thread for this type.  If we aren't the ones
 			// to get the init state, then we have to wait until the thread that did completes
-			currentState = Interlocked.CompareExchange(ref this._resolveState, initState, currentState);
+			currentState = Interlocked.CompareExchange(ref _resolveState, initState, currentState);
 			if (currentState > initState)
 			{
 				return true;
@@ -940,7 +940,7 @@ namespace EpicGames.UHT.Types
 				do
 				{
 					Thread.Yield();
-					currentState = Interlocked.Add(ref this._resolveState, 0);
+					currentState = Interlocked.Add(ref _resolveState, 0);
 				} while (currentState == initState);
 				return true;
 			}
@@ -955,7 +955,7 @@ namespace EpicGames.UHT.Types
 				}
 				finally
 				{
-					Interlocked.Increment(ref this._resolveState);
+					Interlocked.Increment(ref _resolveState);
 				}
 				return result;
 			}
@@ -985,27 +985,27 @@ namespace EpicGames.UHT.Types
 		/// <param name="phase">Resolution phase</param>
 		protected virtual void ResolveChildren(UhtResolvePhase phase)
 		{
-			if (this._children != null)
+			if (_children != null)
 			{
 				if (phase == UhtResolvePhase.InvalidCheck)
 				{
 					int outIndex = 0;
-					for (int index = 0; index < this._children.Count; ++index)
+					for (int index = 0; index < _children.Count; ++index)
 					{
-						UhtType child = this._children[index];
+						UhtType child = _children[index];
 						if (child.Resolve(phase))
 						{
-							this._children[outIndex++] = child;
+							_children[outIndex++] = child;
 						}
 					}
-					if (outIndex < this._children.Count)
+					if (outIndex < _children.Count)
 					{
-						this._children.RemoveRange(outIndex, this._children.Count - outIndex);
+						_children.RemoveRange(outIndex, _children.Count - outIndex);
 					}
 				}
 				else
 				{
-					foreach (UhtType child in this._children)
+					foreach (UhtType child in _children)
 					{
 						child.Resolve(phase);
 					}
@@ -1025,7 +1025,7 @@ namespace EpicGames.UHT.Types
 		/// <returns>The located type of null if not found</returns>
 		public UhtType? FindType(UhtFindOptions options, string name, IUhtMessageSite? messageSite = null, int lineNumber = -1)
 		{
-			return this.Session.FindType(this, options, name, messageSite, lineNumber);
+			return Session.FindType(this, options, name, messageSite, lineNumber);
 		}
 
 		/// <summary>
@@ -1037,7 +1037,7 @@ namespace EpicGames.UHT.Types
 		/// <returns>The located type of null if not found</returns>
 		public UhtType? FindType(UhtFindOptions options, ref UhtToken name, IUhtMessageSite? messageSite = null)
 		{
-			return this.Session.FindType(this, options, ref name, messageSite);
+			return Session.FindType(this, options, ref name, messageSite);
 		}
 
 		/// <summary>
@@ -1050,7 +1050,7 @@ namespace EpicGames.UHT.Types
 		/// <returns>The located type of null if not found</returns>
 		public UhtType? FindType(UhtFindOptions options, UhtTokenList identifiers, IUhtMessageSite? messageSite = null, int lineNumber = -1)
 		{
-			return this.Session.FindType(this, options, identifiers, messageSite, lineNumber);
+			return Session.FindType(this, options, identifiers, messageSite, lineNumber);
 		}
 
 		/// <summary>
@@ -1063,7 +1063,7 @@ namespace EpicGames.UHT.Types
 		/// <returns>The located type of null if not found</returns>
 		public UhtType? FindType(UhtFindOptions options, UhtToken[] identifiers, IUhtMessageSite? messageSite = null, int lineNumber = -1)
 		{
-			return this.Session.FindType(this, options, identifiers, messageSite, lineNumber);
+			return Session.FindType(this, options, identifiers, messageSite, lineNumber);
 		}
 		#endregion
 
@@ -1082,12 +1082,12 @@ namespace EpicGames.UHT.Types
 
 		private void ValidateMetaData()
 		{
-			UhtSpecifierValidatorTable? table = this.SpecifierValidatorTable;
+			UhtSpecifierValidatorTable? table = SpecifierValidatorTable;
 			if (table != null)
 			{
-				if (this.MetaData.Dictionary != null && this.MetaData.Dictionary.Count > 0)
+				if (MetaData.Dictionary != null && MetaData.Dictionary.Count > 0)
 				{
-					foreach (KeyValuePair<UhtMetaDataKey, string> kvp in this.MetaData.Dictionary)
+					foreach (KeyValuePair<UhtMetaDataKey, string> kvp in MetaData.Dictionary)
 					{
 						if (table.TryGetValue(kvp.Key.Name, out UhtSpecifierValidator? specifier))
 						{
@@ -1100,16 +1100,16 @@ namespace EpicGames.UHT.Types
 
 		private void ValidateDocumentationPolicy()
 		{
-			if (UhtDocumentationPolicy.TryGet(this.DocumentationPolicyName, out UhtDocumentationPolicy policy))
+			if (UhtDocumentationPolicy.TryGet(DocumentationPolicyName, out UhtDocumentationPolicy policy))
 			{
 				if (policy.PolicySet)
 				{
-					this.ValidateDocumentationPolicy(policy);
+					ValidateDocumentationPolicy(policy);
 				}
 			}
 			else
 			{
-				this.LogError(this.MetaData.LineNumber, $"Documentation policy '{this.DocumentationPolicyName}' is not known");
+				this.LogError(MetaData.LineNumber, $"Documentation policy '{DocumentationPolicyName}' is not known");
 			}
 		}
 
@@ -1158,7 +1158,7 @@ namespace EpicGames.UHT.Types
 		{
 			if (this != stopOuter)
 			{
-				UhtType? outer = this.Outer;
+				UhtType? outer = Outer;
 				if (outer != null && outer != stopOuter)
 				{
 					outer.AppendPathName(builder, stopOuter);
@@ -1180,7 +1180,7 @@ namespace EpicGames.UHT.Types
 						builder.Append('.');
 					}
 				}
-				builder.Append(this.EngineName);
+				builder.Append(EngineName);
 			}
 		}
 
@@ -1190,7 +1190,7 @@ namespace EpicGames.UHT.Types
 		/// <param name="builder">Destination builder</param>
 		public virtual void AppendFullName(StringBuilder builder)
 		{
-			builder.Append(this.EngineClassName);
+			builder.Append(EngineClassName);
 			builder.Append(' ');
 			AppendPathName(builder);
 		}
@@ -1201,7 +1201,7 @@ namespace EpicGames.UHT.Types
 		/// <returns>The display name for this object.</returns>
 		public string GetDisplayNameText()
 		{
-			if (this.MetaData.TryGetValue(UhtNames.DisplayName, out string? nativeDisplayName))
+			if (MetaData.TryGetValue(UhtNames.DisplayName, out string? nativeDisplayName))
 			{
 				return nativeDisplayName;
 			}
@@ -1219,12 +1219,12 @@ namespace EpicGames.UHT.Types
 			string nativeToolTip = String.Empty;
 			if (shortToolTip)
 			{
-				nativeToolTip = this.MetaData.GetValueOrDefault(UhtNames.ShortToolTip);
+				nativeToolTip = MetaData.GetValueOrDefault(UhtNames.ShortToolTip);
 			}
 
 			if (nativeToolTip.Length == 0)
 			{
-				nativeToolTip = this.MetaData.GetValueOrDefault(UhtNames.ToolTip);
+				nativeToolTip = MetaData.GetValueOrDefault(UhtNames.ToolTip);
 			}
 
 			if (nativeToolTip.Length == 0)
@@ -1243,7 +1243,7 @@ namespace EpicGames.UHT.Types
 		/// <returns>The generated string</returns>
 		protected virtual string GetDisplayStringFromEngineName()
 		{
-			return UhtFCString.NameToDisplayString(this.EngineName, false);
+			return UhtFCString.NameToDisplayString(EngineName, false);
 		}
 		#endregion
 
diff --git a/Engine/Source/Programs/Shared/EpicGames.UHT/Utils/UhtBuffer.cs b/Engine/Source/Programs/Shared/EpicGames.UHT/Utils/UhtBuffer.cs
--- a/Engine/Source/Programs/Shared/EpicGames.UHT/Utils/UhtBuffer.cs
+++ b/Engine/Source/Programs/Shared/EpicGames.UHT/Utils/UhtBuffer.cs
@@ -66,8 +66,8 @@ namespace EpicGames.UHT.Utils
 		/// <param name="bucketSize">The size all blocks in this bucket</param>
 		private UhtBuffer(int size, int bucket, int bucketSize)
 		{
-			this.Block = new char[bucketSize];
-			this.Bucket = bucket;
+			Block = new char[bucketSize];
+			Bucket = bucket;
 			Reset(size);
 		}
 
@@ -77,7 +77,7 @@ namespace EpicGames.UHT.Utils
 		/// <param name="size"></param>
 		public void Reset(int size)
 		{
-			this.Memory = new Memory<char>(this.Block, 0, size);
+			Memory = new Memory<char>(Block, 0, size);
 		}
 
 		/// <summary>
@@ -188,7 +188,7 @@ namespace EpicGames.UHT.Utils
 		/// <param name="size">The size to borrow</param>
 		public UhtBorrowBuffer(int size)
 		{
-			this.Buffer = UhtBuffer.Borrow(size);
+			Buffer = UhtBuffer.Borrow(size);
 		}
 
 		/// <summary>
@@ -197,7 +197,7 @@ namespace EpicGames.UHT.Utils
 		/// <param name="builder">Initial contents of the buffer</param>
 		public UhtBorrowBuffer(StringBuilder builder)
 		{
-			this.Buffer = UhtBuffer.Borrow(builder);
+			Buffer = UhtBuffer.Borrow(builder);
 		}
 
 		/// <summary>
@@ -208,7 +208,7 @@ namespace EpicGames.UHT.Utils
 		/// <param name="length">Length of the data in the builder</param>
 		public UhtBorrowBuffer(StringBuilder builder, int startIndex, int length)
 		{
-			this.Buffer = UhtBuffer.Borrow(builder, startIndex, length);
+			Buffer = UhtBuffer.Borrow(builder, startIndex, length);
 		}
 
 		/// <summary>
@@ -216,7 +216,7 @@ namespace EpicGames.UHT.Utils
 		/// </summary>
 		public void Dispose()
 		{
-			UhtBuffer.Return(this.Buffer);
+			UhtBuffer.Return(Buffer);
 		}
 	}
 }
diff --git a/Engine/Source/Programs/Shared/EpicGames.UHT/Utils/UhtException.cs b/Engine/Source/Programs/Shared/EpicGames.UHT/Utils/UhtException.cs
--- a/Engine/Source/Programs/Shared/EpicGames.UHT/Utils/UhtException.cs
+++ b/Engine/Source/Programs/Shared/EpicGames.UHT/Utils/UhtException.cs
@@ -29,7 +29,7 @@ namespace EpicGames.UHT.Utils
 		/// <param name="message">Text of the error</param>
 		public UhtException(string message)
 		{
-			this.UhtMessage = UhtMessage.MakeMessage(UhtMessageType.Error, null, null, 1, message);
+			UhtMessage = UhtMessage.MakeMessage(UhtMessageType.Error, null, null, 1, message);
 		}
 
 		/// <summary>
@@ -44,7 +44,7 @@ namespace EpicGames.UHT.Utils
 			{
 				message = $"{message} while parsing {UhtMessage.FormatContext(extraContext)}";
 			}
-			this.UhtMessage = UhtMessage.MakeMessage(UhtMessageType.Error, messageSite.MessageSource, null, messageSite.GetLineNumber(), message);
+			UhtMessage = UhtMessage.MakeMessage(UhtMessageType.Error, messageSite.MessageSource, null, messageSite.GetLineNumber(), message);
 		}
 
 		/// <summary>
@@ -60,11 +60,11 @@ namespace EpicGames.UHT.Utils
 			{
 				message = $"{message} while parsing {UhtMessage.FormatContext(extraContext)}";
 			}
-			this.UhtMessage = UhtMessage.MakeMessage(UhtMessageType.Error, messageSite.MessageSource, null, messageSite.GetLineNumber(lineNumber), message);
+			UhtMessage = UhtMessage.MakeMessage(UhtMessageType.Error, messageSite.MessageSource, null, messageSite.GetLineNumber(lineNumber), message);
 		}
 
 		/// <inheritdoc/>
-		public override string Message => this.UhtMessage.Message;
+		public override string Message => UhtMessage.Message;
 	}
 
 	/// <summary>
@@ -86,7 +86,7 @@ namespace EpicGames.UHT.Utils
 			string message = expected != null
 				? $"Found {UhtMessage.FormatContext(got)} when expecting {UhtMessage.FormatContext(expected)}{FormatExtraContext(extraContext)}"
 				: $"Found {UhtMessage.FormatContext(got)}{FormatExtraContext(extraContext)}";
-			this.UhtMessage = UhtMessage.MakeMessage(UhtMessageType.Error, messageSite.MessageSource, null, got.InputLine, message);
+			UhtMessage = UhtMessage.MakeMessage(UhtMessageType.Error, messageSite.MessageSource, null, got.InputLine, message);
 		}
 
 		/// <summary>
@@ -119,7 +119,7 @@ namespace EpicGames.UHT.Utils
 		/// <param name="message">Text of the error</param>
 		public UhtIceException(string message)
 		{
-			this.UhtMessage = UhtMessage.MakeMessage(UhtMessageType.Ice, null, null, 1, message);
+			UhtMessage = UhtMessage.MakeMessage(UhtMessageType.Ice, null, null, 1, message);
 		}
 	}
 }
diff --git a/Engine/Source/Programs/Shared/EpicGames.UHT/Utils/UhtFCString.cs b/Engine/Source/Programs/Shared/EpicGames.UHT/Utils/UhtFCString.cs
--- a/Engine/Source/Programs/Shared/EpicGames.UHT/Utils/UhtFCString.cs
+++ b/Engine/Source/Programs/Shared/EpicGames.UHT/Utils/UhtFCString.cs
@@ -468,7 +468,7 @@ namespace EpicGames.UHT.Utils
 		/// <param name="builder">Destination string builder</param>
 		public static void TabsToSpaces(ReadOnlySpan<char> span, int tabSpacing, bool emulateCrBug, int tabIndex, StringBuilder builder)
 		{
-			// Locate the last \n since all tabs have to be computed relative to this.
+			// Locate the last \n since all tabs have to be computed relative to 
 			int crPos = span[..tabIndex].LastIndexOf('\n') + 1;
 
 			int committed = 0;
diff --git a/Engine/Source/Programs/Shared/EpicGames.UHT/Utils/UhtHash.cs b/Engine/Source/Programs/Shared/EpicGames.UHT/Utils/UhtHash.cs
--- a/Engine/Source/Programs/Shared/EpicGames.UHT/Utils/UhtHash.cs
+++ b/Engine/Source/Programs/Shared/EpicGames.UHT/Utils/UhtHash.cs
@@ -18,7 +18,7 @@ namespace EpicGames.UHT.Utils
 		/// </summary>
 		public void Begin()
 		{
-			this._hash = 0;
+			_hash = 0;
 		}
 
 		/// <summary>
@@ -49,7 +49,7 @@ namespace EpicGames.UHT.Utils
 		/// <returns>Final hash value</returns>
 		public uint End()
 		{
-			return (uint)(this._hash + (this._hash >> 32));
+			return (uint)(_hash + (_hash >> 32));
 		}
 
 		/// <summary>
@@ -73,7 +73,7 @@ namespace EpicGames.UHT.Utils
 				{
 					fixed (char* textPtr = text)
 					{
-						this._hash = CityHash.CityHash64WithSeed((byte*)textPtr, (uint)(text.Length * sizeof(char)), this._hash);
+						_hash = CityHash.CityHash64WithSeed((byte*)textPtr, (uint)(text.Length * sizeof(char)), _hash);
 					}
 				}
 			}
diff --git a/Engine/Source/Programs/Shared/EpicGames.UHT/Utils/UhtMessage.cs b/Engine/Source/Programs/Shared/EpicGames.UHT/Utils/UhtMessage.cs
--- a/Engine/Source/Programs/Shared/EpicGames.UHT/Utils/UhtMessage.cs
+++ b/Engine/Source/Programs/Shared/EpicGames.UHT/Utils/UhtMessage.cs
@@ -296,9 +296,9 @@ namespace EpicGames.UHT.Utils
 
 		#region IUHTMessageSite implementation
 		/// <inheritdoc/>
-		public IUhtMessageSession MessageSession => this._messageSession;
+		public IUhtMessageSession MessageSession => _messageSession;
 		/// <inheritdoc/>
-		public IUhtMessageSource? MessageSource { get => this._messageSource; set => this._messageSource = value; }
+		public IUhtMessageSource? MessageSource { get => _messageSource; set => _messageSource = value; }
 		/// <inheritdoc/>
 		public IUhtMessageLineNumber? MessageLineNumber => null;
 		#endregion
@@ -310,8 +310,8 @@ namespace EpicGames.UHT.Utils
 		/// <param name="messageSource">Source for the messages</param>
 		public UhtSimpleMessageSite(IUhtMessageSession messageSession, IUhtMessageSource? messageSource = null)
 		{
-			this._messageSession = messageSession;
-			this._messageSource = messageSource;
+			_messageSession = messageSession;
+			_messageSource = messageSource;
 		}
 	}
 
@@ -322,9 +322,9 @@ namespace EpicGames.UHT.Utils
 	{
 		#region IUHTMessageSource implementation
 		/// <inheritdoc/>
-		public string MessageFilePath => this._filePath;
+		public string MessageFilePath => _filePath;
 		/// <inheritdoc/>
-		public string MessageFullFilePath => this._filePath;
+		public string MessageFullFilePath => _filePath;
 		/// <inheritdoc/>
 		public bool MessageIsFragment => false;
 		/// <inheritdoc/>
@@ -344,8 +344,8 @@ namespace EpicGames.UHT.Utils
 		/// <param name="filePath">File associated with the site</param>
 		public UhtSimpleFileMessageSite(IUhtMessageSession messageSession, string filePath) : base(messageSession, null)
 		{
-			this._filePath = filePath;
-			this.MessageSource = this;
+			_filePath = filePath;
+			MessageSource = this;
 		}
 	}
 
@@ -518,7 +518,7 @@ namespace EpicGames.UHT.Utils
 		private static readonly ThreadLocal<UhtTlsMessageExtraContext> s_tls = new(() => new UhtTlsMessageExtraContext());
 
 		#region IUHTMessageExtraContext implementation
-		IEnumerable<object?>? IUhtMessageExtraContext.MessageExtraContext => this._extraContexts;
+		IEnumerable<object?>? IUhtMessageExtraContext.MessageExtraContext => _extraContexts;
 		#endregion
 
 		/// <summary>
@@ -527,11 +527,11 @@ namespace EpicGames.UHT.Utils
 		/// <param name="exceptionContext"></param>
 		public void PushExtraContext(object? exceptionContext)
 		{
-			if (this._extraContexts == null)
+			if (_extraContexts == null)
 			{
-				this._extraContexts = new Stack<object?>(8);
+				_extraContexts = new Stack<object?>(8);
 			}
-			this._extraContexts.Push(exceptionContext);
+			_extraContexts.Push(exceptionContext);
 		}
 
 		/// <summary>
@@ -539,9 +539,9 @@ namespace EpicGames.UHT.Utils
 		/// </summary>
 		public void PopExtraContext()
 		{
-			if (this._extraContexts != null)
+			if (_extraContexts != null)
 			{
-				this._extraContexts.Pop();
+				_extraContexts.Pop();
 			}
 		}
 
@@ -571,10 +571,10 @@ namespace EpicGames.UHT.Utils
 		/// <param name="extraContext">Extra context to be added</param>
 		public UhtMessageContext(object? extraContext)
 		{
-			this._stack = UhtTlsMessageExtraContext.GetTls();
-			if (this._stack != null)
+			_stack = UhtTlsMessageExtraContext.GetTls();
+			if (_stack != null)
 			{
-				this._stack.PushExtraContext(extraContext);
+				_stack.PushExtraContext(extraContext);
 			}
 		}
 
@@ -585,10 +585,10 @@ namespace EpicGames.UHT.Utils
 		/// <param name="extraContext">New extra context</param>
 		public void Reset(object? extraContext)
 		{
-			if (this._stack != null)
+			if (_stack != null)
 			{
-				this._stack.PopExtraContext();
-				this._stack.PushExtraContext(extraContext);
+				_stack.PopExtraContext();
+				_stack.PushExtraContext(extraContext);
 			}
 		}
 
@@ -597,9 +597,9 @@ namespace EpicGames.UHT.Utils
 		/// </summary>
 		public void Dispose()
 		{
-			if (this._stack != null)
+			if (_stack != null)
 			{
-				this._stack.PopExtraContext();
+				_stack.PopExtraContext();
 			}
 		}
 	}
diff --git a/Engine/Source/Programs/Shared/EpicGames.UHT/Utils/UhtReferenceCollector.cs b/Engine/Source/Programs/Shared/EpicGames.UHT/Utils/UhtReferenceCollector.cs
--- a/Engine/Source/Programs/Shared/EpicGames.UHT/Utils/UhtReferenceCollector.cs
+++ b/Engine/Source/Programs/Shared/EpicGames.UHT/Utils/UhtReferenceCollector.cs
@@ -106,9 +106,9 @@ namespace EpicGames.UHT.Utils
 			if (obj != null)
 			{
 				int key = GetKey(obj, registered);
-				if (this.Uniques.Add(key))
+				if (Uniques.Add(key))
 				{
-					this.References.Add(key);
+					References.Add(key);
 				}
 			}
 		}
@@ -121,10 +121,10 @@ namespace EpicGames.UHT.Utils
 		public ReadOnlyMemory<string> GetSortedReferences(GetReferenceStringDelegate referenceStringDelegate)
 		{
 			// Collect the unsorted array
-			string[] sorted = new string[this.References.Count];
-			for (int index = 0; index < this.References.Count; ++index)
+			string[] sorted = new string[References.Count];
+			for (int index = 0; index < References.Count; ++index)
 			{
-				int key = this.References[index];
+				int key = References[index];
 				UngetKey(key, out int objectIndex, out bool registered);
 				sorted[index] = referenceStringDelegate(objectIndex, registered);
 			}
@@ -134,7 +134,7 @@ namespace EpicGames.UHT.Utils
 
 			// Remove duplicates.  In some instances the different keys might return the same string.
 			// This removes those duplicates
-			if (this.References.Count > 1)
+			if (References.Count > 1)
 			{
 				int priorOut = 0;
 				for (int index = 1; index < sorted.Length; ++index)
@@ -200,7 +200,7 @@ namespace EpicGames.UHT.Utils
 			CrossModule.Add(obj, registered);
 			if (obj != null && obj is not UhtPackage && registered)
 			{
-				this.ReferencedHeaders.Add(obj.HeaderFile);
+				ReferencedHeaders.Add(obj.HeaderFile);
 			}
 		}
 
@@ -211,7 +211,7 @@ namespace EpicGames.UHT.Utils
 		/// <param name="registered">True if the object being declared must be registered</param>
 		public void AddDeclaration(UhtObject obj, bool registered)
 		{
-			this.Declaration.Add(obj, registered);
+			Declaration.Add(obj, registered);
 		}
 
 		/// <summary>
@@ -220,7 +220,7 @@ namespace EpicGames.UHT.Utils
 		/// <param name="field">Field being added</param>
 		public void AddSingleton(UhtField field)
 		{
-			this.Singletons.Add(field);
+			Singletons.Add(field);
 		}
 
 		/// <summary>
@@ -229,7 +229,7 @@ namespace EpicGames.UHT.Utils
 		/// <param name="field">Type to be exported</param>
 		public void AddExportType(UhtField field)
 		{
-			this.ExportTypes.Add(field);
+			ExportTypes.Add(field);
 		}
 
 		/// <summary>
@@ -240,7 +240,7 @@ namespace EpicGames.UHT.Utils
 		{
 			if (!String.IsNullOrEmpty(declaration))
 			{
-				this.ForwardDeclarations.Add(declaration);
+				ForwardDeclarations.Add(declaration);
 			}
 		}
 	}
