Removed the use of "this." where possible in EpicGames.UHT.

diff --git a/Engine/Source/Programs/Shared/EpicGames.UHT/Utils/UhtSession.cs b/Engine/Source/Programs/Shared/EpicGames.UHT/Utils/UhtSession.cs
--- a/Engine/Source/Programs/Shared/EpicGames.UHT/Utils/UhtSession.cs
+++ b/Engine/Source/Programs/Shared/EpicGames.UHT/Utils/UhtSession.cs
@@ -83,12 +83,12 @@ namespace EpicGames.UHT.Utils
 		/// <summary>
 		/// UHT Session
 		/// </summary>
-		public UhtSession Session => this._session;
+		public UhtSession Session => _session;
 
 		/// <summary>
 		/// Plugin module
 		/// </summary>
-		public UHTManifest.Module? PluginModule => this._pluginModule;
+		public UHTManifest.Module? PluginModule => _pluginModule;
 
 		/// <summary>
 		/// Collection of error from mismatches with the reference files
@@ -123,14 +123,14 @@ namespace EpicGames.UHT.Utils
 		/// <param name="exporter">Exporter being run</param>
 		public UhtExportFactory(UhtSession session, UHTManifest.Module? pluginModule, UhtExporter exporter)
 		{
-			this._session = session;
-			this._pluginModule = pluginModule;
-			this.Exporter = exporter;
-			if (this.Session.ReferenceMode != UhtReferenceMode.None)
+			_session = session;
+			_pluginModule = pluginModule;
+			Exporter = exporter;
+			if (Session.ReferenceMode != UhtReferenceMode.None)
 			{
-				this.ReferenceDirectory = Path.Combine(this.Session.ReferenceDirectory, this.Exporter.Name);
-				this.VerifyDirectory = Path.Combine(this.Session.VerifyDirectory, this.Exporter.Name);
-				Directory.CreateDirectory(this.Session.ReferenceMode == UhtReferenceMode.Reference ? this.ReferenceDirectory : this.VerifyDirectory);
+				ReferenceDirectory = Path.Combine(Session.ReferenceDirectory, Exporter.Name);
+				VerifyDirectory = Path.Combine(Session.VerifyDirectory, Exporter.Name);
+				Directory.CreateDirectory(Session.ReferenceMode == UhtReferenceMode.Reference ? ReferenceDirectory : VerifyDirectory);
 			}
 		}
 
@@ -167,7 +167,7 @@ namespace EpicGames.UHT.Utils
 		/// <returns>Task object or null if the task was immediately executed.</returns>
 		public Task? CreateTask(List<Task?>? prereqs, UhtExportTaskDelegate action)
 		{
-			if (this.Session.GoWide)
+			if (Session.GoWide)
 			{
 				Task[]? prereqTasks = prereqs?.Where(x => x != null).Cast<Task>().ToArray();
 				if (prereqTasks != null && prereqTasks.Length > 0)
@@ -226,11 +226,11 @@ namespace EpicGames.UHT.Utils
 		/// <returns>Output file path</returns>
 		public string MakePath(string fileName, string extension)
 		{
-			if (this.PluginModule == null)
+			if (PluginModule == null)
 			{
 				throw new UhtIceException("MakePath with just a filename and extension can not be called from non-plugin exporters");
 			}
-			return MakePath(this.PluginModule, fileName, extension);
+			return MakePath(PluginModule, fileName, extension);
 		}
 
 		/// <summary>
@@ -239,7 +239,7 @@ namespace EpicGames.UHT.Utils
 		/// <param name="filePath">External dependency to add</param>
 		public void AddExternalDependency(string filePath)
 		{
-			this.ExternalDependencies.Add(filePath);
+			ExternalDependencies.Add(filePath);
 		}
 
 		private string MakePath(UHTManifest.Module module, string fileName, string suffix)
@@ -262,7 +262,7 @@ namespace EpicGames.UHT.Utils
 
 			ReadOnlySpan<char> exportedSpan = exported.Span;
 
-			if (this.Session.ReferenceMode != UhtReferenceMode.None)
+			if (Session.ReferenceMode != UhtReferenceMode.None)
 			{
 				string fileName = Path.GetFileName(filePath);
 
@@ -271,10 +271,10 @@ namespace EpicGames.UHT.Utils
 				{
 					UhtExportFactory.s_writeRefSemaphore.WaitOne();
 					{
-						string outPath = Path.Combine(this.Session.ReferenceMode == UhtReferenceMode.Reference ? this.ReferenceDirectory : this.VerifyDirectory, fileName);
-						if (!this.Session.WriteSource(outPath, exported.Span))
+						string outPath = Path.Combine(Session.ReferenceMode == UhtReferenceMode.Reference ? ReferenceDirectory : VerifyDirectory, fileName);
+						if (!Session.WriteSource(outPath, exported.Span))
 						{
-							new UhtSimpleFileMessageSite(this.Session, outPath).LogWarning($"Unable to write reference file {outPath}");
+							new UhtSimpleFileMessageSite(Session, outPath).LogWarning($"Unable to write reference file {outPath}");
 						}
 					}
 				}
@@ -284,11 +284,11 @@ namespace EpicGames.UHT.Utils
 				}
 
 				// If we are verifying, read the existing file and check the contents
-				if (this.Session.ReferenceMode == UhtReferenceMode.Verify)
+				if (Session.ReferenceMode == UhtReferenceMode.Verify)
 				{
 					string message = String.Empty;
-					string refPath = Path.Combine(this.ReferenceDirectory, fileName);
-					UhtBuffer? existingRef = this.Session.ReadSourceToBuffer(refPath);
+					string refPath = Path.Combine(ReferenceDirectory, fileName);
+					UhtBuffer? existingRef = Session.ReadSourceToBuffer(refPath);
 					if (existingRef != null)
 					{
 						ReadOnlySpan<char> existingSpan = existingRef.Memory.Span;
@@ -306,28 +306,28 @@ namespace EpicGames.UHT.Utils
 					if (!String.IsNullOrEmpty(message))
 					{
 						Log.TraceInformation(message);
-						lock (this.ReferenceErrorMessages)
+						lock (ReferenceErrorMessages)
 						{
-							this.ReferenceErrorMessages.Add(message, true);
+							ReferenceErrorMessages.Add(message, true);
 						}
 					}
 				}
 			}
 
 			// Check to see if the contents have changed
-			UhtBuffer? original = this.Session.ReadSourceToBuffer(filePath);
+			UhtBuffer? original = Session.ReadSourceToBuffer(filePath);
 			bool save = original == null;
 			if (original != null)
 			{
 				ReadOnlySpan<char> originalSpan = original.Memory.Span;
 				if (originalSpan.CompareTo(exportedSpan, StringComparison.Ordinal) != 0)
 				{
-					if (this.Session.FailIfGeneratedCodeChanges)
+					if (Session.FailIfGeneratedCodeChanges)
 					{
 						string conflictPath = filePath + ".conflict";
-						if (!this.Session.WriteSource(conflictPath, exported.Span))
+						if (!Session.WriteSource(conflictPath, exported.Span))
 						{
-							new UhtSimpleFileMessageSite(this.Session, filePath).LogError($"Changes to generated code are not allowed - conflicts written to '{conflictPath}'");
+							new UhtSimpleFileMessageSite(Session, filePath).LogError($"Changes to generated code are not allowed - conflicts written to '{conflictPath}'");
 						}
 					}
 					save = true;
@@ -336,11 +336,11 @@ namespace EpicGames.UHT.Utils
 			}
 
 			// If changed of the original didn't exist, then save the new version
-			if (save && !this.Session.NoOutput)
+			if (save && !Session.NoOutput)
 			{
-				if (!this.Session.WriteSource(tempFilePath, exported.Span))
+				if (!Session.WriteSource(tempFilePath, exported.Span))
 				{
-					new UhtSimpleFileMessageSite(this.Session, filePath).LogWarning($"Failed to save export file: '{tempFilePath}'");
+					new UhtSimpleFileMessageSite(Session, filePath).LogWarning($"Failed to save export file: '{tempFilePath}'");
 				}
 			}
 			else
@@ -349,9 +349,9 @@ namespace EpicGames.UHT.Utils
 			}
 
 			// Add this to the list of outputs
-			lock (this.Outputs)
+			lock (Outputs)
 			{
-				this.Outputs.Add(new Output { FilePath = filePath, TempFilePath = tempFilePath, Saved = save });
+				Outputs.Add(new Output { FilePath = filePath, TempFilePath = tempFilePath, Saved = save });
 			}
 		}
 
@@ -362,10 +362,10 @@ namespace EpicGames.UHT.Utils
 		{
 
 			// Invoke the exported via the delegate
-			this.Exporter.Delegate(this);
+			Exporter.Delegate(this);
 
 			// If outputs were exported
-			if (this.Outputs.Count > 0)
+			if (Outputs.Count > 0)
 			{
 
 				// These outputs are used to cull old outputs from the directories
@@ -373,7 +373,7 @@ namespace EpicGames.UHT.Utils
 				List<UhtExportFactory.Output> saves = new();
 
 				// Collect information about the outputs
-				foreach (UhtExportFactory.Output output in this.Outputs)
+				foreach (UhtExportFactory.Output output in Outputs)
 				{
 
 					// Add this output to the list of expected outputs by directory
@@ -396,7 +396,7 @@ namespace EpicGames.UHT.Utils
 				}
 
 				// Perform the renames
-				if (this.Session.GoWide)
+				if (Session.GoWide)
 				{
 					Parallel.ForEach(saves, (UhtExportFactory.Output output) =>
 					{
@@ -412,10 +412,10 @@ namespace EpicGames.UHT.Utils
 				}
 
 				// Perform the culling of the output directories
-				if (this.Session.CullOutput && !this.Session.NoOutput &&
-					(this.Exporter.CppFilters != null || this.Exporter.HeaderFilters != null || this.Exporter.OtherFilters != null))
+				if (Session.CullOutput && !Session.NoOutput &&
+					(Exporter.CppFilters != null || Exporter.HeaderFilters != null || Exporter.OtherFilters != null))
 				{
-					if (this.Session.GoWide)
+					if (Session.GoWide)
 					{
 						Parallel.ForEach(outputsByDirectory, (KeyValuePair<string, HashSet<string>> kvp) =>
 						{
@@ -440,7 +440,7 @@ namespace EpicGames.UHT.Utils
 		/// <param name="output">The output file to rename</param>
 		private void RenameSource(UhtExportFactory.Output output)
 		{
-			this.Session.RenameSource(output.TempFilePath, output.FilePath);
+			Session.RenameSource(output.TempFilePath, output.FilePath);
 		}
 
 		/// <summary>
@@ -458,9 +458,9 @@ namespace EpicGames.UHT.Utils
 					continue;
 				}
 
-				if (IsFilterMatch(fileName, this.Exporter.CppFilters) ||
-					IsFilterMatch(fileName, this.Exporter.HeaderFilters) ||
-					IsFilterMatch(fileName, this.Exporter.OtherFilters))
+				if (IsFilterMatch(fileName, Exporter.CppFilters) ||
+					IsFilterMatch(fileName, Exporter.HeaderFilters) ||
+					IsFilterMatch(fileName, Exporter.OtherFilters))
 				{
 					try
 					{
@@ -537,13 +537,13 @@ namespace EpicGames.UHT.Utils
 			/// <returns>Index starting at zero</returns>
 			public int GetNext()
 			{
-				return Interlocked.Increment(ref this._count) - 1;
+				return Interlocked.Increment(ref _count) - 1;
 			}
 
 			/// <summary>
 			/// The number of times GetNext was called.
 			/// </summary>
-			public int Count => Interlocked.Add(ref this._count, 0) + 1;
+			public int Count => Interlocked.Add(ref _count, 0) + 1;
 		}
 
 		/// <summary>
@@ -667,29 +667,29 @@ namespace EpicGames.UHT.Utils
 		/// <summary>
 		/// Manifest data from the manifest file
 		/// </summary>
-		public UHTManifest? Manifest => this.ManifestFile?.Manifest;
+		public UHTManifest? Manifest => ManifestFile?.Manifest;
 
 		/// <summary>
 		/// Collection of packages from the manifest
 		/// </summary>
-		public IReadOnlyList<UhtPackage> Packages => this._packages;
+		public IReadOnlyList<UhtPackage> Packages => _packages;
 
 		/// <summary>
 		/// Collection of header files from the manifest.  The header files will also appear as the children 
 		/// of the packages
 		/// </summary>
-		public IReadOnlyList<UhtHeaderFile> HeaderFiles => this._headerFiles;
+		public IReadOnlyList<UhtHeaderFile> HeaderFiles => _headerFiles;
 
 		/// <summary>
 		/// Collection of header files topologically sorted.  This will not be populated until after header files
 		/// are parsed and resolved.
 		/// </summary>
-		public IReadOnlyList<UhtHeaderFile> SortedHeaderFiles => this._sortedHeaderFiles;
+		public IReadOnlyList<UhtHeaderFile> SortedHeaderFiles => _sortedHeaderFiles;
 
 		/// <summary>
 		/// Dictionary of stripped file name to the header file
 		/// </summary>
-		public IReadOnlyDictionary<string, UhtHeaderFile> HeaderFileDictionary => this._headerFileDictionary;
+		public IReadOnlyDictionary<string, UhtHeaderFile> HeaderFileDictionary => _headerFileDictionary;
 
 		/// <summary>
 		/// After headers are parsed, returns the UObject class.
@@ -698,11 +698,11 @@ namespace EpicGames.UHT.Utils
 		{
 			get
 			{
-				if (this._uobject == null)
+				if (_uobject == null)
 				{
 					throw new UhtIceException("UObject was not defined.");
 				}
-				return this._uobject;
+				return _uobject;
 			}
 		}
 
@@ -713,11 +713,11 @@ namespace EpicGames.UHT.Utils
 		{
 			get
 			{
-				if (this._uclass == null)
+				if (_uclass == null)
 				{
 					throw new UhtIceException("UClass was not defined.");
 				}
-				return this._uclass;
+				return _uclass;
 			}
 		}
 
@@ -728,11 +728,11 @@ namespace EpicGames.UHT.Utils
 		{
 			get
 			{
-				if (this._uinterface == null)
+				if (_uinterface == null)
 				{
 					throw new UhtIceException("UInterface was not defined.");
 				}
-				return this._uinterface;
+				return _uinterface;
 			}
 		}
 
@@ -743,11 +743,11 @@ namespace EpicGames.UHT.Utils
 		{
 			get
 			{
-				if (this._iinterface == null)
+				if (_iinterface == null)
 				{
 					throw new UhtIceException("IInterface was not defined.");
 				}
-				return this._iinterface;
+				return _iinterface;
 			}
 		}
 
@@ -791,17 +791,17 @@ namespace EpicGames.UHT.Utils
 		/// <summary>
 		/// The number of errors
 		/// </summary>
-		public long ErrorCount => Interlocked.Read(ref this._errorCount);
+		public long ErrorCount => Interlocked.Read(ref _errorCount);
 
 		/// <summary>
 		/// The number of warnings
 		/// </summary>
-		public long WarningCount => Interlocked.Read(ref this._warningCount);
+		public long WarningCount => Interlocked.Read(ref _warningCount);
 
 		/// <summary>
 		/// True if any errors have occurred or warnings if warnings are to be treated as errors 
 		/// </summary>
-		public bool HasErrors => this.ErrorCount > 0 || (this.WarningsAsErrors && this.WarningCount > 0);
+		public bool HasErrors => ErrorCount > 0 || (WarningsAsErrors && WarningCount > 0);
 
 		#region IUHTMessageSession implementation
 		/// <inheritdoc/>
@@ -818,13 +818,13 @@ namespace EpicGames.UHT.Utils
 		/// <returns>New index</returns>
 		public int GetNextTypeIndex()
 		{
-			return this._typeCounter.GetNext();
+			return _typeCounter.GetNext();
 		}
 
 		/// <summary>
 		/// Return the number of types that have been defined.  This includes all types.
 		/// </summary>
-		public int TypeCount => this._typeCounter.Count;
+		public int TypeCount => _typeCounter.Count;
 
 		/// <summary>
 		/// Return the index for a newly defined packaging
@@ -832,13 +832,13 @@ namespace EpicGames.UHT.Utils
 		/// <returns>New index</returns>
 		public int GetNextHeaderFileTypeIndex()
 		{
-			return this._headerFileTypeCount.GetNext();
+			return _headerFileTypeCount.GetNext();
 		}
 
 		/// <summary>
 		/// Return the number of headers that have been defined
 		/// </summary>
-		public int HeaderFileTypeCount => this._headerFileTypeCount.Count;
+		public int HeaderFileTypeCount => _headerFileTypeCount.Count;
 
 		/// <summary>
 		/// Return the index for a newly defined package
@@ -846,13 +846,13 @@ namespace EpicGames.UHT.Utils
 		/// <returns>New index</returns>
 		public int GetNextPackageTypeIndex()
 		{
-			return this._packageTypeCount.GetNext();
+			return _packageTypeCount.GetNext();
 		}
 
 		/// <summary>
 		/// Return the number of UPackage types that have been defined
 		/// </summary>
-		public int PackageTypeCount => this._packageTypeCount.Count;
+		public int PackageTypeCount => _packageTypeCount.Count;
 
 		/// <summary>
 		/// Return the index for a newly defined object
@@ -860,18 +860,18 @@ namespace EpicGames.UHT.Utils
 		/// <returns>New index</returns>
 		public int GetNextObjectTypeIndex()
 		{
-			return this._objectTypeCount.GetNext();
+			return _objectTypeCount.GetNext();
 		}
 
 		/// <summary>
 		/// Return the total number of UObject types that have been defined
 		/// </summary>
-		public int ObjectTypeCount => this._objectTypeCount.Count;
+		public int ObjectTypeCount => _objectTypeCount.Count;
 
 		/// <summary>
 		/// Return the collection of exporters
 		/// </summary>
-		public UhtExporterTable ExporterTable => this.Tables!.ExporterTable;
+		public UhtExporterTable ExporterTable => Tables!.ExporterTable;
 
 		/// <summary>
 		/// Return the keyword table for the given table name
@@ -880,7 +880,7 @@ namespace EpicGames.UHT.Utils
 		/// <returns>The requested table</returns>
 		public UhtKeywordTable GetKeywordTable(string tableName)
 		{
-			return this.Tables!.KeywordTables.Get(tableName);
+			return Tables!.KeywordTables.Get(tableName);
 		}
 
 		/// <summary>
@@ -890,7 +890,7 @@ namespace EpicGames.UHT.Utils
 		/// <returns>The requested table</returns>
 		public UhtSpecifierTable GetSpecifierTable(string tableName)
 		{
-			return this.Tables!.SpecifierTables.Get(tableName);
+			return Tables!.SpecifierTables.Get(tableName);
 		}
 
 		/// <summary>
@@ -900,7 +900,7 @@ namespace EpicGames.UHT.Utils
 		/// <returns>The requested table</returns>
 		public UhtSpecifierValidatorTable GetSpecifierValidatorTable(string tableName)
 		{
-			return this.Tables!.SpecifierValidatorTables.Get(tableName);
+			return Tables!.SpecifierValidatorTables.Get(tableName);
 		}
 
 		/// <summary>
@@ -910,7 +910,7 @@ namespace EpicGames.UHT.Utils
 		/// <param name="token">Unhandled token</param>
 		public void LogUnhandledKeywordError(IUhtTokenReader tokenReader, UhtToken token)
 		{
-			this.Tables!.KeywordTables.LogUnhandledError(tokenReader, token);
+			Tables!.KeywordTables.LogUnhandledError(tokenReader, token);
 		}
 
 		/// <summary>
@@ -920,7 +920,7 @@ namespace EpicGames.UHT.Utils
 		/// <returns>True if the class name is a property.  False if the class name isn't a property or isn't an engine class.</returns>
 		public bool IsValidPropertyTypeName(StringView name)
 		{
-			return this.Tables!.EngineClassTable.IsValidPropertyTypeName(name);
+			return Tables!.EngineClassTable.IsValidPropertyTypeName(name);
 		}
 
 		/// <summary>
@@ -931,13 +931,13 @@ namespace EpicGames.UHT.Utils
 		/// <returns></returns>
 		public bool TryGetLocTextDefaultValue(StringView name, out UhtLocTextDefaultValue locTextDefaultValue)
 		{
-			return this.Tables!.LocTextDefaultValueTable.TryGet(name, out locTextDefaultValue);
+			return Tables!.LocTextDefaultValueTable.TryGet(name, out locTextDefaultValue);
 		}
 
 		/// <summary>
 		/// Return the default processor
 		/// </summary>
-		public UhtPropertyType DefaultPropertyType => this.Tables!.PropertyTypeTable.Default;
+		public UhtPropertyType DefaultPropertyType => Tables!.PropertyTypeTable.Default;
 
 		/// <summary>
 		/// Return the property type associated with the given name
@@ -947,13 +947,13 @@ namespace EpicGames.UHT.Utils
 		/// <returns></returns>
 		public bool TryGetPropertyType(StringView name, out UhtPropertyType propertyType)
 		{
-			return this.Tables!.PropertyTypeTable.TryGet(name, out propertyType);
+			return Tables!.PropertyTypeTable.TryGet(name, out propertyType);
 		}
 
 		/// <summary>
 		/// Fetch the default sanitizer
 		/// </summary>
-		public UhtStructDefaultValue DefaultStructDefaultValue => this.Tables!.StructDefaultValueTable.Default;
+		public UhtStructDefaultValue DefaultStructDefaultValue => Tables!.StructDefaultValueTable.Default;
 
 		/// <summary>
 		/// Return the structure default value associated with the given name
@@ -963,7 +963,7 @@ namespace EpicGames.UHT.Utils
 		/// <returns></returns>
 		public bool TryGetStructDefaultValue(StringView name, out UhtStructDefaultValue structDefaultValue)
 		{
-			return this.Tables!.StructDefaultValueTable.TryGet(name, out structDefaultValue);
+			return Tables!.StructDefaultValueTable.TryGet(name, out structDefaultValue);
 		}
 
 		/// <summary>
@@ -972,53 +972,53 @@ namespace EpicGames.UHT.Utils
 		/// <param name="manifestFilePath">Path to the manifest file</param>
 		public void Run(string manifestFilePath)
 		{
-			if (this.FileManager == null)
+			if (FileManager == null)
 			{
-				Interlocked.Increment(ref this._errorCount);
+				Interlocked.Increment(ref _errorCount);
 				Log.Logger.LogError("No file manager supplied, aborting.");
 				return;
 			}
 
-			if (this.Config == null)
+			if (Config == null)
 			{
-				Interlocked.Increment(ref this._errorCount);
+				Interlocked.Increment(ref _errorCount);
 				Log.Logger.LogError("No configuration supplied, aborting.");
 				return;
 			}
 
-			if (this.Tables == null)
+			if (Tables == null)
 			{
-				Interlocked.Increment(ref this._errorCount);
+				Interlocked.Increment(ref _errorCount);
 				Log.Logger.LogError("No parsing tables supplied, aborting.");
 				return;
 			}
 
-			switch (this.ReferenceMode)
+			switch (ReferenceMode)
 			{
 				case UhtReferenceMode.None:
 					break;
 
 				case UhtReferenceMode.Reference:
-					if (String.IsNullOrEmpty(this.ReferenceDirectory))
+					if (String.IsNullOrEmpty(ReferenceDirectory))
 					{
 						Log.Logger.LogError("WRITEREF requested but directory not set, ignoring");
-						this.ReferenceMode = UhtReferenceMode.None;
+						ReferenceMode = UhtReferenceMode.None;
 					}
 					break;
 
 				case UhtReferenceMode.Verify:
-					if (String.IsNullOrEmpty(this.ReferenceDirectory) || String.IsNullOrEmpty(this.VerifyDirectory))
+					if (String.IsNullOrEmpty(ReferenceDirectory) || String.IsNullOrEmpty(VerifyDirectory))
 					{
 						Log.Logger.LogError("VERIFYREF requested but directories not set, ignoring");
-						this.ReferenceMode = UhtReferenceMode.None;
+						ReferenceMode = UhtReferenceMode.None;
 					}
 					break;
 			}
 
-			if (this.ReferenceMode != UhtReferenceMode.None)
+			if (ReferenceMode != UhtReferenceMode.None)
 			{
-				string directoryToDelete = this.ReferenceMode == UhtReferenceMode.Reference ? this.ReferenceDirectory : this.VerifyDirectory;
-				this._referenceDeleteTask = Task.Factory.StartNew(() =>
+				string directoryToDelete = ReferenceMode == UhtReferenceMode.Reference ? ReferenceDirectory : VerifyDirectory;
+				_referenceDeleteTask = Task.Factory.StartNew(() =>
 				{
 					try
 					{
@@ -1043,10 +1043,10 @@ namespace EpicGames.UHT.Utils
 			TopologicalSortHeaderFiles();
 
 			// If we are deleting the reference directory, then wait for that task to complete
-			if (this._referenceDeleteTask != null)
+			if (_referenceDeleteTask != null)
 			{
 				Log.Logger.LogTrace("Step - Waiting for reference output to be cleared.");
-				this._referenceDeleteTask.Wait();
+				_referenceDeleteTask.Wait();
 			}
 
 			StepExport();
@@ -1060,7 +1060,7 @@ namespace EpicGames.UHT.Utils
 		/// <param name="action">The lambda to be invoked</param>
 		public void Try(IUhtMessageSource? messageSource, Action action)
 		{
-			if (!this.HasErrors)
+			if (!HasErrors)
 			{
 				try
 				{
@@ -1080,7 +1080,7 @@ namespace EpicGames.UHT.Utils
 		/// <returns>Information about the read source</returns>
 		public UhtSourceFragment ReadSource(string filePath)
 		{
-			if (this.FileManager!.ReadSource(filePath, out UhtSourceFragment fragment))
+			if (FileManager!.ReadSource(filePath, out UhtSourceFragment fragment))
 			{
 				return fragment;
 			}
@@ -1094,7 +1094,7 @@ namespace EpicGames.UHT.Utils
 		/// <returns>Buffer containing the read data or null if not found.  The returned buffer must be returned to the cache via a call to UhtBuffer.Return</returns>
 		public UhtBuffer? ReadSourceToBuffer(string filePath)
 		{
-			return this.FileManager!.ReadOutput(filePath);
+			return FileManager!.ReadOutput(filePath);
 		}
 
 		/// <summary>
@@ -1105,7 +1105,7 @@ namespace EpicGames.UHT.Utils
 		/// <returns>True if the source was written</returns>
 		internal bool WriteSource(string filePath, ReadOnlySpan<char> contents)
 		{
-			return this.FileManager!.WriteOutput(filePath, contents);
+			return FileManager!.WriteOutput(filePath, contents);
 		}
 
 		/// <summary>
@@ -1115,7 +1115,7 @@ namespace EpicGames.UHT.Utils
 		/// <param name="newFilePath">New file path name</param>
 		public void RenameSource(string oldFilePath, string newFilePath)
 		{
-			if (!this.FileManager!.RenameOutput(oldFilePath, newFilePath))
+			if (!FileManager!.RenameOutput(oldFilePath, newFilePath))
 			{
 				new UhtSimpleFileMessageSite(this, newFilePath).LogError($"Failed to rename exported file: '{oldFilePath}'");
 			}
@@ -1129,11 +1129,11 @@ namespace EpicGames.UHT.Utils
 		public UhtEnum? FindRegularEnumValue(string name)
 		{
 			//COMPATIBILITY-TODO - See comment below on a more rebust version of the enum lookup
-			//if (this.RegularEnumValueLookup.TryGetValue(name, out UhtEnum? enumObj))
+			//if (RegularEnumValueLookup.TryGetValue(name, out UhtEnum? enumObj))
 			//{
 			//	return enumObj;
 			//}
-			if (this._fullEnumValueLookup.TryGetValue(name, out EnumAndValue value))
+			if (_fullEnumValueLookup.TryGetValue(name, out EnumAndValue value))
 			{
 				if (value.Value != -1)
 				{
@@ -1141,7 +1141,7 @@ namespace EpicGames.UHT.Utils
 				}
 			}
 
-			if (!name.Contains("::", StringComparison.Ordinal) && this._shortEnumValueLookup.TryGetValue(name, out UhtEnum? enumObj))
+			if (!name.Contains("::", StringComparison.Ordinal) && _shortEnumValueLookup.TryGetValue(name, out UhtEnum? enumObj))
 			{
 				return enumObj;
 			}
@@ -1330,22 +1330,22 @@ namespace EpicGames.UHT.Utils
 			{
 				if (options.HasAnyFlags(UhtFindOptions.CaseCompare))
 				{
-					return this._engineNameSymbolTable.FindCasedType(startingType, options, name);
+					return _engineNameSymbolTable.FindCasedType(startingType, options, name);
 				}
 				else
 				{
-					return this._engineNameSymbolTable.FindCaselessType(startingType, options, name);
+					return _engineNameSymbolTable.FindCaselessType(startingType, options, name);
 				}
 			}
 			else if (options.HasAnyFlags(UhtFindOptions.SourceName))
 			{
 				if (options.HasAnyFlags(UhtFindOptions.CaselessCompare))
 				{
-					return this._sourceNameSymbolTable.FindCaselessType(startingType, options, name);
+					return _sourceNameSymbolTable.FindCaselessType(startingType, options, name);
 				}
 				else
 				{
-					return this._sourceNameSymbolTable.FindCasedType(startingType, options, name);
+					return _sourceNameSymbolTable.FindCasedType(startingType, options, name);
 				}
 			}
 			else
@@ -1376,7 +1376,7 @@ namespace EpicGames.UHT.Utils
 				throw new UhtIceException("No type options specified");
 			}
 
-			if (!this._symbolTablePopulated)
+			if (!_symbolTablePopulated)
 			{
 				throw new UhtIceException("Symbol table has not been populated, don't call FindType until headers are parsed.");
 			}
@@ -1427,7 +1427,7 @@ namespace EpicGames.UHT.Utils
 		/// <returns></returns>
 		public UhtHeaderFile? FindHeaderFile(string name)
 		{
-			if (this._headerFileDictionary.TryGetValue(name, out UhtHeaderFile? headerFile))
+			if (_headerFileDictionary.TryGetValue(name, out UhtHeaderFile? headerFile))
 			{
 				return headerFile;
 			}
@@ -1446,15 +1446,15 @@ namespace EpicGames.UHT.Utils
 				return;
 			}
 
-			lock (this._messages)
+			lock (_messages)
 			{
-				this._messages.Add(message);
+				_messages.Add(message);
 
 				// If we aren't caching messages and this is the first message,
 				// start a task to flush the messages.
-				if (!this.CacheMessages && this._messages.Count == 1)
+				if (!CacheMessages && _messages.Count == 1)
 				{
-					this._messageTask = Task.Factory.StartNew(() => FlushMessages(), CancellationToken.None, TaskCreationOptions.None, TaskScheduler.Default);
+					_messageTask = Task.Factory.StartNew(() => FlushMessages(), CancellationToken.None, TaskCreationOptions.None, TaskScheduler.Default);
 				}
 			}
 
@@ -1462,11 +1462,11 @@ namespace EpicGames.UHT.Utils
 			{
 				case UhtMessageType.Error:
 				case UhtMessageType.Ice:
-					Interlocked.Increment(ref this._errorCount);
+					Interlocked.Increment(ref _errorCount);
 					break;
 
 				case UhtMessageType.Warning:
-					Interlocked.Increment(ref this._warningCount);
+					Interlocked.Increment(ref _warningCount);
 					break;
 
 				case UhtMessageType.Info:
@@ -1485,9 +1485,9 @@ namespace EpicGames.UHT.Utils
 		/// </summary>
 		public void LogMessages()
 		{
-			if (this._messageTask != null)
+			if (_messageTask != null)
 			{
-				this._messageTask.Wait();
+				_messageTask.Wait();
 			}
 
 			foreach (UhtMessage message in FetchOrderedMessages())
@@ -1502,10 +1502,10 @@ namespace EpicGames.UHT.Utils
 		private void FlushMessages()
 		{
 			UhtMessage[]? messageArray = null;
-			lock (this._messages)
+			lock (_messages)
 			{
-				messageArray = this._messages.ToArray();
-				this._messages.Clear();
+				messageArray = _messages.ToArray();
+				_messages.Clear();
 			}
 
 			foreach (UhtMessage message in messageArray)
@@ -1573,10 +1573,10 @@ namespace EpicGames.UHT.Utils
 		/// <param name="newType"></param>
 		public void ReplaceTypeInSymbolTable(UhtType oldType, UhtType newType)
 		{
-			this._sourceNameSymbolTable.Replace(oldType, newType, oldType.SourceName);
+			_sourceNameSymbolTable.Replace(oldType, newType, oldType.SourceName);
 			if (oldType.EngineType.HasEngineName())
 			{
-				this._engineNameSymbolTable.Replace(oldType, newType, oldType.EngineName);
+				_engineNameSymbolTable.Replace(oldType, newType, oldType.EngineName);
 			}
 		}
 
@@ -1587,10 +1587,10 @@ namespace EpicGames.UHT.Utils
 		private IOrderedEnumerable<UhtMessage> FetchOrderedMessages()
 		{
 			List<UhtMessage> messages = new();
-			lock (this._messages)
+			lock (_messages)
 			{
-				messages.AddRange(this._messages);
-				this._messages.Clear();
+				messages.AddRange(_messages);
+				_messages.Clear();
 			}
 			return messages.OrderBy(context => context.FilePath).ThenBy(context => context.LineNumber + context.MessageSource?.MessageFragmentLineNumber);
 		}
@@ -1613,7 +1613,7 @@ namespace EpicGames.UHT.Utils
 			{
 				if (message.MessageSource.MessageIsFragment)
 				{
-					if (this.RelativePathInLog)
+					if (RelativePathInLog)
 					{
 						filePath = message.MessageSource.MessageFragmentFilePath;
 					}
@@ -1626,7 +1626,7 @@ namespace EpicGames.UHT.Utils
 				}
 				else
 				{
-					if (this.RelativePathInLog)
+					if (RelativePathInLog)
 					{
 						filePath = message.MessageSource.MessageFilePath;
 					}
@@ -1699,7 +1699,7 @@ namespace EpicGames.UHT.Utils
 		/// <returns>Normalized path possibly converted to a full path.</returns>
 		private string GetNormalizedFullFilePath(string filePath)
 		{
-			return this.FileManager!.GetFullFilePath(filePath).Replace('\\', '/');
+			return FileManager!.GetFullFilePath(filePath).Replace('\\', '/');
 		}
 
 		private UhtHeaderFileParser? ParseHeaderFile(UhtHeaderFile headerFile)
@@ -1720,33 +1720,33 @@ namespace EpicGames.UHT.Utils
 		#region Run steps
 		private void StepReadManifestFile(string manifestFilePath)
 		{
-			this.ManifestFile = new UhtManifestFile(this, manifestFilePath);
+			ManifestFile = new UhtManifestFile(this, manifestFilePath);
 
-			Try(this.ManifestFile.MessageSource, () =>
+			Try(ManifestFile.MessageSource, () =>
 			{
 				Log.Logger.LogTrace("Step - Read Manifest File");
 
-				this.ManifestFile.Read();
+				ManifestFile.Read();
 
-				if (this.Manifest != null && this.Tables != null)
+				if (Manifest != null && Tables != null)
 				{
-					this.Tables.AddPlugins(this.Manifest.UhtPlugins);
+					Tables.AddPlugins(Manifest.UhtPlugins);
 				}
 			});
 		}
 
 		private void StepPrepareModules()
 		{
-			if (this.ManifestFile == null || this.ManifestFile.Manifest == null)
+			if (ManifestFile == null || ManifestFile.Manifest == null)
 			{
 				return;
 			}
 
-			Try(this.ManifestFile.MessageSource, () =>
+			Try(ManifestFile.MessageSource, () =>
 			{
 				Log.Logger.LogTrace("Step - Prepare Modules");
 
-				foreach (UHTManifest.Module module in this.ManifestFile.Manifest.Modules)
+				foreach (UHTManifest.Module module in ManifestFile.Manifest.Modules)
 				{
 					EPackageFlags packageFlags = EPackageFlags.ContainsScript | EPackageFlags.Compiling;
 
@@ -1788,7 +1788,7 @@ namespace EpicGames.UHT.Utils
 					}
 
 					UhtPackage package = new(this, module, packageFlags);
-					this._packages.Add(package);
+					_packages.Add(package);
 				}
 			});
 		}
@@ -1813,7 +1813,7 @@ namespace EpicGames.UHT.Utils
 					string normalizedExistingFullFilePath = GetNormalizedFullFilePath(existingHeaderFile.FilePath);
 					if (!String.Equals(normalizedFullFilePath, normalizedExistingFullFilePath, StringComparison.OrdinalIgnoreCase))
 					{
-						IUhtMessageSite site = (IUhtMessageSite?)this.ManifestFile ?? this;
+						IUhtMessageSite site = (IUhtMessageSite?)ManifestFile ?? this;
 						site.LogError($"Two headers with the same name is not allowed. '{headerFilePath}' conflicts with '{existingHeaderFile.FilePath}'");
 						continue;
 					}
@@ -1849,16 +1849,16 @@ namespace EpicGames.UHT.Utils
 
 		private void StepPrepareHeaders()
 		{
-			if (this.ManifestFile == null)
+			if (ManifestFile == null)
 			{
 				return;
 			}
 
-			Try(this.ManifestFile.MessageSource, () =>
+			Try(ManifestFile.MessageSource, () =>
 			{
 				Log.Logger.LogTrace("Step - Prepare Headers");
 
-				foreach (UhtPackage package in this._packages)
+				foreach (UhtPackage package in _packages)
 				{
 					if (package.Module != null)
 					{
@@ -1870,9 +1870,9 @@ namespace EpicGames.UHT.Utils
 				}
 
 				// Locate the NoExportTypes.h file and add it to every other header file
-				if (this._headerFileDictionary.TryGetValue("NoExportTypes.h", out UhtHeaderFile? noExportTypes))
+				if (_headerFileDictionary.TryGetValue("NoExportTypes.h", out UhtHeaderFile? noExportTypes))
 				{
-					foreach (UhtPackage package in this._packages)
+					foreach (UhtPackage package in _packages)
 					{
 						foreach (UhtHeaderFile headerFile in package.Children)
 						{
@@ -1888,23 +1888,23 @@ namespace EpicGames.UHT.Utils
 
 		private void StepParseHeaders()
 		{
-			if (this.HasErrors)
+			if (HasErrors)
 			{
 				return;
 			}
 
 			Log.Logger.LogTrace("Step - Parse Headers");
 
-			if (this.GoWide)
+			if (GoWide)
 			{
-				Parallel.ForEach(this._headerFiles, headerFile =>
+				Parallel.ForEach(_headerFiles, headerFile =>
 				{
 					ParseHeaderFile(headerFile);
 				});
 			}
 			else
 			{
-				foreach (UhtHeaderFile headerFile in this._headerFiles)
+				foreach (UhtHeaderFile headerFile in _headerFiles)
 				{
 					ParseHeaderFile(headerFile);
 				}
@@ -1917,17 +1917,17 @@ namespace EpicGames.UHT.Utils
 			{
 				Log.Logger.LogTrace("Step - Populate symbol table");
 
-				this._sourceNameSymbolTable = new UhtSymbolTable(this.TypeCount);
-				this._engineNameSymbolTable = new UhtSymbolTable(this.TypeCount);
+				_sourceNameSymbolTable = new UhtSymbolTable(TypeCount);
+				_engineNameSymbolTable = new UhtSymbolTable(TypeCount);
 
 				PopulateSymbolTable();
 
-				this._uobject = (UhtClass?)FindType(null, UhtFindOptions.SourceName | UhtFindOptions.Class, "UObject");
-				this._uclass = (UhtClass?)FindType(null, UhtFindOptions.SourceName | UhtFindOptions.Class, "UClass");
-				this._uinterface = (UhtClass?)FindType(null, UhtFindOptions.SourceName | UhtFindOptions.Class, "UInterface");
-				this._iinterface = (UhtClass?)FindType(null, UhtFindOptions.SourceName | UhtFindOptions.Class, "IInterface");
-				this.AActor = (UhtClass?)FindType(null, UhtFindOptions.SourceName | UhtFindOptions.Class, "AActor");
-				this.INotifyFieldValueChanged = (UhtClass?)FindType(null, UhtFindOptions.SourceName | UhtFindOptions.Class, "INotifyFieldValueChanged");
+				_uobject = (UhtClass?)FindType(null, UhtFindOptions.SourceName | UhtFindOptions.Class, "UObject");
+				_uclass = (UhtClass?)FindType(null, UhtFindOptions.SourceName | UhtFindOptions.Class, "UClass");
+				_uinterface = (UhtClass?)FindType(null, UhtFindOptions.SourceName | UhtFindOptions.Class, "UInterface");
+				_iinterface = (UhtClass?)FindType(null, UhtFindOptions.SourceName | UhtFindOptions.Class, "IInterface");
+				AActor = (UhtClass?)FindType(null, UhtFindOptions.SourceName | UhtFindOptions.Class, "AActor");
+				INotifyFieldValueChanged = (UhtClass?)FindType(null, UhtFindOptions.SourceName | UhtFindOptions.Class, "INotifyFieldValueChanged");
 			});
 		}
 
@@ -1987,21 +1987,21 @@ namespace EpicGames.UHT.Utils
 
 		private void StepForAllHeaders(StepDelegate stepDelegate)
 		{
-			if (this.HasErrors)
+			if (HasErrors)
 			{
 				return;
 			}
 
-			if (this.GoWide)
+			if (GoWide)
 			{
-				Parallel.ForEach(this._headerFiles, headerFile =>
+				Parallel.ForEach(_headerFiles, headerFile =>
 				{
 					stepDelegate(headerFile);
 				});
 			}
 			else
 			{
-				foreach (UhtHeaderFile headerFile in this._headerFiles)
+				foreach (UhtHeaderFile headerFile in _headerFiles)
 				{
 					stepDelegate(headerFile);
 				}
@@ -2012,11 +2012,11 @@ namespace EpicGames.UHT.Utils
 		#region Symbol table initialization
 		private void PopulateSymbolTable()
 		{
-			foreach (UhtHeaderFile headerFile in this._headerFiles)
+			foreach (UhtHeaderFile headerFile in _headerFiles)
 			{
 				AddTypeToSymbolTable(headerFile);
 			}
-			this._symbolTablePopulated = true;
+			_symbolTablePopulated = true;
 		}
 
 		private void AddTypeToSymbolTable(UhtType type)
@@ -2038,14 +2038,14 @@ namespace EpicGames.UHT.Utils
 				//{
 				//	foreach (UhtEnumValue Value in Enum.EnumValues)
 				//	{
-				//		this.RegularEnumValueLookup.Add(Value.Name.ToString(), Enum);
+				//		RegularEnumValueLookup.Add(Value.Name.ToString(), Enum);
 				//	}
 				//}
 				bool addShortNames = enumObj.CppForm == UhtEnumCppForm.Namespaced || enumObj.CppForm == UhtEnumCppForm.EnumClass;
 				string checkName = $"{enumObj.SourceName}::";
 				foreach (UhtEnumValue value in enumObj.EnumValues)
 				{
-					if (!this._fullEnumValueLookup.TryAdd(value.Name, new EnumAndValue { Enum = enumObj, Value = value.Value }))
+					if (!_fullEnumValueLookup.TryAdd(value.Name, new EnumAndValue { Enum = enumObj, Value = value.Value }))
 					{
 						//TODO - add a warning
 					}
@@ -2053,7 +2053,7 @@ namespace EpicGames.UHT.Utils
 					{
 						if (value.Name.StartsWith(checkName, StringComparison.Ordinal))
 						{
-							this._shortEnumValueLookup.TryAdd(value.Name[checkName.Length..], enumObj);
+							_shortEnumValueLookup.TryAdd(value.Name[checkName.Length..], enumObj);
 						}
 					}
 				}
@@ -2071,7 +2071,7 @@ namespace EpicGames.UHT.Utils
 
 				if (engineExtendedType.HasEngineName() && engineExtendedType.MustBeUnique())
 				{
-					UhtType? existingType = this._engineNameSymbolTable.FindCaselessType(null, engineExtendedType.MustBeUniqueFindOptions(), type.EngineName);
+					UhtType? existingType = _engineNameSymbolTable.FindCaselessType(null, engineExtendedType.MustBeUniqueFindOptions(), type.EngineName);
 					if (existingType != null)
 					{
 						type.HeaderFile.LogError(type.LineNumber,
@@ -2082,10 +2082,10 @@ namespace EpicGames.UHT.Utils
 
 				if (type.VisibleType)
 				{
-					this._sourceNameSymbolTable.Add(type, type.SourceName);
+					_sourceNameSymbolTable.Add(type, type.SourceName);
 					if (engineExtendedType.HasEngineName())
 					{
-						this._engineNameSymbolTable.Add(type, type.EngineName);
+						_engineNameSymbolTable.Add(type, type.EngineName);
 					}
 				}
 			}
@@ -2142,7 +2142,7 @@ namespace EpicGames.UHT.Utils
 						}
 					}
 					states[visit.HeaderFileTypeIndex] = TopologicalState.Permanent;
-					this._sortedHeaderFiles.Add(visit);
+					_sortedHeaderFiles.Add(visit);
 					return null;
 
 				case TopologicalState.Temporary:
@@ -2163,14 +2163,14 @@ namespace EpicGames.UHT.Utils
 				Log.Logger.LogTrace("Step - Topological Sort Header Files");
 
 				// Initialize a scratch table for topological states
-				this._sortedHeaderFiles.Capacity = this.HeaderFileTypeCount;
-				List<TopologicalState> states = new(this.HeaderFileTypeCount);
-				for (int index = 0; index < this.HeaderFileTypeCount; ++index)
+				_sortedHeaderFiles.Capacity = HeaderFileTypeCount;
+				List<TopologicalState> states = new(HeaderFileTypeCount);
+				for (int index = 0; index < HeaderFileTypeCount; ++index)
 				{
 					states.Add(TopologicalState.Unmarked);
 				}
 
-				foreach (UhtHeaderFile headerFile in this.HeaderFiles)
+				foreach (UhtHeaderFile headerFile in HeaderFiles)
 				{
 					if (states[headerFile.HeaderFileTypeIndex] == TopologicalState.Unmarked)
 					{
@@ -2202,9 +2202,9 @@ namespace EpicGames.UHT.Utils
 		{
 
 			// Check for existing value
-			lock (this._scriptStructsValidForNet)
+			lock (_scriptStructsValidForNet)
 			{
-				if (this._scriptStructsValidForNet.Contains(referencedScriptStruct))
+				if (_scriptStructsValidForNet.Contains(referencedScriptStruct))
 				{
 					return true;
 				}
@@ -2235,9 +2235,9 @@ namespace EpicGames.UHT.Utils
 			// Save the results
 			if (isStructValid)
 			{
-				lock (this._scriptStructsValidForNet)
+				lock (_scriptStructsValidForNet)
 				{
-					this._scriptStructsValidForNet.Add(referencedScriptStruct);
+					_scriptStructsValidForNet.Add(referencedScriptStruct);
 				}
 			}
 			return isStructValid;
@@ -2253,7 +2253,7 @@ namespace EpicGames.UHT.Utils
 		/// <param name="enabled">If true, the exporter is to be enabled</param>
 		public void SetExporterStatus(string name, bool enabled)
 		{
-			this._exporterStates[name] = enabled;
+			_exporterStates[name] = enabled;
 		}
 
 		/// <summary>
@@ -2264,22 +2264,22 @@ namespace EpicGames.UHT.Utils
 		/// <returns>True if enabled</returns>
 		public bool IsPluginEnabled(string pluginName, bool includeTargetCheck)
 		{
-			if (this._projectJson == null && this.ProjectDirectory != null && this.ProjectFile != null)
+			if (_projectJson == null && ProjectDirectory != null && ProjectFile != null)
 			{
-				UhtBuffer? contents = this.ReadSourceToBuffer(this.ProjectFile);
+				UhtBuffer? contents = ReadSourceToBuffer(ProjectFile);
 				if (contents != null)
 				{
-					this._projectJson = JsonDocument.Parse(contents.Memory);
+					_projectJson = JsonDocument.Parse(contents.Memory);
 					UhtBuffer.Return(contents);
 				}
 			}
 
-			if (this._projectJson == null)
+			if (_projectJson == null)
 			{
 				return false;
 			}
 
-			JsonObject rootObject = new(this._projectJson.RootElement);
+			JsonObject rootObject = new(_projectJson.RootElement);
 			if (rootObject.TryGetObjectArrayField("Plugins", out JsonObject[]? plugins))
 			{
 				foreach (JsonObject plugin in plugins)
@@ -2292,18 +2292,18 @@ namespace EpicGames.UHT.Utils
 					{
 						return false;
 					}
-					if (includeTargetCheck && this.Manifest != null)
+					if (includeTargetCheck && Manifest != null)
 					{
 						if (plugin.TryGetStringArrayField("TargetAllowList", out string[]? allowList))
 						{
-							if (allowList.Contains(this.Manifest.TargetName, StringComparer.OrdinalIgnoreCase))
+							if (allowList.Contains(Manifest.TargetName, StringComparer.OrdinalIgnoreCase))
 							{
 								return true;
 							}
 						}
 						if (plugin.TryGetStringArrayField("TargetDenyList", out string[]? denyList))
 						{
-							if (denyList.Contains(this.Manifest.TargetName, StringComparer.OrdinalIgnoreCase))
+							if (denyList.Contains(Manifest.TargetName, StringComparer.OrdinalIgnoreCase))
 							{
 								return false;
 							}
@@ -2323,18 +2323,18 @@ namespace EpicGames.UHT.Utils
 			{
 				Log.Logger.LogTrace("Step - Exports");
 
-				foreach (UhtExporter exporter in this.ExporterTable)
+				foreach (UhtExporter exporter in ExporterTable)
 				{
-					if (!this._exporterStates.TryGetValue(exporter.Name, out bool run))
+					if (!_exporterStates.TryGetValue(exporter.Name, out bool run))
 					{
 						run = Config!.IsExporterEnabled(exporter.Name) ||
-							(exporter.Options.HasAnyFlags(UhtExporterOptions.Default) && !this.NoDefaultExporters);
+							(exporter.Options.HasAnyFlags(UhtExporterOptions.Default) && !NoDefaultExporters);
 					}
 
 					UHTManifest.Module? pluginModule = null;
 					if (!String.IsNullOrEmpty(exporter.ModuleName))
 					{
-						foreach (UHTManifest.Module module in this.Manifest!.Modules)
+						foreach (UHTManifest.Module module in Manifest!.Modules)
 						{
 							if (String.Equals(module.Name, exporter.ModuleName, StringComparison.OrdinalIgnoreCase))
 							{
@@ -2373,9 +2373,9 @@ namespace EpicGames.UHT.Utils
 				}
 
 				// Save the collected external dependencies
-				if (!String.IsNullOrEmpty(this.Manifest!.ExternalDependenciesFile))
+				if (!String.IsNullOrEmpty(Manifest!.ExternalDependenciesFile))
 				{
-					using StreamWriter output = new(this.Manifest!.ExternalDependenciesFile);
+					using StreamWriter output = new(Manifest!.ExternalDependenciesFile);
 					foreach (string dep in externalDependencies)
 					{
 						output.WriteLine(dep);
diff --git a/Engine/Source/Programs/Shared/EpicGames.UHT/Utils/UhtStringBuilderExtensions.cs b/Engine/Source/Programs/Shared/EpicGames.UHT/Utils/UhtStringBuilderExtensions.cs
--- a/Engine/Source/Programs/Shared/EpicGames.UHT/Utils/UhtStringBuilderExtensions.cs
+++ b/Engine/Source/Programs/Shared/EpicGames.UHT/Utils/UhtStringBuilderExtensions.cs
@@ -523,9 +523,9 @@ namespace EpicGames.UHT.Utils
 		/// <param name="initialBufferSize">Initial buffer size for newly created StringBuilders</param>
 		public StringBuilderCache(int capacity, int initialBufferSize)
 		{
-			this._capacity = capacity;
-			this._initialBufferSize = initialBufferSize;
-			this._stack = new Stack<StringBuilder>(this._capacity);
+			_capacity = capacity;
+			_initialBufferSize = initialBufferSize;
+			_stack = new Stack<StringBuilder>(_capacity);
 		}
 
 		/// <summary>
@@ -534,15 +534,15 @@ namespace EpicGames.UHT.Utils
 		/// <returns></returns>
 		public StringBuilder Borrow()
 		{
-			lock (this._stack)
+			lock (_stack)
 			{
-				if (this._stack.Count > 0)
+				if (_stack.Count > 0)
 				{
-					return this._stack.Pop();
+					return _stack.Pop();
 				}
 			}
 
-			return new StringBuilder(this._initialBufferSize);
+			return new StringBuilder(_initialBufferSize);
 		}
 
 		/// <summary>
@@ -553,11 +553,11 @@ namespace EpicGames.UHT.Utils
 		{
 			// Sadly, clearing the builder (sets length to 0) will reallocate chunks.
 			builder.Clear();
-			lock (this._stack)
+			lock (_stack)
 			{
-				if (this._stack.Count < this._capacity)
+				if (_stack.Count < _capacity)
 				{
-					this._stack.Push(builder);
+					_stack.Push(builder);
 				}
 			}
 		}
@@ -586,8 +586,8 @@ namespace EpicGames.UHT.Utils
 		/// <param name="cache">String builder cache</param>
 		public BorrowStringBuilder(StringBuilderCache cache)
 		{
-			this.Cache = cache;
-			this.StringBuilder = this.Cache.Borrow();
+			Cache = cache;
+			StringBuilder = Cache.Borrow();
 		}
 
 		/// <summary>
@@ -595,7 +595,7 @@ namespace EpicGames.UHT.Utils
 		/// </summary>
 		public void Dispose()
 		{
-			this.Cache.Return(this.StringBuilder);
+			Cache.Return(StringBuilder);
 		}
 	}
 }
diff --git a/Engine/Source/Programs/Shared/EpicGames.UHT/Utils/UhtSymbolTable.cs b/Engine/Source/Programs/Shared/EpicGames.UHT/Utils/UhtSymbolTable.cs
--- a/Engine/Source/Programs/Shared/EpicGames.UHT/Utils/UhtSymbolTable.cs
+++ b/Engine/Source/Programs/Shared/EpicGames.UHT/Utils/UhtSymbolTable.cs
@@ -91,7 +91,7 @@ namespace EpicGames.UHT.Utils
 		/// <param name="typeCount">Number of types in the table.</param>
 		public UhtSymbolTable(int typeCount)
 		{
-			this._symbols = new Symbol[typeCount];
+			_symbols = new Symbol[typeCount];
 		}
 
 		/// <summary>
@@ -101,25 +101,25 @@ namespace EpicGames.UHT.Utils
 		/// <param name="name">The name of the type which could be the source name or the engine name</param>
 		public void Add(UhtType type, string name)
 		{
-			if (this._casedDictionary.TryGetValue(name, out Lookup existing))
+			if (_casedDictionary.TryGetValue(name, out Lookup existing))
 			{
 				AddExisting(type, existing.CasedIndex, existing.SymbolIndex);
 				return;
 			}
 
-			int casedIndex = ++this._casedCount;
+			int casedIndex = ++_casedCount;
 
-			if (this._caselessDictionary.TryGetValue(name, out int symbolIndex))
+			if (_caselessDictionary.TryGetValue(name, out int symbolIndex))
 			{
-				this._casedDictionary.Add(name, new Lookup { SymbolIndex = symbolIndex, CasedIndex = casedIndex });
+				_casedDictionary.Add(name, new Lookup { SymbolIndex = symbolIndex, CasedIndex = casedIndex });
 				AddExisting(type, casedIndex, symbolIndex);
 				return;
 			}
 
 			symbolIndex = type.TypeIndex;
-			this._symbols[symbolIndex] = new Symbol { Type = type, MatchOptions = type.EngineType.FindOptions(), CasedIndex = casedIndex, NextIndex = 0, LastIndex = symbolIndex };
-			this._caselessDictionary.Add(name, symbolIndex);
-			this._casedDictionary.Add(name, new Lookup { SymbolIndex = symbolIndex, CasedIndex = casedIndex });
+			_symbols[symbolIndex] = new Symbol { Type = type, MatchOptions = type.EngineType.FindOptions(), CasedIndex = casedIndex, NextIndex = 0, LastIndex = symbolIndex };
+			_caselessDictionary.Add(name, symbolIndex);
+			_casedDictionary.Add(name, new Lookup { SymbolIndex = symbolIndex, CasedIndex = casedIndex });
 		}
 
 		/// <summary>
@@ -133,13 +133,13 @@ namespace EpicGames.UHT.Utils
 		/// <exception cref="UhtIceException">Thrown if the symbol wasn't found.</exception>
 		public void Replace(UhtType oldType, UhtType newType, string name)
 		{
-			if (this._caselessDictionary.TryGetValue(name, out int symbolIndex))
+			if (_caselessDictionary.TryGetValue(name, out int symbolIndex))
 			{
-				for (; symbolIndex != 0; symbolIndex = this._symbols[symbolIndex].NextIndex)
+				for (; symbolIndex != 0; symbolIndex = _symbols[symbolIndex].NextIndex)
 				{
-					if (this._symbols[symbolIndex].Type == oldType)
+					if (_symbols[symbolIndex].Type == oldType)
 					{
-						this._symbols[symbolIndex].Type = newType;
+						_symbols[symbolIndex].Type = newType;
 						return;
 					}
 				}
@@ -156,7 +156,7 @@ namespace EpicGames.UHT.Utils
 		/// <returns>Found type or null if not found.</returns>
 		public UhtType? FindCasedType(UhtType? startingType, UhtFindOptions options, string name)
 		{
-			if (this._casedDictionary.TryGetValue(name, out Lookup existing))
+			if (_casedDictionary.TryGetValue(name, out Lookup existing))
 			{
 				return FindType(startingType, options, existing);
 			}
@@ -172,7 +172,7 @@ namespace EpicGames.UHT.Utils
 		/// <returns>Found type or null if not found.</returns>
 		public UhtType? FindCaselessType(UhtType? startingType, UhtFindOptions options, string name)
 		{
-			if (this._caselessDictionary.TryGetValue(name, out int symbolIndex))
+			if (_caselessDictionary.TryGetValue(name, out int symbolIndex))
 			{
 				return FindType(startingType, options, new Lookup { SymbolIndex = symbolIndex, CasedIndex = 0 });
 			}
@@ -207,7 +207,7 @@ namespace EpicGames.UHT.Utils
 					UhtHeaderFile headerFile = startingType.HeaderFile;
 					foreach (UhtHeaderFile includedFile in headerFile.IncludedHeaders)
 					{
-						found = this.FindSymbolChain(includedFile, options, ref lookup);
+						found = FindSymbolChain(includedFile, options, ref lookup);
 						if (found != null)
 						{
 							break;
@@ -223,11 +223,11 @@ namespace EpicGames.UHT.Utils
 			// Global search.  Match anything that has an owner of parent
 			if (!options.HasAnyFlags(UhtFindOptions.NoGlobal))
 			{
-				for (int index = lookup.SymbolIndex; index != 0; index = this._symbols[index].NextIndex)
+				for (int index = lookup.SymbolIndex; index != 0; index = _symbols[index].NextIndex)
 				{
-					if (IsMatch(options, index, lookup.CasedIndex) && this._symbols[index].Type.Outer is UhtHeaderFile)
+					if (IsMatch(options, index, lookup.CasedIndex) && _symbols[index].Type.Outer is UhtHeaderFile)
 					{
-						return this._symbols[index].Type;
+						return _symbols[index].Type;
 					}
 				}
 			}
@@ -330,11 +330,11 @@ namespace EpicGames.UHT.Utils
 		/// <returns>Found type or null if not found.</returns>
 		private UhtType? FindSymbolChain(UhtType owner, UhtFindOptions options, ref Lookup lookup)
 		{
-			for (int index = lookup.SymbolIndex; index != 0; index = this._symbols[index].NextIndex)
+			for (int index = lookup.SymbolIndex; index != 0; index = _symbols[index].NextIndex)
 			{
-				if (IsMatch(options, index, lookup.CasedIndex) && this._symbols[index].Type.Outer == owner)
+				if (IsMatch(options, index, lookup.CasedIndex) && _symbols[index].Type.Outer == owner)
 				{
-					return this._symbols[index].Type;
+					return _symbols[index].Type;
 				}
 			}
 			return null;
@@ -349,9 +349,9 @@ namespace EpicGames.UHT.Utils
 		private void AddExisting(UhtType type, int casedIndex, int symbolIndex)
 		{
 			int typeIndex = type.TypeIndex;
-			this._symbols[typeIndex] = new Symbol { Type = type, MatchOptions = type.EngineType.FindOptions(), CasedIndex = casedIndex, NextIndex = 0, LastIndex = 0 };
-			this._symbols[this._symbols[symbolIndex].LastIndex].NextIndex = typeIndex;
-			this._symbols[symbolIndex].LastIndex = typeIndex;
+			_symbols[typeIndex] = new Symbol { Type = type, MatchOptions = type.EngineType.FindOptions(), CasedIndex = casedIndex, NextIndex = 0, LastIndex = 0 };
+			_symbols[_symbols[symbolIndex].LastIndex].NextIndex = typeIndex;
+			_symbols[symbolIndex].LastIndex = typeIndex;
 		}
 
 		/// <summary>
@@ -363,9 +363,9 @@ namespace EpicGames.UHT.Utils
 		/// <returns>True if the symbol is a match</returns>
 		private bool IsMatch(UhtFindOptions options, int symbolIndex, int casedIndex)
 		{
-			return (casedIndex == 0 || casedIndex == this._symbols[symbolIndex].CasedIndex) &&
-				this._symbols[symbolIndex].Type.VisibleType &&
-				(this._symbols[symbolIndex].MatchOptions & options) != 0;
+			return (casedIndex == 0 || casedIndex == _symbols[symbolIndex].CasedIndex) &&
+				_symbols[symbolIndex].Type.VisibleType &&
+				(_symbols[symbolIndex].MatchOptions & options) != 0;
 		}
 	}
 }
diff --git a/Engine/Source/Programs/Shared/EpicGames.UHT/Utils/UhtTables.cs b/Engine/Source/Programs/Shared/EpicGames.UHT/Utils/UhtTables.cs
--- a/Engine/Source/Programs/Shared/EpicGames.UHT/Utils/UhtTables.cs
+++ b/Engine/Source/Programs/Shared/EpicGames.UHT/Utils/UhtTables.cs
@@ -108,8 +108,8 @@ namespace EpicGames.UHT.Utils
 		/// </summary>
 		public string UserName
 		{
-			get => String.IsNullOrEmpty(this._userName) ? this.TableName : this._userName;
-			set => this._userName = value;
+			get => String.IsNullOrEmpty(_userName) ? TableName : _userName;
+			set => _userName = value;
 		}
 
 		/// <summary>
@@ -178,7 +178,7 @@ namespace EpicGames.UHT.Utils
 		/// </summary>
 		public UhtLookupTable(StringViewComparer comparer)
 		{
-			this._lookup = new Dictionary<StringView, TValue>(comparer);
+			_lookup = new Dictionary<StringView, TValue>(comparer);
 		}
 
 		/// <summary>
@@ -188,7 +188,7 @@ namespace EpicGames.UHT.Utils
 		/// <param name="value">Value to be added</param>
 		public UhtLookupTable<TValue> Add(string key, TValue value)
 		{
-			this._lookup.Add(key, value);
+			_lookup.Add(key, value);
 			return this;
 		}
 
@@ -200,7 +200,7 @@ namespace EpicGames.UHT.Utils
 		/// <returns>True if the key was found, false if not</returns>
 		public bool TryGetValue(StringView key, [MaybeNullWhen(false)] out TValue value)
 		{
-			return this._lookup.TryGetValue(key, out value);
+			return _lookup.TryGetValue(key, out value);
 		}
 
 		/// <summary>
@@ -211,7 +211,7 @@ namespace EpicGames.UHT.Utils
 		{
 			foreach (KeyValuePair<StringView, TValue> kvp in ((UhtLookupTable<TValue>)baseTable)._lookup)
 			{
-				this._lookup.TryAdd(kvp.Key, kvp.Value);
+				_lookup.TryAdd(kvp.Key, kvp.Value);
 			}
 		}
 	}
@@ -239,7 +239,7 @@ namespace EpicGames.UHT.Utils
 		/// <param name="name">The name of the group</param>
 		public UhtLookupTables(string name)
 		{
-			this.Name = name;
+			Name = name;
 		}
 
 		/// <summary>
@@ -249,11 +249,11 @@ namespace EpicGames.UHT.Utils
 		/// <returns>The table associated with the name.</returns>
 		public TTable Get(string tableName)
 		{
-			if (!this.Tables.TryGetValue(tableName, out TTable? table))
+			if (!Tables.TryGetValue(tableName, out TTable? table))
 			{
 				table = new TTable();
 				table.TableName = tableName;
-				this.Tables.Add(tableName, table);
+				Tables.Add(tableName, table);
 			}
 			return table;
 		}
@@ -286,16 +286,16 @@ namespace EpicGames.UHT.Utils
 		/// <exception cref="UhtIceException">Thrown if there are problems with the tables.</exception>
 		public void Merge()
 		{
-			List<TTable> orderedList = new(this.Tables.Count);
-			List<TTable> remainingList = new(this.Tables.Count);
+			List<TTable> orderedList = new(Tables.Count);
+			List<TTable> remainingList = new(Tables.Count);
 			HashSet<UhtLookupTableBase> doneTables = new();
 
 			// Collect all of the tables
-			foreach (KeyValuePair<string, TTable> kvp in this.Tables)
+			foreach (KeyValuePair<string, TTable> kvp in Tables)
 			{
 				if (!kvp.Value.Implemented)
 				{
-					throw new UhtIceException($"{this.Name} table '{kvp.Value.TableName}' has been referenced but not implemented");
+					throw new UhtIceException($"{Name} table '{kvp.Value.TableName}' has been referenced but not implemented");
 				}
 				remainingList.Add(kvp.Value);
 			}
@@ -322,7 +322,7 @@ namespace EpicGames.UHT.Utils
 				}
 				if (!addedOne)
 				{
-					throw new UhtIceException($"Circular dependency in {this.GetType().Name}.{this.Name} tables detected");
+					throw new UhtIceException($"Circular dependency in {GetType().Name}.{Name} tables detected");
 				}
 			}
 
@@ -607,7 +607,7 @@ namespace EpicGames.UHT.Utils
 					}
 					else if (classAttribute is UhtEngineClassAttribute engineClassAttribute)
 					{
-						this.EngineClassTable.OnEngineClassAttribute(engineClassAttribute);
+						EngineClassTable.OnEngineClassAttribute(engineClassAttribute);
 					}
 				}
 			}
@@ -615,13 +615,13 @@ namespace EpicGames.UHT.Utils
 
 		private void PerformPostInitialization()
 		{
-			this.KeywordTables.Merge();
-			this.SpecifierTables.Merge();
-			this.SpecifierValidatorTables.Merge();
+			KeywordTables.Merge();
+			SpecifierTables.Merge();
+			SpecifierValidatorTables.Merge();
 
 			// Invoke this to generate an exception if there is no default
-			_ = this.PropertyTypeTable.Default;
-			_ = this.StructDefaultValueTable.Default;
+			_ = PropertyTypeTable.Default;
+			_ = StructDefaultValueTable.Default;
 		}
 
 		private void HandleUnrealHeaderToolAttribute(Type type, UnrealHeaderToolAttribute parserAttribute)
@@ -645,35 +645,35 @@ namespace EpicGames.UHT.Utils
 				{
 					if (methodAttribute is UhtSpecifierAttribute specifierAttribute)
 					{
-						this.SpecifierTables.OnSpecifierAttribute(type, methodInfo, specifierAttribute);
+						SpecifierTables.OnSpecifierAttribute(type, methodInfo, specifierAttribute);
 					}
 					else if (methodAttribute is UhtSpecifierValidatorAttribute specifierValidatorAttribute)
 					{
-						this.SpecifierValidatorTables.OnSpecifierValidatorAttribute(type, methodInfo, specifierValidatorAttribute);
+						SpecifierValidatorTables.OnSpecifierValidatorAttribute(type, methodInfo, specifierValidatorAttribute);
 					}
 					else if (methodAttribute is UhtKeywordCatchAllAttribute keywordCatchAllAttribute)
 					{
-						this.KeywordTables.OnKeywordCatchAllAttribute(type, methodInfo, keywordCatchAllAttribute);
+						KeywordTables.OnKeywordCatchAllAttribute(type, methodInfo, keywordCatchAllAttribute);
 					}
 					else if (methodAttribute is UhtKeywordAttribute keywordAttribute)
 					{
-						this.KeywordTables.OnKeywordAttribute(type, methodInfo, keywordAttribute);
+						KeywordTables.OnKeywordAttribute(type, methodInfo, keywordAttribute);
 					}
 					else if (methodAttribute is UhtPropertyTypeAttribute propertyTypeAttribute)
 					{
-						this.PropertyTypeTable.OnPropertyTypeAttribute(methodInfo, propertyTypeAttribute);
+						PropertyTypeTable.OnPropertyTypeAttribute(methodInfo, propertyTypeAttribute);
 					}
 					else if (methodAttribute is UhtStructDefaultValueAttribute structDefaultValueAttribute)
 					{
-						this.StructDefaultValueTable.OnStructDefaultValueAttribute(methodInfo, structDefaultValueAttribute);
+						StructDefaultValueTable.OnStructDefaultValueAttribute(methodInfo, structDefaultValueAttribute);
 					}
 					else if (methodAttribute is UhtExporterAttribute exporterAttribute)
 					{
-						this.ExporterTable.OnExporterAttribute(type, methodInfo, exporterAttribute);
+						ExporterTable.OnExporterAttribute(type, methodInfo, exporterAttribute);
 					}
 					else if (methodAttribute is UhtLocTextDefaultValueAttribute locTextDefaultValueAttribute)
 					{
-						this.LocTextDefaultValueTable.OnLocTextDefaultValueAttribute(methodInfo, locTextDefaultValueAttribute);
+						LocTextDefaultValueTable.OnLocTextDefaultValueAttribute(methodInfo, locTextDefaultValueAttribute);
 					}
 				}
 			}
diff --git a/Engine/Source/Programs/Shared/EpicGames.UHT/Utils/UhtUtilities.cs b/Engine/Source/Programs/Shared/EpicGames.UHT/Utils/UhtUtilities.cs
--- a/Engine/Source/Programs/Shared/EpicGames.UHT/Utils/UhtUtilities.cs
+++ b/Engine/Source/Programs/Shared/EpicGames.UHT/Utils/UhtUtilities.cs
@@ -226,13 +226,13 @@ namespace EpicGames.UHT.Utils
 		{
 			get
 			{
-				if (this._useStringBuilder && this._stringBuilder != null)
+				if (_useStringBuilder && _stringBuilder != null)
 				{
-					return this._stringBuilder.Length;
+					return _stringBuilder.Length;
 				}
 				else
 				{
-					return this._stringView.Span.Length;
+					return _stringView.Span.Length;
 				}
 			}
 		}
@@ -243,7 +243,7 @@ namespace EpicGames.UHT.Utils
 		/// <returns>Contents of the builder</returns>
 		public StringView ToStringView()
 		{
-			return this._useStringBuilder ? new StringView(this._stringBuilder!.ToString()) : this._stringView;
+			return _useStringBuilder ? new StringView(_stringBuilder!.ToString()) : _stringView;
 		}
 
 		/// <summary>
@@ -252,7 +252,7 @@ namespace EpicGames.UHT.Utils
 		/// <returns>Contents of the builder</returns>
 		public override string ToString()
 		{
-			return this._useStringBuilder ? this._stringBuilder!.ToString() : this._stringView.ToString();
+			return _useStringBuilder ? _stringBuilder!.ToString() : _stringView.ToString();
 		}
 
 		/// <summary>
@@ -262,18 +262,18 @@ namespace EpicGames.UHT.Utils
 		/// <returns>The string builder</returns>
 		public StringViewBuilder Append(StringView text)
 		{
-			if (this._useStringBuilder)
+			if (_useStringBuilder)
 			{
-				this._stringBuilder!.Append(text.Span);
+				_stringBuilder!.Append(text.Span);
 			}
-			else if (this._stringView.Span.Length > 0)
+			else if (_stringView.Span.Length > 0)
 			{
 				SwitchToStringBuilder();
-				this._stringBuilder!.Append(text.Span);
+				_stringBuilder!.Append(text.Span);
 			}
 			else
 			{
-				this._stringView = text;
+				_stringView = text;
 			}
 			return this;
 		}
@@ -286,7 +286,7 @@ namespace EpicGames.UHT.Utils
 		public StringViewBuilder Append(char c)
 		{
 			SwitchToStringBuilder();
-			this._stringBuilder!.Append(c);
+			_stringBuilder!.Append(c);
 			return this;
 		}
 
@@ -295,14 +295,14 @@ namespace EpicGames.UHT.Utils
 		/// </summary>
 		private void SwitchToStringBuilder()
 		{
-			if (!this._useStringBuilder)
+			if (!_useStringBuilder)
 			{
-				if (this._stringBuilder == null)
+				if (_stringBuilder == null)
 				{
-					this._stringBuilder = new StringBuilder();
+					_stringBuilder = new StringBuilder();
 				}
-				this._useStringBuilder = true;
-				this._stringBuilder.Append(this._stringView.Span);
+				_useStringBuilder = true;
+				_stringBuilder.Append(_stringView.Span);
 			}
 		}
 	}
