Improved support for contextual logging in UHT

diff --git a/Engine/Source/Programs/Shared/EpicGames.Core/KnownLogEvents.cs b/Engine/Source/Programs/Shared/EpicGames.Core/KnownLogEvents.cs
--- a/Engine/Source/Programs/Shared/EpicGames.Core/KnownLogEvents.cs
+++ b/Engine/Source/Programs/Shared/EpicGames.Core/KnownLogEvents.cs
@@ -35,6 +35,11 @@ namespace EpicGames.Core
 		/// </summary>
 		public static EventId Compiler { get; } = new EventId(100);
 
+		/// <summary>
+		/// UHT compiler message
+		/// </summary>
+		public static EventId UHT { get; } = new EventId(110);
+
 		/// <summary>
 		/// Linker error
 		/// </summary>
diff --git a/Engine/Source/Programs/Shared/EpicGames.UHT/Utils/UhtSession.cs b/Engine/Source/Programs/Shared/EpicGames.UHT/Utils/UhtSession.cs
--- a/Engine/Source/Programs/Shared/EpicGames.UHT/Utils/UhtSession.cs
+++ b/Engine/Source/Programs/Shared/EpicGames.UHT/Utils/UhtSession.cs
@@ -562,6 +562,11 @@ namespace EpicGames.UHT.Utils
 
 		#region Configurable settings
 
+		/// <summary>
+		/// Logger interface.  If not set, Log.Logger will be used
+		/// </summary>
+		public ILogger? Logger { get; set; }
+
 		/// <summary>
 		/// Interface used to read/write files
 		/// </summary>
@@ -1517,12 +1522,11 @@ namespace EpicGames.UHT.Utils
 		/// <summary>
 		/// Log the given message
 		/// </summary>
-		/// <param name="message">The message to be logged</param>
-		private void LogMessage(UhtMessage message)
+		/// <param name="msg">The message to be logged</param>
+		private void LogMessage(UhtMessage msg)
 		{
-			string formattedMessage = FormatMessage(message);
 			LogLevel logLevel;
-			switch (message.MessageType)
+			switch (msg.MessageType)
 			{
 				default:
 				case UhtMessageType.Error:
@@ -1547,7 +1551,13 @@ namespace EpicGames.UHT.Utils
 					break;
 			}
 
-			Log.Logger.Log(logLevel, "{FormattedMessage}", formattedMessage);
+			GetMessageParts(msg, out string filePath, out int lineNumber, out string fragmentPath, out string severity, out string message);
+			(Logger??Log.Logger).Log(logLevel, KnownLogEvents.UHT, "{File}({Line}){FileFragment}: {Severity}: {Message}", 
+				new FileReference(filePath), 
+				new LogValue(LogValueType.LineNumber, lineNumber.ToString()), 
+				fragmentPath, 
+				new LogValue(LogValueType.Severity, severity), 
+				new LogValue(LogValueType.Message, message));
 		}
 
 		/// <summary>
@@ -1557,9 +1567,10 @@ namespace EpicGames.UHT.Utils
 		public List<string> CollectMessages()
 		{
 			List<string> messages = new();
-			foreach (UhtMessage message in FetchOrderedMessages())
+			foreach (UhtMessage msg in FetchOrderedMessages())
 			{
-				messages.Add(FormatMessage(message));
+				GetMessageParts(msg, out string filePath, out int lineNumber, out string fragmentPath, out string severity, out string message);
+				messages.Add($"{filePath}({lineNumber}){fragmentPath}: {severity}: {message}");
 			}
 			return messages;
 		}
@@ -1596,43 +1607,46 @@ namespace EpicGames.UHT.Utils
 		}
 
 		/// <summary>
-		/// Format the given message
+		/// Extract the part of a message
 		/// </summary>
-		/// <param name="message">Message to be formatted</param>
-		/// <returns>Text of the formatted message</returns>
-		private string FormatMessage(UhtMessage message)
+		/// <param name="msg">Message being consider</param>
+		/// <param name="filePath">Full path of file where error happened</param>
+		/// <param name="lineNumber">Line number where error happened</param>
+		/// <param name="fragmentPath">For multi-part files (test harness), the inner file</param>
+		/// <param name="severity">Severity of the error</param>
+		/// <param name="message">Text of the message</param>
+		private void GetMessageParts(UhtMessage msg, out string filePath, out int lineNumber, out string fragmentPath, out string severity, out string message)
 		{
-			string filePath;
-			string fragmentPath = "";
-			int lineNumber = message.LineNumber;
-			if (message.FilePath != null)
+			fragmentPath = "";
+			lineNumber = msg.LineNumber;
+			if (msg.FilePath != null)
 			{
-				filePath = message.FilePath;
+				filePath = msg.FilePath;
 			}
-			else if (message.MessageSource != null)
+			else if (msg.MessageSource != null)
 			{
-				if (message.MessageSource.MessageIsFragment)
+				if (msg.MessageSource.MessageIsFragment)
 				{
 					if (RelativePathInLog)
 					{
-						filePath = message.MessageSource.MessageFragmentFilePath;
+						filePath = msg.MessageSource.MessageFragmentFilePath;
 					}
 					else
 					{
-						filePath = message.MessageSource.MessageFragmentFullFilePath;
+						filePath = msg.MessageSource.MessageFragmentFullFilePath;
 					}
-					fragmentPath = $"[{message.MessageSource.MessageFilePath}]";
-					lineNumber += message.MessageSource.MessageFragmentLineNumber;
+					fragmentPath = $"[{msg.MessageSource.MessageFilePath}]";
+					lineNumber += msg.MessageSource.MessageFragmentLineNumber;
 				}
 				else
 				{
 					if (RelativePathInLog)
 					{
-						filePath = message.MessageSource.MessageFilePath;
+						filePath = msg.MessageSource.MessageFilePath;
 					}
 					else
 					{
-						filePath = message.MessageSource.MessageFullFilePath;
+						filePath = msg.MessageSource.MessageFullFilePath;
 					}
 				}
 			}
@@ -1641,22 +1655,30 @@ namespace EpicGames.UHT.Utils
 				filePath = "UnknownSource";
 			}
 
-			switch (message.MessageType)
+			switch (msg.MessageType)
 			{
 				case UhtMessageType.Error:
-					return $"{filePath}({lineNumber}){fragmentPath}: Error: {message.Message}";
+					severity = "Error";
+					break;
 				case UhtMessageType.Warning:
-					return $"{filePath}({lineNumber}){fragmentPath}: Warning: {message.Message}";
+					severity = "Warning";
+					break;
 				case UhtMessageType.Info:
-					return $"{filePath}({lineNumber}){fragmentPath}: Info: {message.Message}";
+					severity = "Info";
+					break;
 				case UhtMessageType.Trace:
-					return $"{filePath}({lineNumber}){fragmentPath}: Trace: {message.Message}";
+					severity = "Trace";
+					break;
 				case UhtMessageType.Deprecation:
-					return $"{filePath}({lineNumber}){fragmentPath}: Deprecation: {message.Message}";
+					severity = "Deprecation";
+					break;
 				default:
 				case UhtMessageType.Ice:
-					return $"{filePath}({lineNumber}){fragmentPath}:  Error: Internal Compiler Error - {message.Message}";
+					severity = "Internal Compiler Error";
+					break;
 			}
+
+			message = msg.Message;
 		}
 
 		/// <summary>
diff --git a/Engine/Source/Programs/UnrealBuildTool/Modes/UnrealHeaderToolMode.cs b/Engine/Source/Programs/UnrealBuildTool/Modes/UnrealHeaderToolMode.cs
--- a/Engine/Source/Programs/UnrealBuildTool/Modes/UnrealHeaderToolMode.cs
+++ b/Engine/Source/Programs/UnrealBuildTool/Modes/UnrealHeaderToolMode.cs
@@ -543,6 +543,7 @@
 			UhtTestFileManager TestFileManager = new(TestDirectory);
 			UhtSession Session = new()
 			{
+				Logger = Logger,
 				Tables = Tables,
 				Config = Config,
 				FileManager = TestFileManager,
@@ -1093,6 +1094,7 @@
 
 				UhtSession Session = new()
 				{
+					Logger = Logger,
 					Tables = Tables,
 					Config = Config,
 					FileManager = new UhtStdFileManager(),
