Horde: Use a connection pool for all Redis commands

Only expose ConnectionPool from RedisService to avoid using Database or ConnectionMultiplexer directly.

diff --git a/Engine/Source/Programs/Horde/Horde.Build.Tests/DatabaseIntegrationTest.cs b/Engine/Source/Programs/Horde/Horde.Build.Tests/DatabaseIntegrationTest.cs
--- a/Engine/Source/Programs/Horde/Horde.Build.Tests/DatabaseIntegrationTest.cs
+++ b/Engine/Source/Programs/Horde/Horde.Build.Tests/DatabaseIntegrationTest.cs
@@ -202,10 +202,11 @@ namespace Horde.Build.Tests
 
 				(string host, int port) = redisRunner.GetListenAddress();
 				_redisService = new RedisService($"{host}:{port},allowAdmin=true", RedisDbNum, _loggerFactory);
+				IConnectionMultiplexer cm = _redisService.ConnectionPool.GetConnection();
 
-				foreach (EndPoint endpoint in _redisService.Multiplexer.GetEndPoints())
+				foreach (EndPoint endpoint in cm.GetEndPoints())
 				{
-					_redisService.Multiplexer.GetServer(endpoint).FlushDatabase(RedisDbNum);
+					cm.GetServer(endpoint).FlushDatabase(RedisDbNum);
 				}
 			}
 			return _redisService;
diff --git a/Engine/Source/Programs/Horde/Horde.Build/Agents/AgentCollection.cs b/Engine/Source/Programs/Horde/Horde.Build/Agents/AgentCollection.cs
--- a/Engine/Source/Programs/Horde/Horde.Build/Agents/AgentCollection.cs
+++ b/Engine/Source/Programs/Horde/Horde.Build/Agents/AgentCollection.cs
@@ -571,13 +571,13 @@ namespace Horde.Build.Agents
 		/// <inheritdoc/>
 		Task PublishUpdateEventAsync(AgentId agentId)
 		{
-			return _redisService.Database.PublishAsync(_updateEventChannel, agentId);
+			return _redisService.GetDatabase().PublishAsync(_updateEventChannel, agentId);
 		}
 
 		/// <inheritdoc/>
 		public async Task<IDisposable> SubscribeToUpdateEventsAsync(Action<AgentId> onUpdate)
 		{
-			return await _redisService.Multiplexer.GetSubscriber().SubscribeAsync(_updateEventChannel, (channel, agentId) => onUpdate(agentId));
+			return await _redisService.GetDatabase().Multiplexer.GetSubscriber().SubscribeAsync(_updateEventChannel, (channel, agentId) => onUpdate(agentId));
 		}
 	}
 }
diff --git a/Engine/Source/Programs/Horde/Horde.Build/Agents/AgentService.cs b/Engine/Source/Programs/Horde/Horde.Build/Agents/AgentService.cs
--- a/Engine/Source/Programs/Horde/Horde.Build/Agents/AgentService.cs
+++ b/Engine/Source/Programs/Horde/Horde.Build/Agents/AgentService.cs
@@ -110,7 +110,7 @@ namespace Horde.Build.Agents
 			_sessions = sessions;
 			_aclService = aclService;
 			_downtimeService = downtimeService;
-			_agentRateTableData = new RedisString<AgentRateTable>(redisService.Database, "agent-rates");
+			_agentRateTableData = new RedisString<AgentRateTable>(redisService.ConnectionPool, "agent-rates");
 			_agentRateTable = new AsyncCachedValue<AgentRateTable?>(() => _agentRateTableData.GetAsync(), TimeSpan.FromSeconds(2.0));//.FromMinutes(5.0));
 			_poolsList = new AsyncCachedValue<List<IPool>>(() => poolCollection.GetAsync(), TimeSpan.FromSeconds(30.0));
 			_dogStatsd = dogStatsd;
@@ -915,7 +915,7 @@ namespace Horde.Build.Agents
 			RedisKey key = $"agent-rate/{agentId}";
 
 			// Try to get the current value
-			RedisValue value = await _redisService.Database.StringGetAsync(key);
+			RedisValue value = await _redisService.GetDatabase().StringGetAsync(key);
 			if (!value.IsNull)
 			{
 				double rate = (double)value;
@@ -951,7 +951,7 @@ namespace Horde.Build.Agents
 				}
 
 				// Cache it for future reference
-				await _redisService.Database.StringSetAsync(key, rate, TimeSpan.FromMinutes(5.0), flags: CommandFlags.FireAndForget);
+				await _redisService.GetDatabase().StringSetAsync(key, rate, TimeSpan.FromMinutes(5.0), flags: CommandFlags.FireAndForget);
 				return rate;
 			}
 		}
diff --git a/Engine/Source/Programs/Horde/Horde.Build/Compute/ComputeService.cs b/Engine/Source/Programs/Horde/Horde.Build/Compute/ComputeService.cs
--- a/Engine/Source/Programs/Horde/Horde.Build/Compute/ComputeService.cs
+++ b/Engine/Source/Programs/Horde/Horde.Build/Compute/ComputeService.cs
@@ -137,8 +137,8 @@ namespace Horde.Build.Compute
 		public ComputeService(MongoService mongoService, RedisService redisService, IStorageClient storageClient, IClock clock, ILogger<ComputeService> logger)
 		{
 			_storageClient = storageClient;
-			_taskScheduler = new RedisTaskScheduler<QueueKey, ComputeTaskInfo>(redisService.Database, "compute/tasks/", logger);
-			_messageQueue = new RedisMessageQueue<ComputeTaskStatus>(redisService.Database, "compute/messages/");
+			_taskScheduler = new RedisTaskScheduler<QueueKey, ComputeTaskInfo>(redisService.ConnectionPool, "compute/tasks/", logger);
+			_messageQueue = new RedisMessageQueue<ComputeTaskStatus>(redisService.GetDatabase(), "compute/messages/");
 			_expireTasksTicker = clock.AddTicker<ComputeService>(TimeSpan.FromMinutes(2.0), ExpireTasksAsync, logger);
 			_requirementsCache = new MemoryCache(new MemoryCacheOptions());
 			_globals = new LazyCachedValue<Task<Globals>>(() => mongoService.GetGlobalsAsync(), TimeSpan.FromSeconds(120.0));
diff --git a/Engine/Source/Programs/Horde/Horde.Build/Compute/RedisTaskScheduler.cs b/Engine/Source/Programs/Horde/Horde.Build/Compute/RedisTaskScheduler.cs
--- a/Engine/Source/Programs/Horde/Horde.Build/Compute/RedisTaskScheduler.cs
+++ b/Engine/Source/Programs/Horde/Horde.Build/Compute/RedisTaskScheduler.cs
@@ -77,7 +77,7 @@ namespace Horde.Build.Compute
 			}
 		}
 
-		readonly IDatabase _redis;
+		readonly RedisConnectionPool _redisConnectionPool;
 		readonly RedisKey _baseKey;
 		readonly RedisSet<TQueueId> _queueIndex;
 		readonly RedisHash<TQueueId, DateTime> _activeQueues; // Queues which are actively being dequeued from
@@ -92,15 +92,15 @@ namespace Horde.Build.Compute
 		/// <summary>
 		/// Constructor
 		/// </summary>
-		/// <param name="redis">The Redis database instance</param>
+		/// <param name="redisConnectionPool">The Redis connection pool</param>
 		/// <param name="baseKey">Base key for all keys used by this scheduler</param>
 		/// <param name="logger"></param>
-		public RedisTaskScheduler(IDatabase redis, RedisKey baseKey, ILogger logger)
+		public RedisTaskScheduler(RedisConnectionPool redisConnectionPool, RedisKey baseKey, ILogger logger)
 		{
-			_redis = redis;
+			_redisConnectionPool = redisConnectionPool;
 			_baseKey = baseKey;
-			_queueIndex = new RedisSet<TQueueId>(redis, baseKey.Append("index"));
-			_activeQueues = new RedisHash<TQueueId, DateTime>(redis, baseKey.Append("active"));
+			_queueIndex = new RedisSet<TQueueId>(_redisConnectionPool, baseKey.Append("index"));
+			_activeQueues = new RedisHash<TQueueId, DateTime>(_redisConnectionPool, baseKey.Append("active"));
 			_newQueueChannel = new RedisChannel<TQueueId>(baseKey.Append("new_queues").ToString());
 			_logger = logger;
 
@@ -141,7 +141,7 @@ namespace Horde.Build.Compute
 		/// <returns></returns>
 		RedisList<TTask> GetQueue(TQueueId queueId)
 		{
-			return new RedisList<TTask>(_redis, GetQueueKey(queueId));
+			return new RedisList<TTask>(_redisConnectionPool, GetQueueKey(queueId));
 		}
 
 		/// <summary>
@@ -177,14 +177,15 @@ namespace Horde.Build.Compute
 					break;
 				}
 
-				ITransaction transaction = _redis.CreateTransaction();
+				IDatabase redis = _redisConnectionPool.GetDatabase();
+				ITransaction transaction = redis.CreateTransaction();
 				_ = transaction.With(_queueIndex).AddAsync(queueId, CommandFlags.FireAndForget);
 				_ = PushTaskAsync(transaction.With(list), task, When.Always, CommandFlags.FireAndForget, atFront);
 
 				if (await transaction.ExecuteAsync())
 				{
 					_logger.LogInformation("Created queue {QueueId}", queueId);
-					await _redis.PublishAsync(_newQueueChannel, queueId);
+					await redis.PublishAsync(_newQueueChannel, queueId);
 					break;
 				}
 
@@ -285,7 +286,7 @@ namespace Horde.Build.Compute
 			TTask? item = await GetQueue(queueId).LeftPopAsync();
 			if (item == null)
 			{
-				ITransaction transaction = _redis.CreateTransaction();
+				ITransaction transaction = _redisConnectionPool.GetDatabase().CreateTransaction();
 				transaction.AddCondition(Condition.KeyNotExists(GetQueueKey(queueId)));
 				Task<bool> wasRemoved = transaction.With(_queueIndex).RemoveAsync(queueId);
 
@@ -378,7 +379,7 @@ namespace Horde.Build.Compute
 		{
 			Channel<TQueueId> newQueues = Channel.CreateUnbounded<TQueueId>();
 
-			ISubscriber subscriber = _redis.Multiplexer.GetSubscriber();
+			ISubscriber subscriber = _redisConnectionPool.GetDatabase().Multiplexer.GetSubscriber();
 			await using RedisChannelSubscription<TQueueId>? _ = await subscriber.SubscribeAsync(_newQueueChannel, (_, v) => newQueues.Writer.TryWrite(v));
 
 			while (await newQueues.Reader.WaitToReadAsync(cancellationToken))
diff --git a/Engine/Source/Programs/Horde/Horde.Build/Issues/IssueCollection.cs b/Engine/Source/Programs/Horde/Horde.Build/Issues/IssueCollection.cs
--- a/Engine/Source/Programs/Horde/Horde.Build/Issues/IssueCollection.cs
+++ b/Engine/Source/Programs/Horde/Horde.Build/Issues/IssueCollection.cs
@@ -468,7 +468,7 @@ namespace Horde.Build.Issues
 			Stopwatch timer = Stopwatch.StartNew();
 			TimeSpan nextNotifyTime = TimeSpan.FromSeconds(2.0);
 
-			RedisLock issueLock = new RedisLock(_redisService.Database, "issues/lock");
+			RedisLock issueLock = new (_redisService.GetDatabase(), "issues/lock");
 			while (!await issueLock.AcquireAsync(TimeSpan.FromMinutes(1)))
 			{
 				if (timer.Elapsed > nextNotifyTime)
diff --git a/Engine/Source/Programs/Horde/Horde.Build/Jobs/Schedules/ScheduleService.cs b/Engine/Source/Programs/Horde/Horde.Build/Jobs/Schedules/ScheduleService.cs
--- a/Engine/Source/Programs/Horde/Horde.Build/Jobs/Schedules/ScheduleService.cs
+++ b/Engine/Source/Programs/Horde/Horde.Build/Jobs/Schedules/ScheduleService.cs
@@ -172,7 +172,7 @@ namespace Horde.Build.Jobs.Schedules
 			_redis = redis;
 			_baseLockKey = "scheduler/locks";
 			_tickLockKey = _baseLockKey.Append("/tick");
-			_queue = new RedisSortedSet<QueueItem>(redis.Database, "scheduler/queue");
+			_queue = new RedisSortedSet<QueueItem>(redis.GetDatabase(), "scheduler/queue");
 			if (mongoService.ReadOnlyMode)
 			{
 				_ticker = new NullTicker();
@@ -206,7 +206,7 @@ namespace Horde.Build.Jobs.Schedules
 			DateTime utcNow = _clock.UtcNow;
 
 			// Update the current queue
-			await using (RedisLock sharedLock = new RedisLock(_redis.Database, _tickLockKey))
+			await using (RedisLock sharedLock = new (_redis.GetDatabase(), _tickLockKey))
 			{
 				if (await sharedLock.AcquireAsync(TimeSpan.FromMinutes(1.0), false))
 				{
@@ -225,7 +225,7 @@ namespace Horde.Build.Jobs.Schedules
 				}
 
 				// Acquire the lock for this schedule and update it
-				await using (RedisLock sharedLock = new RedisLock<QueueItem>(_redis.Database, _baseLockKey, item))
+				await using (RedisLock sharedLock = new RedisLock<QueueItem>(_redis.GetDatabase(), _baseLockKey, item))
 				{
 					if (await sharedLock.AcquireAsync(TimeSpan.FromMinutes(1.0)))
 					{
@@ -264,8 +264,9 @@ namespace Horde.Build.Jobs.Schedules
 
 		internal async Task ResetAsync()
 		{
-			await _redis.Database.KeyDeleteAsync(_queue.Key);
-			await _redis.Database.KeyDeleteAsync(_tickLockKey);
+			IDatabase redis = _redis.GetDatabase();
+			await redis.KeyDeleteAsync(_queue.Key);
+			await redis.KeyDeleteAsync(_tickLockKey);
 		}
 
 		internal async Task TickForTestingAsync()
diff --git a/Engine/Source/Programs/Horde/Horde.Build/Logs/Builder/RedisLogBuilder.cs b/Engine/Source/Programs/Horde/Horde.Build/Logs/Builder/RedisLogBuilder.cs
--- a/Engine/Source/Programs/Horde/Horde.Build/Logs/Builder/RedisLogBuilder.cs
+++ b/Engine/Source/Programs/Horde/Horde.Build/Logs/Builder/RedisLogBuilder.cs
@@ -6,6 +6,7 @@ using System.Globalization;
 using System.Text.RegularExpressions;
 using System.Threading.Tasks;
 using EpicGames.Core;
+using EpicGames.Redis;
 using Horde.Build.Logs.Data;
 using Horde.Build.Utilities;
 using Microsoft.Extensions.Logging;
diff --git a/Engine/Source/Programs/Horde/Horde.Build/Notifications/NotificationService.cs b/Engine/Source/Programs/Horde/Horde.Build/Notifications/NotificationService.cs
--- a/Engine/Source/Programs/Horde/Horde.Build/Notifications/NotificationService.cs
+++ b/Engine/Source/Programs/Horde/Horde.Build/Notifications/NotificationService.cs
@@ -8,6 +8,7 @@ using System.Security.Claims;
 using System.Threading;
 using System.Threading.Tasks;
 using EpicGames.Core;
+using EpicGames.Redis;
 using Horde.Build.Agents.Pools;
 using Horde.Build.Devices;
 using Horde.Build.Issues;
@@ -88,9 +89,9 @@ namespace Horde.Build.Notifications
 		private readonly IDogStatsd _dogStatsd;
 
 		/// <summary>
-		/// Redis database
+		/// Connection pool for Redis databases
 		/// </summary>
-		private readonly IDatabase _redisDatabase;
+		private readonly RedisConnectionPool _redisConnectionPool;
 		
 		/// <summary>
 		/// Settings for the application.
@@ -141,7 +142,7 @@ namespace Horde.Build.Notifications
 			_issueService = issueService;
 			_logFileService = logFileService;
 			_dogStatsd = dogStatsd;
-			_redisDatabase = redisService.Database;
+			_redisConnectionPool = redisService.ConnectionPool;
 
 			issueService.OnIssueUpdated += NotifyIssueUpdated;
 			jobService.OnJobStepComplete += NotifyJobStepComplete;
@@ -290,7 +291,7 @@ namespace Horde.Build.Notifications
 			try
 			{
 				byte[] data = JsonSerializer.SerializeToUtf8Bytes(notification);
-				await _redisDatabase.ListRightPushAsync(RedisQueueListKey(typeof(T).ToString()), data);
+				await _redisConnectionPool.GetDatabase().ListRightPushAsync(RedisQueueListKey(typeof(T).ToString()), data);
 			}
 			catch (Exception e)
 			{
@@ -318,11 +319,13 @@ namespace Horde.Build.Notifications
 		/// <exception cref="Exception"></exception>
 		private async Task<List<T>> GetAllQueuedNotificationsAsync<T>() where T : INotification
 		{
+			IDatabase redis = _redisConnectionPool.GetDatabase();
+
 			// Reading and deleting the entire list from Redis in this way is not thread-safe.
 			// But likely not a big deal considering the alternative of distributed locks.
 			string redisKey = RedisQueueListKey(typeof(T).ToString());
-			RedisValue[] rawNotifications = await _redisDatabase.ListRangeAsync(redisKey);
-			await _redisDatabase.KeyDeleteAsync(redisKey);
+			RedisValue[] rawNotifications = await redis.ListRangeAsync(redisKey);
+			await redis.KeyDeleteAsync(redisKey);
 
 			List<T> notifications = new List<T>();
 			foreach (byte[] data in rawNotifications)
diff --git a/Engine/Source/Programs/Horde/Horde.Build/Perforce/CommitService.cs b/Engine/Source/Programs/Horde/Horde.Build/Perforce/CommitService.cs
--- a/Engine/Source/Programs/Horde/Horde.Build/Perforce/CommitService.cs
+++ b/Engine/Source/Programs/Horde/Horde.Build/Perforce/CommitService.cs
@@ -85,7 +85,7 @@ namespace Horde.Build.Perforce
 		public CommitServiceOptions Options { get; }
 
 		// Redis
-		readonly IDatabase _redis;
+		readonly RedisConnectionPool _redisConnectionPool;
 
 		// Collections
 		readonly ICommitCollection _commitCollection;
@@ -111,8 +111,7 @@ namespace Horde.Build.Perforce
 		{
 			Options = options.Value;
 
-			_redis = redisService.Database;
-
+			_redisConnectionPool = redisService.ConnectionPool;
 			_commitCollection = commitCollection;
 			_streamCollection = streamCollection;
 			_perforceService = perforceService;
@@ -231,7 +230,7 @@ namespace Horde.Build.Perforce
 			Dictionary<IStream, int> streamToNextChange = new Dictionary<IStream, int>();
 			foreach (IStream stream in streams)
 			{
-				RedisValue lastChange = await _redis.StringGetAsync(LastChangeKey(stream.Id));
+				RedisValue lastChange = await _redisConnectionPool.GetDatabase().StringGetAsync(LastChangeKey(stream.Id));
 
 				int nextChange = 0;
 				if (lastChange.HasValue && lastChange.IsInteger)
diff --git a/Engine/Source/Programs/Horde/Horde.Build/Perforce/ReplicationService.cs b/Engine/Source/Programs/Horde/Horde.Build/Perforce/ReplicationService.cs
--- a/Engine/Source/Programs/Horde/Horde.Build/Perforce/ReplicationService.cs
+++ b/Engine/Source/Programs/Horde/Horde.Build/Perforce/ReplicationService.cs
@@ -158,7 +158,7 @@ namespace Horde.Build.Perforce
 		public ReplicationServiceOptions Options { get; }
 
 		// Redis
-		readonly IDatabase _redis;
+		readonly RedisConnectionPool _redisConnectionPool;
 
 		// Collections
 		readonly CommitService _commitService;
@@ -190,9 +190,9 @@ namespace Horde.Build.Perforce
 		{
 			Options = options.Value;
 
-			_redis = redisService.Database;
-			_redisDirtyStreams = new RedisSet<StreamId>(_redis, RedisBaseKey.Append("streams"));
-			_redisReservations = new RedisSortedSet<StreamId>(_redis, RedisBaseKey.Append("reservations"));
+			_redisConnectionPool = redisService.ConnectionPool;
+			_redisDirtyStreams = new RedisSet<StreamId>(_redisConnectionPool, RedisBaseKey.Append("streams"));
+			_redisReservations = new RedisSortedSet<StreamId>(_redisConnectionPool, RedisBaseKey.Append("reservations"));
 
 			_commitService = commitService;
 			_commitListener = _commitService.AddListener(OnCommitAdded);
@@ -216,7 +216,7 @@ namespace Horde.Build.Perforce
 			{
 				_stopping = false;
 				_updateStreamsEvent.Reset();
-				_redisUpdateSubscription = await _redis.Multiplexer.GetSubscriber().SubscribeAsync(RedisUpdateChannel, (_, _) => _updateStreamsEvent.Pulse());
+				_redisUpdateSubscription = await _redisConnectionPool.GetDatabase().Multiplexer.GetSubscriber().SubscribeAsync(RedisUpdateChannel, (_, _) => _updateStreamsEvent.Pulse());
 				_streamUpdateTask = Task.Run(() => UpdateContentAsync(), cancellationToken);
 			}
 		}
@@ -238,7 +238,7 @@ namespace Horde.Build.Perforce
 			}
 		}
 
-		RedisList<int> RedisStreamChanges(StreamId streamId) => new RedisList<int>(_redis, RedisBaseKey.Append($"stream/{streamId}/changes"));
+		RedisList<int> RedisStreamChanges(StreamId streamId) => new RedisList<int>(_redisConnectionPool, RedisBaseKey.Append($"stream/{streamId}/changes"));
 
 		async Task OnCommitAdded(ICommit commit)
 		{
@@ -248,7 +248,7 @@ namespace Horde.Build.Perforce
 
 			// Signal to any listeners that we have new data to process
 			await _redisDirtyStreams.AddAsync(commit.StreamId);
-			await _redis.PublishAsync(RedisUpdateChannel, commit.StreamId);
+			await _redisConnectionPool.GetDatabase().PublishAsync(RedisUpdateChannel, commit.StreamId);
 		}
 
 		async Task UpdateContentAsync()
@@ -391,7 +391,7 @@ namespace Horde.Build.Perforce
 				}
 
 				// Remove this stream from the dirty list if it's empty
-				ITransaction transaction = _redis.CreateTransaction();
+				ITransaction transaction = _redisConnectionPool.GetDatabase().CreateTransaction();
 				transaction.AddCondition(Condition.ListLengthLessThan(streamChanges.Key, 2));
 				_ = transaction.With(_redisDirtyStreams).RemoveAsync(streamId);
 				if (await transaction.ExecuteAsync())
@@ -441,7 +441,7 @@ namespace Horde.Build.Perforce
 					prevChange = values[1];
 
 					// Remove the first item from the list
-					ITransaction transaction = _redis.CreateTransaction();
+					ITransaction transaction = _redisConnectionPool.GetDatabase().CreateTransaction();
 					transaction.AddCondition(Condition.ListIndexEqual(changes.Key, 0, values[0]));
 					transaction.AddCondition(Condition.ListIndexEqual(changes.Key, 1, values[1]));
 					_ = transaction.With(changes).LeftPopAsync();
diff --git a/Engine/Source/Programs/Horde/Horde.Build/Server/LifetimeService.cs b/Engine/Source/Programs/Horde/Horde.Build/Server/LifetimeService.cs
--- a/Engine/Source/Programs/Horde/Horde.Build/Server/LifetimeService.cs
+++ b/Engine/Source/Programs/Horde/Horde.Build/Server/LifetimeService.cs
@@ -168,8 +168,9 @@ namespace Horde.Build.Server
 			try
 			{
 				string key = "HordeLifetimeService-Health-Check";
-				await _redisService.Database.StringSetAsync(key, "ok");
-				await _redisService.Database.StringGetAsync(key);
+				IDatabase redis = _redisService.GetDatabase();
+				await redis.StringSetAsync(key, "ok");
+				await redis.StringGetAsync(key);
 				isHealthy = true;
 			}
 			catch (Exception e)
diff --git a/Engine/Source/Programs/Horde/Horde.Build/Server/MongoService.cs b/Engine/Source/Programs/Horde/Horde.Build/Server/MongoService.cs
--- a/Engine/Source/Programs/Horde/Horde.Build/Server/MongoService.cs
+++ b/Engine/Source/Programs/Horde/Horde.Build/Server/MongoService.cs
@@ -931,7 +931,7 @@ namespace Horde.Build.Server
 				Func<CancellationToken, Task> updateIndexTask = await _mongoService.ReadNextUpgradeTask(cancellationToken);
 				for (; ; )
 				{
-					using (RedisLock schemaLock = new RedisLock(_redisService.Database, s_schemaLockKey))
+					using (RedisLock schemaLock = new (_redisService.GetDatabase(), s_schemaLockKey))
 					{
 						if (await schemaLock.AcquireAsync(TimeSpan.FromMinutes(5.0)) && await UpdateOneAsync(updateIndexTask, cancellationToken))
 						{
diff --git a/Engine/Source/Programs/Horde/Horde.Build/Server/RedisService.cs b/Engine/Source/Programs/Horde/Horde.Build/Server/RedisService.cs
--- a/Engine/Source/Programs/Horde/Horde.Build/Server/RedisService.cs
+++ b/Engine/Source/Programs/Horde/Horde.Build/Server/RedisService.cs
@@ -9,6 +9,7 @@ using System.Net.NetworkInformation;
 using System.Runtime.InteropServices;
 using System.Threading.Tasks;
 using EpicGames.Core;
+using EpicGames.Redis;
 using Horde.Build.Utilities;
 using Microsoft.Extensions.Logging;
 using Microsoft.Extensions.Options;
@@ -39,10 +40,11 @@ namespace Horde.Build.Server
 		/// <summary>
 		/// Connection multiplexer
 		/// </summary>
-		public ConnectionMultiplexer Multiplexer { get; }
+		private readonly ConnectionMultiplexer _multiplexer;
 
 		/// <summary>
-		/// The database interface
+		/// The database interface.
+		/// If possible, use ConnectionPool instead. 
 		/// </summary>
 		public IDatabase Database { get; }
 
@@ -105,15 +107,25 @@ namespace Horde.Build.Server
 				}
 			}
 
-			Multiplexer = ConnectionMultiplexer.Connect(connectionString);
-			Database = Multiplexer.GetDatabase(dbNum);
+			_multiplexer = ConnectionMultiplexer.Connect(connectionString);
+			Database = _multiplexer.GetDatabase(dbNum);
 			ConnectionPool = new RedisConnectionPool(20, connectionString, dbNum);
 		}
 
+		/// <summary>
+		/// Get the least-loaded Redis database from the connection pool
+		/// Don't store the returned object and try to resolve this as late as possible to ensure load is balanced.
+		/// </summary>
+		/// <returns>A Redis database</returns>
+		public IDatabase GetDatabase()
+		{
+			return ConnectionPool.GetDatabase();
+		}
+
 		/// <inheritdoc/>
 		public void Dispose()
 		{
-			Multiplexer.Dispose();
+			_multiplexer.Dispose();
 
 			if (_redisProcess != null)
 			{
diff --git a/Engine/Source/Programs/Horde/Horde.Build/Startup.cs b/Engine/Source/Programs/Horde/Horde.Build/Startup.cs
--- a/Engine/Source/Programs/Horde/Horde.Build/Startup.cs
+++ b/Engine/Source/Programs/Horde/Horde.Build/Startup.cs
@@ -324,8 +324,6 @@ namespace Horde.Build
 			RedisService redisService = new RedisService(settings);
 #pragma warning restore CA2000 // Dispose objects before losing scope
 			services.AddSingleton<RedisService>(sp => redisService);
-			services.AddSingleton<IDatabase>(redisService.Database);
-			services.AddSingleton<ConnectionMultiplexer>(sp => redisService.Multiplexer);
 			services.AddDataProtection().PersistKeysToStackExchangeRedis(() => redisService.Database, "aspnet-data-protection");
 
 			if (settings.CorsEnabled)
diff --git a/Engine/Source/Programs/Horde/Horde.Build/Tools/ToolCollection.cs b/Engine/Source/Programs/Horde/Horde.Build/Tools/ToolCollection.cs
--- a/Engine/Source/Programs/Horde/Horde.Build/Tools/ToolCollection.cs
+++ b/Engine/Source/Programs/Horde/Horde.Build/Tools/ToolCollection.cs
@@ -177,7 +177,7 @@ namespace Horde.Build.Tools
 		private async Task ClearCachedIndexAsync()
 		{
 			CachedIndex newIndexData = new CachedIndex { Rev = ObjectId.GenerateNewId().ToString() };
-			await _redisService.Database.StringSetAsync(s_indexKey, CbSerializer.SerializeToByteArray(newIndexData));
+			await _redisService.GetDatabase().StringSetAsync(s_indexKey, CbSerializer.SerializeToByteArray(newIndexData));
 		}
 
 		/// <summary>
@@ -186,11 +186,12 @@ namespace Horde.Build.Tools
 		/// <returns>Sequence of tool documents</returns>
 		async Task<List<ToolId>> FindAllIdsAsync()
 		{
+			IDatabase redis = _redisService.GetDatabase();
 			CachedIndex? index;
 			for (; ; )
 			{
 				// Get the cached index, and check if it has valid data
-				RedisValue value = await _redisService.Database.StringGetAsync(s_indexKey);
+				RedisValue value = await redis.StringGetAsync(s_indexKey);
 				if (!value.IsNullOrEmpty)
 				{
 					index = CbSerializer.Deserialize<CachedIndex>((byte[])value);
@@ -210,14 +211,14 @@ namespace Horde.Build.Tools
 				RedisValue newValue = CbSerializer.SerializeToByteArray(index);
 				if (value.IsNull)
 				{
-					if (await _redisService.Database.StringSetAsync(s_indexKey, newValue, when: When.NotExists))
+					if (await redis.StringSetAsync(s_indexKey, newValue, when: When.NotExists))
 					{
 						break;
 					}
 				}
 				else
 				{
-					ITransaction transaction = _redisService.Database.CreateTransaction();
+					ITransaction transaction = redis.CreateTransaction();
 					transaction.AddCondition(Condition.StringEqual(s_indexKey, value));
 					_ = transaction.StringSetAsync(s_indexKey, newValue);
 
diff --git a/Engine/Source/Programs/Horde/Horde.Build/Utilities/Clock.cs b/Engine/Source/Programs/Horde/Horde.Build/Utilities/Clock.cs
--- a/Engine/Source/Programs/Horde/Horde.Build/Utilities/Clock.cs
+++ b/Engine/Source/Programs/Horde/Horde.Build/Utilities/Clock.cs
@@ -237,7 +237,7 @@ namespace HordeCommon
 
 		async ValueTask TriggerSharedAsync(RedisKey key, TimeSpan interval, Func<CancellationToken, ValueTask> tickAsync, CancellationToken cancellationToken)
 		{
-			using (RedisLock sharedLock = new RedisLock(_redis.Database, key))
+			using (RedisLock sharedLock = new (_redis.GetDatabase(), key))
 			{
 				if (await sharedLock.AcquireAsync(interval, false))
 				{
diff --git a/Engine/Source/Programs/Horde/Horde.Build/Utilities/RedisConnectionPool.cs b/Engine/Source/Programs/Horde/Horde.Build/Utilities/RedisConnectionPool.cs
deleted file mode 100644
diff --git a/Engine/Source/Programs/Horde/Horde.Build/Utilities/VersionedCollection.cs b/Engine/Source/Programs/Horde/Horde.Build/Utilities/VersionedCollection.cs
--- a/Engine/Source/Programs/Horde/Horde.Build/Utilities/VersionedCollection.cs
+++ b/Engine/Source/Programs/Horde/Horde.Build/Utilities/VersionedCollection.cs
@@ -5,6 +5,7 @@ using System.Collections.Generic;
 using System.Linq;
 using System.Text.Json.Serialization;
 using System.Threading.Tasks;
+using EpicGames.Redis;
 using Horde.Build.Server;
 using MongoDB.Bson;
 using MongoDB.Bson.IO;
@@ -124,7 +125,7 @@ namespace Horde.Build.Utilities
 
 		private static FilterDefinitionBuilder<VersionedDocument<TId, TLatest>> FilterBuilder { get; } = Builders<VersionedDocument<TId, TLatest>>.Filter;
 
-		private readonly IDatabase _redis;
+		private readonly RedisConnectionPool _redisConnectionPool;
 		private readonly RedisKey _baseKey;
 
 		private static readonly HashSet<Type> s_registeredTypes = new HashSet<Type>();
@@ -149,7 +150,7 @@ namespace Horde.Build.Utilities
 		/// <param name="types">Types to serialize from the collection</param>
 		public VersionedCollection(MongoService mongoService, string collectionName, RedisService redisService, RedisKey baseKey, IReadOnlyDictionary<int, Type> types)
 		{
-			_redis = redisService.Database;
+			_redisConnectionPool = redisService.ConnectionPool;
 			_baseKey = baseKey;
 
 			lock (s_registeredTypes)
@@ -208,12 +209,12 @@ namespace Horde.Build.Utilities
 			KeyValuePair<RedisKey, RedisValue>[] pairs = new KeyValuePair<RedisKey, RedisValue>[2];
 			pairs[0] = new KeyValuePair<RedisKey, RedisValue>(docKey, doc.ToBson(typeof(VersionedDocument<TId, TLatest>)));
 			pairs[1] = new KeyValuePair<RedisKey, RedisValue>(docKey.Append(s_timeSuffix), doc.LastUpdateTime.Ticks);
-			_redis.StringSetAsync(pairs, When.NotExists, flags: CommandFlags.FireAndForget);
+			_redisConnectionPool.GetDatabase().StringSetAsync(pairs, When.NotExists, flags: CommandFlags.FireAndForget);
 		}
 
 		private async ValueTask<VersionedDocument<TId, TLatest>?> GetCachedValueAsync(RedisKey docKey)
 		{
-			RedisValue cacheValue = await _redis.StringGetAsync(docKey);
+			RedisValue cacheValue = await _redisConnectionPool.GetDatabase().StringGetAsync(docKey);
 			if (!cacheValue.IsNull)
 			{
 				try
@@ -230,7 +231,7 @@ namespace Horde.Build.Utilities
 
 		private async ValueTask<bool> UpdateCachedValueAsync(RedisKey docKey, VersionedDocument<TId, TLatest> prevDoc, TLatest doc)
 		{
-			ITransaction transaction = _redis.CreateTransaction();
+			ITransaction transaction = _redisConnectionPool.GetDatabase().CreateTransaction();
 			transaction.AddCondition(Condition.StringEqual(docKey.Append(s_timeSuffix), prevDoc.LastUpdateTime.Ticks));
 
 			KeyValuePair<RedisKey, RedisValue>[] pairs = new KeyValuePair<RedisKey, RedisValue>[2];
@@ -250,7 +251,7 @@ namespace Horde.Build.Utilities
 
 		private async ValueTask DeleteCachedValueAsync(RedisKey docKey)
 		{
-			await _redis.KeyDeleteAsync(new[] { docKey, docKey.Append(s_timeSuffix) });
+			await _redisConnectionPool.GetDatabase().KeyDeleteAsync(new[] { docKey, docKey.Append(s_timeSuffix) });
 		}
 
 		/// <summary>
diff --git a/Engine/Source/Programs/Shared/EpicGames.Redis/RedisConnectionPool.cs b/Engine/Source/Programs/Shared/EpicGames.Redis/RedisConnectionPool.cs
new file mode 100644
--- /dev/null
+++ b/Engine/Source/Programs/Shared/EpicGames.Redis/RedisConnectionPool.cs
@@ -0,0 +1,72 @@
+// Copyright Epic Games, Inc. All Rights Reserved.
+
+using System;
+using System.Collections.Concurrent;
+using System.Collections.Generic;
+using System.Linq;
+using StackExchange.Redis;
+
+namespace EpicGames.Redis
+{
+	/// <summary>
+	/// A connection pool for Redis client
+	///
+	/// Wraps multiple ConnectionMultiplexers as lazy values and initializes them as needed.
+	/// If full, the least loaded connection will be picked.
+	/// </summary>
+	public class RedisConnectionPool
+	{
+		private readonly ConcurrentBag<Lazy<ConnectionMultiplexer>> _connections;
+		private readonly int _defaultDatabaseIndex;
+
+		/// <summary>
+		/// Constructor
+		/// </summary>
+		/// <param name="poolSize">Size of the pool (i.e max number of connections)</param>
+		/// <param name="redisConfString">Configuration string for Redis</param>
+		/// <param name="defaultDatabaseIndex">The Redis database index to use. Use -1 for the default one</param>
+		public RedisConnectionPool(int poolSize, string redisConfString, int defaultDatabaseIndex = -1)
+		{
+			_defaultDatabaseIndex = defaultDatabaseIndex;
+			_connections = new ConcurrentBag<Lazy<ConnectionMultiplexer>>();
+			for (int i = 0; i < poolSize; i++)
+			{
+				_connections.Add(new Lazy<ConnectionMultiplexer>(() => ConnectionMultiplexer.Connect(redisConfString)));
+			}
+		}
+
+		/// <summary>
+		/// Get a connection from the pool
+		///
+		/// It will pick the least loaded connection or create a new one (if pool size allows)
+		/// </summary>
+		/// <returns>A Redis database connection</returns>
+		public IConnectionMultiplexer GetConnection()
+		{
+			Lazy<ConnectionMultiplexer> lazyConnection;
+			IEnumerable<Lazy<ConnectionMultiplexer>>? lazyConnections = _connections.Where(x => x.IsValueCreated);
+
+			if (lazyConnections.Count() == _connections.Count)
+			{
+				// No more new connections can be created, pick the least loaded one
+				lazyConnection = _connections.OrderBy(x => x.Value.GetCounters().TotalOutstanding).First();
+			}
+			else
+			{
+				// Create a new connection by picking a not yet initialized lazy value 
+				lazyConnection = _connections.First(x => !x.IsValueCreated);
+			}
+
+			return lazyConnection.Value;
+		}
+
+		/// <summary>
+		/// Shortcut to getting a IDatabase
+		/// </summary>
+		/// <returns>A Redis database</returns>
+		public IDatabase GetDatabase()
+		{
+			return GetConnection().GetDatabase(_defaultDatabaseIndex);
+		}
+	}
+}
\ No newline at end of file
diff --git a/Engine/Source/Programs/Shared/EpicGames.Redis/RedisHash.cs b/Engine/Source/Programs/Shared/EpicGames.Redis/RedisHash.cs
--- a/Engine/Source/Programs/Shared/EpicGames.Redis/RedisHash.cs
+++ b/Engine/Source/Programs/Shared/EpicGames.Redis/RedisHash.cs
@@ -15,45 +15,71 @@ namespace EpicGames.Redis
 	/// <typeparam name="TValue">The value type for the hash</typeparam>
 	public readonly struct RedisHash<TName, TValue>
 	{
-		internal readonly IDatabaseAsync Database;
+		internal readonly RedisConnectionPool? ConnectionPool = null;
+		internal readonly IDatabaseAsync? Database = null;
 
 		/// <summary>
-		/// The key for the hash
+		/// The key for the list
 		/// </summary>
 		public readonly RedisKey Key { get; }
 
 		/// <summary>
 		/// Constructor
 		/// </summary>
+		/// <param name="database">Explicitly set Redis database (overrides use of any connection pool)</param>
+		/// <param name="key">Redis key this type is using</param>
 		public RedisHash(IDatabaseAsync database, RedisKey key)
 		{
 			Database = database;
 			Key = key;
 		}
+		
+		/// <summary>
+		/// Constructor
+		/// </summary>
+		/// <param name="connectionPool">Connection pool for Redis databases</param>
+		/// <param name="key">Redis key this type is using</param>
+		public RedisHash(RedisConnectionPool connectionPool, RedisKey key)
+		{
+			ConnectionPool = connectionPool;
+			Key = key;
+		}
+		
+		/// <summary>
+		/// Get the Redis database in use
+		/// </summary>
+		/// <returns>A connection pool or explicitly set Redis database</returns>
+		/// <exception cref="InvalidOperationException">If neither are set</exception>
+		public IDatabaseAsync GetDatabase()
+		{
+			if (Database != null) return Database;
+			if (ConnectionPool != null) return ConnectionPool.GetDatabase();
+			throw new InvalidOperationException($"Neither {nameof(Database)} or {nameof(ConnectionPool)} has been set!");
+		}
 
 		/// <inheritdoc cref="IDatabaseAsync.HashDeleteAsync(RedisKey, RedisValue, CommandFlags)"/>
 		public Task<bool> DeleteAsync(TName name, CommandFlags flags = CommandFlags.None)
 		{
-			return Database.HashDeleteAsync(Key, RedisSerializer.Serialize(name), flags);
+			return GetDatabase().HashDeleteAsync(Key, RedisSerializer.Serialize(name), flags);
 		}
 
 		/// <inheritdoc cref="IDatabaseAsync.HashDeleteAsync(RedisKey, RedisValue[], CommandFlags)"/>
 		public Task<long> DeleteAsync(IEnumerable<TName> names, CommandFlags flags = CommandFlags.None)
 		{
 			RedisValue[] nameArray = names.Select(x => RedisSerializer.Serialize(x)).ToArray();
-			return Database.HashDeleteAsync(Key, nameArray, flags);
+			return GetDatabase().HashDeleteAsync(Key, nameArray, flags);
 		}
 
 		/// <inheritdoc cref="IDatabaseAsync.HashExistsAsync(RedisKey, RedisValue, CommandFlags)"/>
 		public Task<bool> ExistsAsync(TName name, CommandFlags flags = CommandFlags.None)
 		{
-			return Database.HashExistsAsync(Key, RedisSerializer.Serialize(name), flags);
+			return GetDatabase().HashExistsAsync(Key, RedisSerializer.Serialize(name), flags);
 		}
 
 		/// <inheritdoc cref="IDatabaseAsync.HashGetAsync(RedisKey, RedisValue, CommandFlags)"/>
 		public async Task<TValue> GetAsync(TName name, CommandFlags flags = CommandFlags.None)
 		{
-			RedisValue value = await Database.HashGetAsync(Key, RedisSerializer.Serialize(name), flags);
+			RedisValue value = await GetDatabase().HashGetAsync(Key, RedisSerializer.Serialize(name), flags);
 			return value.IsNull ? default! : RedisSerializer.Deserialize<TValue>(value)!;
 		}
 
@@ -61,46 +87,46 @@ namespace EpicGames.Redis
 		public async Task<TValue[]> GetAsync(IEnumerable<TName> names, CommandFlags flags = CommandFlags.None)
 		{
 			RedisValue[] nameArray = names.Select(x => RedisSerializer.Serialize(x)).ToArray();
-			RedisValue[] valueArray = await Database.HashGetAsync(Key, nameArray, flags);
+			RedisValue[] valueArray = await GetDatabase().HashGetAsync(Key, nameArray, flags);
 			return Array.ConvertAll(valueArray, x => RedisSerializer.Deserialize<TValue>(x)!);
 		}
 
 		/// <inheritdoc cref="IDatabaseAsync.HashGetAllAsync(RedisKey, CommandFlags)"/>
 		public async Task<HashEntry<TName, TValue>[]> GetAllAsync(CommandFlags flags = CommandFlags.None)
 		{
-			HashEntry[] entries = await Database.HashGetAllAsync(Key, flags);
+			HashEntry[] entries = await GetDatabase().HashGetAllAsync(Key, flags);
 			return Array.ConvertAll(entries, x => new HashEntry<TName, TValue>(RedisSerializer.Deserialize<TName>(x.Name)!, RedisSerializer.Deserialize<TValue>(x.Value)!));
 		}
 
 		/// <inheritdoc cref="IDatabaseAsync.HashKeysAsync(RedisKey, CommandFlags)"/>
 		public async Task<TName[]> KeysAsync(CommandFlags flags = CommandFlags.None)
 		{
-			RedisValue[] nameArray = await Database.HashKeysAsync(Key, flags);
+			RedisValue[] nameArray = await GetDatabase().HashKeysAsync(Key, flags);
 			return Array.ConvertAll(nameArray, x => RedisSerializer.Deserialize<TName>(x)!);
 		}
 
 		/// <inheritdoc cref="IDatabaseAsync.HashLengthAsync(RedisKey, CommandFlags)"/>
 		public Task<long> LengthAsync(CommandFlags flags = CommandFlags.None)
 		{
-			return Database.HashLengthAsync(Key, flags);
+			return GetDatabase().HashLengthAsync(Key, flags);
 		}
 
 		/// <inheritdoc cref="IDatabaseAsync.HashSetAsync(RedisKey, RedisValue, RedisValue, When, CommandFlags)"/>
 		public Task SetAsync(TName name, TValue value, When when = When.Always, CommandFlags flags = CommandFlags.None)
 		{
-			return Database.HashSetAsync(Key, RedisSerializer.Serialize(name), RedisSerializer.Serialize(value), when, flags);
+			return GetDatabase().HashSetAsync(Key, RedisSerializer.Serialize(name), RedisSerializer.Serialize(value), when, flags);
 		}
 
 		/// <inheritdoc cref="IDatabaseAsync.HashSetAsync(RedisKey, HashEntry[], CommandFlags)"/>
 		public Task SetAsync(IEnumerable<HashEntry<TName, TValue>> entries, CommandFlags flags = CommandFlags.None)
 		{
-			return Database.HashSetAsync(Key, entries.Select(x => (HashEntry)x).ToArray(), flags);
+			return GetDatabase().HashSetAsync(Key, entries.Select(x => (HashEntry)x).ToArray(), flags);
 		}
 
 		/// <inheritdoc cref="IDatabaseAsync.HashValuesAsync(RedisKey, CommandFlags)"/>
 		public async Task<TValue[]> ValuesAsync(CommandFlags flags = CommandFlags.None)
 		{
-			RedisValue[] values = await Database.HashValuesAsync(Key, flags);
+			RedisValue[] values = await GetDatabase().HashValuesAsync(Key, flags);
 			return Array.ConvertAll(values, x => RedisSerializer.Deserialize<TValue>(x)!);
 		}
 	}
@@ -152,25 +178,25 @@ namespace EpicGames.Redis
 		/// <inheritdoc cref="IDatabaseAsync.HashDecrementAsync(RedisKey, RedisValue, Int64, CommandFlags)"/>
 		public static Task<long> DecrementAsync<TName>(this RedisHash<TName, long> hash, TName name, long value = 1L, CommandFlags flags = CommandFlags.None)
 		{
-			return hash.Database.HashDecrementAsync(hash.Key, RedisSerializer.Serialize<TName>(name), value, flags);
+			return hash.GetDatabase().HashDecrementAsync(hash.Key, RedisSerializer.Serialize<TName>(name), value, flags);
 		}
 
 		/// <inheritdoc cref="IDatabaseAsync.HashDecrementAsync(RedisKey, RedisValue, Double, CommandFlags)"/>
 		public static Task<double> DecrementAsync<TName>(this RedisHash<TName, long> hash, TName name, double value = 1.0, CommandFlags flags = CommandFlags.None)
 		{
-			return hash.Database.HashDecrementAsync(hash.Key, RedisSerializer.Serialize<TName>(name), value, flags);
+			return hash.GetDatabase().HashDecrementAsync(hash.Key, RedisSerializer.Serialize<TName>(name), value, flags);
 		}
 
 		/// <inheritdoc cref="IDatabaseAsync.HashIncrementAsync(RedisKey, RedisValue, Int64, CommandFlags)"/>
 		public static Task<long> IncrementAsync<TName>(this RedisHash<TName, long> hash, TName name, long value = 1L, CommandFlags flags = CommandFlags.None)
 		{
-			return hash.Database.HashIncrementAsync(hash.Key, RedisSerializer.Serialize<TName>(name), value, flags);
+			return hash.GetDatabase().HashIncrementAsync(hash.Key, RedisSerializer.Serialize<TName>(name), value, flags);
 		}
 
 		/// <inheritdoc cref="IDatabaseAsync.HashIncrementAsync(RedisKey, RedisValue, Double, CommandFlags)"/>
 		public static Task<double> IncrementAsync<TName>(this RedisHash<TName, long> hash, TName name, double value = 1.0, CommandFlags flags = CommandFlags.None)
 		{
-			return hash.Database.HashIncrementAsync(hash.Key, RedisSerializer.Serialize<TName>(name), value, flags);
+			return hash.GetDatabase().HashIncrementAsync(hash.Key, RedisSerializer.Serialize<TName>(name), value, flags);
 		}
 
 		/// <summary>
diff --git a/Engine/Source/Programs/Shared/EpicGames.Redis/RedisList.cs b/Engine/Source/Programs/Shared/EpicGames.Redis/RedisList.cs
--- a/Engine/Source/Programs/Shared/EpicGames.Redis/RedisList.cs
+++ b/Engine/Source/Programs/Shared/EpicGames.Redis/RedisList.cs
@@ -14,10 +14,8 @@ namespace EpicGames.Redis
 	/// <typeparam name="TElement">The type of element stored in the list</typeparam>
 	public readonly struct RedisList<TElement>
 	{
-		/// <summary>
-		/// The database containing this object
-		/// </summary>
-		IDatabaseAsync Database { get; }
+		internal readonly RedisConnectionPool? ConnectionPool = null;
+		internal readonly IDatabaseAsync? Database = null;
 
 		/// <summary>
 		/// The key for the list
@@ -27,16 +25,41 @@ namespace EpicGames.Redis
 		/// <summary>
 		/// Constructor
 		/// </summary>
+		/// <param name="database">Explicitly set Redis database (overrides use of any connection pool)</param>
+		/// <param name="key">Redis key this type is using</param>
 		public RedisList(IDatabaseAsync database, RedisKey key)
 		{
 			Database = database;
 			Key = key;
 		}
+		
+		/// <summary>
+		/// Constructor
+		/// </summary>
+		/// <param name="connectionPool">Connection pool for Redis databases</param>
+		/// <param name="key">Redis key this type is using</param>
+		public RedisList(RedisConnectionPool connectionPool, RedisKey key)
+		{
+			ConnectionPool = connectionPool;
+			Key = key;
+		}
+		
+		/// <summary>
+		/// Get the Redis database in use
+		/// </summary>
+		/// <returns>A connection pool or explicitly set Redis database</returns>
+		/// <exception cref="InvalidOperationException">If neither are set</exception>
+		public IDatabaseAsync GetDatabase()
+		{
+			if (Database != null) return Database;
+			if (ConnectionPool != null) return ConnectionPool.GetDatabase();
+			throw new InvalidOperationException($"Neither {nameof(Database)} or {nameof(ConnectionPool)} has been set!");
+		}
 
 		/// <inheritdoc cref="IDatabaseAsync.ListGetByIndexAsync(RedisKey, Int64, CommandFlags)"/>
 		public async Task<TElement> GetByIndexAsync(long index, CommandFlags flags = CommandFlags.None)
 		{
-			RedisValue value = await Database.ListGetByIndexAsync(Key, index, flags);
+			RedisValue value = await GetDatabase().ListGetByIndexAsync(Key, index, flags);
 			return value.IsNull? default! : RedisSerializer.Deserialize<TElement>(value)!;
 		}
 
@@ -45,7 +68,7 @@ namespace EpicGames.Redis
 		{
 			RedisValue pivotValue = RedisSerializer.Serialize(pivot);
 			RedisValue itemValue = RedisSerializer.Serialize(item);
-			return Database.ListInsertAfterAsync(Key, pivotValue, itemValue, flags);
+			return GetDatabase().ListInsertAfterAsync(Key, pivotValue, itemValue, flags);
 		}
 
 		/// <inheritdoc cref="IDatabaseAsync.ListInsertBeforeAsync(RedisKey, RedisValue, RedisValue, CommandFlags)"/>
@@ -53,39 +76,39 @@ namespace EpicGames.Redis
 		{
 			RedisValue pivotValue = RedisSerializer.Serialize(pivot);
 			RedisValue itemValue = RedisSerializer.Serialize(item);
-			return Database.ListInsertBeforeAsync(Key, pivotValue, itemValue, flags);
+			return GetDatabase().ListInsertBeforeAsync(Key, pivotValue, itemValue, flags);
 		}
 
 		/// <inheritdoc cref="IDatabaseAsync.ListLeftPopAsync(RedisKey, CommandFlags)"/>
 		public async Task<TElement> LeftPopAsync(CommandFlags flags = CommandFlags.None)
 		{
-			RedisValue value = await Database.ListLeftPopAsync(Key, flags);
+			RedisValue value = await GetDatabase().ListLeftPopAsync(Key, flags);
 			return value.IsNull ? default! : RedisSerializer.Deserialize<TElement>(value)!;
 		}
 
 		/// <inheritdoc cref="IDatabaseAsync.ListLeftPushAsync(RedisKey, RedisValue, When, CommandFlags)"/>
 		public Task<long> LeftPushAsync(TElement item, When when = When.Always, CommandFlags flags = CommandFlags.None)
 		{
-			return Database.ListLeftPushAsync(Key, RedisSerializer.Serialize(item), when, flags);
+			return GetDatabase().ListLeftPushAsync(Key, RedisSerializer.Serialize(item), when, flags);
 		}
 
 		/// <inheritdoc cref="IDatabaseAsync.ListLeftPushAsync(RedisKey, RedisValue[], When, CommandFlags)"/>
 		public Task<long> LeftPushAsync(IEnumerable<TElement> items, When when = When.Always, CommandFlags flags = CommandFlags.None)
 		{
 			RedisValue[] values = items.Select(x => RedisSerializer.Serialize(x)).ToArray();
-			return Database.ListLeftPushAsync(Key, values, when, flags);
+			return GetDatabase().ListLeftPushAsync(Key, values, when, flags);
 		}
 
 		/// <inheritdoc cref="IDatabaseAsync.ListLengthAsync(RedisKey, CommandFlags)"/>
 		public Task<long> LengthAsync()
 		{
-			return Database.ListLengthAsync(Key);
+			return GetDatabase().ListLengthAsync(Key);
 		}
 
 		/// <inheritdoc cref="IDatabaseAsync.ListRangeAsync(RedisKey, Int64, Int64, CommandFlags)"/>
 		public async Task<TElement[]> RangeAsync(long start = 0, long stop = -1, CommandFlags flags = CommandFlags.None)
 		{
-			RedisValue[] values = await Database.ListRangeAsync(Key, start, stop, flags);
+			RedisValue[] values = await GetDatabase().ListRangeAsync(Key, start, stop, flags);
 			return Array.ConvertAll(values, (Converter<RedisValue, TElement>)(x => (TElement)RedisSerializer.Deserialize<TElement>(x)!));
 		}
 
@@ -93,40 +116,40 @@ namespace EpicGames.Redis
 		public Task<long> RemoveAsync(TElement item, long count = 0L, CommandFlags flags = CommandFlags.None)
 		{
 			RedisValue value = RedisSerializer.Serialize(item);
-			return Database.ListRemoveAsync(Key, value, count, flags);
+			return GetDatabase().ListRemoveAsync(Key, value, count, flags);
 		}
 
 		/// <inheritdoc cref="IDatabaseAsync.ListRightPopAsync(RedisKey, CommandFlags)"/>
 		public async Task<TElement> RightPopAsync(CommandFlags flags = CommandFlags.None)
 		{
-			RedisValue value = await Database.ListRightPopAsync(Key, flags);
+			RedisValue value = await GetDatabase().ListRightPopAsync(Key, flags);
 			return value.IsNull? default! : RedisSerializer.Deserialize<TElement>(value)!;
 		}
 
 		/// <inheritdoc cref="IDatabaseAsync.ListRightPushAsync(RedisKey, RedisValue, When, CommandFlags)"/>
 		public Task<long> RightPushAsync(TElement item, When when = When.Always, CommandFlags flags = CommandFlags.None)
 		{
-			return Database.ListRightPushAsync(Key, RedisSerializer.Serialize(item), when, flags);
+			return GetDatabase().ListRightPushAsync(Key, RedisSerializer.Serialize(item), when, flags);
 		}
 
 		/// <inheritdoc cref="IDatabaseAsync.ListRightPushAsync(RedisKey, RedisValue[], When, CommandFlags)"/>
 		public Task<long> RightPushAsync(IEnumerable<TElement> items, When when = When.Always, CommandFlags flags = CommandFlags.None)
 		{
 			RedisValue[] values = items.Select(x => RedisSerializer.Serialize(x)).ToArray();
-			return Database.ListRightPushAsync(Key, values, when, flags);
+			return GetDatabase().ListRightPushAsync(Key, values, when, flags);
 		}
 
 		/// <inheritdoc cref="IDatabaseAsync.ListSetByIndexAsync(RedisKey, Int64, RedisValue, CommandFlags)"/>
 		public Task SetByIndexAsync(long index, TElement item, CommandFlags flags = CommandFlags.None)
 		{
 			RedisValue value = RedisSerializer.Serialize(item);
-			return Database.ListSetByIndexAsync(Key, index, value, flags);
+			return GetDatabase().ListSetByIndexAsync(Key, index, value, flags);
 		}
 
 		/// <inheritdoc cref="IDatabaseAsync.ListTrimAsync(RedisKey, Int64, Int64, CommandFlags)"/>
 		public Task TrimAsync(long start, long stop, CommandFlags flags = CommandFlags.None)
 		{
-			return Database.ListTrimAsync(Key, start, stop, flags);
+			return GetDatabase().ListTrimAsync(Key, start, stop, flags);
 		}
 	}
 
diff --git a/Engine/Source/Programs/Shared/EpicGames.Redis/RedisSet.cs b/Engine/Source/Programs/Shared/EpicGames.Redis/RedisSet.cs
--- a/Engine/Source/Programs/Shared/EpicGames.Redis/RedisSet.cs
+++ b/Engine/Source/Programs/Shared/EpicGames.Redis/RedisSet.cs
@@ -14,7 +14,8 @@ namespace EpicGames.Redis
 	/// <typeparam name="TElement">The type of element stored in the list</typeparam>
 	public readonly struct RedisSet<TElement>
 	{
-		readonly IDatabaseAsync _database;
+		internal readonly RedisConnectionPool? ConnectionPool = null;
+		internal readonly IDatabaseAsync? Database = null;
 
 		/// <summary>
 		/// The key for the list
@@ -24,70 +25,95 @@ namespace EpicGames.Redis
 		/// <summary>
 		/// Constructor
 		/// </summary>
+		/// <param name="database">Explicitly set Redis database (overrides use of any connection pool)</param>
+		/// <param name="key">Redis key this type is using</param>
 		public RedisSet(IDatabaseAsync database, RedisKey key)
 		{
-			_database = database;
+			Database = database;
 			Key = key;
 		}
+		
+		/// <summary>
+		/// Constructor
+		/// </summary>
+		/// <param name="connectionPool">Connection pool for Redis databases</param>
+		/// <param name="key">Redis key this type is using</param>
+		public RedisSet(RedisConnectionPool connectionPool, RedisKey key)
+		{
+			ConnectionPool = connectionPool;
+			Key = key;
+		}
+		
+		/// <summary>
+		/// Get the Redis database in use
+		/// </summary>
+		/// <returns>A connection pool or explicitly set Redis database</returns>
+		/// <exception cref="InvalidOperationException">If neither are set</exception>
+		public IDatabaseAsync GetDatabase()
+		{
+			if (Database != null) return Database;
+			if (ConnectionPool != null) return ConnectionPool.GetDatabase();
+			throw new InvalidOperationException($"Neither {nameof(Database)} or {nameof(ConnectionPool)} has been set!");
+		}
 
 		/// <inheritdoc cref="IDatabaseAsync.SetAddAsync(RedisKey, RedisValue, CommandFlags)"/>
 		public Task<bool> AddAsync(TElement item, CommandFlags flags = CommandFlags.None)
 		{
-			return _database.SetAddAsync(Key, RedisSerializer.Serialize(item), flags);
+			return GetDatabase().SetAddAsync(Key, RedisSerializer.Serialize(item), flags);
 		}
 
 		/// <inheritdoc cref="IDatabaseAsync.SetAddAsync(RedisKey, RedisValue[], CommandFlags)"/>
 		public Task<long> AddAsync(IEnumerable<TElement> items, CommandFlags flags = CommandFlags.None)
 		{
 			RedisValue[] values = items.Select(x => RedisSerializer.Serialize(x)).ToArray();
-			return _database.SetAddAsync(Key, values, flags);
+			return GetDatabase().SetAddAsync(Key, values, flags);
 		}
 
 		/// <inheritdoc cref="IDatabaseAsync.SetContainsAsync(RedisKey, RedisValue, CommandFlags)"/>
 		public Task<bool> ContainsAsync(TElement item, CommandFlags flags = CommandFlags.None)
 		{
 			RedisValue value = RedisSerializer.Serialize(item);
-			return _database.SetContainsAsync(Key, value, flags);
+			return GetDatabase().SetContainsAsync(Key, value, flags);
 		}
 
 		/// <inheritdoc cref="IDatabaseAsync.SetLengthAsync(RedisKey, CommandFlags)"/>
 		public Task<long> LengthAsync(CommandFlags flags = CommandFlags.None)
 		{
-			return _database.SetLengthAsync(Key, flags);
+			return GetDatabase().SetLengthAsync(Key, flags);
 		}
 
 		/// <inheritdoc cref="IDatabaseAsync.SetRemoveAsync(RedisKey, RedisValue, CommandFlags)"/>
 		public async Task<TElement[]> MembersAsync(CommandFlags flags = CommandFlags.None)
 		{
-			RedisValue[] values = await _database.SetMembersAsync(Key, flags);
+			RedisValue[] values = await GetDatabase().SetMembersAsync(Key, flags);
 			return Array.ConvertAll(values, (Converter<RedisValue, TElement>)(x => RedisSerializer.Deserialize<TElement>(x)!));
 		}
 
 		/// <inheritdoc cref="IDatabaseAsync.SetPopAsync(RedisKey, CommandFlags)"/>
 		public async Task<TElement> PopAsync(CommandFlags flags = CommandFlags.None)
 		{
-			RedisValue value = await _database.SetPopAsync(Key, flags);
+			RedisValue value = await GetDatabase().SetPopAsync(Key, flags);
 			return value.IsNull ? default! : RedisSerializer.Deserialize<TElement>(value)!;
 		}
 
 		/// <inheritdoc cref="IDatabaseAsync.SetPopAsync(RedisKey, Int64, CommandFlags)"/>
 		public async Task<TElement[]> PopAsync(long count, CommandFlags flags = CommandFlags.None)
 		{
-			RedisValue[] values = await _database.SetPopAsync(Key, count, flags);
+			RedisValue[] values = await GetDatabase().SetPopAsync(Key, count, flags);
 			return Array.ConvertAll<RedisValue, TElement>(values, x => RedisSerializer.Deserialize<TElement>(x)!);
 		}
 
 		/// <inheritdoc cref="IDatabaseAsync.SetRemoveAsync(RedisKey, RedisValue, CommandFlags)"/>
 		public Task<bool> RemoveAsync(TElement item, CommandFlags flags = CommandFlags.None)
 		{
-			return _database.SetRemoveAsync(Key, RedisSerializer.Serialize(item), flags);
+			return GetDatabase().SetRemoveAsync(Key, RedisSerializer.Serialize(item), flags);
 		}
 
 		/// <inheritdoc cref="IDatabaseAsync.SetRemoveAsync(RedisKey, RedisValue[], CommandFlags)"/>
 		public Task<long> RemoveAsync(IEnumerable<TElement> items, CommandFlags flags = CommandFlags.None)
 		{
 			RedisValue[] values = items.Select(x => RedisSerializer.Serialize(x)).ToArray();
-			return _database.SetRemoveAsync(Key, values, flags);
+			return GetDatabase().SetRemoveAsync(Key, values, flags);
 		}
 	}
 
diff --git a/Engine/Source/Programs/Shared/EpicGames.Redis/RedisSortedSet.cs b/Engine/Source/Programs/Shared/EpicGames.Redis/RedisSortedSet.cs
--- a/Engine/Source/Programs/Shared/EpicGames.Redis/RedisSortedSet.cs
+++ b/Engine/Source/Programs/Shared/EpicGames.Redis/RedisSortedSet.cs
@@ -70,7 +70,8 @@ namespace EpicGames.Redis
 	/// <typeparam name="TElement">The type of element stored in the list</typeparam>
 	public readonly struct RedisSortedSet<TElement>
 	{
-		readonly IDatabaseAsync _database;
+		internal readonly RedisConnectionPool? ConnectionPool = null;
+		internal readonly IDatabaseAsync? Database = null;
 
 		/// <summary>
 		/// The key for the list
@@ -80,30 +81,55 @@ namespace EpicGames.Redis
 		/// <summary>
 		/// Constructor
 		/// </summary>
+		/// <param name="database">Explicitly set Redis database (overrides use of any connection pool)</param>
+		/// <param name="key">Redis key this type is using</param>
 		public RedisSortedSet(IDatabaseAsync database, RedisKey key)
 		{
-			_database = database;
+			Database = database;
 			Key = key;
 		}
+		
+		/// <summary>
+		/// Constructor
+		/// </summary>
+		/// <param name="connectionPool">Connection pool for Redis databases</param>
+		/// <param name="key">Redis key this type is using</param>
+		public RedisSortedSet(RedisConnectionPool connectionPool, RedisKey key)
+		{
+			ConnectionPool = connectionPool;
+			Key = key;
+		}
+		
+		/// <summary>
+		/// Get the Redis database in use
+		/// </summary>
+		/// <returns>A connection pool or explicitly set Redis database</returns>
+		/// <exception cref="InvalidOperationException">If neither are set</exception>
+		public IDatabaseAsync GetDatabase()
+		{
+			if (Database != null) return Database;
+			if (ConnectionPool != null) return ConnectionPool.GetDatabase();
+			throw new InvalidOperationException($"Neither {nameof(Database)} or {nameof(ConnectionPool)} has been set!");
+		}
 
 		/// <inheritdoc cref="IDatabaseAsync.SortedSetAddAsync(RedisKey, RedisValue, Double, CommandFlags)"/>
 		public Task<bool> AddAsync(TElement item, double score, When when = When.Always, CommandFlags flags = CommandFlags.None)
 		{
 			RedisValue value = RedisSerializer.Serialize(item);
-			return _database.SortedSetAddAsync(Key, value, score, when, flags);
+			return GetDatabase().SortedSetAddAsync(Key, value, score, when, flags);
 		}
 
 		/// <inheritdoc cref="IDatabaseAsync.SortedSetAddAsync(RedisKey, RedisValue, Double, When, CommandFlags)"/>
 		public Task<long> AddAsync(SortedSetEntry<TElement>[] values, When when = When.Always, CommandFlags flags = CommandFlags.None)
 		{
 			SortedSetEntry[] untyped = Array.ConvertAll(values, x => new SortedSetEntry(x.ElementValue, x.Score));
-			return _database.SortedSetAddAsync(Key, untyped, when, flags);
+			return GetDatabase().SortedSetAddAsync(Key, untyped, when, flags);
 		}
 
 		/// <inheritdoc cref="IDatabaseAsync.SortedSetLengthAsync(RedisKey, Double, Double, Exclude, CommandFlags)"/>
 		public Task<long> LengthAsync(double min = Double.NegativeInfinity, double max = Double.PositiveInfinity, Exclude exclude = Exclude.None, CommandFlags flags = CommandFlags.None)
 		{
-			return _database.SortedSetLengthAsync(Key, min, max, exclude, flags);
+			return GetDatabase().SortedSetLengthAsync(Key, min, max, exclude, flags);
 		}
 
 		/// <inheritdoc cref="IDatabaseAsync.SortedSetLengthByValueAsync(RedisKey, RedisValue, RedisValue, Exclude, CommandFlags)"/>
@@ -111,34 +137,34 @@ namespace EpicGames.Redis
 		{
 			RedisValue minValue = RedisSerializer.Serialize(min);
 			RedisValue maxValue = RedisSerializer.Serialize(max);
-			return _database.SortedSetLengthByValueAsync(Key, minValue, maxValue, exclude, flags);
+			return GetDatabase().SortedSetLengthByValueAsync(Key, minValue, maxValue, exclude, flags);
 		}
 
 		/// <inheritdoc cref="IDatabaseAsync.SortedSetRangeByRankAsync(RedisKey, Int64, Int64, Order, CommandFlags)"/>
 		public async Task<TElement[]> RangeByRankAsync(long start, long stop = -1, Order order = Order.Ascending, CommandFlags flags = CommandFlags.None)
 		{
-			RedisValue[] values = await _database.SortedSetRangeByRankAsync(Key, start, stop, order, flags);
+			RedisValue[] values = await GetDatabase().SortedSetRangeByRankAsync(Key, start, stop, order, flags);
 			return Array.ConvertAll(values, x => RedisSerializer.Deserialize<TElement>(x)!);
 		}
 
 		/// <inheritdoc cref="IDatabaseAsync.SortedSetRangeByRankWithScoresAsync(RedisKey, Int64, Int64, Order, CommandFlags)"/>
 		public async Task<SortedSetEntry<TElement>[]> RangeByRankWithScoresAsync(long start, long stop = -1, Order order = Order.Ascending, CommandFlags flags = CommandFlags.None)
 		{
-			SortedSetEntry[] values = await _database.SortedSetRangeByRankWithScoresAsync(Key, start, stop, order, flags);
+			SortedSetEntry[] values = await GetDatabase().SortedSetRangeByRankWithScoresAsync(Key, start, stop, order, flags);
 			return Array.ConvertAll(values, x => new SortedSetEntry<TElement>(x));
 		}
 
 		/// <inheritdoc cref="IDatabaseAsync.SortedSetRangeByScoreAsync(RedisKey, Double, Double, Exclude, Order, Int64, Int64, CommandFlags)"/>
 		public async Task<TElement[]> RangeByScoreAsync(double start = Double.NegativeInfinity, double stop = Double.PositiveInfinity, Exclude exclude = Exclude.None, Order order = Order.Ascending, long skip = 0L, long take = -1L, CommandFlags flags = CommandFlags.None)
 		{
-			RedisValue[] values = await _database.SortedSetRangeByScoreAsync(Key, start, stop, exclude, order, skip, take, flags);
+			RedisValue[] values = await GetDatabase().SortedSetRangeByScoreAsync(Key, start, stop, exclude, order, skip, take, flags);
 			return Array.ConvertAll(values, x => RedisSerializer.Deserialize<TElement>(x)!);
 		}
 
 		/// <inheritdoc cref="IDatabaseAsync.SortedSetRangeByScoreWithScoresAsync(RedisKey, Double, Double, Exclude, Order, Int64, Int64, CommandFlags)"/>
 		public async Task<SortedSetEntry<TElement>[]> RangeByScoreWithScoresAsync(double start = Double.NegativeInfinity, double stop = Double.PositiveInfinity, Exclude exclude = Exclude.None, Order order = Order.Ascending, long skip = 0L, long take = -1L, CommandFlags flags = CommandFlags.None)
 		{
-			SortedSetEntry[] values = await _database.SortedSetRangeByScoreWithScoresAsync(Key, start, stop, exclude, order, skip, take, flags);
+			SortedSetEntry[] values = await GetDatabase().SortedSetRangeByScoreWithScoresAsync(Key, start, stop, exclude, order, skip, take, flags);
 			return Array.ConvertAll(values, x => new SortedSetEntry<TElement>(x));
 		}
 
@@ -147,39 +173,39 @@ namespace EpicGames.Redis
 		{
 			RedisValue minValue = RedisSerializer.Serialize(min);
 			RedisValue maxValue = RedisSerializer.Serialize(max);
-			RedisValue[] values = await _database.SortedSetRangeByValueAsync(Key, minValue, maxValue, exclude, order, skip, take, flags);
+			RedisValue[] values = await GetDatabase().SortedSetRangeByValueAsync(Key, minValue, maxValue, exclude, order, skip, take, flags);
 			return Array.ConvertAll(values, x => RedisSerializer.Deserialize<TElement>(x)!);
 		}
 
 		/// <inheritdoc cref="IDatabaseAsync.SortedSetRankAsync(RedisKey, RedisValue, Order, CommandFlags)"/>
 		public Task<long?> RankAsync(TElement item, Order order = Order.Ascending, CommandFlags flags = CommandFlags.None)
 		{
-			return _database.SortedSetRankAsync(Key, RedisSerializer.Serialize(item), order, flags);
+			return GetDatabase().SortedSetRankAsync(Key, RedisSerializer.Serialize(item), order, flags);
 		}
 
 		/// <inheritdoc cref="IDatabaseAsync.SortedSetRemoveAsync(RedisKey, RedisValue, CommandFlags)"/>
 		public Task<bool> RemoveAsync(TElement value, CommandFlags flags = CommandFlags.None)
 		{
-			return _database.SortedSetRemoveAsync(Key, RedisSerializer.Serialize(value), flags);
+			return GetDatabase().SortedSetRemoveAsync(Key, RedisSerializer.Serialize(value), flags);
 		}
 
 		/// <inheritdoc cref="IDatabaseAsync.SortedSetRemoveAsync(RedisKey, RedisValue[], CommandFlags)"/>
 		public Task<long> RemoveAsync(IEnumerable<TElement> items, CommandFlags flags = CommandFlags.None)
 		{
 			RedisValue[] values = items.Select(x => RedisSerializer.Serialize(x)).ToArray();
-			return _database.SortedSetRemoveAsync(Key, values, flags);
+			return GetDatabase().SortedSetRemoveAsync(Key, values, flags);
 		}
 
 		/// <inheritdoc cref="IDatabaseAsync.SortedSetRemoveRangeByRankAsync(RedisKey, Int64, Int64, CommandFlags)"/>
 		public Task<long> RemoveRangeByRankAsync(long start, long stop, CommandFlags flags = CommandFlags.None)
 		{
-			return _database.SortedSetRemoveRangeByRankAsync(Key, start, stop, flags);
+			return GetDatabase().SortedSetRemoveRangeByRankAsync(Key, start, stop, flags);
 		}
 
 		/// <inheritdoc cref="IDatabaseAsync.SortedSetRemoveRangeByScoreAsync(RedisKey, Double, Double, Exclude, CommandFlags)"/>
 		public Task<long> RemoveRangeByScoreAsync(double start, double stop, Exclude exclude = Exclude.None, CommandFlags flags = CommandFlags.None)
 		{
-			return _database.SortedSetRemoveRangeByScoreAsync(Key, start, stop, exclude, flags);
+			return GetDatabase().SortedSetRemoveRangeByScoreAsync(Key, start, stop, exclude, flags);
 		}
 
 		/// <inheritdoc cref="IDatabaseAsync.SortedSetRemoveRangeByValueAsync(RedisKey, RedisValue, RedisValue, Exclude, CommandFlags)"/>
@@ -187,13 +213,13 @@ namespace EpicGames.Redis
 		{
 			RedisValue minValue = RedisSerializer.Serialize(min);
 			RedisValue maxValue = RedisSerializer.Serialize(max);
-			return await _database.SortedSetRemoveRangeByValueAsync(Key, minValue, maxValue, exclude, flags);
+			return await GetDatabase().SortedSetRemoveRangeByValueAsync(Key, minValue, maxValue, exclude, flags);
 		}
 
 		/// <inheritdoc cref="IDatabaseAsync.SortedSetAddAsync(RedisKey, RedisValue, Double, CommandFlags)"/>
 		public async IAsyncEnumerable<SortedSetEntry<TElement>> ScanAsync(RedisValue pattern = default, int pageSize = 250, long cursor = 0, int pageOffset = 0, CommandFlags flags = CommandFlags.None)
 		{
-			await foreach (SortedSetEntry entry in _database.SortedSetScanAsync(Key, pattern, pageSize, cursor, pageOffset, flags))
+			await foreach (SortedSetEntry entry in GetDatabase().SortedSetScanAsync(Key, pattern, pageSize, cursor, pageOffset, flags))
 			{
 				yield return new SortedSetEntry<TElement>(entry);
 			}
@@ -203,7 +229,7 @@ namespace EpicGames.Redis
 		public Task<double?> ScoreAsync(TElement member, CommandFlags flags = CommandFlags.None)
 		{
 			RedisValue value = RedisSerializer.Serialize(member);
-			return _database.SortedSetScoreAsync(Key, value, flags);
+			return GetDatabase().SortedSetScoreAsync(Key, value, flags);
 		}
 	}
 
diff --git a/Engine/Source/Programs/Shared/EpicGames.Redis/RedisString.cs b/Engine/Source/Programs/Shared/EpicGames.Redis/RedisString.cs
--- a/Engine/Source/Programs/Shared/EpicGames.Redis/RedisString.cs
+++ b/Engine/Source/Programs/Shared/EpicGames.Redis/RedisString.cs
@@ -12,7 +12,8 @@ namespace EpicGames.Redis
 	/// <typeparam name="TElement">The type of element stored in the list</typeparam>
 	public readonly struct RedisString<TElement>
 	{
-		internal readonly IDatabaseAsync Database;
+		internal readonly RedisConnectionPool? ConnectionPool = null;
+		internal readonly IDatabaseAsync? Database = null;
 
 		/// <summary>
 		/// The key for the list
@@ -29,17 +30,40 @@ namespace EpicGames.Redis
 			Database = database;
 			Key = key;
 		}
+		
+		/// <summary>
+		/// Constructor
+		/// </summary>
+		/// <param name="connectionPool"></param>
+		/// <param name="key"></param>
+		public RedisString(RedisConnectionPool connectionPool, RedisKey key)
+		{
+			ConnectionPool = connectionPool;
+			Key = key;
+		}
+
+		/// <summary>
+		/// Get the Redis database in use
+		/// </summary>
+		/// <returns>A connection pool or explicitly set Redis database</returns>
+		/// <exception cref="InvalidOperationException">If neither are set</exception>
+		public IDatabaseAsync GetDatabase()
+		{
+			if (Database != null) return Database;
+			if (ConnectionPool != null) return ConnectionPool.GetDatabase();
+			throw new InvalidOperationException($"Neither {nameof(Database)} or {nameof(ConnectionPool)} has been set!");
+		}
 
 		/// <inheritdoc cref="IDatabaseAsync.StringLengthAsync(RedisKey, CommandFlags)"/>
 		public Task<long> LengthAsync(CommandFlags flags = CommandFlags.None)
 		{
-			return Database.StringLengthAsync(Key, flags);
+			return GetDatabase().StringLengthAsync(Key, flags);
 		}
 
 		/// <inheritdoc cref="IDatabaseAsync.StringSetAsync(RedisKey, RedisValue, TimeSpan?, When, CommandFlags)"/>
 		public Task<bool> SetAsync(TElement value, TimeSpan? expiry = null, When when = When.Always, CommandFlags flags = CommandFlags.None)
 		{
-			return Database.StringSetAsync(Key, RedisSerializer.Serialize(value), expiry, when, flags);
+			return GetDatabase().StringSetAsync(Key, RedisSerializer.Serialize(value), expiry, when, flags);
 		}
 	}
 
@@ -51,13 +75,13 @@ namespace EpicGames.Redis
 		/// <inheritdoc cref="IDatabaseAsync.KeyDeleteAsync(RedisKey, CommandFlags)"/>
 		public static Task<bool> DeleteAsync<TElement>(this RedisString<TElement> str, CommandFlags flags = CommandFlags.None)
 		{
-			return str.Database.KeyDeleteAsync(str.Key, flags);
+			return str.GetDatabase().KeyDeleteAsync(str.Key, flags);
 		}
 
 		/// <inheritdoc cref="IDatabaseAsync.StringGetAsync(RedisKey, CommandFlags)"/>
 		public static async Task<TElement?> GetAsync<TElement>(this RedisString<TElement> str, CommandFlags flags = CommandFlags.None) where TElement : class
 		{
-			RedisValue value = await str.Database.StringGetAsync(str.Key, flags);
+			RedisValue value = await str.GetDatabase().StringGetAsync(str.Key, flags);
 			if (value.IsNullOrEmpty)
 			{
 				return null;
@@ -68,7 +92,7 @@ namespace EpicGames.Redis
 		/// <inheritdoc cref="IDatabaseAsync.StringGetAsync(RedisKey, CommandFlags)"/>
 		public static async Task<TElement?> GetValueAsync<TElement>(this RedisString<TElement> str, CommandFlags flags = CommandFlags.None) where TElement : struct
 		{
-			RedisValue value = await str.Database.StringGetAsync(str.Key, flags);
+			RedisValue value = await str.GetDatabase().StringGetAsync(str.Key, flags);
 			if (value.IsNullOrEmpty)
 			{
 				return default(TElement);
@@ -79,25 +103,25 @@ namespace EpicGames.Redis
 		/// <inheritdoc cref="IDatabaseAsync.StringDecrementAsync(RedisKey, Int64, CommandFlags)"/>
 		public static Task<long> DecrementAsync(this RedisString<long> str, long value = 1L, CommandFlags flags = CommandFlags.None)
 		{
-			return str.Database.StringDecrementAsync(str.Key, value, flags);
+			return str.GetDatabase().StringDecrementAsync(str.Key, value, flags);
 		}
 
 		/// <inheritdoc cref="IDatabaseAsync.StringDecrementAsync(RedisKey, Double, CommandFlags)"/>
 		public static Task<double> DecrementAsync(this RedisString<double> str, double value = 1.0, CommandFlags flags = CommandFlags.None)
 		{
-			return str.Database.StringDecrementAsync(str.Key, value, flags);
+			return str.GetDatabase().StringDecrementAsync(str.Key, value, flags);
 		}
 
 		/// <inheritdoc cref="IDatabaseAsync.StringDecrementAsync(RedisKey, Double, CommandFlags)"/>
 		public static Task<long> IncrementAsync(this RedisString<long> str, long value = 1L, CommandFlags flags = CommandFlags.None)
 		{
-			return str.Database.StringIncrementAsync(str.Key, value, flags);
+			return str.GetDatabase().StringIncrementAsync(str.Key, value, flags);
 		}
 
 		/// <inheritdoc cref="IDatabaseAsync.StringDecrementAsync(RedisKey, Double, CommandFlags)"/>
 		public static Task<double> IncrementAsync(this RedisString<double> str, double value = 1.0, CommandFlags flags = CommandFlags.None)
 		{
-			return str.Database.StringIncrementAsync(str.Key, value, flags);
+			return str.GetDatabase().StringIncrementAsync(str.Key, value, flags);
 		}
 
 		/// <summary>
