Fixed output issues with the Json exporter in C# UHT.

#rnx

diff --git a/Engine/Source/Programs/Shared/EpicGames.UHT/Types/Properties/UhtBoolProperty.cs b/Engine/Source/Programs/Shared/EpicGames.UHT/Types/Properties/UhtBoolProperty.cs
--- a/Engine/Source/Programs/Shared/EpicGames.UHT/Types/Properties/UhtBoolProperty.cs
+++ b/Engine/Source/Programs/Shared/EpicGames.UHT/Types/Properties/UhtBoolProperty.cs
@@ -2,6 +2,7 @@
 
 using System.Diagnostics.CodeAnalysis;
 using System.Text;
+using System.Text.Json.Serialization;
 using EpicGames.Core;
 using EpicGames.UHT.Tables;
 using EpicGames.UHT.Tokenizer;
@@ -106,6 +107,7 @@ namespace EpicGames.UHT.Types
 		/// <summary>
 		/// Type of the boolean
 		/// </summary>
+		[JsonConverter(typeof(JsonStringEnumConverter))]
 		public UhtBoolType BoolType { get; }
 
 		/// <summary>
diff --git a/Engine/Source/Programs/Shared/EpicGames.UHT/Types/Properties/UhtEnumProperty.cs b/Engine/Source/Programs/Shared/EpicGames.UHT/Types/Properties/UhtEnumProperty.cs
--- a/Engine/Source/Programs/Shared/EpicGames.UHT/Types/Properties/UhtEnumProperty.cs
+++ b/Engine/Source/Programs/Shared/EpicGames.UHT/Types/Properties/UhtEnumProperty.cs
@@ -40,16 +40,19 @@ namespace EpicGames.UHT.Types
 		/// <summary>
 		/// Underlying property set when the enum has an underlying interger type
 		/// </summary>
+		[JsonIgnore]
 		public UhtProperty? UnderlyingProperty { get; set; }
 
 		/// <summary>
 		/// Underlying type which defaults to int32 if the referenced enum doesn't have an underlying type
 		/// </summary>
+		[JsonConverter(typeof(JsonStringEnumConverter))]
 		public UhtEnumUnderlyingType UnderlyingType => Enum.UnderlyingType != UhtEnumUnderlyingType.Unspecified ? Enum.UnderlyingType : UhtEnumUnderlyingType.int32;
 
 		/// <summary>
 		/// Underlying type size.  Defaults to unsized if the referenced enum doesn't have an underlying type
 		/// </summary>
+		[JsonConverter(typeof(JsonStringEnumConverter))]
 		public UhtPropertyIntType UnderlyingTypeSize => Enum.UnderlyingType != UhtEnumUnderlyingType.Unspecified ? UhtPropertyIntType.Sized : UhtPropertyIntType.Unsized;
 
 		/// <summary>
diff --git a/Engine/Source/Programs/Shared/EpicGames.UHT/Types/Properties/UhtNumericProperty.cs b/Engine/Source/Programs/Shared/EpicGames.UHT/Types/Properties/UhtNumericProperty.cs
--- a/Engine/Source/Programs/Shared/EpicGames.UHT/Types/Properties/UhtNumericProperty.cs
+++ b/Engine/Source/Programs/Shared/EpicGames.UHT/Types/Properties/UhtNumericProperty.cs
@@ -1,6 +1,7 @@
 // Copyright Epic Games, Inc. All Rights Reserved.
 
 using System.Text;
+using System.Text.Json.Serialization;
 
 namespace EpicGames.UHT.Types
 {
@@ -19,6 +20,7 @@ namespace EpicGames.UHT.Types
 		/// <summary>
 		/// Describes the integer as either being sized on unsized
 		/// </summary>
+		[JsonConverter(typeof(JsonStringEnumConverter))]
 		public UhtPropertyIntType IntType { get; }
 
 		/// <summary>
diff --git a/Engine/Source/Programs/Shared/EpicGames.UHT/Types/Properties/UhtStructProperty.cs b/Engine/Source/Programs/Shared/EpicGames.UHT/Types/Properties/UhtStructProperty.cs
--- a/Engine/Source/Programs/Shared/EpicGames.UHT/Types/Properties/UhtStructProperty.cs
+++ b/Engine/Source/Programs/Shared/EpicGames.UHT/Types/Properties/UhtStructProperty.cs
@@ -209,11 +209,11 @@ namespace EpicGames.UHT.Types
 		///<inheritdoc/>
 		public override bool ContainsEditorOnlyProperties()
 		{
-			foreach (UhtType Child in ScriptStruct.Children)
+			foreach (UhtType child in ScriptStruct.Children)
 			{
-				if (Child is UhtProperty Property)
+				if (child is UhtProperty property)
 				{
-					if (Property.PropertyFlags.HasAnyFlags(EPropertyFlags.EditorOnly) || Property.ContainsEditorOnlyProperties())
+					if (property.PropertyFlags.HasAnyFlags(EPropertyFlags.EditorOnly) || property.ContainsEditorOnlyProperties())
 					{
 						return true;
 					}
diff --git a/Engine/Source/Programs/Shared/EpicGames.UHT/Types/UhtClass.cs b/Engine/Source/Programs/Shared/EpicGames.UHT/Types/UhtClass.cs
--- a/Engine/Source/Programs/Shared/EpicGames.UHT/Types/UhtClass.cs
+++ b/Engine/Source/Programs/Shared/EpicGames.UHT/Types/UhtClass.cs
@@ -287,6 +287,7 @@ namespace EpicGames.UHT.Types
 		/// <summary>
 		/// If needed, the #if block define for the serializer
 		/// </summary>
+		[JsonIgnore]
 		public string EnclosingDefine { get; set; } = String.Empty;
 
 		/// <summary>
@@ -316,6 +317,7 @@ namespace EpicGames.UHT.Types
 		/// <summary>
 		/// Type of the class
 		/// </summary>
+		[JsonConverter(typeof(JsonStringEnumConverter))]
 		public UhtClassType ClassType { get; set; } = UhtClassType.Class;
 
 		/// <summary>
@@ -338,6 +340,7 @@ namespace EpicGames.UHT.Types
 		/// <summary>
 		/// If this, this class is a UINTERFACE and NativeInterface is the associated native interface
 		/// </summary>
+		[JsonIgnore]
 		public UhtClass? NativeInterface { get; set; } = null;
 
 		/// <summary>
diff --git a/Engine/Source/Programs/Shared/EpicGames.UHT/Types/UhtFunction.cs b/Engine/Source/Programs/Shared/EpicGames.UHT/Types/UhtFunction.cs
--- a/Engine/Source/Programs/Shared/EpicGames.UHT/Types/UhtFunction.cs
+++ b/Engine/Source/Programs/Shared/EpicGames.UHT/Types/UhtFunction.cs
@@ -226,6 +226,7 @@ namespace EpicGames.UHT.Types
 		/// <summary>
 		/// The type of function
 		/// </summary>
+		[JsonConverter(typeof(JsonStringEnumConverter))]
 		public UhtFunctionType FunctionType { get; set; } = UhtFunctionType.Function;
 
 		/// <summary>
diff --git a/Engine/Source/Programs/Shared/EpicGames.UHT/Types/UhtHeaderFile.cs b/Engine/Source/Programs/Shared/EpicGames.UHT/Types/UhtHeaderFile.cs
--- a/Engine/Source/Programs/Shared/EpicGames.UHT/Types/UhtHeaderFile.cs
+++ b/Engine/Source/Programs/Shared/EpicGames.UHT/Types/UhtHeaderFile.cs
@@ -122,16 +122,19 @@ namespace EpicGames.UHT.Types
 		/// <summary>
 		/// File name without the extension
 		/// </summary>
+		[JsonIgnore]
 		public string FileNameWithoutExtension { get; }
 
 		/// <summary>
 		/// Required name for the generated.h file name.  Used to validate parsed code
 		/// </summary>
+		[JsonIgnore]
 		public string GeneratedHeaderFileName { get; }
 
 		/// <summary>
 		/// True if this header is NoExportTypes.h
 		/// </summary>
+		[JsonIgnore]
 		public bool IsNoExportTypes { get; }
 
 		/// <summary>
@@ -147,11 +150,13 @@ namespace EpicGames.UHT.Types
 		/// <summary>
 		/// Location where the header file was found
 		/// </summary>
+		[JsonConverter(typeof(JsonStringEnumConverter))]
 		public UhtHeaderFileType HeaderFileType { get; set; } = UhtHeaderFileType.Private;
 
 		/// <summary>
 		/// Unique index of the header file
 		/// </summary>
+		[JsonIgnore]
 		public int HeaderFileTypeIndex { get; }
 
 		/// <summary>
diff --git a/Engine/Source/Programs/Shared/EpicGames.UHT/Types/UhtObject.cs b/Engine/Source/Programs/Shared/EpicGames.UHT/Types/UhtObject.cs
--- a/Engine/Source/Programs/Shared/EpicGames.UHT/Types/UhtObject.cs
+++ b/Engine/Source/Programs/Shared/EpicGames.UHT/Types/UhtObject.cs
@@ -1,5 +1,6 @@
 // Copyright Epic Games, Inc. All Rights Reserved.
 
+using System.Text.Json.Serialization;
 using EpicGames.Core;
 using EpicGames.UHT.Tables;
 using EpicGames.UHT.Utils;
@@ -16,17 +17,20 @@ namespace EpicGames.UHT.Types
 		/// <summary>
 		/// Internal object flags.
 		/// </summary>
+		[JsonConverter(typeof(JsonStringEnumConverter))]
 		public EInternalObjectFlags InternalObjectFlags { get; set; } = EInternalObjectFlags.None;
 
 		/// <summary>
 		/// Unique index of the object
 		/// </summary>
+		[JsonIgnore]
 		public int ObjectTypeIndex { get; }
 
 		/// <summary>
 		/// The alternate object is used by the interface system where the native interface will
 		/// update this setting to point to the UInterface derived companion object.
 		/// </summary>
+		[JsonIgnore]
 		public UhtObject? AlternateObject { get; set; } = null;
 
 		/// <inheritdoc/>
diff --git a/Engine/Source/Programs/Shared/EpicGames.UHT/Types/UhtPackage.cs b/Engine/Source/Programs/Shared/EpicGames.UHT/Types/UhtPackage.cs
--- a/Engine/Source/Programs/Shared/EpicGames.UHT/Types/UhtPackage.cs
+++ b/Engine/Source/Programs/Shared/EpicGames.UHT/Types/UhtPackage.cs
@@ -20,11 +20,13 @@ namespace EpicGames.UHT.Types
 		/// <summary>
 		/// Unique index of the package
 		/// </summary>
+		[JsonIgnore]
 		public int PackageTypeIndex { get; }
 
 		/// <summary>
 		/// Engine package flags
 		/// </summary>
+		[JsonConverter(typeof(JsonStringEnumConverter))]
 		public EPackageFlags PackageFlags { get; set; } = EPackageFlags.None;
 
 		/// <summary>
diff --git a/Engine/Source/Programs/Shared/EpicGames.UHT/Types/UhtProperty.cs b/Engine/Source/Programs/Shared/EpicGames.UHT/Types/UhtProperty.cs
--- a/Engine/Source/Programs/Shared/EpicGames.UHT/Types/UhtProperty.cs
+++ b/Engine/Source/Programs/Shared/EpicGames.UHT/Types/UhtProperty.cs
@@ -884,6 +884,7 @@ namespace EpicGames.UHT.Types
 		/// <summary>
 		/// Property category
 		/// </summary>
+		[JsonConverter(typeof(JsonStringEnumConverter))]
 		public UhtPropertyCategory PropertyCategory { get; set; } = UhtPropertyCategory.Member;
 
 		/// <summary>
@@ -902,26 +903,32 @@ namespace EpicGames.UHT.Types
 		/// <summary>
 		/// Engine flags that are disallowed on this property
 		/// </summary>
+		[JsonIgnore]
+		[JsonConverter(typeof(JsonStringEnumConverter))]
 		public EPropertyFlags DisallowPropertyFlags { get; set; } = EPropertyFlags.None;
 
 		/// <summary>
 		/// UHT specified property flags
 		/// </summary>
+		[JsonConverter(typeof(JsonStringEnumConverter))]
 		public UhtPropertyExportFlags PropertyExportFlags { get; set; } = UhtPropertyExportFlags.Public;
 
 		/// <summary>
 		/// Reference type of the property
 		/// </summary>
+		[JsonConverter(typeof(JsonStringEnumConverter))]
 		public UhtPropertyRefQualifier RefQualifier { get; set; } = UhtPropertyRefQualifier.None;
 
 		/// <summary>
 		/// Pointer type of the property
 		/// </summary>
+		[JsonConverter(typeof(JsonStringEnumConverter))]
 		public UhtPointerType PointerType { get; set; } = UhtPointerType.None;
 
 		/// <summary>
 		/// Allocator to be used with containers
 		/// </summary>
+		[JsonConverter(typeof(JsonStringEnumConverter))]
 		public UhtPropertyAllocator Allocator { get; set; } = UhtPropertyAllocator.Default;
 
 		/// <summary>
@@ -965,6 +972,7 @@ namespace EpicGames.UHT.Types
 		public override bool Deprecated => PropertyFlags.HasAnyFlags(EPropertyFlags.Deprecated);
 
 		///<inheritdoc/>
+		[JsonIgnore]
 		protected override UhtSpecifierValidatorTable? SpecifierValidatorTable
 		{
 			get
diff --git a/Engine/Source/Programs/Shared/EpicGames.UHT/Types/UhtScriptStruct.cs b/Engine/Source/Programs/Shared/EpicGames.UHT/Types/UhtScriptStruct.cs
--- a/Engine/Source/Programs/Shared/EpicGames.UHT/Types/UhtScriptStruct.cs
+++ b/Engine/Source/Programs/Shared/EpicGames.UHT/Types/UhtScriptStruct.cs
@@ -110,6 +110,7 @@ namespace EpicGames.UHT.Types
 		/// <summary>
 		/// Property associated with the RigVM parameter
 		/// </summary>
+		[JsonConverter(typeof(UhtTypeSourceNameJsonConverter<UhtProperty>))]
 		public UhtProperty? Property { get; }
 
 		/// <summary>
@@ -125,16 +126,19 @@ namespace EpicGames.UHT.Types
 		/// <summary>
 		/// Cast name
 		/// </summary>
+		[JsonIgnore]
 		public string? CastName { get; }
 
 		/// <summary>
 		/// Cast type
 		/// </summary>
+		[JsonIgnore]
 		public string? CastType { get; }
 
 		/// <summary>
 		/// Flags associated with the parameter
 		/// </summary>
+		[JsonConverter(typeof(JsonStringEnumConverter))]
 		public UhtRigVMParameterFlags ParameterFlags { get; set; } = UhtRigVMParameterFlags.None;
 
 		/// <summary>
diff --git a/Engine/Source/Programs/Shared/EpicGames.UHT/Types/UhtStruct.cs b/Engine/Source/Programs/Shared/EpicGames.UHT/Types/UhtStruct.cs
--- a/Engine/Source/Programs/Shared/EpicGames.UHT/Types/UhtStruct.cs
+++ b/Engine/Source/Programs/Shared/EpicGames.UHT/Types/UhtStruct.cs
@@ -23,6 +23,7 @@ namespace EpicGames.UHT.Types
 		/// <summary>
 		/// Generated code version of the type.  Set via specifiers
 		/// </summary>
+		[JsonConverter(typeof(JsonStringEnumConverter))]
 		public EGeneratedCodeVersion GeneratedCodeVersion { get; set; } = EGeneratedCodeVersion.None;
 
 		/// <summary>
@@ -81,13 +82,13 @@ namespace EpicGames.UHT.Types
 		/// <summary>
 		/// Super type
 		/// </summary>
+		[JsonConverter(typeof(UhtNullableTypeSourceNameJsonConverter<UhtStruct>))]
 		public UhtStruct? Super { get; set; } = null;
 
 		/// <summary>
 		/// Base types
 		/// </summary>
-		//[JsonConverter(typeof(UhtNullableTypeListJsonConverter<UhtStruct>))]
-		[JsonIgnore]
+		[JsonConverter(typeof(UhtTypeReadOnlyListSourceNameJsonConverter<UhtStruct>))]
 		public IReadOnlyList<UhtStruct> Bases => _bases ?? s_emptyBases;
 
 		/// <summary>
diff --git a/Engine/Source/Programs/Shared/EpicGames.UHT/Types/UhtType.cs b/Engine/Source/Programs/Shared/EpicGames.UHT/Types/UhtType.cs
--- a/Engine/Source/Programs/Shared/EpicGames.UHT/Types/UhtType.cs
+++ b/Engine/Source/Programs/Shared/EpicGames.UHT/Types/UhtType.cs
@@ -743,6 +743,7 @@ namespace EpicGames.UHT.Types
 		/// <summary>
 		/// If true, then this type is visible to the symbol table.
 		/// </summary>
+		[JsonIgnore]
 		public bool VisibleType { get; set; } = true;
 
 		/// <summary>
@@ -793,12 +794,13 @@ namespace EpicGames.UHT.Types
 		/// <summary>
 		/// Return the specifier validation table to be used to validate the meta data on this type
 		/// </summary>
+		[JsonIgnore]
 		protected virtual UhtSpecifierValidatorTable? SpecifierValidatorTable => null;
 
 		/// <summary>
 		/// Meta data associated with the type
 		/// </summary>
-		[JsonIgnore] //ETSTODO - disable for now
+		[JsonConverter(typeof(UhtMetaDataConverter))]
 		public UhtMetaData MetaData { get; set; }
 
 		/// <summary>
diff --git a/Engine/Source/Programs/Shared/EpicGames.UHT/Utils/UhtJsonConverters.cs b/Engine/Source/Programs/Shared/EpicGames.UHT/Utils/UhtJsonConverters.cs
--- a/Engine/Source/Programs/Shared/EpicGames.UHT/Utils/UhtJsonConverters.cs
+++ b/Engine/Source/Programs/Shared/EpicGames.UHT/Utils/UhtJsonConverters.cs
@@ -9,6 +9,44 @@ using EpicGames.UHT.Types;
 namespace EpicGames.UHT.Utils
 {
 
+	/// <summary>
+	/// JSON converter to output the meta data
+	/// </summary>
+	public class UhtMetaDataConverter : JsonConverter<UhtMetaData> 
+	{
+		/// <summary>
+		/// Read the JSON value
+		/// </summary>
+		/// <param name="reader">Source reader</param>
+		/// <param name="typeToConvert">Type to convert</param>
+		/// <param name="options">Serialization options</param>
+		/// <returns>Read value</returns>
+		/// <exception cref="NotImplementedException"></exception>
+		public override UhtMetaData Read(ref Utf8JsonReader reader, Type typeToConvert, JsonSerializerOptions options)
+		{
+			throw new NotImplementedException();
+		}
+
+		/// <summary>
+		/// Write the JSON value
+		/// </summary>
+		/// <param name="writer">Destination writer</param>
+		/// <param name="type">Value being written</param>
+		/// <param name="options">Serialization options</param>
+		public override void Write(Utf8JsonWriter writer, UhtMetaData type, JsonSerializerOptions options)
+		{
+			writer.WriteStartArray();
+			foreach (KeyValuePair<string, string> kvp in type.GetSorted())
+			{
+				writer.WriteStartObject();
+				writer.WriteString("Key", kvp.Key);
+				writer.WriteString("Value", kvp.Value);
+				writer.WriteEndObject();
+			}
+			writer.WriteEndArray();
+		}
+	}
+
 	/// <summary>
 	/// JSON converter to output the source name of a type
 	/// </summary>
@@ -74,6 +112,42 @@ namespace EpicGames.UHT.Utils
 		}
 	}
 
+	/// <summary>
+	/// Read/Write type list source name for an optional type value
+	/// </summary>
+	/// <typeparam name="T"></typeparam>
+	public class UhtTypeReadOnlyListSourceNameJsonConverter<T> : JsonConverter<IReadOnlyList<T>> where T : UhtType
+	{
+		/// <summary>
+		/// Read the JSON value
+		/// </summary>
+		/// <param name="reader">Source reader</param>
+		/// <param name="typeToConvert">Type to convert</param>
+		/// <param name="options">Serialization options</param>
+		/// <returns>Read value</returns>
+		/// <exception cref="NotImplementedException"></exception>
+		public override IReadOnlyList<T> Read(ref Utf8JsonReader reader, Type typeToConvert, JsonSerializerOptions options)
+		{
+			throw new NotImplementedException();
+		}
+
+		/// <summary>
+		/// Write the JSON value
+		/// </summary>
+		/// <param name="writer">Destination writer</param>
+		/// <param name="collection">Value being written</param>
+		/// <param name="options">Serialization options</param>
+		public override void Write(Utf8JsonWriter writer, IReadOnlyList<T> collection, JsonSerializerOptions options)
+		{
+			writer.WriteStartArray();
+			foreach (UhtType type in collection)
+			{
+				writer.WriteStringValue(type.SourceName);
+			}
+			writer.WriteEndArray();
+		}
+	}
+
 	/// <summary>
 	/// Serialize a list of types
 	/// </summary>
