Horde: Perforce replication improvements.

* Sync state is periodically flushed, and stored with an optimized list of paths to conform a workspace back to the same state.
* Writes are automatically queued to run asynchrnously once a bundle of data is ready to be stored, rather than clients having to trigger a write.
* BundleStore supports different contexts for different reads, to prevent hashed blobs being merged together when reading from different sources.

diff --git a/Engine/Source/Programs/Horde/Horde.Build/Commands/Bundles/CreateCommand.cs b/Engine/Source/Programs/Horde/Horde.Build/Commands/Bundles/CreateCommand.cs
--- a/Engine/Source/Programs/Horde/Horde.Build/Commands/Bundles/CreateCommand.cs
+++ b/Engine/Source/Programs/Horde/Horde.Build/Commands/Bundles/CreateCommand.cs
@@ -41,10 +41,13 @@ namespace Horde.Build.Commands.Bundles
 			using ServiceProvider serviceProvider = Startup.CreateServiceProvider(_configuration, _loggerProvider);
 
 			ITreeStore store = serviceProvider.GetRequiredService<ITreeStore<CommitService>>();
+			ITreeWriter writer = store.CreateTreeWriter(RefName);
 
 			DirectoryNode node = new DirectoryNode();
-			await node.CopyFromDirectoryAsync(InputDir.ToDirectoryInfo(), new ChunkingOptions(), logger, CancellationToken.None);
-			await store.WriteTreeAsync(RefName, node);
+			await node.CopyFromDirectoryAsync(InputDir.ToDirectoryInfo(), new ChunkingOptions(), writer, logger, CancellationToken.None);
+
+			TreeNodeRef<DirectoryNode> root = new TreeNodeRef<DirectoryNode>(null!, node);
+			await root.CollapseAsync(writer, CancellationToken.None);
 
 			return 0;
 		}
diff --git a/Engine/Source/Programs/Horde/Horde.Build/Commands/Bundles/ExtractCommand.cs b/Engine/Source/Programs/Horde/Horde.Build/Commands/Bundles/ExtractCommand.cs
--- a/Engine/Source/Programs/Horde/Horde.Build/Commands/Bundles/ExtractCommand.cs
+++ b/Engine/Source/Programs/Horde/Horde.Build/Commands/Bundles/ExtractCommand.cs
@@ -35,10 +35,11 @@ namespace Horde.Build.Commands.Bundles
 		{
 			using ServiceProvider serviceProvider = Startup.CreateServiceProvider(_configuration, _loggerProvider);
 
-			ITreeStore store = serviceProvider.GetRequiredService<ITreeStore<CommitService>>();
+			ITreeStore store = serviceProvider.GetRequiredService<ITreeStore<ReplicationService>>();
 
-			DirectoryNode node = await store.ReadTreeAsync<DirectoryNode>(RefName);
-			await node.CopyToDirectoryAsync(OutputDir.ToDirectoryInfo(), logger, CancellationToken.None);
+			ReplicationNode node = await store.ReadTreeAsync<ReplicationNode>(RefName);
+			DirectoryNode root = await node.Contents.ExpandAsync();
+			await root.CopyToDirectoryAsync(OutputDir.ToDirectoryInfo(), logger, CancellationToken.None);
 
 			return 0;
 		}
diff --git a/Engine/Source/Programs/Horde/Horde.Build/Perforce/ReplicationService.cs b/Engine/Source/Programs/Horde/Horde.Build/Perforce/ReplicationService.cs
--- a/Engine/Source/Programs/Horde/Horde.Build/Perforce/ReplicationService.cs
+++ b/Engine/Source/Programs/Horde/Horde.Build/Perforce/ReplicationService.cs
@@ -88,31 +88,30 @@ namespace Horde.Build.Perforce
 			Contents = new TreeNodeRef<DirectoryNode>(this, contents);
 		}
 
-		private ReplicationNode(IEnumerable<Utf8String> paths, ITreeBlob contents)
+		private ReplicationNode(IEnumerable<Utf8String> paths, ITreeBlobRef contents)
 		{
 			Paths = new List<Utf8String>(paths);
 			Contents = new TreeNodeRef<DirectoryNode>(this, contents);
 		}
 
-		internal async ValueTask<ITreeBlob> SerializeAsync(ITreeBlobWriter writer, CancellationToken cancellationToken)
+		/// <inheritdoc/>
+		public override NewTreeBlob Serialize()
 		{
-			List<ITreeBlob> references = new List<ITreeBlob>();
-			references.Add(await Contents.CollapseAsync(writer, cancellationToken));
+			List<TreeNodeRef> references = new List<TreeNodeRef>();
+			references.Add(Contents);
 
 			ByteArrayBuilder builder = new ByteArrayBuilder();
 			builder.WriteVariableLengthArray(Paths, x => builder.WriteUtf8String(x));
-			return await writer.WriteBlobAsync(builder.AsSequence(), references, cancellationToken);
+
+			return new NewTreeBlob(builder.AsSequence(), references);
 		}
 
-		internal static async ValueTask<ReplicationNode> DeserializeAsync(ITreeBlob blob, CancellationToken cancellationToken)
+		internal static ReplicationNode Deserialize(ITreeBlob blob)
 		{
-			IReadOnlyList<ITreeBlob> children = await blob.GetReferencesAsync(cancellationToken);
-			ReadOnlySequence<byte> data = await blob.GetDataAsync(cancellationToken);
-
-			MemoryReader reader = new MemoryReader(data.AsSingleSegment());
+			MemoryReader reader = new MemoryReader(blob.Data.AsSingleSegment());
 			Utf8String[] paths = reader.ReadVariableLengthArray(() => reader.ReadUtf8String());
 
-			return new ReplicationNode(paths, children[0]);
+			return new ReplicationNode(paths, blob.Refs[0]);
 		}
 
 		/// <inheritdoc/>
@@ -127,18 +126,13 @@ namespace Horde.Build.Perforce
 	class ReplicationNodeSerializer : TreeNodeSerializer<ReplicationNode>
 	{
 		/// <inheritdoc/>
-		public override ValueTask<ReplicationNode> DeserializeAsync(ITreeBlob node, CancellationToken cancellationToken)
-			=> ReplicationNode.DeserializeAsync(node, cancellationToken);
-
-		/// <inheritdoc/>
-		public override ValueTask<ITreeBlob> SerializeAsync(ITreeBlobWriter writer, ReplicationNode node, CancellationToken cancellationToken)
-			=> node.SerializeAsync(writer, cancellationToken);
+		public override ReplicationNode Deserialize(ITreeBlob node) => ReplicationNode.Deserialize(node);
 	}
 
 	/// <summary>
 	/// Service which replicates content from Perforce
 	/// </summary>
-	class ReplicationService : IHostedService, IAsyncDisposable
+	sealed class ReplicationService : IHostedService, IDisposable
 	{
 		/// <summary>
 		/// Metadata about a stream that needs to have commits mirrored
@@ -204,9 +198,9 @@ namespace Horde.Build.Perforce
 		}
 
 		/// <inheritdoc/>
-		public async ValueTask DisposeAsync()
+		public void Dispose()
 		{
-			await _commitListener.DisposeAsync();
+			_commitListener.DisposeAsync().AsTask().Wait();
 		}
 
 		/// <inheritdoc/>
@@ -465,14 +459,14 @@ namespace Horde.Build.Perforce
 		/// <returns></returns>
 		static RefName GetRefName(StreamId streamId, int change, string? filter, bool revisionsOnly)
 		{
-			StringBuilder builder = new StringBuilder($"v3/{streamId}/{change}");
+			StringBuilder builder = new StringBuilder($"v4/{streamId}/{change}");
 			if (filter != null)
 			{
-				builder.Append($"_filter_{IoHash.Compute(Encoding.UTF8.GetBytes(filter))}");
+				builder.Append($"-flt_{IoHash.Compute(Encoding.UTF8.GetBytes(filter)).ToString().Substring(0, 6)}");
 			}
 			if (revisionsOnly)
 			{
-				builder.Append("_revs");
+				builder.Append("-revs");
 			}
 			return new RefName(builder.ToString());
 		}
@@ -574,7 +568,8 @@ namespace Horde.Build.Perforce
 			string queryPath = $"//{clientInfo.Client.Name}/{filterOrDefault}";
 
 			RefName refName = GetRefName(stream.Id, change, filter, revisionsOnly);
-			RefName incRefName = new RefName($"{refName}_inc");
+			RefName incRefName = new RefName($"{refName}-inc");
+			ITreeWriter writer = _treeStore.CreateTreeWriter(incRefName);
 
 			// Get the current sync state for this change
 			ReplicationNode? syncNode = await _treeStore.TryReadTreeAsync<ReplicationNode>(incRefName, cancellationToken);
@@ -590,7 +585,6 @@ namespace Horde.Build.Perforce
 
 			if (revisionsOnly)
 			{
-				int count = 0;
 				await foreach (PerforceResponse<SyncRecord> record in perforce.StreamCommandAsync<SyncRecord>("sync", new[] { "-k", $"{queryPath}@{change}" }, null, default, cancellationToken))
 				{
 					PerforceError? error = record.Error;
@@ -610,14 +604,8 @@ namespace Horde.Build.Perforce
 
 					using ReadOnlyMemoryStream dataStream = new ReadOnlyMemoryStream(data);
 
-					FileEntry entry = await root.AddFileByPathAsync(path, FileEntryFlags.PerforceDepotPathAndRevision, cancellationToken);
-					await entry.AppendAsync(dataStream, Options.Chunking, cancellationToken);
-
-					if (++count > 1000)
-					{
-						await _treeStore.WriteTreeAsync(refName, syncNode, false, cancellationToken);
-						count = 0;
-					}
+					FileEntry entry = await root.AddFileByPathAsync(path, FileEntryFlags.PerforceDepotPathAndRevision, writer, cancellationToken);
+					await entry.AppendAsync(dataStream, Options.Chunking, writer, cancellationToken);
 				}
 			}
 			else
@@ -661,6 +649,9 @@ namespace Horde.Build.Perforce
 				long totalSize = files.Sum(x => x.Size);
 				_logger.LogInformation("Total sync size: {Size:n1}mb", totalSize / (1024.0 * 1024.0));
 
+				// Previously scheduled flush task
+				Task flushTask = Task.CompletedTask;
+
 				// Sync incrementally
 				long syncedSize = 0;
 				while (files.Count > 0)
@@ -668,7 +659,7 @@ namespace Horde.Build.Perforce
 					// Save the incremental state
 					if (syncedSize > 0)
 					{
-						await _treeStore.WriteTreeAsync(incRefName, syncNode, cancellationToken: cancellationToken);
+						await writer.FlushAsync(syncNode, cancellationToken);
 						deleteIncRef = true;
 					}
 
@@ -677,10 +668,16 @@ namespace Horde.Build.Perforce
 					(int idx, Utf8String path, long size) = GetSyncBatch(files, MaxBatchSize, clientInfo.ServerInfo.Utf8PathComparer, _logger);
 					syncedSize += size;
 
-					const long TrimInterval = 1024 * 1024 * 256;
+					// If we're syncing the entire view, use the given filter instead. It will be tighter than the base path from all files.
+					if (filter != null && path.Equals("..."))
+					{
+						path = filter;
+					}
+
+					const long FlushInterval = 1024 * 1024 * 256;
 
 					long dataSize = 0;
-					long trimDataSize = TrimInterval;
+					long flushDataSize = FlushInterval;
 
 					string syncPath = $"//{clientInfo.Client.Name}/{path}@{change}";
 					double syncPct = (totalSize == 0) ? 100.0 : (syncedSize * 100.0) / totalSize;
@@ -688,9 +685,8 @@ namespace Horde.Build.Perforce
 
 					Stopwatch syncTimer = Stopwatch.StartNew();
 					Stopwatch processTimer = new Stopwatch();
-					Stopwatch trimTimer = new Stopwatch();
+					Stopwatch flushTimer = new Stopwatch();
 					Stopwatch gcTimer = new Stopwatch();
-					Task trimTask = Task.CompletedTask;
 
 					Dictionary<int, FileEntry> handles = new Dictionary<int, FileEntry>();
 					await foreach (PerforceResponse response in perforce.StreamCommandAsync("sync", Array.Empty<string>(), new string[] { syncPath }, null, typeof(SyncRecord), true, default))
@@ -710,12 +706,12 @@ namespace Horde.Build.Perforce
 							if (io.Command == PerforceIoCommand.Open)
 							{
 								Utf8String file = GetClientRelativePath(io.Payload, clientInfo.Client.Root);
-								handles[io.File] = await root.AddFileByPathAsync(file, FileEntryFlags.None, cancellationToken);
+								handles[io.File] = await root.AddFileByPathAsync(file, FileEntryFlags.None, writer, cancellationToken);
 							}
 							else if (io.Command == PerforceIoCommand.Write)
 							{
 								FileEntry entry = handles[io.File];
-								await entry.AppendAsync(io.Payload, Options.Chunking, cancellationToken);
+								await entry.AppendAsync(io.Payload, Options.Chunking, writer, cancellationToken);
 								dataSize += io.Payload.Length;
 							}
 							else if (io.Command == PerforceIoCommand.Close)
@@ -731,45 +727,70 @@ namespace Horde.Build.Perforce
 							{
 								_logger.LogWarning("Unhandled command code {Code}", io.Command);
 							}
+						}
 
-							if (dataSize > trimDataSize)
+						processTimer.Stop();
+					}
+					_logger.LogInformation("Completed batch in {TimeSeconds:n1}s ({ProcessTimeSeconds:n1}s processing, {TrimTimeSeconds:n1}s stalled for flush)", syncTimer.Elapsed.TotalSeconds, processTimer.Elapsed.TotalSeconds, flushTimer.Elapsed.TotalSeconds);
+
+					// Combine all the existing sync paths together with a new wildcard.
+					List<Utf8String> paths = syncNode.Paths;
+					if (paths.Count > 0)
+					{
+						Utf8String lastPath = syncNode.Paths[^1];
+						for (int endIdx = 0; endIdx < lastPath.Length; endIdx++)
+						{
+							if (lastPath[endIdx] == '/')
 							{
-								_logger.LogInformation("Trimming working set after receiving {Size:n1}mb...", dataSize / (1024 * 1024));
-								if (!trimTask.IsCompleted)
+								Utf8String prefix = lastPath.Substring(0, endIdx + 1);
+								if (!path.StartsWith(prefix, clientInfo.ServerInfo.Utf8PathComparer))
 								{
-									_logger.LogInformation("Waiting for previous trim to complete...");
+									// Find all the paths that start with this prefix
+									int pathIdx = paths.Count - 1;
+									while (pathIdx > 0 && paths[pathIdx - 1].StartsWith(prefix, clientInfo.ServerInfo.Utf8PathComparer))
+									{
+										pathIdx--;
+									}
+
+									// Replace the paths with a wildcard
+									paths.RemoveRange(pathIdx, paths.Count - pathIdx);
+									paths.Add(prefix + "...");
+									break;
 								}
-								await trimTask;
-								trimTask = Task.Run(async () =>
-								{
-									trimTimer.Start();
-									await _treeStore.WriteTreeAsync(refName, syncNode, false, cancellationToken);
-									trimTimer.Stop();
-									gcTimer.Start();
-									GC.Collect();
-									gcTimer.Stop();
-								}, cancellationToken);
-								trimDataSize = dataSize + TrimInterval;
-								_logger.LogInformation("Trimming complete. Next trim at {Size:n1}mb.", trimDataSize / (1024 * 1024));
 							}
 						}
+					}
+					paths.Add(path);
 
-						processTimer.Stop();
+					// Print the new state
+					if (dataSize > flushDataSize)
+					{
+						_logger.LogInformation("Queuing flush of working set after receiving {Size:n1}mb. Filter:", dataSize / (1024 * 1024));
+						foreach (Utf8String nextPath in paths)
+						{
+							_logger.LogInformation("  {Path}", nextPath);
+						}
+
+						flushTimer.Start();
+						await flushTask;
+						flushTask = writer.FlushAsync(syncNode, cancellationToken);
+						flushTimer.Stop();
+
+						flushDataSize = dataSize + FlushInterval;
+						_logger.LogInformation("Next flush at {Size:n1}mb.", flushDataSize / (1024 * 1024));
 					}
-					await trimTask;
-					_logger.LogInformation("Completed batch in {TimeSeconds:n1}s ({ProcessTimeSeconds:n1}s processing, {TrimTimeSeconds:n1}s trimming, {GcTimeSeconds:n1}s gc)", syncTimer.Elapsed.TotalSeconds, processTimer.Elapsed.TotalSeconds, trimTimer.Elapsed.TotalSeconds, gcTimer.Elapsed.TotalSeconds);
 
-					// Update the root sync node
-					syncNode.Paths.Add(path);
+					// Remove the paths we've synced
 					files.RemoveRange(idx, files.Count - idx);
 				}
+				await flushTask;
 			}
 
 			clientInfo.Change = change;
 
 			// Return the new root object
 			_logger.LogInformation("Writing ref {RefId} for {StreamId} change {Change}", refName, stream.Id, change);
-			await _treeStore.WriteTreeAsync(refName, syncNode, true, cancellationToken);
+			await _treeStore.WriteTreeAsync(refName, syncNode, cancellationToken);
 
 			// Delete the incremental state
 			if (deleteIncRef)
diff --git a/Engine/Source/Programs/Shared/EpicGames.Horde/Storage/Bundles/BundleStore.cs b/Engine/Source/Programs/Shared/EpicGames.Horde/Storage/Bundles/BundleStore.cs
--- a/Engine/Source/Programs/Shared/EpicGames.Horde/Storage/Bundles/BundleStore.cs
+++ b/Engine/Source/Programs/Shared/EpicGames.Horde/Storage/Bundles/BundleStore.cs
@@ -5,7 +5,6 @@ using System.Buffers;
 using System.Collections.Concurrent;
 using System.Collections.Generic;
 using System.Diagnostics;
-using System.IO;
 using System.Linq;
 using System.Threading;
 using System.Threading.Tasks;
@@ -47,78 +46,45 @@ namespace EpicGames.Horde.Storage.Bundles
 	}
 
 	/// <summary>
-	/// Statistics for the state of a bundle
+	/// Implementation of <see cref="ITreeStore"/> which packs nodes together into <see cref="Bundle"/> objects, then stores them
+	/// in a <see cref="IBlobStore"/>.
 	/// </summary>
-	public class BundleStats
-	{
-		/// <summary>
-		/// Number of blobs that have been written
-		/// </summary>
-		public int NewBlobCount { get; set; }
-
-		/// <summary>
-		/// Total size of the blobs that have been written
-		/// </summary>
-		public long NewBlobBytes { get; set; }
-
-		/// <summary>
-		/// Number of refs that have been written
-		/// </summary>
-		public int NewRefCount { get; set; }
-
-		/// <summary>
-		/// Total size of the refs that have been written
-		/// </summary>
-		public long NewRefBytes { get; set; }
-
-		/// <summary>
-		/// Number of nodes in the tree that have been written
-		/// </summary>
-		public int NewNodeCount { get; set; }
-
-		/// <summary>
-		/// Total size of the nodes that have been serialized
-		/// </summary>
-		public long NewNodeBytes { get; set; }
-	}
-
-	/// <summary>
-	/// Manages a collection of bundles.
-	/// </summary>
-	public class BundleStore : ITreeStore, ITreeBlobWriter, IDisposable
+	public class BundleStore : ITreeStore, IDisposable
 	{
 		/// <summary>
 		/// Information about a node within a bundle.
 		/// </summary>
 		[DebuggerDisplay("{Hash}")]
-		sealed class NodeInfo : ITreeBlob
+		sealed class NodeInfo : ITreeBlobRef
 		{
 			public readonly BundleStore Owner;
-			public IoHash Hash { get; }
-			public NodeState State => _state;
-
-			NodeState _state;
+			public readonly IoHash Hash;
+			public NodeState State { get; private set; }
 
 			public NodeInfo(BundleStore owner, IoHash hash, NodeState state)
 			{
 				Owner = owner;
 				Hash = hash;
-				_state = state;
+				State = state;
 			}
 
-			public bool TrySetState(NodeState prevState, NodeState nextState)
+			/// <summary>
+			/// Updates the current state. The only allowed transition is to an exported node state.
+			/// </summary>
+			/// <param name="state">The new state</param>
+			public void SetState(ExportedNodeState state)
 			{
-				return Interlocked.CompareExchange(ref _state, nextState, prevState) == prevState;
+				State = state;
 			}
 
 			/// <inheritdoc/>
-			public ValueTask<ReadOnlySequence<byte>> GetDataAsync(CancellationToken cancellationToken) => Owner.GetDataAsync(this, cancellationToken);
-
-			/// <inheritdoc/>
-			public async ValueTask<IReadOnlyList<ITreeBlob>> GetReferencesAsync(CancellationToken cancellationToken) => await Owner.GetReferencesAsync(this, cancellationToken);
+			public ValueTask<ITreeBlob> GetTargetAsync(CancellationToken cancellationToken = default) => Owner.GetNodeAsync(this, cancellationToken);
 		}
 
-		sealed class PacketInfo
+		/// <summary>
+		/// Metadata about a compression packet within a bundle
+		/// </summary>
+		class PacketInfo
 		{
 			public readonly int Offset;
 			public readonly int DecodedLength;
@@ -132,26 +98,33 @@ namespace EpicGames.Horde.Storage.Bundles
 			}
 		}
 
+		/// <summary>
+		/// Base class for the state of a node
+		/// </summary>
 		abstract class NodeState
 		{
 		}
 
-		// State for data imported from a bundle, but whose location within it is not yet known because the bundle's header has not been read yet.
-		// Can be transitioned to ExportedNodeState.
+		/// <summary>
+		/// State for data imported from a bundle, but whose location within it is not yet known because the bundle's header has not been read yet.
+		/// Can be transitioned to ExportedNodeState.
+		/// </summary>
 		class ImportedNodeState : NodeState
 		{
-			public readonly BlobInfo BundleInfo;
+			public readonly BundleInfo BundleInfo;
 			public readonly int ExportIdx;
 
-			public ImportedNodeState(BlobInfo bundleInfo, int exportIdx)
+			public ImportedNodeState(BundleInfo bundleInfo, int exportIdx)
 			{
 				BundleInfo = bundleInfo;
 				ExportIdx = exportIdx;
 			}
 		}
 
-		// State for data persisted to a bundle. Decoded data is cached outside this structure in the store's MemoryCache.
-		// Cannot be transitioned to any other state.
+		/// <summary>
+		/// State for data persisted to a bundle. Decoded data is cached outside this structure in the store's MemoryCache.
+		/// Cannot be transitioned to any other state.
+		/// </summary>
 		class ExportedNodeState : ImportedNodeState
 		{
 			public readonly PacketInfo PacketInfo;
@@ -159,7 +132,7 @@ namespace EpicGames.Horde.Storage.Bundles
 			public readonly int Length;
 			public readonly IReadOnlyList<NodeInfo> References;
 
-			public ExportedNodeState(BlobInfo bundleInfo, int exportIdx, PacketInfo packetInfo, int offset, int length, IReadOnlyList<NodeInfo> references)
+			public ExportedNodeState(BundleInfo bundleInfo, int exportIdx, PacketInfo packetInfo, int offset, int length, IReadOnlyList<NodeInfo> references)
 				: base(bundleInfo, exportIdx)
 			{
 				PacketInfo = packetInfo;
@@ -169,13 +142,18 @@ namespace EpicGames.Horde.Storage.Bundles
 			}
 		}
 
-		// Data for a node which exists in memory, but which does NOT exist in storage yet. Nodes which 
-		// are read in from storage are represented by ExportedNodeState and a MemoryCache. Can transition to an ExportedNodeState only.
-		class InMemoryNodeState : NodeState
+		/// <summary>
+		/// Data for a node which exists in memory, but which does NOT exist in storage yet. Nodes which 
+		/// are read in from storage are represented by ExportedNodeState and a MemoryCache. Can transition to an ExportedNodeState only.
+		/// </summary>
+		class InMemoryNodeState : NodeState, ITreeBlob
 		{
 			public readonly ReadOnlySequence<byte> Data;
 			public readonly IReadOnlyList<NodeInfo> References;
 
+			ReadOnlySequence<byte> ITreeBlob.Data => Data;
+			IReadOnlyList<ITreeBlobRef> ITreeBlob.Refs => References;
+
 			public InMemoryNodeState(ReadOnlySequence<byte> data, IReadOnlyList<NodeInfo> references)
 			{
 				Data = data;
@@ -184,30 +162,69 @@ namespace EpicGames.Horde.Storage.Bundles
 		}
 
 		/// <summary>
-		/// Information about a blob
+		/// Stores a lookup from blob id to bundle info, and from there to the nodes it contains.
+		/// </summary>
+		class BundleContext
+		{
+			readonly ConcurrentDictionary<BlobId, BundleInfo> _blobIdToBundle = new ConcurrentDictionary<BlobId, BundleInfo>();
+
+			public BundleInfo FindOrAddBundle(BlobId blobId, int exportCount)
+			{
+				BundleInfo bundleInfo = _blobIdToBundle.GetOrAdd(blobId, new BundleInfo(this, blobId, exportCount));
+				Debug.Assert(bundleInfo.Exports.Length == exportCount);
+				return bundleInfo;
+			}
+		}
+
+		/// <summary>
+		/// Information about an imported bundle
 		/// </summary>
 		[DebuggerDisplay("{BlobId}")]
-		class BlobInfo
+		class BundleInfo
 		{
+			public readonly BundleContext Context;
 			public readonly BlobId BlobId;
 			public readonly NodeInfo?[] Exports;
-			public Task<NodeInfo>? _mountTask;
 
-			public BlobInfo(BlobId blobId, int exportCount)
-				: this(blobId, new NodeInfo?[exportCount])
+			public bool Mounted;
+			public Task MountTask = Task.CompletedTask;
+
+			public BundleInfo(BundleContext context, BlobId blobId, int exportCount)
 			{
+				Context = context;
+				BlobId = blobId;
+				Exports = new NodeInfo?[exportCount];
 			}
+		}
+
+		class BundleWriteContext : BundleContext, ITreeWriter
+		{
+			readonly BundleStore _owner;
+
+			public readonly RefName Name;
+			public readonly ConcurrentDictionary<IoHash, NodeInfo> HashToNode = new ConcurrentDictionary<IoHash, NodeInfo>();
+
+			public Task WriteTask = Task.CompletedTask;
+			public readonly List<NodeInfo> WriteNodes = new List<NodeInfo>();
+			public readonly HashSet<NodeInfo> WriteNodesSet = new HashSet<NodeInfo>();
+			public long WriteLength;
 
-			public BlobInfo(BlobId blobId, NodeInfo?[] exports)
+			public byte[] PacketBuffer = Array.Empty<byte>();
+
+			public BundleWriteContext(BundleStore owner, RefName name)
 			{
-				BlobId = blobId;
-				Exports = exports;
+				_owner = owner;
+				Name = name;
 			}
+
+			/// <inheritdoc/>
+			public Task FlushAsync(ITreeBlob root, CancellationToken cancellationToken = default) => _owner.FlushAsync(this, root, cancellationToken);
+
+			/// <inheritdoc/>
+			public Task<ITreeBlobRef> WriteNodeAsync(ReadOnlySequence<byte> data, IReadOnlyList<ITreeBlobRef> refs, CancellationToken cancellationToken = default) => _owner.WriteNodeAsync(this, data, refs, cancellationToken);
 		}
 
 		readonly IBlobStore _blobStore;
-		readonly ConcurrentDictionary<IoHash, NodeInfo> _hashToNode = new ConcurrentDictionary<IoHash, NodeInfo>();
-		readonly ConcurrentDictionary<BlobId, BlobInfo> _blobIdToBundleInfo = new ConcurrentDictionary<BlobId, BlobInfo>();
 		readonly IMemoryCache _cache;
 
 		readonly object _lockObject = new object();
@@ -218,19 +235,8 @@ namespace EpicGames.Horde.Storage.Bundles
 		// so that other threads can also await it.
 		readonly Dictionary<string, Task<Bundle>> _readTasks = new Dictionary<string, Task<Bundle>>();
 
-		Task _writeTask = Task.CompletedTask;
-
 		readonly BundleOptions _options;
 
-		/// <summary>
-		/// Tracks statistics for the size of written data
-		/// </summary>
-		public BundleStats Stats { get; private set; } = new BundleStats();
-
-		BundleStats _nextStats = new BundleStats();
-		byte[] _blockBuffer = Array.Empty<byte>();
-		byte[] _encodedBuffer = Array.Empty<byte>();
-
 		/// <summary>
 		/// Constructor
 		/// </summary>
@@ -242,6 +248,12 @@ namespace EpicGames.Horde.Storage.Bundles
 			_readSema = new SemaphoreSlim(options.MaxConcurrentReads);
 		}
 
+		/// <inheritdoc/>
+		public ITreeWriter CreateTreeWriter(RefName name) => new BundleWriteContext(this, name);
+
+		/// <inheritdoc/>
+		public Task DeleteTreeAsync(RefName name, CancellationToken cancellationToken) => _blobStore.DeleteRefAsync(name, cancellationToken);
+
 		/// <inheritdoc/>
 		public void Dispose()
 		{
@@ -269,137 +281,65 @@ namespace EpicGames.Horde.Storage.Bundles
 			{
 				return null;
 			}
-			return MountInMemoryBundle(bundle);
-		}
 
-		/// <inheritdoc/>
-		public async Task<T?> TryReadTreeAsync<T>(RefName name, CancellationToken cancellationToken = default) where T : TreeNode
-		{
-			ITreeBlob? node = await TryReadTreeAsync(name, cancellationToken);
-			if (node == null)
-			{
-				return null;
-			}
-			return await TreeNode.DeserializeAsync<T>(node, cancellationToken);
-		}
+			// Create a new context for this bundle and its references
+			BundleContext context = new BundleContext();
 
-		/// <inheritdoc/>
-		public async Task WriteTreeAsync(RefName name, ITreeBlob root, bool flush = true, CancellationToken cancellationToken = default)
-		{
-			await SequenceWrite(() => WriteTreeInternalAsync(name, root, flush, cancellationToken), cancellationToken);
-		}
+			// Create all the imports for the root bundle
+			BundleHeader header = bundle.Header;
+			List<NodeInfo> nodes = CreateImports(context, header);
 
-		async Task SequenceWrite(Func<Task> writeFunc, CancellationToken cancellationToken)
-		{
-			Task task;
-			lock (_lockObject)
+			// Decompress the bundle data and mount all the nodes in memory
+			int exportIdx = 0;
+			int packetOffset = 0;
+			foreach (BundlePacket packet in header.Packets)
 			{
-				Task prevTask = _writeTask;
-
-				// Wait for the previous write to complete first, ignoring any exceptions
-				Func<Task> wrappedWriteFunc = async () =>
-				{
-					await prevTask.ContinueWith(x => { }, cancellationToken, TaskContinuationOptions.None, TaskScheduler.Default);
-					await writeFunc();
-				};
-				task = Task.Run(wrappedWriteFunc, cancellationToken);
-
-				_writeTask = task;
-			}
-			await task;
-		}
-
-		async Task WriteTreeInternalAsync(RefName name, ITreeBlob root, bool flush = true, CancellationToken cancellationToken = default)
-		{
-			// Find all the referenced in-memory nodes from the root
-			List<NodeInfo> nodes = new List<NodeInfo>();
-			FindNodesToWrite((NodeInfo)root, nodes, new HashSet<NodeInfo>(), flush);
+				ReadOnlyMemory<byte> encodedPacket = bundle.Payload.Slice(packetOffset, packet.EncodedLength);
 
-			// Group the nodes into bundles
-			long bundleSize = 0;
-			List<NodeInfo> bundleNodes = new List<NodeInfo>();
+				byte[] decodedPacket = new byte[packet.DecodedLength];
+				LZ4Codec.Decode(encodedPacket.Span, decodedPacket);
 
-			foreach (NodeInfo node in nodes)
-			{
-				// Check if we need to flush the current bundle
-				long nodeSize = ((InMemoryNodeState)node.State).Data.Length;
-				if (bundleNodes.Count > 0 && bundleSize + nodeSize > _options.MaxBlobSize)
+				int nodeOffset = 0;
+				for (; exportIdx < header.Exports.Count && nodeOffset + header.Exports[exportIdx].Length <= packet.DecodedLength; exportIdx++)
 				{
-					Bundle bundle = CreateBundle(bundleNodes);
-					BlobId blobId = await _blobStore.WriteBundleBlobAsync(name, bundle, cancellationToken);
-
-					BlobInfo bundleInfo = new BlobInfo(blobId, bundleNodes.ToArray());
-					MountExportedBundle(bundleInfo, bundle);
+					BundleExport export = header.Exports[exportIdx];
 
-					bundleSize = 0;
-					bundleNodes.Clear();
-				}
+					ReadOnlyMemory<byte> nodeData = decodedPacket.AsMemory(nodeOffset, export.Length);
+					List<NodeInfo> nodeRefs = export.References.ConvertAll(x => nodes[x]);
 
-				// Add this node to the list
-				bundleSize += nodeSize;
-				bundleNodes.Add(node);
-			}
+					NodeInfo node = new NodeInfo(this, export.Hash, new InMemoryNodeState(new ReadOnlySequence<byte>(nodeData), nodeRefs));
+					nodes.Add(node);
 
-			// If we're doing a full flush, also write whatever is left
-			if (flush)
-			{
-				// Write the main ref
-				Bundle rootBundle = CreateBundle(bundleNodes);
-				await _blobStore.WriteBundleRefAsync(name, rootBundle, cancellationToken);
+					nodeOffset += export.Length;
+				}
 
-				// Copy the stats over
-				Stats = _nextStats;
-				_nextStats = new BundleStats();
+				packetOffset += packet.EncodedLength;
 			}
-		}
 
-		/// <inheritdoc/>
-		public async Task WriteTreeAsync<T>(RefName name, T root, bool flush = true, CancellationToken cancellationToken = default) where T : TreeNode
-		{
-			ITreeBlob rootBlob = await TreeNode.SerializeAsync(root, this, cancellationToken);
-			await WriteTreeAsync(name, rootBlob, flush, cancellationToken);
+			// Return the last node as the root
+			return (InMemoryNodeState)nodes[^1].State;
 		}
 
-		/// <inheritdoc/>
-		public Task DeleteTreeAsync(RefName name, CancellationToken cancellationToken) => _blobStore.DeleteRefAsync(name, cancellationToken);
-
 		/// <summary>
 		/// Gets the data for a given node
 		/// </summary>
 		/// <param name="node">The node to get the data for</param>
 		/// <param name="cancellationToken">Cancellation token for the operation</param>
 		/// <returns>The node data</returns>
-		async ValueTask<ReadOnlySequence<byte>> GetDataAsync(NodeInfo node, CancellationToken cancellationToken)
+		async ValueTask<ITreeBlob> GetNodeAsync(NodeInfo node, CancellationToken cancellationToken)
 		{
 			if (node.State is InMemoryNodeState inMemoryState)
 			{
-				return inMemoryState.Data;
+				return inMemoryState;
 			}
 
 			ExportedNodeState exportedState = await GetExportedStateAsync(node, cancellationToken);
 
-			Bundle bundle = await ReadBundleAsync(exportedState.BundleInfo, cancellationToken);
-			ReadOnlyMemory<byte> packetData = DecodePacket(exportedState.BundleInfo, exportedState.PacketInfo, bundle.Payload);
+			Bundle bundle = await ReadBundleAsync(exportedState.BundleInfo.BlobId, cancellationToken);
+			ReadOnlyMemory<byte> packetData = DecodePacket(exportedState.BundleInfo.BlobId, exportedState.PacketInfo, bundle.Payload);
 			ReadOnlyMemory<byte> nodeData = packetData.Slice(exportedState.Offset, exportedState.Length);
 
-			return new ReadOnlySequence<byte>(nodeData);
-		}
-
-		/// <summary>
-		/// Gets the references from a given node
-		/// </summary>
-		/// <param name="node">The node to get the data for</param>
-		/// <param name="cancellationToken">Cancellation token for the operation</param>
-		/// <returns>The node data</returns>
-		async ValueTask<IReadOnlyList<NodeInfo>> GetReferencesAsync(NodeInfo node, CancellationToken cancellationToken)
-		{
-			if (node.State is InMemoryNodeState inMemoryState)
-			{
-				return inMemoryState.References;
-			}
-
-			ExportedNodeState exportedState = await GetExportedStateAsync(node, cancellationToken);
-			return exportedState.References;
+			return new InMemoryNodeState(new ReadOnlySequence<byte>(nodeData), exportedState.References);
 		}
 
 		async ValueTask<ExportedNodeState> GetExportedStateAsync(NodeInfo node, CancellationToken cancellationToken)
@@ -414,334 +354,388 @@ namespace EpicGames.Horde.Storage.Bundles
 			return exportedState;
 		}
 
-		/// <inheritdoc/>
-		public Task<ITreeBlob> WriteBlobAsync(ReadOnlySequence<byte> data, IReadOnlyList<ITreeBlob> references, CancellationToken cancellationToken)
-		{
-			List<NodeInfo> typedReferences = references.ConvertAll(x => (NodeInfo)x);
-			if (typedReferences.Any(x => x.Owner != this))
-			{
-				throw new InvalidDataException("Referenced node does not belong to the current tree.");
-			}
-
-			IoHash hash = ComputeHash(data, typedReferences);
+		#region Reading bundles
 
-			InMemoryNodeState state = new InMemoryNodeState(data, typedReferences);
-			NodeInfo node = FindOrAddNode(hash, state);
-
-			_nextStats.NewNodeCount++;
-			_nextStats.NewNodeBytes += data.Length;
-
-			return Task.FromResult<ITreeBlob>(node);
-		}
-
-		NodeInfo FindOrAddNode(IoHash hash, NodeState state)
+		/// <summary>
+		/// Reads a bundle from the given blob id, or retrieves it from the cache
+		/// </summary>
+		/// <param name="blobId"></param>
+		/// <param name="cancellationToken"></param>
+		/// <returns></returns>
+		async Task<Bundle> ReadBundleAsync(BlobId blobId, CancellationToken cancellationToken = default)
 		{
-			NodeInfo? node;
-			if (!_hashToNode.TryGetValue(hash, out node))
-			{
-				node = _hashToNode.GetOrAdd(hash, new NodeInfo(this, hash, state));
-			}
-
-			if (state is ExportedNodeState)
+			string cacheKey = GetBundleCacheKey(blobId);
+			if (!_cache.TryGetValue<Bundle>(cacheKey, out Bundle? bundle))
 			{
-				// Transition from any non-exported state
-				for (; ; )
+				Task<Bundle>? readTask;
+				lock (_lockObject)
 				{
-					NodeState prevState = node.State;
-					if (prevState is ExportedNodeState)
-					{
-						break;
-					}
-					else if (node.TrySetState(prevState, state))
+					if (!_readTasks.TryGetValue(cacheKey, out readTask))
 					{
-						break;
+						readTask = Task.Run(() => ReadBundleInternalAsync(cacheKey, blobId, cancellationToken), cancellationToken);
+						_readTasks.Add(cacheKey, readTask);
 					}
 				}
+				bundle = await readTask;
 			}
-			else if (state is ImportedNodeState)
+			return bundle;
+		}
+
+		async Task<Bundle> ReadBundleInternalAsync(string cacheKey, BlobId blobId, CancellationToken cancellationToken)
+		{
+			// Perform another (sequenced) check whether an object has been added to the cache, to counteract the race between a read task being added and a task completing.
+			lock (_lockObject)
 			{
-				// Transition from in memory state
-				for (; ; )
+				if (_cache.TryGetValue<Bundle>(cacheKey, out Bundle? cachedObject))
 				{
-					NodeState prevState = node.State;
-					if (!(prevState is InMemoryNodeState))
-					{
-						break;
-					}
-					else if (node.TrySetState(prevState, state))
-					{
-						break;
-					}
+					return cachedObject;
 				}
 			}
 
-			return node;
-		}
+			// Wait for the read semaphore to avoid triggering too many operations at once.
+			await _readSema.WaitAsync(cancellationToken);
 
-		static IoHash ComputeHash(ReadOnlySequence<byte> data, IReadOnlyList<NodeInfo> references)
-		{
-			byte[] buffer = new byte[IoHash.NumBytes * (references.Count + 1)];
+			// Read the data from storage
+			Bundle bundle;
+			try
+			{
+				bundle = await _blobStore.ReadBundleAsync(blobId, cancellationToken);
+			}
+			finally
+			{
+				_readSema.Release();
+			}
 
-			Span<byte> span = buffer;
-			for (int idx = 0; idx < references.Count; idx++)
+			// Add the object to the cache
+			using (ICacheEntry entry = _cache.CreateEntry(cacheKey))
 			{
-				references[idx].Hash.CopyTo(span);
-				span = span[IoHash.NumBytes..];
+				entry.SetSize(bundle.Payload.Length);
+				entry.SetValue(bundle);
 			}
-			IoHash.Compute(data).CopyTo(span);
 
-			return IoHash.Compute(buffer);
+			// Remove this object from the list of read tasks
+			lock (_lockObject)
+			{
+				_readTasks.Remove(cacheKey);
+			}
+			return bundle;
 		}
 
-		void FindNodesToWrite(NodeInfo root, List<NodeInfo> nodes, HashSet<NodeInfo> uniqueNodes, bool flush)
+		static string GetBundleCacheKey(BlobId blobId) => $"bundle:{blobId}";
+
+		void AddBundleToCache(string cacheKey, Bundle bundle)
 		{
-			if (root.State is InMemoryNodeState inMemoryState)
+			using (ICacheEntry entry = _cache.CreateEntry(cacheKey))
 			{
-				if (uniqueNodes.Add(root))
-				{
-					int count = inMemoryState.References.Count;
-					for (int idx = 0; idx < count; idx++)
-					{
-						FindNodesToWrite(inMemoryState.References[idx], nodes, uniqueNodes, flush || idx < count - 1);
-					}
-					if (flush)
-					{
-						nodes.Add(root);
-					}
-				}
+				entry.SetSize(bundle.Payload.Length);
+				entry.SetValue(bundle);
 			}
 		}
 
-		async Task<NodeInfo> MountBundleAsync(BlobInfo bundle, CancellationToken cancellationToken)
+		async Task MountBundleAsync(BundleInfo bundleInfo, CancellationToken cancellationToken)
 		{
-			Task<NodeInfo>? mountTask = bundle._mountTask;
-			if (mountTask == null)
+			if (!bundleInfo.Mounted)
 			{
-				lock (bundle)
+				Task mountTask;
+				lock (bundleInfo)
 				{
-					bundle._mountTask ??= Task.Run(() => MountBundleInternalAsync(bundle, cancellationToken));
-					mountTask = bundle._mountTask;
+					Task prevMountTask = bundleInfo.MountTask.ContinueWith(x => { }, cancellationToken, TaskContinuationOptions.None, TaskScheduler.Default);
+					mountTask = Task.Run(() => MountBundleInternalAsync(prevMountTask, bundleInfo, cancellationToken), cancellationToken);
+					bundleInfo.MountTask = mountTask;
 				}
+				await mountTask;
 			}
-			return await mountTask;
 		}
 
-		async Task<NodeInfo> MountBundleInternalAsync(BlobInfo bundleInfo, CancellationToken cancellationToken)
+		async Task MountBundleInternalAsync(Task prevMountTask, BundleInfo bundleInfo, CancellationToken cancellationToken)
 		{
-			Bundle bundle = await ReadBundleAsync(bundleInfo, cancellationToken);
-			return MountExportedBundle(bundleInfo, bundle);
+			// Wait for the previous mount task to complete. This may succeed or be cancelled.
+			await prevMountTask;
+
+			// If it didn't succeed, try again
+			if (!bundleInfo.Mounted)
+			{
+				// Read the bundle data
+				Bundle bundle = await ReadBundleAsync(bundleInfo.BlobId, cancellationToken);
+				BundleHeader header = bundle.Header;
+
+				// Create all the imported nodes
+				List<NodeInfo> nodes = CreateImports(bundleInfo.Context, header);
+
+				// Create the exported nodes, or update the state of any imported nodes to exported
+				int exportIdx = 0;
+				int packetOffset = 0;
+				foreach (BundlePacket packet in header.Packets)
+				{
+					PacketInfo packetInfo = new PacketInfo(packetOffset, packet);
+
+					int nodeOffset = 0;
+					for (; exportIdx < header.Exports.Count && nodeOffset + header.Exports[exportIdx].Length <= packet.DecodedLength; exportIdx++)
+					{
+						BundleExport export = header.Exports[exportIdx];
+						List<NodeInfo> nodeRefs = export.References.ConvertAll(x => nodes[x]);
+
+						ExportedNodeState state = new ExportedNodeState(bundleInfo, exportIdx, packetInfo, nodeOffset, export.Length, nodeRefs);
+
+						NodeInfo? node = bundleInfo.Exports[exportIdx];
+						if (node == null)
+						{
+							node = bundleInfo.Exports[exportIdx] = new NodeInfo(this, export.Hash, state);
+						}
+						else
+						{
+							node.SetState(state);
+						}
+						nodes.Add(node);
+
+						nodeOffset += export.Length;
+					}
+
+					packetOffset += packet.EncodedLength;
+				}
+
+				// Mark the bundle as mounted
+				bundleInfo.Mounted = true;
+			}
 		}
 
-		List<NodeInfo> RegisterImports(BundleHeader header)
+		List<NodeInfo> CreateImports(BundleContext context, BundleHeader header)
 		{
-			// Create all the nodes as imports first, so we can fix up references later.
-			List<NodeInfo> nodes = new List<NodeInfo>(header.Imports.Sum(x => x.Exports.Count) + header.Exports.Count);
+			List<NodeInfo> indexedNodes = new List<NodeInfo>(header.Imports.Sum(x => x.Exports.Count) + header.Exports.Count);
 			foreach (BundleImport import in header.Imports)
 			{
-				BlobInfo importBundle = _blobIdToBundleInfo.GetOrAdd(import.BlobId, id => new BlobInfo(id, import.ExportCount));
+				BundleInfo importBundle = context.FindOrAddBundle(import.BlobId, import.ExportCount);
 				foreach ((int exportIdx, IoHash exportHash) in import.Exports)
 				{
-					NodeInfo? export = importBundle.Exports[exportIdx];
-					if (export == null)
+					NodeInfo? node = importBundle.Exports[exportIdx];
+					if (node == null)
 					{
-						export = FindOrAddNode(exportHash, new ImportedNodeState(importBundle, exportIdx));
-						importBundle.Exports[exportIdx] = export;
+						node = new NodeInfo(this, exportHash, new ImportedNodeState(importBundle, exportIdx));
+						importBundle.Exports[exportIdx] = node;
 					}
-					nodes.Add(export);
+					indexedNodes.Add(node);
 				}
 			}
-			return nodes;
+			return indexedNodes;
 		}
-
-		NodeInfo MountInMemoryBundle(Bundle bundle)
+		/// <summary>
+		/// Gets a decoded block from the store
+		/// </summary>
+		/// <param name="blobId">Information about the bundle</param>
+		/// <param name="packetInfo">The decoded block location and size</param>
+		/// <param name="payload">The bundle payload</param>
+		/// <returns>The decoded data</returns>
+		ReadOnlyMemory<byte> DecodePacket(BlobId blobId, PacketInfo packetInfo, ReadOnlyMemory<byte> payload)
 		{
-			BundleHeader header = bundle.Header;
-			List<NodeInfo> nodes = RegisterImports(header);
-
-			int exportIdx = 0;
-			int packetOffset = 0;
-			foreach (BundlePacket packet in header.Packets)
+			string cacheKey = $"bundle-packet:{blobId}@{packetInfo.Offset}";
+			return _cache.GetOrCreate<ReadOnlyMemory<byte>>(cacheKey, entry =>
 			{
-				ReadOnlyMemory<byte> encodedPacket = bundle.Payload.Slice(packetOffset, packet.EncodedLength);
-
-				byte[] decodedPacket = new byte[packet.DecodedLength];
-				LZ4Codec.Decode(encodedPacket.Span, decodedPacket);
-
-				int nodeOffset = 0;
-				for (; exportIdx < header.Exports.Count && nodeOffset + header.Exports[exportIdx].Length <= packet.DecodedLength; exportIdx++)
-				{
-					BundleExport export = header.Exports[exportIdx];
+				ReadOnlyMemory<byte> decodedPacket = DecodePacketUncached(packetInfo, payload);
+				entry.SetSize(packetInfo.DecodedLength);
+				return decodedPacket;
+			});
+		}
 
-					ReadOnlyMemory<byte> nodeData = decodedPacket.AsMemory(nodeOffset, export.Length);
-					List<NodeInfo> nodeRefs = export.References.ConvertAll(x => nodes[x]);
+		static ReadOnlyMemory<byte> DecodePacketUncached(PacketInfo packetInfo, ReadOnlyMemory<byte> payload)
+		{
+			ReadOnlyMemory<byte> encodedPacket = payload.Slice(packetInfo.Offset, packetInfo.EncodedLength);
 
-					InMemoryNodeState nodeState = new InMemoryNodeState(new ReadOnlySequence<byte>(nodeData), nodeRefs);
-					NodeInfo node = FindOrAddNode(export.Hash, nodeState);
-					nodes.Add(node);
+			byte[] decodedPacket = new byte[packetInfo.DecodedLength];
+			LZ4Codec.Decode(encodedPacket.Span, decodedPacket);
 
-					nodeOffset += header.Exports[exportIdx].Length;
-				}
+			return decodedPacket;
+		}
 
-				packetOffset += packet.EncodedLength;
-			}
+		#endregion
 
-			return nodes[^1];
-		}
+		#region Writing bundles
 
-		NodeInfo MountExportedBundle(BlobInfo bundleInfo, Bundle bundle)
+		Task<ITreeBlobRef> WriteNodeAsync(BundleWriteContext context, ReadOnlySequence<byte> data, IReadOnlyList<ITreeBlobRef> refs, CancellationToken cancellationToken = default)
 		{
-			BundleHeader header = bundle.Header;
-			List<NodeInfo> nodes = RegisterImports(header);
+			IReadOnlyList<NodeInfo> typedRefs = refs.Select(x => (NodeInfo)x).ToList();
 
-			int exportIdx = 0;
-			int packetOffset = 0;
-			foreach (BundlePacket packet in header.Packets)
+			IoHash hash = ComputeHash(data, typedRefs);
+
+			NodeInfo? node;
+			if (!context.HashToNode.TryGetValue(hash, out node))
 			{
-				PacketInfo packetInfo = new PacketInfo(packetOffset, packet);
+				InMemoryNodeState state = new InMemoryNodeState(data, typedRefs);
+				NodeInfo newNode = new NodeInfo(this, hash, state);
 
-				int nodeOffset = 0;
-				for (; exportIdx < header.Exports.Count && nodeOffset + header.Exports[exportIdx].Length <= packet.DecodedLength; exportIdx++)
+				// Try to add the node again. If we succeed, check if we need to flush the current bundle being built.
+				node = context.HashToNode.GetOrAdd(hash, newNode);
+				if (node == newNode)
 				{
-					BundleExport export = header.Exports[exportIdx];
-
-					List<NodeInfo> nodeRefs = export.References.ConvertAll(x => nodes[x]);
-					ExportedNodeState exportedState = new ExportedNodeState(bundleInfo, exportIdx, packetInfo, nodeOffset, export.Length, nodeRefs);
-
-					NodeInfo node = FindOrAddNode(export.Hash, exportedState);
-					bundleInfo.Exports[exportIdx] = node;
-					nodes.Add(node);
-
-					nodeOffset += header.Exports[exportIdx].Length;
+					lock (context)
+					{
+						WriteInMemoryNodes(context, node);
+					}
 				}
-
-				packetOffset += packet.EncodedLength;
 			}
 
-			return nodes[^1];
+			return Task.FromResult<ITreeBlobRef>(node);
 		}
 
-		async Task<Bundle> ReadBundleAsync(BlobInfo bundleInfo, CancellationToken cancellationToken)
+		void WriteInMemoryNodes(BundleWriteContext context, NodeInfo root)
 		{
-			string cacheKey = $"bundle:{bundleInfo.BlobId}";
-			if (!_cache.TryGetValue<Bundle>(cacheKey, out Bundle? bundle))
+			if (root.State is InMemoryNodeState inMemoryState)
 			{
-				Task<Bundle>? readTask;
-				lock (_lockObject)
+				if (context.WriteNodesSet.Add(root))
 				{
-					if (!_readTasks.TryGetValue(cacheKey, out readTask))
+					// Write all the references
+					foreach (NodeInfo reference in inMemoryState.References)
 					{
-						readTask = Task.Run(() => ReadBundleInternalAsync(cacheKey, bundleInfo.BlobId, cancellationToken));
-						_readTasks.Add(cacheKey, readTask);
+						WriteInMemoryNodes(context, reference);
 					}
+
+					// If this node pushes the current blob over the max size, queue it to be written now
+					long length = inMemoryState.Data.Length;
+					if (context.WriteLength + length > _options.MaxBlobSize)
+					{
+						// Start the write, but don't wait for it to complete. Once we've copied the list of nodes they can be written asynchronously.
+						NodeInfo[] writeNodes = context.WriteNodes.ToArray();
+						_ = SequenceWriteAsync(context, () => WriteBundleAsync(context, writeNodes, CancellationToken.None), CancellationToken.None);
+
+						context.WriteNodes.Clear();
+						context.WriteLength = 0;
+					}
+
+					// Add this node to the list of nodes to be written
+					context.WriteNodes.Add(root);
+					context.WriteLength += length;
 				}
-				bundle = await readTask;
 			}
-			return bundle;
 		}
 
-		async Task<Bundle> ReadBundleInternalAsync(string cacheKey, BlobId blobId, CancellationToken cancellationToken)
+		static IoHash ComputeHash(ReadOnlySequence<byte> data, IReadOnlyList<NodeInfo> references)
 		{
-			// Perform another (sequenced) check whether an object has been added to the cache, to counteract the race between a read task being added and a task completing.
-			lock (_lockObject)
+			byte[] buffer = new byte[IoHash.NumBytes * (references.Count + 1)];
+
+			Span<byte> span = buffer;
+			for (int idx = 0; idx < references.Count; idx++)
 			{
-				if (_cache.TryGetValue<Bundle>(cacheKey, out Bundle? cachedObject))
-				{
-					return cachedObject;
-				}
+				references[idx].Hash.CopyTo(span);
+				span = span[IoHash.NumBytes..];
 			}
+			IoHash.Compute(data).CopyTo(span);
 
-			// Wait for the read semaphore to avoid triggering too many operations at once.
-			await _readSema.WaitAsync(cancellationToken);
+			return IoHash.Compute(buffer);
+		}
 
-			// Read the data from storage
-			Bundle bundle;
-			try
-			{
-				bundle = await _blobStore.ReadBundleAsync(blobId, cancellationToken);
-			}
-			finally
+		async Task WriteBundleAsync(BundleWriteContext context, NodeInfo[] nodes, CancellationToken cancellationToken)
+		{
+			// Create the bundle
+			Bundle bundle = CreateBundle(context, nodes);
+			BundleHeader header = bundle.Header;
+
+			// Write it to storage
+			BlobId blobId = await _blobStore.WriteBundleBlobAsync(context.Name, bundle, cancellationToken);
+			string cacheKey = GetBundleCacheKey(blobId);
+			AddBundleToCache(cacheKey, bundle);
+
+			// Create a BundleInfo for it
+			BundleInfo bundleInfo = context.FindOrAddBundle(blobId, nodes.Length);
+			for (int idx = 0; idx < nodes.Length; idx++)
 			{
-				_readSema.Release();
+				bundleInfo.Exports[idx] = nodes[idx];
 			}
 
-			// Add the object to the cache
-			using (ICacheEntry entry = _cache.CreateEntry(cacheKey))
+			// Update the node states to reference the written bundle
+			int exportIdx = 0;
+			int packetOffset = 0;
+			foreach (BundlePacket packet in bundle.Header.Packets)
 			{
-				entry.SetSize(bundle.Payload.Length);
-				entry.SetValue(bundle);
+				PacketInfo packetInfo = new PacketInfo(packetOffset, packet);
+
+				int nodeOffset = 0;
+				for (; exportIdx < header.Exports.Count && nodeOffset + bundle.Header.Exports[exportIdx].Length <= packet.DecodedLength; exportIdx++)
+				{
+					InMemoryNodeState inMemoryState = (InMemoryNodeState)nodes[exportIdx].State;
+
+					BundleExport export = header.Exports[exportIdx];
+					nodes[exportIdx].SetState(new ExportedNodeState(bundleInfo, exportIdx, packetInfo, nodeOffset, export.Length, inMemoryState.References));
+
+					nodeOffset += header.Exports[exportIdx].Length;
+				}
+
+				packetOffset += packet.EncodedLength;
 			}
 
-			// Remove this object from the list of read tasks
-			lock (_lockObject)
+			// Remove all the newly written nodes from the dirty set
+			lock (context)
 			{
-				_readTasks.Remove(cacheKey);
+				context.WriteNodesSet.ExceptWith(nodes);
 			}
-			return bundle;
 		}
 
-		/// <summary>
-		/// Gets a decoded block from the store
-		/// </summary>
-		/// <param name="bundleInfo">Information about the bundle</param>
-		/// <param name="packetInfo">The decoded block location and size</param>
-		/// <param name="payload">The bundle payload</param>
-		/// <returns>The decoded data</returns>
-		ReadOnlyMemory<byte> DecodePacket(BlobInfo bundleInfo, PacketInfo packetInfo, ReadOnlyMemory<byte> payload)
+		async Task FlushAsync(BundleWriteContext context, ITreeBlob root, CancellationToken cancellationToken = default)
 		{
-			string cacheKey = $"bundle-packet:{bundleInfo.BlobId}@{packetInfo.Offset}";
-			return _cache.GetOrCreate<ReadOnlyMemory<byte>>(cacheKey, entry =>
+			// Create the root node, and add everything it references to the write buffer
+			IoHash rootHash = ComputeHash(root.Data, root.Refs.ConvertAll(x => (NodeInfo)x));
+			NodeInfo rootNode = new NodeInfo(this, rootHash, (InMemoryNodeState)root);
+
+			// Flush all the nodes to bundles, and snapshot the remaining ref
+			Task task;
+			lock (context)
 			{
-				ReadOnlyMemory<byte> decodedPacket = DecodePacketUncached(packetInfo, payload);
-				entry.SetSize(packetInfo.DecodedLength);
-				return decodedPacket;
-			});
+				WriteInMemoryNodes(context, rootNode);
+				NodeInfo[] writeNodes = context.WriteNodes.ToArray();
+				task = SequenceWriteAsync(context, () => FlushInternalAsync(context, writeNodes, cancellationToken), cancellationToken);
+			}
+			await task;
 		}
 
-		static ReadOnlyMemory<byte> DecodePacketUncached(PacketInfo packetInfo, ReadOnlyMemory<byte> payload)
+		async Task FlushInternalAsync(BundleWriteContext context, NodeInfo[] writeNodes, CancellationToken cancellationToken)
 		{
-			ReadOnlyMemory<byte> encodedPacket = payload.Slice(packetInfo.Offset, packetInfo.EncodedLength);
-
-			byte[] decodedPacket = new byte[packetInfo.DecodedLength];
-			LZ4Codec.Decode(encodedPacket.Span, decodedPacket);
-
-			return decodedPacket;
+			Bundle bundle = CreateBundle(context, writeNodes);
+			await _blobStore.WriteBundleRefAsync(context.Name, bundle, cancellationToken);
 		}
 
-		async Task FindLiveNodes(NodeInfo node, Dictionary<BlobInfo, HashSet<NodeInfo>> bundleToLiveNodes, HashSet<NodeInfo> nodes, CancellationToken cancellationToken)
+		void FindNodesToWrite(NodeInfo root, List<NodeInfo> nodes, HashSet<NodeInfo> uniqueNodes)
 		{
-			if (nodes.Add(node))
+			if (root.State is InMemoryNodeState inMemoryState)
 			{
-				if (node.State is ImportedNodeState importedState)
+				if (uniqueNodes.Add(root))
 				{
-					HashSet<NodeInfo>? liveNodes;
-					if (!bundleToLiveNodes.TryGetValue(importedState.BundleInfo, out liveNodes))
+					foreach (NodeInfo reference in inMemoryState.References)
 					{
-						liveNodes = new HashSet<NodeInfo>();
-						bundleToLiveNodes.Add(importedState.BundleInfo, liveNodes);
+						FindNodesToWrite(reference, nodes, uniqueNodes);
 					}
-					liveNodes.Add(node);
+					nodes.Add(root);
 				}
+			}
+		}
 
-				foreach (NodeInfo reference in await GetReferencesAsync(node, cancellationToken))
+		static async Task SequenceWriteAsync(BundleWriteContext context, Func<Task> writeFunc, CancellationToken cancellationToken)
+		{
+			Task task;
+			lock (context)
+			{
+				Task prevTask = context.WriteTask;
+
+				// Wait for the previous write to complete first, ignoring any exceptions
+				Func<Task> wrappedWriteFunc = async () =>
 				{
-					await FindLiveNodes(reference, bundleToLiveNodes, nodes, cancellationToken);
-				}
+					await prevTask.ContinueWith(x => { }, cancellationToken, TaskContinuationOptions.None, TaskScheduler.Default);
+					await writeFunc();
+				};
+				task = Task.Run(wrappedWriteFunc, cancellationToken);
+
+				context.WriteTask = task;
 			}
+			await task;
 		}
 
 		/// <summary>
 		/// Creates a Bundle containing a set of nodes. 
-		///
-		/// WARNING: The <see cref="Bundle.Payload"/> member will be set to the value of <see cref="_encodedBuffer"/> on return, which will be
-		/// reused on a subsequent call. If the object needs to be persisted across the lifetime of other objects, this field must be duplicated.
 		/// </summary>
-		Bundle CreateBundle(IReadOnlyList<NodeInfo> nodes)
+		Bundle CreateBundle(BundleWriteContext context, NodeInfo[] nodes)
 		{
 			// Create a set from the nodes to be written. We use this to determine references that are imported.
 			HashSet<NodeInfo> nodeSet = new HashSet<NodeInfo>(nodes);
 
 			// Find all the imported nodes by bundle
-			Dictionary<BlobInfo, List<NodeInfo>> bundleToImportedNodes = new Dictionary<BlobInfo, List<NodeInfo>>();
+			Dictionary<BundleInfo, List<NodeInfo>> bundleToImportedNodes = new Dictionary<BundleInfo, List<NodeInfo>>();
 			foreach (NodeInfo node in nodes)
 			{
 				InMemoryNodeState state = (InMemoryNodeState)node.State;
@@ -751,7 +745,7 @@ namespace EpicGames.Horde.Storage.Bundles
 					{
 						// Get the persisted node info
 						ImportedNodeState importedState = (ImportedNodeState)reference.State;
-						BlobInfo bundleInfo = importedState.BundleInfo;
+						BundleInfo bundleInfo = importedState.BundleInfo;
 
 						// Get the list of nodes within it
 						List<NodeInfo>? importedNodes;
@@ -770,7 +764,7 @@ namespace EpicGames.Horde.Storage.Bundles
 
 			// Add all the imports and assign them identifiers
 			Dictionary<NodeInfo, int> nodeToIndex = new Dictionary<NodeInfo, int>();
-			foreach ((BlobInfo bundleInfo, List<NodeInfo> importedNodes) in bundleToImportedNodes)
+			foreach ((BundleInfo bundleInfo, List<NodeInfo> importedNodes) in bundleToImportedNodes)
 			{
 				(int, IoHash)[] exportInfos = new (int, IoHash)[importedNodes.Count];
 				for (int idx = 0; idx < importedNodes.Count; idx++)
@@ -785,7 +779,7 @@ namespace EpicGames.Horde.Storage.Bundles
 			}
 
 			// Preallocate data in the encoded buffer to reduce fragmentation if we have to resize
-			CreateFreeSpace(ref _encodedBuffer, 0, _options.MaxBlobSize);
+			byte[] outputBuffer = new byte[_options.MaxBlobSize];
 
 			// Create the export list
 			List<BundleExport> exports = new List<BundleExport>();
@@ -794,8 +788,8 @@ namespace EpicGames.Horde.Storage.Bundles
 			// Size of data currently stored in the block buffer
 			int blockSize = 0;
 
-			// Compress all the nodes into the encoded buffer
-			int packetOffset = 0;
+			// Compress all the nodes into the output buffer buffer
+			int outputOffset = 0;
 			foreach (NodeInfo node in nodes)
 			{
 				InMemoryNodeState nodeState = (InMemoryNodeState)node.State;
@@ -804,7 +798,7 @@ namespace EpicGames.Horde.Storage.Bundles
 				// If we can't fit this data into the current block, flush the contents of it first
 				if (blockSize > 0 && blockSize + nodeData.Length > _options.MinCompressionPacketSize)
 				{
-					FlushPacket(_blockBuffer.AsMemory(0, blockSize), ref packetOffset, packets);
+					FlushPacket(context.PacketBuffer.AsMemory(0, blockSize), ref outputBuffer, ref outputOffset, packets);
 					blockSize = 0;
 				}
 
@@ -818,41 +812,41 @@ namespace EpicGames.Horde.Storage.Bundles
 				if (nodeData.Length < _options.MinCompressionPacketSize || !nodeData.IsSingleSegment)
 				{
 					int requiredSize = Math.Max(blockSize + (int)nodeData.Length, (int)(_options.MaxBlobSize * 1.2));
-					CreateFreeSpace(ref _blockBuffer, blockSize, requiredSize);
+					CreateFreeSpace(ref context.PacketBuffer, blockSize, requiredSize);
 
 					foreach (ReadOnlyMemory<byte> nodeSegment in nodeData)
 					{
-						nodeSegment.CopyTo(_blockBuffer.AsMemory(blockSize));
+						nodeSegment.CopyTo(context.PacketBuffer.AsMemory(blockSize));
 						blockSize += nodeSegment.Length;
 					}
 				}
 				else
 				{
-					FlushPacket(nodeData.First, ref packetOffset, packets);
+					FlushPacket(nodeData.First, ref outputBuffer, ref outputOffset, packets);
 				}
 			}
-			FlushPacket(_blockBuffer.AsMemory(0, blockSize), ref packetOffset, packets);
+			FlushPacket(context.PacketBuffer.AsMemory(0, blockSize), ref outputBuffer, ref outputOffset, packets);
 
 			// Flush the data
 			BundleHeader header = new BundleHeader(imports, exports, packets);
-			ReadOnlyMemory<byte> payload = _encodedBuffer.AsMemory(0, packets.Sum(x => x.EncodedLength));
+			ReadOnlyMemory<byte> payload = outputBuffer.AsMemory(0, packets.Sum(x => x.EncodedLength));
 			return new Bundle(header, payload);
 		}
 
-		void FlushPacket(ReadOnlyMemory<byte> inputData, ref int packetOffset, List<BundlePacket> packets)
+		static void FlushPacket(ReadOnlyMemory<byte> inputData, ref byte[] outputBuffer, ref int outputOffset, List<BundlePacket> packets)
 		{
 			if (inputData.Length > 0)
 			{
 				int minFreeSpace = LZ4Codec.MaximumOutputSize(inputData.Length);
-				CreateFreeSpace(ref _encodedBuffer, packetOffset, packetOffset + minFreeSpace);
+				CreateFreeSpace(ref outputBuffer, outputOffset, outputOffset + minFreeSpace);
 
 				ReadOnlySpan<byte> inputSpan = inputData.Span;
-				Span<byte> outputSpan = _encodedBuffer.AsSpan(packetOffset);
+				Span<byte> outputSpan = outputBuffer.AsSpan(outputOffset);
 
 				int encodedLength = LZ4Codec.Encode(inputSpan, outputSpan);
 				packets.Add(new BundlePacket(encodedLength, inputData.Length));
 
-				packetOffset += encodedLength;
+				outputOffset += encodedLength;
 
 				Debug.Assert(encodedLength >= 0);
 			}
@@ -867,5 +861,7 @@ namespace EpicGames.Horde.Storage.Bundles
 				buffer = newBuffer;
 			}
 		}
+
+		#endregion
 	}
 }
diff --git a/Engine/Source/Programs/Shared/EpicGames.Horde/Storage/ITreeStore.cs b/Engine/Source/Programs/Shared/EpicGames.Horde/Storage/ITreeStore.cs
--- a/Engine/Source/Programs/Shared/EpicGames.Horde/Storage/ITreeStore.cs
+++ b/Engine/Source/Programs/Shared/EpicGames.Horde/Storage/ITreeStore.cs
@@ -1,47 +1,40 @@
-// Copyright Epic Games, Inc. All Rights Reserved.
+﻿// Copyright Epic Games, Inc. All Rights Reserved.
 
 using System;
 using System.Buffers;
 using System.Collections.Generic;
 using System.Threading;
 using System.Threading.Tasks;
-using EpicGames.Core;
 
 namespace EpicGames.Horde.Storage
 {
 	/// <summary>
-	/// Interface for a blob of data for a tree node
+	/// Interface for a node within a tree, represented by an opaque blob of data
 	/// </summary>
 	public interface ITreeBlob
 	{
 		/// <summary>
-		/// Gets the data for a node with the given hash
+		/// Data for the node
 		/// </summary>
-		/// <param name="cancellationToken">Cancellation token for the operation</param>
-		/// <returns>Data for the blob</returns>
-		ValueTask<ReadOnlySequence<byte>> GetDataAsync(CancellationToken cancellationToken = default);
+		ReadOnlySequence<byte> Data { get; }
 
 		/// <summary>
-		/// Find the outward references for a node
+		/// References to other tree blobs
 		/// </summary>
-		/// <param name="cancellationToken">Cancellation token for the operation</param>
-		/// <returns>List of outward references from the blob</returns>
-		ValueTask<IReadOnlyList<ITreeBlob>> GetReferencesAsync(CancellationToken cancellationToken = default);
+		IReadOnlyList<ITreeBlobRef> Refs { get; }
 	}
 
 	/// <summary>
-	/// Writer for trees
+	/// Reference to a <see cref="ITreeBlob"/>
 	/// </summary>
-	public interface ITreeBlobWriter
+	public interface ITreeBlobRef
 	{
 		/// <summary>
-		/// Write a blob to storage
+		/// Gets the target of a reference
 		/// </summary>
-		/// <param name="data"></param>
-		/// <param name="references"></param>
 		/// <param name="cancellationToken">Cancellation token for the operation</param>
-		/// <returns></returns>
-		public Task<ITreeBlob> WriteBlobAsync(ReadOnlySequence<byte> data, IReadOnlyList<ITreeBlob> references, CancellationToken cancellationToken);
+		/// <returns>The referenced node</returns>
+		ValueTask<ITreeBlob> GetTargetAsync(CancellationToken cancellationToken = default);
 	}
 
 	/// <summary>
@@ -50,12 +43,10 @@ namespace EpicGames.Horde.Storage
 	public interface ITreeStore : IDisposable
 	{
 		/// <summary>
-		/// Tests whether a tree exists with the given name
+		/// Creates a new context for reading trees.
 		/// </summary>
-		/// <param name="name">Name of the ref used to store the tree</param>
-		/// <param name="cancellationToken">Cancellation token for the operation</param>
-		/// <returns>Tree reader instance</returns>
-		Task<bool> HasTreeAsync(RefName name, CancellationToken cancellationToken = default);
+		/// <returns>New context instance</returns>
+		ITreeWriter CreateTreeWriter(RefName name);
 
 		/// <summary>
 		/// Deletes a tree with the given name
@@ -65,40 +56,20 @@ namespace EpicGames.Horde.Storage
 		Task DeleteTreeAsync(RefName name, CancellationToken cancellationToken = default);
 
 		/// <summary>
-		/// Reads the root blob from a tree in storage
-		/// </summary>
-		/// <param name="name">Identifier for the tree</param>
-		/// <param name="cancellationToken">Cancellation token for the operation</param>
-		/// <returns></returns>
-		Task<ITreeBlob?> TryReadTreeAsync(RefName name, CancellationToken cancellationToken = default);
-
-		/// <summary>
-		/// Reads a root node from the store
-		/// </summary>
-		/// <param name="name">Name of the tree to fetch</param>
-		/// <param name="cancellationToken">Cancellation token for the operation</param>
-		/// <returns>Deserialized node stored with this name</returns>
-		Task<T?> TryReadTreeAsync<T>(RefName name, CancellationToken cancellationToken = default) where T : TreeNode;
-
-		/// <summary>
-		/// Creates a writer for a tree with the given name
+		/// Tests whether a tree exists with the given name
 		/// </summary>
 		/// <param name="name">Name of the ref used to store the tree</param>
-		/// <param name="root">Node to flush</param>
-		/// <param name="flush">Whether to flush the complete tree state. If false, an implementation may choose to buffer some writes to allow packing multiple nodes together.</param>
 		/// <param name="cancellationToken">Cancellation token for the operation</param>
-		/// <returns>Tree writer instance</returns>
-		Task WriteTreeAsync(RefName name, ITreeBlob root, bool flush = true, CancellationToken cancellationToken = default);
+		/// <returns>Tree reader instance</returns>
+		Task<bool> HasTreeAsync(RefName name, CancellationToken cancellationToken = default);
 
 		/// <summary>
-		/// Creates a writer for a tree with the given name
+		/// Attempts to read a tree with the given name
 		/// </summary>
 		/// <param name="name">Name of the ref used to store the tree</param>
-		/// <param name="root">Node to flush</param>
-		/// <param name="flush">Whether to flush the complete tree state. If false, an implementation may choose to buffer some writes to allow packing multiple nodes together.</param>
 		/// <param name="cancellationToken">Cancellation token for the operation</param>
-		/// <returns>Tree writer instance</returns>
-		Task WriteTreeAsync<T>(RefName name, T root, bool flush = true, CancellationToken cancellationToken = default) where T : TreeNode;
+		/// <returns>Reference to the root of the tree. Null if the ref does not exist.</returns>
+		Task<ITreeBlob?> TryReadTreeAsync(RefName name, CancellationToken cancellationToken = default);
 	}
 
 	/// <summary>
@@ -108,6 +79,28 @@ namespace EpicGames.Horde.Storage
 	{
 	}
 
+	/// <summary>
+	/// Allows incrementally writing new trees
+	/// </summary>
+	public interface ITreeWriter
+	{
+		/// <summary>
+		/// Adds a new node to this tree
+		/// </summary>
+		/// <param name="data">Data for the node</param>
+		/// <param name="refs">References to other nodes</param>
+		/// <param name="cancellationToken">Cancellation token for the operation</param>
+		/// <returns>Reference to the node that was added</returns>
+		Task<ITreeBlobRef> WriteNodeAsync(ReadOnlySequence<byte> data, IReadOnlyList<ITreeBlobRef> refs, CancellationToken cancellationToken = default);
+
+		/// <summary>
+		/// Flush the tree with the given root.
+		/// </summary>
+		/// <param name="root">The root node</param>
+		/// <param name="cancellationToken">Cancellation token for the operation</param>
+		Task FlushAsync(ITreeBlob root, CancellationToken cancellationToken = default);
+	}
+
 	/// <summary>
 	/// Extension methods for tree store instances
 	/// </summary>
@@ -122,6 +115,9 @@ namespace EpicGames.Horde.Storage
 				_inner = inner;
 			}
 
+			/// <inheritdoc/>
+			public ITreeWriter CreateTreeWriter(RefName name) => _inner.CreateTreeWriter(name);
+
 			/// <inheritdoc/>
 			public Task DeleteTreeAsync(RefName name, CancellationToken cancellationToken = default) => _inner.DeleteTreeAsync(name, cancellationToken);
 
@@ -133,15 +129,6 @@ namespace EpicGames.Horde.Storage
 
 			/// <inheritdoc/>
 			public Task<ITreeBlob?> TryReadTreeAsync(RefName name, CancellationToken cancellationToken = default) => _inner.TryReadTreeAsync(name, cancellationToken);
-
-			/// <inheritdoc/>
-			public Task<TNode?> TryReadTreeAsync<TNode>(RefName name, CancellationToken cancellationToken = default) where TNode : TreeNode => _inner.TryReadTreeAsync<TNode>(name, cancellationToken);
-
-			/// <inheritdoc/>
-			public Task WriteTreeAsync(RefName name, ITreeBlob root, bool flush = true, CancellationToken cancellationToken = default) => _inner.WriteTreeAsync(name, root, flush, cancellationToken);
-
-			/// <inheritdoc/>
-			public Task WriteTreeAsync<TNode>(RefName name, TNode root, bool flush = true, CancellationToken cancellationToken = default) where TNode : TreeNode => _inner.WriteTreeAsync<TNode>(name, root, flush, cancellationToken);
 		}
 
 		/// <summary>
@@ -154,39 +141,5 @@ namespace EpicGames.Horde.Storage
 		{
 			return new TypedTreeStore<T>(store);
 		}
-
-		/// <summary>
-		/// Reads a tree, throwing an exception if it doesn't exist
-		/// </summary>
-		/// <param name="store">Store to read from</param>
-		/// <param name="name">Name of the ref</param>
-		/// <param name="cancellationToken">Cancellation token for the operation</param>
-		/// <returns></returns>
-		public static async Task<ITreeBlob> ReadTreeAsync(this ITreeStore store, RefName name, CancellationToken cancellationToken = default)
-		{
-			ITreeBlob? blob = await store.TryReadTreeAsync(name, cancellationToken);
-			if (blob == null)
-			{
-				throw new RefNameNotFoundException(name);
-			}
-			return blob;
-		}
-
-		/// <summary>
-		/// Reads a tree, throwing an exception if it doesn't exist
-		/// </summary>
-		/// <param name="store">Store to read from</param>
-		/// <param name="name">Name of the ref</param>
-		/// <param name="cancellationToken">Cancellation token for the operation</param>
-		/// <returns></returns>
-		public static async Task<T> ReadTreeAsync<T>(this ITreeStore store, RefName name, CancellationToken cancellationToken = default) where T : TreeNode
-		{
-			T? node = await store.TryReadTreeAsync<T>(name, cancellationToken);
-			if (node == null)
-			{
-				throw new RefNameNotFoundException(name);
-			}
-			return node;
-		}
 	}
 }
diff --git a/Engine/Source/Programs/Shared/EpicGames.Horde/Storage/Nodes/DirectoryNode.cs b/Engine/Source/Programs/Shared/EpicGames.Horde/Storage/Nodes/DirectoryNode.cs
--- a/Engine/Source/Programs/Shared/EpicGames.Horde/Storage/Nodes/DirectoryNode.cs
+++ b/Engine/Source/Programs/Shared/EpicGames.Horde/Storage/Nodes/DirectoryNode.cs
@@ -90,7 +90,7 @@ namespace EpicGames.Horde.Storage.Nodes
 		/// <summary>
 		/// Constructor
 		/// </summary>
-		public FileEntry(DirectoryNode owner, Utf8String name, FileEntryFlags flags, long length, ITreeBlob target)
+		public FileEntry(DirectoryNode owner, Utf8String name, FileEntryFlags flags, long length, ITreeBlobRef target)
 			: base(owner, target)
 		{
 			Name = name;
@@ -104,11 +104,12 @@ namespace EpicGames.Horde.Storage.Nodes
 		/// </summary>
 		/// <param name="data">Data to append to the file</param>
 		/// <param name="options">Options for chunking the data</param>
+		/// <param name="writer">Writer for new node data</param>
 		/// <param name="cancellationToken">Cancellation token for the operation</param>
-		public async Task AppendAsync(ReadOnlyMemory<byte> data, ChunkingOptions options, CancellationToken cancellationToken)
+		public async Task AppendAsync(ReadOnlyMemory<byte> data, ChunkingOptions options, ITreeWriter writer, CancellationToken cancellationToken)
 		{
 			FileNode node = await ExpandAsync(cancellationToken);
-			Node = node.Append(data, options);
+			Node = await node.AppendAsync(data, options, writer, cancellationToken);
 			MarkAsDirty();
 		}
 
@@ -117,8 +118,9 @@ namespace EpicGames.Horde.Storage.Nodes
 		/// </summary>
 		/// <param name="stream">Data to append to the file</param>
 		/// <param name="options">Options for chunking the data</param>
+		/// <param name="writer">Writer for new node data</param>
 		/// <param name="cancellationToken">Cancellation token for the operation</param>
-		public async Task AppendAsync(Stream stream, ChunkingOptions options, CancellationToken cancellationToken)
+		public async Task AppendAsync(Stream stream, ChunkingOptions options, ITreeWriter writer, CancellationToken cancellationToken)
 		{
 			byte[] buffer = new byte[4 * 1024];
 			for (; ; )
@@ -128,7 +130,7 @@ namespace EpicGames.Horde.Storage.Nodes
 				{
 					break;
 				}
-				await AppendAsync(buffer.AsMemory(0, numBytes), options, cancellationToken);
+				await AppendAsync(buffer.AsMemory(0, numBytes), options, writer, cancellationToken);
 			}
 		}
 
@@ -174,7 +176,7 @@ namespace EpicGames.Horde.Storage.Nodes
 		/// <summary>
 		/// Constructor
 		/// </summary>
-		public DirectoryEntry(DirectoryNode owner, Utf8String name, long length, ITreeBlob target)
+		public DirectoryEntry(DirectoryNode owner, Utf8String name, long length, ITreeBlobRef target)
 			: base(owner, target)
 		{
 			Name = name;
@@ -195,7 +197,7 @@ namespace EpicGames.Horde.Storage.Nodes
 	/// <summary>
 	/// A directory node
 	/// </summary>
-	[TreeSerializer(typeof(DirectoryNodeDeserializer))]
+	[TreeSerializer(typeof(DirectoryNodeSerializer))]
 	public class DirectoryNode : TreeNode
 	{
 		internal const byte TypeId = (byte)'d';
@@ -275,12 +277,13 @@ namespace EpicGames.Horde.Storage.Nodes
 		/// <param name="flags">Flags for the new file</param>
 		/// <param name="stream">The stream to read from</param>
 		/// <param name="options">Options for chunking the data</param>
+		/// <param name="writer">Writer for new node data</param>
 		/// <param name="cancellationToken">Cancellation token for the operation</param>
 		/// <returns>The new directory object</returns>
-		public async Task<FileEntry> AddFileAsync(Utf8String name, FileEntryFlags flags, Stream stream, ChunkingOptions options, CancellationToken cancellationToken = default)
+		public async Task<FileEntry> AddFileAsync(Utf8String name, FileEntryFlags flags, Stream stream, ChunkingOptions options, ITreeWriter writer, CancellationToken cancellationToken = default)
 		{
 			FileEntry entry = AddFile(name, flags);
-			await entry.AppendAsync(stream, options, cancellationToken);
+			await entry.AppendAsync(stream, options, writer, cancellationToken);
 			return entry;
 		}
 
@@ -289,9 +292,10 @@ namespace EpicGames.Horde.Storage.Nodes
 		/// </summary>
 		/// <param name="path">Path to the file</param>
 		/// <param name="flags">Flags for the new file</param>
+		/// <param name="writer">Writer for new node data</param>
 		/// <param name="cancellationToken">Cancellation token for the operation</param>
 		/// <returns>The new directory object</returns>
-		public async ValueTask<FileEntry> AddFileByPathAsync(Utf8String path, FileEntryFlags flags, CancellationToken cancellationToken = default)
+		public async ValueTask<FileEntry> AddFileByPathAsync(Utf8String path, FileEntryFlags flags, ITreeWriter writer, CancellationToken cancellationToken = default)
 		{
 			DirectoryNode directory = this;
 
@@ -473,24 +477,23 @@ namespace EpicGames.Horde.Storage.Nodes
 		#endregion
 
 		/// <inheritdoc/>
-		public async ValueTask<ITreeBlob> SerializeAsync(ITreeBlobWriter writer, CancellationToken cancellationToken)
+		public override NewTreeBlob Serialize()
 		{
-			List<ITreeBlob> childNodes = new List<ITreeBlob>();
+			List<TreeNodeRef> refs = new List<TreeNodeRef>();
 
 			List<FileEntry> fileEntries = _nameToFileEntry.Values.OrderBy(x => x.Name).ToList();
 			foreach (FileEntry fileEntry in fileEntries)
 			{
-				childNodes.Add(await fileEntry.CollapseAsync(writer, cancellationToken));
+				refs.Add(fileEntry);
 			}
 
 			List<DirectoryEntry> directoryEntries = _nameToDirectoryEntry.Values.OrderBy(x => x.Name).ToList();
 			foreach (DirectoryEntry directoryEntry in directoryEntries)
 			{
-				childNodes.Add(await directoryEntry.CollapseAsync(writer, cancellationToken));
+				refs.Add(directoryEntry);
 			}
 
-			ReadOnlySequence<byte> data = SerializeData(fileEntries, directoryEntries);
-			return await writer.WriteBlobAsync(data, childNodes, cancellationToken);
+			return new NewTreeBlob(SerializeData(fileEntries, directoryEntries), refs);
 		}
 
 		static ReadOnlySequence<byte> SerializeData(List<FileEntry> fileEntries, List<DirectoryEntry> directoryEntries)
@@ -549,12 +552,11 @@ namespace EpicGames.Horde.Storage.Nodes
 		/// <summary>
 		/// Deserialize a directory node from data
 		/// </summary>
-		/// <param name="payload">The data to read from</param>
-		/// <param name="children">Child nodes</param>
+		/// <param name="source">The source node to deserialize</param>
 		/// <returns>The deserialized directory node</returns>
-		public static DirectoryNode Deserialize(ReadOnlySequence<byte> payload, IReadOnlyList<ITreeBlob> children)
+		public static DirectoryNode Deserialize(ITreeBlob source)
 		{
-			ReadOnlySpan<byte> span = payload.AsSingleSegment().Span;
+			ReadOnlySpan<byte> span = source.Data.AsSingleSegment().Span;
 			if(span[0] != TypeId)
 			{
 				throw new InvalidOperationException("Invalid signature byte for directory");
@@ -582,7 +584,8 @@ namespace EpicGames.Horde.Storage.Nodes
 				long length = (long)VarInt.ReadUnsigned(span, out int lengthBytes);
 				span = span.Slice(lengthBytes);
 
-				node._nameToFileEntry[name] = new FileEntry(node, name, flags, length, children[childIdx++]);
+				ITreeBlobRef child = source.Refs[childIdx++];
+				node._nameToFileEntry[name] = new FileEntry(node, name, flags, length, child);
 			}
 
 			int directoryCount = (int)VarInt.ReadUnsigned(span, out int directoryCountBytes);
@@ -598,7 +601,8 @@ namespace EpicGames.Horde.Storage.Nodes
 				long length = (long)VarInt.ReadUnsigned(span, out int lengthBytes);
 				span = span.Slice(lengthBytes);
 
-				node._nameToDirectoryEntry[name] = new DirectoryEntry(node, name, length, children[childIdx++]);
+				ITreeBlobRef child = source.Refs[childIdx++];
+				node._nameToDirectoryEntry[name] = new DirectoryEntry(node, name, length, child);
 			}
 
 			Debug.Assert(span.Length == 0);
@@ -610,22 +614,23 @@ namespace EpicGames.Horde.Storage.Nodes
 		/// </summary>
 		/// <param name="directoryInfo"></param>
 		/// <param name="options">Options for chunking file content</param>
+		/// <param name="writer">Writer for new node data</param>
 		/// <param name="logger"></param>
 		/// <param name="cancellationToken">Cancellation token for the operation</param>
 		/// <returns></returns>
-		public async Task CopyFromDirectoryAsync(DirectoryInfo directoryInfo, ChunkingOptions options, ILogger logger, CancellationToken cancellationToken)
+		public async Task CopyFromDirectoryAsync(DirectoryInfo directoryInfo, ChunkingOptions options, ITreeWriter writer, ILogger logger, CancellationToken cancellationToken)
 		{
 			foreach (DirectoryInfo subDirectoryInfo in directoryInfo.EnumerateDirectories())
 			{
 				logger.LogInformation("Adding {Directory}", subDirectoryInfo.FullName);
 				DirectoryNode subDirectoryNode = AddDirectory(subDirectoryInfo.Name);
-				await subDirectoryNode.CopyFromDirectoryAsync(subDirectoryInfo, options, logger, cancellationToken);
+				await subDirectoryNode.CopyFromDirectoryAsync(subDirectoryInfo, options, writer, logger, cancellationToken);
 			}
 			foreach (FileInfo fileInfo in directoryInfo.EnumerateFiles())
 			{
 				logger.LogInformation("Adding {File}", fileInfo.FullName);
 				using Stream stream = fileInfo.OpenRead();
-				await AddFileAsync(fileInfo.Name, 0, stream, options, cancellationToken);
+				await AddFileAsync(fileInfo.Name, 0, stream, options, writer, cancellationToken);
 			}
 		}
 
@@ -663,13 +668,10 @@ namespace EpicGames.Horde.Storage.Nodes
 	/// <summary>
 	/// Factory for creating and serializing <see cref="DirectoryNode"/> objects
 	/// </summary>
-	public class DirectoryNodeDeserializer : TreeNodeSerializer<DirectoryNode>
+	public class DirectoryNodeSerializer : TreeNodeSerializer<DirectoryNode>
 	{
 		/// <inheritdoc/>
-		public override async ValueTask<DirectoryNode> DeserializeAsync(ITreeBlob blob, CancellationToken cancellationToken) => DirectoryNode.Deserialize(await blob.GetDataAsync(cancellationToken), await blob.GetReferencesAsync(cancellationToken));
-
-		/// <inheritdoc/>
-		public override ValueTask<ITreeBlob> SerializeAsync(ITreeBlobWriter writer, DirectoryNode node, CancellationToken cancellationToken) => node.SerializeAsync(writer, cancellationToken);
+		public override DirectoryNode Deserialize(ITreeBlob blob) => DirectoryNode.Deserialize(blob);
 	}
 
 	/// <summary>
diff --git a/Engine/Source/Programs/Shared/EpicGames.Horde/Storage/Nodes/FileNode.cs b/Engine/Source/Programs/Shared/EpicGames.Horde/Storage/Nodes/FileNode.cs
--- a/Engine/Source/Programs/Shared/EpicGames.Horde/Storage/Nodes/FileNode.cs
+++ b/Engine/Source/Programs/Shared/EpicGames.Horde/Storage/Nodes/FileNode.cs
@@ -87,22 +87,15 @@ namespace EpicGames.Horde.Storage.Nodes
 		/// </summary>
 		public abstract uint RollingHash { get; }
 
-		/// <summary>
-		/// Serailize this node
-		/// </summary>
-		/// <param name="writer"></param>
-		/// <param name="cancellationToken"></param>
-		/// <returns></returns>
-		public abstract ValueTask<ITreeBlob> SerializeAsync(ITreeBlobWriter writer, CancellationToken cancellationToken);
-
 		/// <summary>
 		/// Creates a file node from a stream
 		/// </summary>
 		/// <param name="stream">The stream to read from</param>
 		/// <param name="options">Options for chunking the data</param>
+		/// <param name="writer">Writer for new tree nodes</param>
 		/// <param name="cancellationToken">Cancellation token for the operation</param>
 		/// <returns>New file node</returns>
-		public static async Task<FileNode> CreateAsync(Stream stream, ChunkingOptions options, CancellationToken cancellationToken)
+		public static async Task<FileNode> CreateAsync(Stream stream, ChunkingOptions options, ITreeWriter writer, CancellationToken cancellationToken)
 		{
 			FileNode node = new LeafFileNode();
 
@@ -114,7 +107,7 @@ namespace EpicGames.Horde.Storage.Nodes
 				{
 					break;
 				}
-				node = node.Append(buffer.AsMemory(0, numBytes), options);
+				node = await node.AppendAsync(buffer.AsMemory(0, numBytes), options, writer, cancellationToken);
 			}
 
 			return node;
@@ -125,17 +118,19 @@ namespace EpicGames.Horde.Storage.Nodes
 		/// </summary>
 		/// <param name="input">The data to write</param>
 		/// <param name="options">Settings for chunking the data</param>
-		public FileNode Append(ReadOnlyMemory<byte> input, ChunkingOptions options)
+		/// <param name="writer">Writer for new tree nodes</param>
+		/// <param name="cancellationToken">Cancellation token for the operation</param>
+		public ValueTask<FileNode> AppendAsync(ReadOnlyMemory<byte> input, ChunkingOptions options, ITreeWriter writer, CancellationToken cancellationToken)
 		{
-			return Append(this, input, options);
+			return AppendAsync(this, input, options, writer, cancellationToken);
 		}
 
-		static FileNode Append(FileNode root, ReadOnlyMemory<byte> input, ChunkingOptions options)
+		static async ValueTask<FileNode> AppendAsync(FileNode root, ReadOnlyMemory<byte> input, ChunkingOptions options, ITreeWriter writer, CancellationToken cancellationToken)
 		{
 			for (; ; )
 			{
 				// Append as much data as possible to the existing tree
-				input = root.AppendToNode(input, options);
+				input = await root.AppendToNodeAsync(input, options, writer, cancellationToken);
 				if (input.IsEmpty)
 				{
 					break;
@@ -147,7 +142,7 @@ namespace EpicGames.Horde.Storage.Nodes
 			return root;
 		}
 
-		private ReadOnlyMemory<byte> AppendToNode(ReadOnlyMemory<byte> appendData, ChunkingOptions options)
+		private async Task<ReadOnlyMemory<byte>> AppendToNodeAsync(ReadOnlyMemory<byte> appendData, ChunkingOptions options, ITreeWriter writer, CancellationToken cancellationToken)
 		{
 			if (appendData.Length == 0 || IsReadOnly())
 			{
@@ -155,7 +150,7 @@ namespace EpicGames.Horde.Storage.Nodes
 			}
 			else
 			{
-				return AppendData(appendData, options);
+				return await AppendDataAsync(appendData, options, writer, cancellationToken);
 			}
 		}
 
@@ -164,8 +159,10 @@ namespace EpicGames.Horde.Storage.Nodes
 		/// </summary>
 		/// <param name="newData">The data to append</param>
 		/// <param name="options">Options for chunking the data</param>
+		/// <param name="writer">Writer for new tree nodes</param>
+		/// <param name="cancellationToken">Cancellation token for the operation</param>
 		/// <returns>Remaining data in the buffer</returns>
-		public abstract ReadOnlyMemory<byte> AppendData(ReadOnlyMemory<byte> newData, ChunkingOptions options);
+		public abstract ValueTask<ReadOnlyMemory<byte>> AppendDataAsync(ReadOnlyMemory<byte> newData, ChunkingOptions options, ITreeWriter writer, CancellationToken cancellationToken);
 
 		/// <summary>
 		/// Copies the contents of this node and its children to the given output stream
@@ -296,13 +293,15 @@ namespace EpicGames.Horde.Storage.Nodes
 		public override IReadOnlyList<TreeNodeRef> GetReferences() => Array.Empty<TreeNodeRef>();
 
 		/// <inheritdoc/>
-		public override async ValueTask<ITreeBlob> SerializeAsync(ITreeBlobWriter writer, CancellationToken cancellationToken)
+		public override NewTreeBlob Serialize() => new NewTreeBlob(_writtenSequence, Array.Empty<TreeNodeRef>());
+
+		/// <inheritdoc/>
+		public override ValueTask<ReadOnlyMemory<byte>> AppendDataAsync(ReadOnlyMemory<byte> newData, ChunkingOptions options, ITreeWriter writer, CancellationToken cancellationToken)
 		{
-			return await writer.WriteBlobAsync(_writtenSequence, Array.Empty<ITreeBlob>(), cancellationToken);
+			return new ValueTask<ReadOnlyMemory<byte>>(AppendData(newData, options));
 		}
 
-		/// <inheritdoc/>
-		public override ReadOnlyMemory<byte> AppendData(ReadOnlyMemory<byte> newData, ChunkingOptions options)
+		ReadOnlyMemory<byte> AppendData(ReadOnlyMemory<byte> newData, ChunkingOptions options)
 		{
 			if (_isReadOnly)
 			{
@@ -458,7 +457,7 @@ namespace EpicGames.Horde.Storage.Nodes
 		/// </summary>
 		/// <param name="data"></param>
 		/// <param name="children"></param>
-		public InteriorFileNode(ReadOnlySequence<byte> data, IReadOnlyList<ITreeBlob> children)
+		public InteriorFileNode(ReadOnlySequence<byte> data, IReadOnlyList<ITreeBlobRef> children)
 		{
 			ReadOnlySpan<byte> span = data.AsSingleSegment().Span;
 			if (span[0] != TypeId)
@@ -485,14 +484,8 @@ namespace EpicGames.Horde.Storage.Nodes
 		public override IReadOnlyList<TreeNodeRef> GetReferences() => _children;
 
 		/// <inheritdoc/>
-		public override async ValueTask<ITreeBlob> SerializeAsync(ITreeBlobWriter writer, CancellationToken cancellationToken)
+		public override NewTreeBlob Serialize()
 		{
-			List<ITreeBlob> children = new List<ITreeBlob>();
-			foreach (TreeNodeRef<FileNode> child in _children)
-			{
-				children.Add(await child.CollapseAsync(writer, cancellationToken));
-			}
-
 			byte[] data = new byte[1 + sizeof(uint) + VarInt.MeasureUnsigned((ulong)_length) + 1];
 			data[0] = TypeId;
 
@@ -508,13 +501,12 @@ namespace EpicGames.Horde.Storage.Nodes
 			pos++;
 
 			Debug.Assert(pos == data.Length);
-			return await writer.WriteBlobAsync(new ReadOnlySequence<byte>(data), children, cancellationToken);
+			return new NewTreeBlob(data, _children);
 		}
 
 		/// <inheritdoc/>
-		public override ReadOnlyMemory<byte> AppendData(ReadOnlyMemory<byte> newData, ChunkingOptions options)
+		public override async ValueTask<ReadOnlyMemory<byte>> AppendDataAsync(ReadOnlyMemory<byte> newData, ChunkingOptions options, ITreeWriter writer, CancellationToken cancellationToken)
 		{
-			Span<byte> hashData = stackalloc byte[4];
 			for (; ; )
 			{
 				Debug.Assert(_children != null);
@@ -527,15 +519,14 @@ namespace EpicGames.Horde.Storage.Nodes
 					{
 						// Update the length to match the new node
 						_length -= lastNode.Length;
-						newData = lastNode.AppendData(newData, options);
+						newData = await lastNode.AppendDataAsync(newData, options, writer, cancellationToken);
 						_length += lastNode.Length;
 
 						// If the last node is complete, write it to the buffer
 						if (lastNode.IsReadOnly())
 						{
 							// Update the hash
-							BinaryPrimitives.WriteUInt32LittleEndian(hashData, lastNode.RollingHash);
-							_rollingHash = BuzHash.Add(_rollingHash, hashData);
+							AppendChildHash(lastNode.RollingHash);
 
 							// Check if it's time to finish this chunk
 							uint hashThreshold = (uint)(((1L << 32) * IoHash.NumBytes) / options.LeafOptions.TargetSize);
@@ -551,6 +542,9 @@ namespace EpicGames.Horde.Storage.Nodes
 						{
 							return newData;
 						}
+
+						// Collapse the final node
+						await Children[^1].CollapseAsync(writer, cancellationToken);
 					}
 				}
 
@@ -559,6 +553,17 @@ namespace EpicGames.Horde.Storage.Nodes
 			}
 		}
 
+		/// <summary>
+		/// Updates the rolling hash to append a child hash
+		/// </summary>
+		/// <param name="childHash">The child hash to append</param>
+		void AppendChildHash(uint childHash)
+		{
+			Span<byte> hashData = stackalloc byte[4];
+			BinaryPrimitives.WriteUInt32LittleEndian(hashData, childHash);
+			_rollingHash = BuzHash.Add(_rollingHash, hashData);
+		}
+
 		/// <inheritdoc/>
 		public override async Task CopyToStreamAsync(Stream outputStream, CancellationToken cancellationToken)
 		{
@@ -576,21 +581,18 @@ namespace EpicGames.Horde.Storage.Nodes
 	public class FileNodeSerializer : TreeNodeSerializer<FileNode>
 	{
 		/// <inheritdoc/>
-		public override async ValueTask<FileNode> DeserializeAsync(ITreeBlob blob, CancellationToken cancellationToken)
+		public override FileNode Deserialize(ITreeBlob blob)
 		{
-			ReadOnlySequence<byte> data = await blob.GetDataAsync(cancellationToken);
+			ReadOnlySequence<byte> data = blob.Data;
 			switch (data.FirstSpan[0])
 			{
 				case LeafFileNode.TypeId:
 					return new LeafFileNode(data);
 				case InteriorFileNode.TypeId:
-					return new InteriorFileNode(data, await blob.GetReferencesAsync(cancellationToken));
+					return new InteriorFileNode(data, blob.Refs);
 				default:
 					throw new InvalidDataException("Unknown type id while deserializing file node");
 			};
 		}
-
-		/// <inheritdoc/>
-		public override ValueTask<ITreeBlob> SerializeAsync(ITreeBlobWriter writer, FileNode node, CancellationToken cancellationToken) => node.SerializeAsync(writer, cancellationToken);
 	}
 }
diff --git a/Engine/Source/Programs/Shared/EpicGames.Horde/Storage/TreeNode.cs b/Engine/Source/Programs/Shared/EpicGames.Horde/Storage/TreeNode.cs
--- a/Engine/Source/Programs/Shared/EpicGames.Horde/Storage/TreeNode.cs
+++ b/Engine/Source/Programs/Shared/EpicGames.Horde/Storage/TreeNode.cs
@@ -1,11 +1,14 @@
-// Copyright Epic Games, Inc. All Rights Reserved.
+﻿// Copyright Epic Games, Inc. All Rights Reserved.
 
 using System;
 using System.Buffers;
 using System.Collections.Generic;
+using System.Diagnostics;
 using System.Reflection;
 using System.Threading;
 using System.Threading.Tasks;
+using System.Xml.Linq;
+using EpicGames.Serialization;
 
 namespace EpicGames.Horde.Storage
 {
@@ -46,28 +49,58 @@ namespace EpicGames.Horde.Storage
 		}
 
 		/// <summary>
-		/// Serialize a node to data
+		/// Serialize a node to a block of memory
+		/// </summary>
+		/// <returns>New data to be stored into a blob</returns>
+		public abstract NewTreeBlob Serialize();
+
+		/// <summary>
+		/// Deserialize a node from data
 		/// </summary>
 		/// <typeparam name="T"></typeparam>
-		/// <param name="node"></param>
-		/// <param name="writer"></param>
-		/// <param name="cancellationToken"></param>
+		/// <param name="blob">Blob to deserialize</param>
 		/// <returns></returns>
-		public static ValueTask<ITreeBlob> SerializeAsync<T>(T node, ITreeBlobWriter writer, CancellationToken cancellationToken) where T : TreeNode
+		public static T Deserialize<T>(ITreeBlob blob) where T : TreeNode
 		{
-			return SerializerInstance<T>.Serializer.SerializeAsync(writer, node, cancellationToken);
+			return SerializerInstance<T>.Serializer.Deserialize(blob);
 		}
+	}
 
+	/// <summary>
+	/// Data to be stored in a tree blob
+	/// </summary>
+	public struct NewTreeBlob
+	{
 		/// <summary>
-		/// Deserialize a node from data
+		/// The opaque data payload
 		/// </summary>
-		/// <typeparam name="T"></typeparam>
-		/// <param name="data"></param>
-		/// <param name="cancellationToken"></param>
-		/// <returns></returns>
-		public static ValueTask<T> DeserializeAsync<T>(ITreeBlob data, CancellationToken cancellationToken) where T : TreeNode
+		public ReadOnlySequence<byte> Data { get; }
+
+		/// <summary>
+		/// References to other tree nodes
+		/// </summary>
+		public IReadOnlyList<TreeNodeRef> Refs { get; }
+
+		/// <summary>
+		/// Constructor
+		/// </summary>
+		/// <param name="data">Data payload</param>
+		/// <param name="refs">References to other nodes</param>
+		public NewTreeBlob(ReadOnlyMemory<byte> data, IReadOnlyList<TreeNodeRef> refs)
+		{
+			Data = new ReadOnlySequence<byte>(data);
+			Refs = refs;
+		}
+
+		/// <summary>
+		/// Constructor
+		/// </summary>
+		/// <param name="data">Data payload</param>
+		/// <param name="refs">References to other nodes</param>
+		public NewTreeBlob(ReadOnlySequence<byte> data, IReadOnlyList<TreeNodeRef> refs)
 		{
-			return SerializerInstance<T>.Serializer.DeserializeAsync(data, cancellationToken);
+			Data = data;
+			Refs = refs;
 		}
 	}
 
@@ -77,22 +110,12 @@ namespace EpicGames.Horde.Storage
 	/// <typeparam name="T">The type of node returned</typeparam>
 	public abstract class TreeNodeSerializer<T> where T : TreeNode
 	{
-		/// <summary>
-		/// Serialize a given node
-		/// </summary>
-		/// <param name="writer">Writer for node data</param>
-		/// <param name="node">The user type to be serialized</param>
-		/// <param name="cancellationToken">Cancellation token for the operation</param>
-		/// <returns>Serialized node instance</returns>
-		public abstract ValueTask<ITreeBlob> SerializeAsync(ITreeBlobWriter writer, T node, CancellationToken cancellationToken);
-
 		/// <summary>
 		/// Deserializes data from the given data
 		/// </summary>
-		/// <param name="node">Node to deserialize from</param>
-		/// <param name="cancellationToken">Cancellation token for the operation</param>
+		/// <param name="blob">The typed blob</param>
 		/// <returns>New node parsed from the data</returns>
-		public abstract ValueTask<T> DeserializeAsync(ITreeBlob node, CancellationToken cancellationToken);
+		public abstract T Deserialize(ITreeBlob blob);
 	}
 
 	/// <summary>
@@ -111,4 +134,73 @@ namespace EpicGames.Horde.Storage
 		/// </summary>
 		public TreeSerializerAttribute(Type type) => Type = type;
 	}
+
+	/// <summary>
+	/// Extension methods for serializing <see cref="TreeNode"/> objects
+	/// </summary>
+	public static class TreeNodeExtensions
+	{
+		/// <summary>
+		/// Flush the tree with the given root node to storage
+		/// </summary>
+		/// <param name="writer">Writer to output the nodes to</param>
+		/// <param name="root">Root node to serialize</param>
+		/// <param name="cancellationToken">Cancellation token for the operation</param>
+		public static async Task FlushAsync(this ITreeWriter writer, TreeNode root, CancellationToken cancellationToken = default)
+		{
+			ITreeBlobRef rootBlobRef = await writer.WriteTreeAsync(root, cancellationToken);
+			ITreeBlob rootBlob = await rootBlobRef.GetTargetAsync(cancellationToken);
+
+			await writer.FlushAsync(rootBlob, cancellationToken);
+		}
+
+		/// <summary>
+		/// Write a tree of nodes to an <see cref="ITreeWriter"/>
+		/// </summary>
+		/// <param name="writer">Writer to output the nodes to</param>
+		/// <param name="node">Root node to serialize</param>
+		/// <param name="cancellationToken">Cancellation token for the operation</param>
+		/// <returns>Reference to the new tree of nodes</returns>
+		public static async Task<ITreeBlobRef> WriteTreeAsync(this ITreeWriter writer, TreeNode node, CancellationToken cancellationToken = default)
+		{
+			NewTreeBlob blob = node.Serialize();
+
+			List<ITreeBlobRef> targetRefs = new List<ITreeBlobRef>();
+			foreach (TreeNodeRef typedRef in blob.Refs)
+			{
+				targetRefs.Add(await typedRef.CollapseAsync(writer, cancellationToken));
+			}
+
+			return await writer.WriteNodeAsync(blob.Data, targetRefs, cancellationToken);
+		}
+
+		/// <inheritdoc/>
+		public static async Task<T?> TryReadTreeAsync<T>(this ITreeStore store, RefName name, CancellationToken cancellationToken = default) where T : TreeNode
+		{
+			ITreeBlob? root = await store.TryReadTreeAsync(name, cancellationToken);
+			if (root == null)
+			{
+				return null;
+			}
+			return TreeNode.Deserialize<T>(root);
+		}
+
+		/// <inheritdoc/>
+		public static async Task<T> ReadTreeAsync<T>(this ITreeStore store, RefName name, CancellationToken cancellationToken = default) where T : TreeNode
+		{
+			T? result = await store.TryReadTreeAsync<T>(name, cancellationToken);
+			if (result == null)
+			{
+				throw new RefNameNotFoundException(name);
+			}
+			return result;
+		}
+
+		/// <inheritdoc/>
+		public static async Task WriteTreeAsync(this ITreeStore store, RefName name, TreeNode root, CancellationToken cancellationToken = default)
+		{
+			ITreeWriter writer = store.CreateTreeWriter(name);
+			await writer.FlushAsync(root, cancellationToken);
+		}
+	}
 }
diff --git a/Engine/Source/Programs/Shared/EpicGames.Horde/Storage/TreeNodeRef.cs b/Engine/Source/Programs/Shared/EpicGames.Horde/Storage/TreeNodeRef.cs
--- a/Engine/Source/Programs/Shared/EpicGames.Horde/Storage/TreeNodeRef.cs
+++ b/Engine/Source/Programs/Shared/EpicGames.Horde/Storage/TreeNodeRef.cs
@@ -1,10 +1,12 @@
 // Copyright Epic Games, Inc. All Rights Reserved.
 
 using System;
+using System.Collections.Generic;
 using System.Diagnostics;
 using System.Reflection;
 using System.Threading;
 using System.Threading.Tasks;
+using EpicGames.Serialization;
 
 namespace EpicGames.Horde.Storage
 {
@@ -36,7 +38,7 @@ namespace EpicGames.Horde.Storage
 		/// <summary>
 		/// Handle to the node in storage. May be set to null when the node is in memory.
 		/// </summary>
-		internal ITreeBlob? _target;
+		internal ITreeBlobRef? _target;
 
 		/// <summary>
 		/// Reference to the node.
@@ -55,7 +57,7 @@ namespace EpicGames.Horde.Storage
 		/// Creates a reference to a node with the given hash
 		/// </summary>
 		/// <param name="persistedNode">Hash of the referenced node</param>
-		public TreeNodeRef(ITreeBlob persistedNode)
+		public TreeNodeRef(ITreeBlobRef persistedNode)
 		{
 			_target = persistedNode;
 		}
@@ -65,7 +67,7 @@ namespace EpicGames.Horde.Storage
 		/// </summary>
 		/// <param name="owner">The node which owns the reference</param>
 		/// <param name="persistedNode">Hash of the referenced node</param>
-		public TreeNodeRef(TreeNode owner, ITreeBlob persistedNode)
+		public TreeNodeRef(TreeNode owner, ITreeBlobRef persistedNode)
 		{
 			_owner = owner;
 			_target = persistedNode;
@@ -148,7 +150,7 @@ namespace EpicGames.Horde.Storage
 		/// <summary>
 		/// Marks this reference as clean
 		/// </summary>
-		internal void MarkAsClean(ITreeBlob persistedNode)
+		internal void MarkAsClean(ITreeBlobRef persistedNode)
 		{
 			_target = persistedNode;
 
@@ -163,6 +165,23 @@ namespace EpicGames.Horde.Storage
 			_lastModifiedTime = 0;
 		}
 
+		/// <summary>
+		/// Collapse this reference and write its node to storage
+		/// </summary>
+		/// <param name="writer">Writer to output the node to</param>
+		/// <param name="cancellationToken">Cancellation token for the operation</param>
+		/// <returns></returns>
+		public async ValueTask<ITreeBlobRef> CollapseAsync(ITreeWriter writer, CancellationToken cancellationToken)
+		{
+			ITreeBlobRef? target = _target;
+			if (target == null)
+			{
+				target = await writer.WriteTreeAsync(Node!, cancellationToken);
+				MarkAsClean(target);
+			}
+			return target;
+		}
+
 		/// <summary>
 		/// Callback for when a reference is collapsed to a hash value. At the time of calling, both the hash and node data will be valid.
 		/// </summary>
@@ -191,7 +210,7 @@ namespace EpicGames.Horde.Storage
 		/// </summary>
 		/// <param name="owner">The node which owns the reference</param>
 		/// <param name="target">Hash of the referenced node</param>
-		public TreeNodeRef(TreeNode owner, ITreeBlob target) : base(owner, target)
+		public TreeNodeRef(TreeNode owner, ITreeBlobRef target) : base(owner, target)
 		{
 		}
 
@@ -218,36 +237,20 @@ namespace EpicGames.Horde.Storage
 			return (TreeNodeSerializer<T>)Activator.CreateInstance(attribute.Type)!;
 		}
 
-		/// <summary>
-		/// Collapse this reference and write its node to storage
-		/// </summary>
-		/// <param name="writer"></param>
-		/// <param name="cancellationToken">Cancellation token for the operation</param>
-		/// <returns></returns>
-		public async ValueTask<ITreeBlob> CollapseAsync(ITreeBlobWriter writer, CancellationToken cancellationToken)
-		{
-			ITreeBlob? node = _target;
-			if (node == null)
-			{
-				node = await Serializer.SerializeAsync(writer, Node!, cancellationToken);
-				MarkAsClean(node);
-			}
-			return node;
-		}
-
 		/// <summary>
 		/// Resolve this reference to a concrete node
 		/// </summary>
 		/// <param name="cancellationToken">Cancellation token for the operation</param>
 		/// <returns></returns>
-		public async ValueTask<T> ExpandAsync(CancellationToken cancellationToken)
+		public async ValueTask<T> ExpandAsync(CancellationToken cancellationToken = default)
 		{
 			T? result = (T?)ConvertToStrongRef();
 			if (result == null)
 			{
 				Debug.Assert(_target != null);
+				ITreeBlob node = await _target.GetTargetAsync(cancellationToken);
 
-				_strongRef = await Serializer.DeserializeAsync(_target, cancellationToken);
+				_strongRef = Serializer.Deserialize(node);
 				_strongRef.IncomingRef = this;
 
 				result = (T)_strongRef;
