OIdcToken - Option to control if claims are loaded from user profile, disabled by default as this option does not work for Azure AD and we do not need it for the epic deployments

diff --git a/Engine/Source/Programs/Shared/EpicGames.OIDC/OidcTokenManager.cs b/Engine/Source/Programs/Shared/EpicGames.OIDC/OidcTokenManager.cs
--- a/Engine/Source/Programs/Shared/EpicGames.OIDC/OidcTokenManager.cs
+++ b/Engine/Source/Programs/Shared/EpicGames.OIDC/OidcTokenManager.cs
@@ -153,6 +153,7 @@ namespace EpicGames.OIDC
 
 	class OidcTokenClient
 	{
+		private readonly ProviderInfo _providerInfo;
 		private readonly Uri _authorityUri;
 		private readonly string _clientId;
 		private readonly Uri _redirectUri;
@@ -171,6 +172,7 @@ namespace EpicGames.OIDC
 		public OidcTokenClient(ILogger<OidcTokenClient> logger, ProviderInfo providerInfo)
 		{
 			_logger = logger;
+			_providerInfo = providerInfo;
 
 			_authorityUri = providerInfo.ServerUri;
 			_clientId = providerInfo.ClientId;
@@ -180,6 +182,7 @@ namespace EpicGames.OIDC
 
 		public OidcTokenClient(ProviderInfo providerInfo)
 		{
+			_providerInfo = providerInfo;
 			_logger = null;
 
 			_authorityUri = providerInfo.ServerUri;
@@ -203,6 +206,7 @@ namespace EpicGames.OIDC
 				Scope = _scopes,
 				FilterClaims = false,
 				RedirectUri = _redirectUri.ToString(),
+				LoadProfile = _providerInfo.LoadClaimsFromUserProfile,
 				//Flow = OidcClientOptions.AuthenticationFlow.AuthorizationCode,
 				//ResponseMode = OidcClientOptions.AuthorizeResponseMode.Redirect
 			};
@@ -481,6 +485,7 @@ namespace EpicGames.OIDC
 		[Required] public string DisplayName { get; set; } = null!;
 
 		[Required] public Uri RedirectUri { get; set; } = null!;
+		[Required] public bool LoadClaimsFromUserProfile { get; set; } = false;
 
 		public string Scopes { get; set; } = "openid profile offline_access";
 	}
