Initialize UseShellExecute to the default of the .NetFramework to avoid issues.

The default used to be true in the Framework but was changed to false since Core.
https://docs.microsoft.com/en-us/dotnet/api/system.diagnostics.processstartinfo.useshellexecute?view=net-6.0

Reported on UDN case 453862

diff --git a/Engine/Source/Programs/Shared/EpicGames.Build/UEBuildPlatformSDK.cs b/Engine/Source/Programs/Shared/EpicGames.Build/UEBuildPlatformSDK.cs
--- a/Engine/Source/Programs/Shared/EpicGames.Build/UEBuildPlatformSDK.cs
+++ b/Engine/Source/Programs/Shared/EpicGames.Build/UEBuildPlatformSDK.cs
@@ -1342,6 +1342,9 @@ namespace EpicGames.Core
 					{
 						// installers may require administrator access to succeed. so run as an admin.
 						HookProcess.StartInfo.Verb = "runas";
+
+						//Forcing the old .Net Framework default to prevent processes from failing
+						HookProcess.StartInfo.UseShellExecute = true;
 					}
 					else
 					{
