Horde: Fix OOM exceptions on server due to unbounded cache growth. Allocate a custom cache for building bundles.

diff --git a/Engine/Source/Programs/Shared/EpicGames.Horde/EpicGames.Horde.csproj b/Engine/Source/Programs/Shared/EpicGames.Horde/EpicGames.Horde.csproj
--- a/Engine/Source/Programs/Shared/EpicGames.Horde/EpicGames.Horde.csproj
+++ b/Engine/Source/Programs/Shared/EpicGames.Horde/EpicGames.Horde.csproj
@@ -39,6 +39,7 @@
     <PackageReference Include="K4os.Compression.LZ4" Version="1.2.6" />
     <PackageReference Include="Microsoft.CodeAnalysis.CSharp.Workspaces" Version="3.11.0" />
     <PackageReference Include="Microsoft.Extensions.Caching.Abstractions" Version="6.0.0" />
+    <PackageReference Include="Microsoft.Extensions.Caching.Memory" Version="6.0.0" />
     <PackageReference Include="Microsoft.Extensions.DependencyInjection" Version="6.0.0" />
     <PackageReference Include="Microsoft.Extensions.Http" Version="3.1.22" />
     <PackageReference Include="Microsoft.Extensions.Logging" Version="5.0.0" />
diff --git a/Engine/Source/Programs/Shared/EpicGames.Horde/Storage/Bundles/BundleStore.cs b/Engine/Source/Programs/Shared/EpicGames.Horde/Storage/Bundles/BundleStore.cs
--- a/Engine/Source/Programs/Shared/EpicGames.Horde/Storage/Bundles/BundleStore.cs
+++ b/Engine/Source/Programs/Shared/EpicGames.Horde/Storage/Bundles/BundleStore.cs
@@ -244,7 +244,7 @@ namespace EpicGames.Horde.Storage.Bundles
 		{
 			_blobStore = blobStore;
 			_options = options;
-			_cache = cache;
+			_cache = new MemoryCache(new MemoryCacheOptions { SizeLimit = 512 * 1024 * 1024 });
 			_readSema = new SemaphoreSlim(options.MaxConcurrentReads);
 		}
 
@@ -267,7 +267,11 @@ namespace EpicGames.Horde.Storage.Bundles
 		/// <param name="disposing"></param>
 		protected virtual void Dispose(bool disposing)
 		{
-			_readSema.Dispose();
+			if (disposing)
+			{
+				_readSema.Dispose();
+				_cache.Dispose();
+			}
 		}
 
 		/// <inheritdoc/>
