Horde: Attempts to reduce memory usage for bundle stores.

* Add configurable size for bundle cache, and disable it by default.
* Sum up sizes of directories to sync without duplicating paths of every file.
* Use ByteArrayBuilder to create segments for bundles rather than concatenating to one large array.

diff --git a/Engine/Source/Programs/Horde/Horde.Build/Perforce/ReplicationService.cs b/Engine/Source/Programs/Horde/Horde.Build/Perforce/ReplicationService.cs
--- a/Engine/Source/Programs/Horde/Horde.Build/Perforce/ReplicationService.cs
+++ b/Engine/Source/Programs/Horde/Horde.Build/Perforce/ReplicationService.cs
@@ -512,31 +512,10 @@ namespace Horde.Build.Perforce
 			return await WriteCommitTreeAsync(stream, change, baseChange, baseContents, stream.ReplicationFilter, revisionsOnly, cancellationToken);
 		}
 
-		/// <summary>
-		/// Sorts paths by their folder, then by their filename.
-		/// </summary>
-		class FolderFirstSorter : IComparer<Utf8String>
+		class DirectoryToSync
 		{
-			readonly Utf8StringComparer _comparer;
-
-			public FolderFirstSorter(Utf8StringComparer comparer)
-			{
-				_comparer = comparer;
-			}
-
-			public int Compare(Utf8String x, Utf8String y)
-			{
-				Utf8String pathX = x.Substring(0, x.LastIndexOf('/') + 1);
-				Utf8String pathY = y.Substring(0, y.LastIndexOf('/') + 1);
-
-				int result = _comparer.Compare(pathX, pathY);
-				if (result != 0)
-				{
-					return result;
-				}
-
-				return _comparer.Compare(x, y);
-			}
+			public Utf8String _path;
+			public long _size;
 		}
 
 		/// <summary>
@@ -619,7 +598,7 @@ namespace Horde.Build.Perforce
 				}
 
 				// Do a sync preview to find everything that's left, and sort the remaining list of paths
-				List<(Utf8String Path, long Size)> files = new List<(Utf8String, long)>();
+				Dictionary<Utf8String, DirectoryToSync> pathToDirectory = new Dictionary<Utf8String, DirectoryToSync>(clientInfo.ServerInfo.Utf8PathComparer);
 				await foreach (PerforceResponse<SyncRecord> response in perforce.StreamCommandAsync<SyncRecord>("sync", new[] { "-n" }, new string[] { $"{queryPath}@{change}" }, null, cancellationToken))
 				{
 					PerforceError? error = response.Error;
@@ -629,24 +608,36 @@ namespace Horde.Build.Perforce
 						continue;
 					}
 
-					Utf8String file = response.Data.Path.Clone();
-					if (!file.StartsWith(clientRoot, Utf8StringComparer.Ordinal))
+					Utf8String path = response.Data.Path;
+					if (!path.StartsWith(clientRoot, Utf8StringComparer.Ordinal))
+					{
+						throw new ArgumentException($"Unable to make path {path} relative to client root {clientRoot}");
+					}
+
+					int fileIdx = path.Length;
+					while (fileIdx > 0 && path[fileIdx - 1] != '/' && path[fileIdx - 1] != '\\')
 					{
-						throw new ArgumentException($"Unable to make path {file} relative to client root {clientRoot}");
+						fileIdx--;
 					}
 
-					byte[] path = new byte[file.Length - clientRoot.Length];
-					for (int idx = clientRoot.Length; idx < file.Length; idx++)
+					Utf8String directoryPath = path.Slice(clientRoot.Length, fileIdx - clientRoot.Length);
+
+					DirectoryToSync? directory;
+					if (!pathToDirectory.TryGetValue(directoryPath, out directory))
 					{
-						path[idx - clientRoot.Length] = (file[idx] == '\\') ? (byte)'/' : file[idx];
+						directory = new DirectoryToSync();
+						directory._path = directoryPath.Clone();
+						pathToDirectory.Add(directory._path, directory);
 					}
 
-					files.Add((new Utf8String(path), response.Data.FileSize));
+					directory._size += response.Data.FileSize;
 				}
-				files.SortBy(x => x.Path, new FolderFirstSorter(clientInfo.ServerInfo.Utf8PathComparer));
+
+				// Sort the directories by name
+				List<DirectoryToSync> directories = pathToDirectory.Values.OrderBy(x => x._path, clientInfo.ServerInfo.Utf8PathComparer).ToList();
 
 				// Output some stats for the sync
-				long totalSize = files.Sum(x => x.Size);
+				long totalSize = directories.Sum(x => x._size);
 				_logger.LogInformation("Total sync size: {Size:n1}mb", totalSize / (1024.0 * 1024.0));
 
 				// Previously scheduled flush task
@@ -654,7 +645,7 @@ namespace Horde.Build.Perforce
 
 				// Sync incrementally
 				long syncedSize = 0;
-				while (files.Count > 0)
+				while (directories.Count > 0)
 				{
 					// Save the incremental state
 					if (syncedSize > 0)
@@ -665,9 +656,17 @@ namespace Horde.Build.Perforce
 
 					// Find the next path to sync
 					const long MaxBatchSize = 10L * 1024 * 1024 * 1024;
-					(int idx, Utf8String path, long size) = GetSyncBatch(files, MaxBatchSize, clientInfo.ServerInfo.Utf8PathComparer, _logger);
+					(int dirIdx, Utf8String path, long size) = GetSyncBatch(directories, MaxBatchSize, clientInfo.ServerInfo.Utf8PathComparer, _logger);
 					syncedSize += size;
 
+					// Correct the path separators
+					byte[] newPath = new byte[path.Length];
+					for (int idx = 0; idx < path.Length; idx++)
+					{
+						newPath[idx] = (path[idx] == '\\') ? (byte)'/' : path[idx];
+					}
+					path = new Utf8String(newPath);
+
 					// If we're syncing the entire view, use the given filter instead. It will be tighter than the base path from all files.
 					if (filter != null && path.Equals("..."))
 					{
@@ -781,7 +780,7 @@ namespace Horde.Build.Perforce
 					}
 
 					// Remove the paths we've synced
-					files.RemoveRange(idx, files.Count - idx);
+					directories.RemoveRange(dirIdx, directories.Count - dirIdx);
 				}
 				await flushTask;
 			}
@@ -800,46 +799,45 @@ namespace Horde.Build.Perforce
 			return syncNode;
 		}
 
-		static (int, Utf8String, long) GetSyncBatch(List<(Utf8String Path, long Size)> files, long maxSize, Utf8StringComparer comparer, ILogger logger)
+		static (int, Utf8String, long) GetSyncBatch(List<DirectoryToSync> directories, long maxSize, Utf8StringComparer comparer, ILogger logger)
 		{
-			int idx = files.Count - 1;
-			Utf8String path = files[idx].Path;
-			long size = files[idx].Size;
+			int idx = directories.Count - 1;
+			long size = directories[idx]._size;
 
-			Utf8String prefix = path;
+			Utf8String prefix = directories[idx]._path;
 			while (prefix.Length > 0)
 			{
 				// Get the length of the parent directory
 				int endIdx = prefix.Length - 1;
-				while (endIdx > 0 && path[endIdx - 1] != (byte)'/')
+				while (endIdx > 0 && prefix[endIdx - 1] != (byte)'/' && prefix[endIdx - 1] != (byte)'\\')
 				{
 					endIdx--;
 				}
 
 				// Get the parent directory prefix, ending with a slash
-				prefix = prefix.Substring(0, endIdx);
+				Utf8String nextPrefix = prefix.Substring(0, endIdx);
 
 				// Include all files with this prefix
 				int nextIdx = idx;
 				long nextSize = size;
-				while (nextIdx > 0 && files[nextIdx - 1].Path.StartsWith(prefix, comparer))
+				while (nextIdx > 0 && directories[nextIdx - 1]._path.StartsWith(nextPrefix, comparer))
 				{
-					nextSize += files[nextIdx - 1].Size;
+					nextSize += directories[nextIdx - 1]._size;
 					nextIdx--;
 				}
 				if (nextSize > maxSize)
 				{
-					logger.LogInformation("Next filter {NextFilter} would exceed max size ({NextSize:n1}mb > {MaxSize:n1}mb); using {Filter} ({Size:n1}mb) instead.", prefix + "...", nextSize / (1024.0 * 1024.0), maxSize / (1024.0 * 1024.0), path, size / (1024.0 * 1024.0));
-					return (idx, path, size);
+					logger.LogInformation("Next filter {NextFilter} would exceed max size ({NextSize:n1}mb > {MaxSize:n1}mb); using {Filter} ({Size:n1}mb) instead.", nextPrefix + "...", nextSize / (1024.0 * 1024.0), maxSize / (1024.0 * 1024.0), prefix, size / (1024.0 * 1024.0));
+					return (idx, prefix + "...", size);
 				}
 				size = nextSize;
 				idx = nextIdx;
 
-				// Update the sync path to be the prefixed directory
-				path = prefix + "...";
+				// Update the prefix
+				prefix = nextPrefix;
 			}
 
-			return (idx, path, size);
+			return (idx, prefix + "...", size);
 		}
 
 		async Task FlushWorkspaceAsync(ReplicationClient clientInfo, IPerforceConnection perforce, int change)
diff --git a/Engine/Source/Programs/Shared/EpicGames.Core/ByteArrayBuilder.cs b/Engine/Source/Programs/Shared/EpicGames.Core/ByteArrayBuilder.cs
--- a/Engine/Source/Programs/Shared/EpicGames.Core/ByteArrayBuilder.cs
+++ b/Engine/Source/Programs/Shared/EpicGames.Core/ByteArrayBuilder.cs
@@ -36,6 +36,15 @@ namespace EpicGames.Core
 		/// </summary>
 		public int Length { get; private set; }
 
+		/// <summary>
+		/// Constructor
+		/// </summary>
+		/// <param name="chunkSize"></param>
+		public ByteArrayBuilder(int chunkSize)
+			: this(chunkSize, chunkSize)
+		{
+		}
+
 		/// <summary>
 		/// Constructor
 		/// </summary>
diff --git a/Engine/Source/Programs/Shared/EpicGames.Horde/Storage/Bundles/Bundle.cs b/Engine/Source/Programs/Shared/EpicGames.Horde/Storage/Bundles/Bundle.cs
--- a/Engine/Source/Programs/Shared/EpicGames.Horde/Storage/Bundles/Bundle.cs
+++ b/Engine/Source/Programs/Shared/EpicGames.Horde/Storage/Bundles/Bundle.cs
@@ -42,12 +42,12 @@ namespace EpicGames.Horde.Storage.Bundles
 		/// <summary>
 		/// Packet data as described in the header
 		/// </summary>
-		public ReadOnlyMemory<byte> Payload { get; }
+		public ReadOnlySequence<byte> Payload { get; }
 
 		/// <summary>
 		/// Constructor
 		/// </summary>
-		public Bundle(BundleHeader header, ReadOnlyMemory<byte> payload)
+		public Bundle(BundleHeader header, ReadOnlySequence<byte> payload)
 		{
 			Header = header;
 			Payload = payload;
@@ -60,7 +60,7 @@ namespace EpicGames.Horde.Storage.Bundles
 		{
 			Header = new BundleHeader(reader);
 			int length = (int)reader.ReadUnsignedVarInt();
-			Payload = reader.ReadFixedLengthBytes(length);
+			Payload = new ReadOnlySequence<byte>(reader.ReadFixedLengthBytes(length));
 		}
 
 		/// <summary>
@@ -71,7 +71,7 @@ namespace EpicGames.Horde.Storage.Bundles
 		{
 			ByteArrayBuilder builder = new ByteArrayBuilder();
 			Header.Write(builder);
-			builder.WriteUnsignedVarInt(Payload.Length);
+			builder.WriteUnsignedVarInt((ulong)Payload.Length);
 
 			ReadOnlySequenceBuilder<byte> sequence = new ReadOnlySequenceBuilder<byte>();
 			sequence.Append(builder.AsSequence());
diff --git a/Engine/Source/Programs/Shared/EpicGames.Horde/Storage/Bundles/BundleStore.cs b/Engine/Source/Programs/Shared/EpicGames.Horde/Storage/Bundles/BundleStore.cs
--- a/Engine/Source/Programs/Shared/EpicGames.Horde/Storage/Bundles/BundleStore.cs
+++ b/Engine/Source/Programs/Shared/EpicGames.Horde/Storage/Bundles/BundleStore.cs
@@ -43,6 +43,11 @@ namespace EpicGames.Horde.Storage.Bundles
 		/// Number of reads from storage to allow concurrently. This has an impact on memory usage as well as network throughput.
 		/// </summary>
 		public int MaxConcurrentReads { get; set; } = 5;
+
+		/// <summary>
+		/// Size of the bundle/decode cache
+		/// </summary>
+		public long CacheSize { get; set; }
 	}
 
 	/// <summary>
@@ -225,7 +230,7 @@ namespace EpicGames.Horde.Storage.Bundles
 		}
 
 		readonly IBlobStore _blobStore;
-		readonly IMemoryCache _cache;
+		readonly IMemoryCache? _cache;
 
 		readonly object _lockObject = new object();
 
@@ -244,7 +249,10 @@ namespace EpicGames.Horde.Storage.Bundles
 		{
 			_blobStore = blobStore;
 			_options = options;
-			_cache = new MemoryCache(new MemoryCacheOptions { SizeLimit = 512 * 1024 * 1024 });
+			if (_options.CacheSize > 0)
+			{
+				_cache = new MemoryCache(new MemoryCacheOptions { SizeLimit = _options.CacheSize });
+			}
 			_readSema = new SemaphoreSlim(options.MaxConcurrentReads);
 		}
 
@@ -270,7 +278,7 @@ namespace EpicGames.Horde.Storage.Bundles
 			if (disposing)
 			{
 				_readSema.Dispose();
-				_cache.Dispose();
+				_cache?.Dispose();
 			}
 		}
 
@@ -298,7 +306,7 @@ namespace EpicGames.Horde.Storage.Bundles
 			int packetOffset = 0;
 			foreach (BundlePacket packet in header.Packets)
 			{
-				ReadOnlyMemory<byte> encodedPacket = bundle.Payload.Slice(packetOffset, packet.EncodedLength);
+				ReadOnlyMemory<byte> encodedPacket = bundle.Payload.Slice(packetOffset, packet.EncodedLength).AsSingleSegment();
 
 				byte[] decodedPacket = new byte[packet.DecodedLength];
 				LZ4Codec.Decode(encodedPacket.Span, decodedPacket);
@@ -369,7 +377,7 @@ namespace EpicGames.Horde.Storage.Bundles
 		async Task<Bundle> ReadBundleAsync(BlobId blobId, CancellationToken cancellationToken = default)
 		{
 			string cacheKey = GetBundleCacheKey(blobId);
-			if (!_cache.TryGetValue<Bundle>(cacheKey, out Bundle? bundle))
+			if (_cache == null || !_cache.TryGetValue<Bundle>(cacheKey, out Bundle? bundle))
 			{
 				Task<Bundle>? readTask;
 				lock (_lockObject)
@@ -390,7 +398,7 @@ namespace EpicGames.Horde.Storage.Bundles
 			// Perform another (sequenced) check whether an object has been added to the cache, to counteract the race between a read task being added and a task completing.
 			lock (_lockObject)
 			{
-				if (_cache.TryGetValue<Bundle>(cacheKey, out Bundle? cachedObject))
+				if (_cache != null && _cache.TryGetValue<Bundle>(cacheKey, out Bundle? cachedObject))
 				{
 					return cachedObject;
 				}
@@ -411,11 +419,7 @@ namespace EpicGames.Horde.Storage.Bundles
 			}
 
 			// Add the object to the cache
-			using (ICacheEntry entry = _cache.CreateEntry(cacheKey))
-			{
-				entry.SetSize(bundle.Payload.Length);
-				entry.SetValue(bundle);
-			}
+			AddBundleToCache(cacheKey, bundle);
 
 			// Remove this object from the list of read tasks
 			lock (_lockObject)
@@ -429,10 +433,13 @@ namespace EpicGames.Horde.Storage.Bundles
 
 		void AddBundleToCache(string cacheKey, Bundle bundle)
 		{
-			using (ICacheEntry entry = _cache.CreateEntry(cacheKey))
+			if (_cache != null)
 			{
-				entry.SetSize(bundle.Payload.Length);
-				entry.SetValue(bundle);
+				using (ICacheEntry entry = _cache.CreateEntry(cacheKey))
+				{
+					entry.SetSize(bundle.Payload.Length);
+					entry.SetValue(bundle);
+				}
 			}
 		}
 
@@ -529,20 +536,27 @@ namespace EpicGames.Horde.Storage.Bundles
 		/// <param name="packetInfo">The decoded block location and size</param>
 		/// <param name="payload">The bundle payload</param>
 		/// <returns>The decoded data</returns>
-		ReadOnlyMemory<byte> DecodePacket(BlobId blobId, PacketInfo packetInfo, ReadOnlyMemory<byte> payload)
+		ReadOnlyMemory<byte> DecodePacket(BlobId blobId, PacketInfo packetInfo, ReadOnlySequence<byte> payload)
 		{
-			string cacheKey = $"bundle-packet:{blobId}@{packetInfo.Offset}";
-			return _cache.GetOrCreate<ReadOnlyMemory<byte>>(cacheKey, entry =>
+			if (_cache == null)
+			{
+				return DecodePacketUncached(packetInfo, payload);
+			}
+			else
 			{
-				ReadOnlyMemory<byte> decodedPacket = DecodePacketUncached(packetInfo, payload);
-				entry.SetSize(packetInfo.DecodedLength);
-				return decodedPacket;
-			});
+				string cacheKey = $"bundle-packet:{blobId}@{packetInfo.Offset}";
+				return _cache.GetOrCreate<ReadOnlyMemory<byte>>(cacheKey, entry =>
+				{
+					ReadOnlyMemory<byte> decodedPacket = DecodePacketUncached(packetInfo, payload);
+					entry.SetSize(packetInfo.DecodedLength);
+					return decodedPacket;
+				});
+			}
 		}
 
-		static ReadOnlyMemory<byte> DecodePacketUncached(PacketInfo packetInfo, ReadOnlyMemory<byte> payload)
+		static ReadOnlyMemory<byte> DecodePacketUncached(PacketInfo packetInfo, ReadOnlySequence<byte> payload)
 		{
-			ReadOnlyMemory<byte> encodedPacket = payload.Slice(packetInfo.Offset, packetInfo.EncodedLength);
+			ReadOnlyMemory<byte> encodedPacket = payload.Slice(packetInfo.Offset, packetInfo.EncodedLength).AsSingleSegment();
 
 			byte[] decodedPacket = new byte[packetInfo.DecodedLength];
 			LZ4Codec.Decode(encodedPacket.Span, decodedPacket);
@@ -783,7 +797,7 @@ namespace EpicGames.Horde.Storage.Bundles
 			}
 
 			// Preallocate data in the encoded buffer to reduce fragmentation if we have to resize
-			byte[] outputBuffer = new byte[_options.MaxBlobSize];
+			ByteArrayBuilder builder = new ByteArrayBuilder(_options.MinCompressionPacketSize * 2);
 
 			// Create the export list
 			List<BundleExport> exports = new List<BundleExport>();
@@ -793,7 +807,6 @@ namespace EpicGames.Horde.Storage.Bundles
 			int blockSize = 0;
 
 			// Compress all the nodes into the output buffer buffer
-			int outputOffset = 0;
 			foreach (NodeInfo node in nodes)
 			{
 				InMemoryNodeState nodeState = (InMemoryNodeState)node.State;
@@ -802,7 +815,7 @@ namespace EpicGames.Horde.Storage.Bundles
 				// If we can't fit this data into the current block, flush the contents of it first
 				if (blockSize > 0 && blockSize + nodeData.Length > _options.MinCompressionPacketSize)
 				{
-					FlushPacket(context.PacketBuffer.AsMemory(0, blockSize), ref outputBuffer, ref outputOffset, packets);
+					FlushPacket(context.PacketBuffer.AsMemory(0, blockSize), builder, packets);
 					blockSize = 0;
 				}
 
@@ -826,33 +839,30 @@ namespace EpicGames.Horde.Storage.Bundles
 				}
 				else
 				{
-					FlushPacket(nodeData.First, ref outputBuffer, ref outputOffset, packets);
+					FlushPacket(nodeData.First, builder, packets);
 				}
 			}
-			FlushPacket(context.PacketBuffer.AsMemory(0, blockSize), ref outputBuffer, ref outputOffset, packets);
+			FlushPacket(context.PacketBuffer.AsMemory(0, blockSize), builder, packets);
 
 			// Flush the data
 			BundleHeader header = new BundleHeader(imports, exports, packets);
-			ReadOnlyMemory<byte> payload = outputBuffer.AsMemory(0, packets.Sum(x => x.EncodedLength));
-			return new Bundle(header, payload);
+			return new Bundle(header, builder.AsSequence());
 		}
 
-		static void FlushPacket(ReadOnlyMemory<byte> inputData, ref byte[] outputBuffer, ref int outputOffset, List<BundlePacket> packets)
+		static void FlushPacket(ReadOnlyMemory<byte> inputData, ByteArrayBuilder builder, List<BundlePacket> packets)
 		{
 			if (inputData.Length > 0)
 			{
-				int minFreeSpace = LZ4Codec.MaximumOutputSize(inputData.Length);
-				CreateFreeSpace(ref outputBuffer, outputOffset, outputOffset + minFreeSpace);
+				int maxSize = LZ4Codec.MaximumOutputSize(inputData.Length);
 
 				ReadOnlySpan<byte> inputSpan = inputData.Span;
-				Span<byte> outputSpan = outputBuffer.AsSpan(outputOffset);
+				Span<byte> outputSpan = builder.GetMemory(maxSize).Span;
 
 				int encodedLength = LZ4Codec.Encode(inputSpan, outputSpan);
 				packets.Add(new BundlePacket(encodedLength, inputData.Length));
 
-				outputOffset += encodedLength;
-
 				Debug.Assert(encodedLength >= 0);
+				builder.Advance(encodedLength);
 			}
 		}
 
