[UnrealBuildTool]
* Added support for ".ubtignore" files that can be put in a folder to make ubt ignore the folder and its subfolder
* Fixed bug where project's plugin and source folder were visited multiple times (look in FileMetadataPreftch.cs)
* Changed so Lazy<File/DirectoryItem> is resolved directly on windows platform and only do the extra refresh on mac/linux (there seems to be some bugs on those platforms and we don't want to pay extra cost on windows)
* Added runtimeconfig.template.json and changed UnrealBuildTool to use server gc and a few more settings that improve performance
* Added so all file traversing paths is aware of .uplugin
* Added debug code to be able to write out which files that have been touched during UBT execution. This can serve as a way to identify places to add .ubtignore files

diff --git a/Engine/Source/Programs/Shared/EpicGames.Build/System/DirectoryItem.cs b/Engine/Source/Programs/Shared/EpicGames.Build/System/DirectoryItem.cs
--- a/Engine/Source/Programs/Shared/EpicGames.Build/System/DirectoryItem.cs
+++ b/Engine/Source/Programs/Shared/EpicGames.Build/System/DirectoryItem.cs
@@ -47,11 +47,19 @@ namespace UnrealBuildBase
 		private DirectoryItem(DirectoryReference Location, DirectoryInfo Info)
 		{
 			this.Location = Location;
-			this.Info = new Lazy<DirectoryInfo>(() =>
+			if (RuntimePlatform.IsWindows)
 			{
-				Info.Refresh();
-				return Info;
-			});
+				this.Info = new Lazy<DirectoryInfo>(Info);
+			}
+			else
+			{
+				// For some reason we need to call an extra Refresh on linux/mac to not get wrong results from "Exists"
+				this.Info = new Lazy<DirectoryInfo>(() =>
+				{
+					Info.Refresh();
+					return Info;
+				});
+			}
 		}
 
 		/// <summary>
@@ -117,7 +125,11 @@ namespace UnrealBuildBase
 		/// <returns>The directory item for this location</returns>
 		public static DirectoryItem GetItemByDirectoryReference(DirectoryReference Location)
 		{
-			return LocationToItem.GetOrAdd(Location, _ => new DirectoryItem(Location, new DirectoryInfo(Location.FullName)));
+			if (LocationToItem.TryGetValue(Location, out DirectoryItem? Result))
+			{
+				return Result;
+			}
+			return LocationToItem.GetOrAdd(Location, new DirectoryItem(Location, new DirectoryInfo(Location.FullName)));
 		}
 
 		/// <summary>
@@ -128,7 +140,11 @@ namespace UnrealBuildBase
 		public static DirectoryItem GetItemByDirectoryInfo(DirectoryInfo Info)
 		{
 			DirectoryReference Location = new DirectoryReference(Info);
-			return LocationToItem.GetOrAdd(Location, _ => new DirectoryItem(Location, Info));
+			if (LocationToItem.TryGetValue(Location, out DirectoryItem? Result))
+			{
+				return Result;
+			}
+			return LocationToItem.GetOrAdd(Location, new DirectoryItem(Location, Info));
 		}
 
 		/// <summary>
@@ -190,25 +206,20 @@ namespace UnrealBuildBase
 		{
 			if (Directories == null)
 			{
-				Dictionary<string, DirectoryItem> NewDirectories = new Dictionary<string, DirectoryItem>(DirectoryReference.Comparer);
+				Dictionary<string, DirectoryItem> NewDirectories;
 				if (Info.Value.Exists)
 				{
-					foreach (DirectoryInfo SubDirectoryInfo in Info.Value.EnumerateDirectories())
+					DirectoryInfo[] Directories = Info.Value.GetDirectories();
+					NewDirectories = new Dictionary<string, DirectoryItem>(Directories.Length, DirectoryReference.Comparer);
+					foreach (DirectoryInfo SubDirectoryInfo in Directories)
 					{
-						if (SubDirectoryInfo.Name.Length == 1 && SubDirectoryInfo.Name[0] == '.')
-						{
-							continue;
-						}
-						else if (SubDirectoryInfo.Name.Length == 2 && SubDirectoryInfo.Name[0] == '.' && SubDirectoryInfo.Name[1] == '.')
-						{
-							continue;
-						}
-						else
-						{
-							NewDirectories[SubDirectoryInfo.Name] = DirectoryItem.GetItemByDirectoryInfo(SubDirectoryInfo);
-						}
+						NewDirectories.Add(SubDirectoryInfo.Name, DirectoryItem.GetItemByDirectoryInfo(SubDirectoryInfo));
 					}
 				}
+				else
+				{
+					NewDirectories = new Dictionary<string, DirectoryItem>(DirectoryReference.Comparer);
+				}
 				Directories = NewDirectories;
 			}
 		}
@@ -256,16 +267,22 @@ namespace UnrealBuildBase
 		{
 			if (Files == null)
 			{
-				Dictionary<string, FileItem> NewFiles = new Dictionary<string, FileItem>(FileReference.Comparer);
+				Dictionary<string, FileItem> NewFiles;
 				if (Info.Value.Exists)
 				{
-					foreach (FileInfo FileInfo in Info.Value.EnumerateFiles())
+					FileInfo[] FileInfos = Info.Value.GetFiles();
+					NewFiles = new Dictionary<string, FileItem>(FileInfos.Length, FileReference.Comparer);
+					foreach (FileInfo FileInfo in FileInfos)
 					{
 						FileItem FileItem = FileItem.GetItemByFileInfo(FileInfo);
 						FileItem.UpdateCachedDirectory(this);
 						NewFiles[FileInfo.Name] = FileItem;
 					}
 				}
+				else
+				{
+					NewFiles = new Dictionary<string, FileItem>(FileReference.Comparer);
+				}
 				Files = NewFiles;
 			}
 		}
@@ -301,6 +318,25 @@ namespace UnrealBuildBase
 		{
 			return Location.FullName;
 		}
+
+		/// <summary>
+		/// Writes out all the enumerated files full names sorted to OutFile
+		/// </summary>
+		public static void WriteDebugFileWithAllEnumeratedFiles(string OutFile)
+		{
+			SortedSet<string> AllFiles = new SortedSet<string>();
+			foreach (DirectoryItem Item in DirectoryItem.LocationToItem.Values)
+			{
+				if (Item.Files != null)
+				{
+					foreach (FileItem File in Item.EnumerateFiles())
+					{
+						AllFiles.Add(File.FullName);
+					}
+				}
+			}
+			File.WriteAllLines(OutFile, AllFiles);
+		}
 	}
 
 	/// <summary>
diff --git a/Engine/Source/Programs/Shared/EpicGames.Build/System/EnumeratePlugins.cs b/Engine/Source/Programs/Shared/EpicGames.Build/System/EnumeratePlugins.cs
--- a/Engine/Source/Programs/Shared/EpicGames.Build/System/EnumeratePlugins.cs
+++ b/Engine/Source/Programs/Shared/EpicGames.Build/System/EnumeratePlugins.cs
@@ -98,6 +98,10 @@ namespace UnrealBuildBase
 						}
 						bSearchSubDirectories = false;
 					}
+					else if (PluginFile.Name == ".ubtignore")
+					{
+						bSearchSubDirectories = false;
+					}
 				}
 
 				if (bSearchSubDirectories)
diff --git a/Engine/Source/Programs/Shared/EpicGames.Build/System/FileItem.cs b/Engine/Source/Programs/Shared/EpicGames.Build/System/FileItem.cs
--- a/Engine/Source/Programs/Shared/EpicGames.Build/System/FileItem.cs
+++ b/Engine/Source/Programs/Shared/EpicGames.Build/System/FileItem.cs
@@ -44,11 +44,19 @@ namespace UnrealBuildBase
 		private FileItem(FileReference Location, FileInfo Info)
 		{
 			this.Location = Location;
-			this.Info = new Lazy<FileInfo>(() =>
+			if (RuntimePlatform.IsWindows)
 			{
-				Info.Refresh();
-				return Info;
-			});
+				this.Info = new Lazy<FileInfo>(Info);
+			}
+			else
+			{
+				// For some reason we need to call an extra Refresh on linux/mac to not get wrong results from "Exists"
+				this.Info = new Lazy<FileInfo>(() =>
+				{
+					Info.Refresh();
+					return Info;
+				});
+			}			
 		}
 
 		/// <summary>
@@ -148,7 +156,11 @@ namespace UnrealBuildBase
 		public static FileItem GetItemByFileInfo(FileInfo Info)
 		{
 			FileReference Location = new FileReference(Info);
-			return UniqueSourceFileMap.GetOrAdd(Location, _ => new FileItem(Location, Info));
+			if (UniqueSourceFileMap.TryGetValue(Location, out FileItem? Result)) // 99.9% reads, so faster to front with a TryGet before GetOrAdd
+			{
+				return Result;
+			}
+			return UniqueSourceFileMap.GetOrAdd(Location, new FileItem(Location, Info));
 		}
 
 		/// <summary>
@@ -158,7 +170,11 @@ namespace UnrealBuildBase
 		/// <returns>The FileItem that represents the given a full file path.</returns>
 		public static FileItem GetItemByFileReference(FileReference Location)
 		{
-			return UniqueSourceFileMap.GetOrAdd(Location, _ => new FileItem(Location, Location.ToFileInfo()));
+			if (UniqueSourceFileMap.TryGetValue(Location, out FileItem? Result)) // 99.9% reads, so faster to front with a TryGet before GetOrAdd
+			{
+				return Result;
+			}
+			return UniqueSourceFileMap.GetOrAdd(Location, new FileItem(Location, Location.ToFileInfo()));
 		}
 
 		/// <summary>
diff --git a/Engine/Source/Programs/Shared/EpicGames.Build/System/Rules.cs b/Engine/Source/Programs/Shared/EpicGames.Build/System/Rules.cs
--- a/Engine/Source/Programs/Shared/EpicGames.Build/System/Rules.cs
+++ b/Engine/Source/Programs/Shared/EpicGames.Build/System/Rules.cs
@@ -137,6 +137,10 @@ namespace UnrealBuildBase
 				}
 			}
 
+			foreach (DirectoryReference f in Folders)
+				if (f.FullName[0] == 'e')
+					System.Console.WriteLine("!LOWERCASE!!!!!!!!!!");
+
 			return FindAllRulesFiles(Folders, RulesFileType);
 		}
 
@@ -270,6 +274,8 @@ namespace UnrealBuildBase
 		{
 			// Scan all the files in this directory
 			bool bSearchSubFolders = true;
+			bool bIsPlugin = false;
+
 			foreach (FileItem File in Directory.EnumerateFiles())
 			{
 				if (File.HasExtension(".build.cs"))
@@ -303,6 +309,24 @@ namespace UnrealBuildBase
 					}
 					bSearchSubFolders = false;
 				}
+				else if (File.HasExtension(".uplugin"))
+				{
+					bIsPlugin = true;
+					bSearchSubFolders = false;
+				}
+				else if (File.Name == ".ubtignore")
+				{
+					bSearchSubFolders = false;
+				}
+			}
+
+			if (bIsPlugin)
+			{
+				DirectoryItem? SourceDir;
+				if (Directory.TryGetDirectory("Source", out SourceDir))
+				{
+					Queue.Enqueue(() => FindAllRulesFilesRecursively(SourceDir, Cache, Queue));
+				}
 			}
 
 			// If we didn't find anything to stop the search, search all the subdirectories too
diff --git a/Engine/Source/Programs/UnrealBuildTool/Configuration/UEBuildModuleCPP.cs b/Engine/Source/Programs/UnrealBuildTool/Configuration/UEBuildModuleCPP.cs
--- a/Engine/Source/Programs/UnrealBuildTool/Configuration/UEBuildModuleCPP.cs
+++ b/Engine/Source/Programs/UnrealBuildTool/Configuration/UEBuildModuleCPP.cs
@@ -1660,7 +1660,15 @@ namespace UnrealBuildTool
 		/// <param name="InputFiles">Collection of source files, categorized by type</param>
 		static void FindInputFilesFromDirectoryRecursive(DirectoryItem BaseDirectory, ReadOnlyHashSet<string> ExcludedNames, HashSet<DirectoryReference> SourceDirectories, Dictionary<DirectoryItem, FileItem[]> DirectoryToSourceFiles, InputFileCollection InputFiles)
 		{
-			foreach(DirectoryItem SubDirectory in BaseDirectory.EnumerateDirectories())
+			bool bIgnoreFileFound;
+			FileItem[] SourceFiles = FindInputFilesFromDirectory(BaseDirectory, InputFiles, out bIgnoreFileFound);
+
+			if (bIgnoreFileFound)
+			{
+				return;
+			}
+
+			foreach (DirectoryItem SubDirectory in BaseDirectory.EnumerateDirectories())
 			{
 				if (!ExcludedNames.Contains(SubDirectory.Name))
 				{
@@ -1668,7 +1676,6 @@ namespace UnrealBuildTool
 				}
 			}
 
-			FileItem[] SourceFiles = FindInputFilesFromDirectory(BaseDirectory, InputFiles);
 			if(SourceFiles.Length > 0)
 			{
 				SourceDirectories.Add(BaseDirectory.Location);
@@ -1681,9 +1688,11 @@ namespace UnrealBuildTool
 		/// </summary>
 		/// <param name="BaseDirectory"></param>
 		/// <param name="InputFiles"></param>
+		/// <param name="bIgnoreFileFound"></param>
 		/// <returns>Array of source files</returns>
-		static FileItem[] FindInputFilesFromDirectory(DirectoryItem BaseDirectory, InputFileCollection InputFiles)
+		static FileItem[] FindInputFilesFromDirectory(DirectoryItem BaseDirectory, InputFileCollection InputFiles, out bool bIgnoreFileFound)
 		{
+			bIgnoreFileFound = false;
 			List<FileItem> SourceFiles = new List<FileItem>();
 			foreach(FileItem InputFile in BaseDirectory.EnumerateFiles())
 			{
@@ -1730,6 +1739,10 @@ namespace UnrealBuildTool
 					SourceFiles.Add(InputFile);
 					InputFiles.ISPCFiles.Add(InputFile);
 				}
+				else if (InputFile.Name == ".ubtignore")
+				{
+					bIgnoreFileFound = true;
+				}
 			}
 			return SourceFiles.ToArray();
 		}
@@ -1741,7 +1754,13 @@ namespace UnrealBuildTool
 		/// <returns>Array of source files</returns>
 		public static FileItem[] GetSourceFiles(DirectoryItem Directory)
 		{
-			return FindInputFilesFromDirectory(Directory, new InputFileCollection());
+			bool bIgnoreFileFound;
+			FileItem[] Files = FindInputFilesFromDirectory(Directory, new InputFileCollection(), out bIgnoreFileFound);
+			if (bIgnoreFileFound)
+			{
+				return Array.Empty<FileItem>();
+			}
+			return Files;
 		}
 
 		/// <summary>
diff --git a/Engine/Source/Programs/UnrealBuildTool/System/ExternalExecution.cs b/Engine/Source/Programs/UnrealBuildTool/System/ExternalExecution.cs
--- a/Engine/Source/Programs/UnrealBuildTool/System/ExternalExecution.cs
+++ b/Engine/Source/Programs/UnrealBuildTool/System/ExternalExecution.cs
@@ -433,18 +433,24 @@ namespace UnrealBuildTool
 		/// <param name="Headers">Receives the list of headers that was found</param>
 		static void FindHeaders(DirectoryItem BaseDirectory, ReadOnlyHashSet<string> ExcludeFolders, List<FileItem> Headers)
 		{
-			foreach(DirectoryItem SubDirectory in BaseDirectory.EnumerateDirectories())
+			foreach (FileItem File in BaseDirectory.EnumerateFiles())
 			{
-				if(!ExcludeFolders.Contains(SubDirectory.Name))
+				if (File.Name == ".ubtignore")
 				{
-					FindHeaders(SubDirectory, ExcludeFolders, Headers);
+					return;
+				}
+
+				if (File.HasExtension(".h"))
+				{
+					Headers.Add(File);
 				}
 			}
-			foreach(FileItem File in BaseDirectory.EnumerateFiles())
+
+			foreach (DirectoryItem SubDirectory in BaseDirectory.EnumerateDirectories())
 			{
-				if(File.HasExtension(".h"))
+				if(!ExcludeFolders.Contains(SubDirectory.Name))
 				{
-					Headers.Add(File);
+					FindHeaders(SubDirectory, ExcludeFolders, Headers);
 				}
 			}
 		}
diff --git a/Engine/Source/Programs/UnrealBuildTool/System/FileMetadataPrefetch.cs b/Engine/Source/Programs/UnrealBuildTool/System/FileMetadataPrefetch.cs
--- a/Engine/Source/Programs/UnrealBuildTool/System/FileMetadataPrefetch.cs
+++ b/Engine/Source/Programs/UnrealBuildTool/System/FileMetadataPrefetch.cs
@@ -139,15 +139,15 @@ namespace UnrealBuildTool
 		/// <param name="ProjectDirectory">The project directory to search</param>
 		static void ScanProjectDirectory(DirectoryItem ProjectDirectory)
 		{
-			foreach (DirectoryReference ExtensionDir in Unreal.GetExtensionDirs(Unreal.EngineDirectory))
+			foreach (DirectoryReference ExtensionDir in Unreal.GetExtensionDirs(ProjectDirectory.Location))
 			{
 				DirectoryItem BaseDirectory = DirectoryItem.GetItemByDirectoryReference(ExtensionDir);
 				BaseDirectory.CacheDirectories();
 
-				DirectoryItem BasePluginsDirectory = DirectoryItem.Combine(ProjectDirectory, "Plugins");
+				DirectoryItem BasePluginsDirectory = DirectoryItem.Combine(BaseDirectory, "Plugins");
 				Enqueue(() => ScanPluginFolder(BasePluginsDirectory));
 
-				DirectoryItem BaseSourceDirectory = DirectoryItem.Combine(ProjectDirectory, "Source");
+				DirectoryItem BaseSourceDirectory = DirectoryItem.Combine(BaseDirectory, "Source");
 				Enqueue(() => ScanDirectoryTree(BaseSourceDirectory));
 			}
 		}
@@ -160,11 +160,27 @@ namespace UnrealBuildTool
 		{
 			foreach(DirectoryItem SubDirectory in Directory.EnumerateDirectories())
 			{
-				if(SubDirectory.EnumerateFiles().Any(x => x.HasExtension(".uplugin")))
+				bool bIsUPlugin = false;
+				bool bIsIgnore = false;
+				foreach (FileItem FileItem in SubDirectory.EnumerateFiles())
+				{
+					if (FileItem.HasExtension(".uplugin"))
+					{
+						bIsUPlugin = true;
+						break;
+					}
+					if (FileItem.Name == ".ubtignore")
+					{
+						bIsIgnore = true;
+						break;
+					}
+				}
+
+				if (bIsUPlugin)
 				{
 					Enqueue(() => ScanDirectoryTree(DirectoryItem.Combine(SubDirectory, "Source")));
 				}
-				else
+				else if (!bIsIgnore)
 				{
 					Enqueue(() => ScanPluginFolder(SubDirectory));
 				}
@@ -177,11 +193,18 @@ namespace UnrealBuildTool
 		/// <param name="Directory">Root of the directory tree</param>
 		static void ScanDirectoryTree(DirectoryItem Directory)
 		{
-			foreach(DirectoryItem SubDirectory in Directory.EnumerateDirectories())
+			foreach (FileItem FileItem in Directory.EnumerateFiles())
+			{
+				if (FileItem.Name == ".ubtignore")
+				{
+					return;
+				}
+			}
+
+			foreach (DirectoryItem SubDirectory in Directory.EnumerateDirectories())
 			{
 				Enqueue(() => ScanDirectoryTree(SubDirectory));
 			}
-			Directory.CacheFiles();
 		}
 	}
 }
diff --git a/Engine/Source/Programs/UnrealBuildTool/System/TargetMakefile.cs b/Engine/Source/Programs/UnrealBuildTool/System/TargetMakefile.cs
--- a/Engine/Source/Programs/UnrealBuildTool/System/TargetMakefile.cs
+++ b/Engine/Source/Programs/UnrealBuildTool/System/TargetMakefile.cs
@@ -823,21 +823,26 @@ namespace UnrealBuildTool
 		/// <param name="Queue">Queue to add sub-tasks to</param>
 		static void FindFilesWithMarkup(DirectoryItem Directory, SourceFileMetadataCache MetadataCache, ReadOnlyHashSet<string> ExcludedFolderNames, ConcurrentBag<FileItem> FilesWithMarkup, ThreadPoolWorkQueue Queue)
 		{
-			// Search through all the subfolders
-			foreach(DirectoryItem SubDirectory in Directory.EnumerateDirectories())
+			// Check for all the headers in this folder
+			foreach (FileItem File in Directory.EnumerateFiles())
 			{
-				if(!ExcludedFolderNames.Contains(SubDirectory.Name))
+				if (File.Name == ".ubtignore")
 				{
-					Queue.Enqueue(() => FindFilesWithMarkup(SubDirectory, MetadataCache, ExcludedFolderNames, FilesWithMarkup, Queue));
+					return;
+				}
+
+				if (File.HasExtension(".h") && MetadataCache.ContainsReflectionMarkup(File))
+				{
+					FilesWithMarkup.Add(File);
 				}
 			}
 
-			// Check for all the headers in this folder
-			foreach(FileItem File in Directory.EnumerateFiles())
+			// Search through all the subfolders
+			foreach (DirectoryItem SubDirectory in Directory.EnumerateDirectories())
 			{
-				if(File.HasExtension(".h") && MetadataCache.ContainsReflectionMarkup(File))
+				if(!ExcludedFolderNames.Contains(SubDirectory.Name))
 				{
-					FilesWithMarkup.Add(File);
+					Queue.Enqueue(() => FindFilesWithMarkup(SubDirectory, MetadataCache, ExcludedFolderNames, FilesWithMarkup, Queue));
 				}
 			}
 		}
diff --git a/Engine/Source/Programs/UnrealBuildTool/UnrealBuildTool.cs b/Engine/Source/Programs/UnrealBuildTool/UnrealBuildTool.cs
--- a/Engine/Source/Programs/UnrealBuildTool/UnrealBuildTool.cs
+++ b/Engine/Source/Programs/UnrealBuildTool/UnrealBuildTool.cs
@@ -664,6 +664,10 @@ namespace UnrealBuildTool
 					FileMetadataPrefetch.Stop();
 				}
 
+				// Uncomment this to output a file that contains all files that UBT has scanned.
+				// Useful when investigating why UBT takes time.
+				//DirectoryItem.WriteDebugFileWithAllEnumeratedFiles(@"c:\temp\AllFiles.txt");
+
 				Utils.LogWriteFileIfChangedActivity(Logger);
 
 				// Print out all the performance info
diff --git a/Engine/Source/Programs/UnrealBuildTool/UnrealBuildTool.csproj b/Engine/Source/Programs/UnrealBuildTool/UnrealBuildTool.csproj
--- a/Engine/Source/Programs/UnrealBuildTool/UnrealBuildTool.csproj
+++ b/Engine/Source/Programs/UnrealBuildTool/UnrealBuildTool.csproj
@@ -22,12 +22,12 @@
     <RunAnalyzersDuringBuild>false</RunAnalyzersDuringBuild>
   </PropertyGroup>
   <PropertyGroup Condition=" '$(Configuration)|$(Platform)' == 'Development|AnyCPU' ">
-    <Optimize>false</Optimize>
+    <Optimize>True</Optimize>
     <OutputPath>..\..\..\Binaries\DotNET\UnrealBuildTool</OutputPath>
     <DocumentationFile>..\..\..\Binaries\DotNET\UnrealBuildTool\UnrealBuildTool.xml</DocumentationFile>
   </PropertyGroup>
   <PropertyGroup Condition=" '$(Configuration)|$(Platform)' == 'Release|AnyCPU' ">
-    <Optimize>false</Optimize>
+    <Optimize>True</Optimize>
     <DebugType>pdbonly</DebugType>
     <PlatformTarget>AnyCPU</PlatformTarget>
     <ErrorReport>prompt</ErrorReport>
diff --git a/Engine/Source/Programs/UnrealBuildTool/app.config b/Engine/Source/Programs/UnrealBuildTool/app.config
deleted file mode 100644
diff --git a/Engine/Source/Programs/UnrealBuildTool/runtimeconfig.template.json b/Engine/Source/Programs/UnrealBuildTool/runtimeconfig.template.json
new file mode 100644
--- /dev/null
+++ b/Engine/Source/Programs/UnrealBuildTool/runtimeconfig.template.json
@@ -0,0 +1,8 @@
+{
+	"configProperties": {
+		"System.GC.Server": true,
+		"System.GC.HeapCount": 32,
+		"System.Globalization.Invariant": true,
+		"System.Globalization.UseNls": true
+	}
+}
\ No newline at end of file
