Horde: Support for generating Git-compatible SHA1 hashes for file and directory objects in bundles.

diff --git a/Engine/Source/Programs/Horde/Horde.Build/Commands/Bundles/CreateCommand.cs b/Engine/Source/Programs/Horde/Horde.Build/Commands/Bundles/CreateCommand.cs
--- a/Engine/Source/Programs/Horde/Horde.Build/Commands/Bundles/CreateCommand.cs
+++ b/Engine/Source/Programs/Horde/Horde.Build/Commands/Bundles/CreateCommand.cs
@@ -43,7 +43,7 @@ namespace Horde.Build.Commands.Bundles
 			ITreeStore store = serviceProvider.GetRequiredService<ITreeStore<CommitService>>();
 			ITreeWriter writer = store.CreateTreeWriter(RefName);
 
-			DirectoryNode node = new DirectoryNode();
+			DirectoryNode node = new DirectoryNode(DirectoryFlags.None);
 			await node.CopyFromDirectoryAsync(InputDir.ToDirectoryInfo(), new ChunkingOptions(), writer, logger, CancellationToken.None);
 
 			TreeNodeRef<DirectoryNode> root = new TreeNodeRef<DirectoryNode>(null!, node);
diff --git a/Engine/Source/Programs/Horde/Horde.Build/Commands/Bundles/PerforceCommand.cs b/Engine/Source/Programs/Horde/Horde.Build/Commands/Bundles/PerforceCommand.cs
--- a/Engine/Source/Programs/Horde/Horde.Build/Commands/Bundles/PerforceCommand.cs
+++ b/Engine/Source/Programs/Horde/Horde.Build/Commands/Bundles/PerforceCommand.cs
@@ -90,7 +90,7 @@ namespace Horde.Build.Commands.Bundles
 			ReplicationNode baseContents;
 			if (BaseChange == 0)
 			{
-				baseContents = new ReplicationNode(new DirectoryNode());
+				baseContents = new ReplicationNode(new DirectoryNode(DirectoryFlags.WithGitHashes));
 			}
 			else
 			{
diff --git a/Engine/Source/Programs/Horde/Horde.Build/Perforce/ReplicationService.cs b/Engine/Source/Programs/Horde/Horde.Build/Perforce/ReplicationService.cs
--- a/Engine/Source/Programs/Horde/Horde.Build/Perforce/ReplicationService.cs
+++ b/Engine/Source/Programs/Horde/Horde.Build/Perforce/ReplicationService.cs
@@ -7,6 +7,7 @@ using System.Diagnostics;
 using System.Globalization;
 using System.IO;
 using System.Linq;
+using System.Security.Cryptography;
 using System.Text;
 using System.Text.RegularExpressions;
 using System.Threading;
@@ -14,6 +15,7 @@ using System.Threading.Tasks;
 using EpicGames.Core;
 using EpicGames.Horde.Storage;
 using EpicGames.Horde.Storage.Bundles;
+using EpicGames.Horde.Storage.Git;
 using EpicGames.Horde.Storage.Nodes;
 using EpicGames.Perforce;
 using EpicGames.Redis;
@@ -34,6 +36,16 @@ namespace Horde.Build.Perforce
 	using CommitId = ObjectId<ICommit>;
 	using StreamId = StringId<IStream>;
 
+	/// <summary>
+	/// Exception triggered during content replication
+	/// </summary>
+	public sealed class ReplicationException : Exception
+	{
+		internal ReplicationException(string message) : base(message)
+		{
+		}
+	}
+
 	/// <summary>
 	/// Options for the commit service
 	/// </summary>
@@ -75,7 +87,7 @@ namespace Horde.Build.Perforce
 		/// Constructor
 		/// </summary>
 		public ReplicationNode()
-			: this(new DirectoryNode())
+			: this(new DirectoryNode(DirectoryFlags.WithGitHashes))
 		{
 		}
 
@@ -512,10 +524,59 @@ namespace Horde.Build.Perforce
 			return await WriteCommitTreeAsync(stream, change, baseChange, baseContents, stream.ReplicationFilter, revisionsOnly, cancellationToken);
 		}
 
+		[DebuggerDisplay("{_path}")]
 		class DirectoryToSync
 		{
-			public Utf8String _path;
+			public readonly Utf8String _path;
+			public readonly Dictionary<Utf8String, long> _fileNameToSize;
 			public long _size;
+
+			public DirectoryToSync(Utf8String path, Utf8StringComparer comparer)
+			{
+				_path = path;
+				_fileNameToSize = new Dictionary<Utf8String, long>(comparer);
+			}
+		}
+
+		class FileWriter : IDisposable
+		{
+			readonly FileEntry _entry;
+			readonly IncrementalHash _hash;
+			readonly long _size;
+
+			public FileWriter(FileEntry entry, long size)
+			{
+				_entry = entry;
+				_hash = IncrementalHash.CreateHash(HashAlgorithmName.SHA1);
+				_size = size;
+
+				GitObject.WriteHeader(GitObjectType.Blob, size, _hash);
+			}
+
+			public void Dispose() => _hash.Dispose();
+
+			public async Task AppendAsync(ReadOnlyMemory<byte> data, ChunkingOptions options, ITreeWriter writer, CancellationToken cancellationToken)
+			{
+				_hash.AppendData(data.Span);
+				await _entry.AppendAsync(data, options, writer, cancellationToken);
+			}
+
+			public async Task FinishAsync(ITreeWriter writer, CancellationToken cancellationToken)
+			{
+				if (_entry.Length != _size)
+				{
+					throw new ReplicationException($"Invalid size for replicated file '{_entry.Name}'. Expected {_size}, got {_entry.Length}.");
+				}
+				UpdateHash();
+				await _entry.CollapseAsync(writer, cancellationToken);
+			}
+
+			void UpdateHash()
+			{
+				Span<byte> buffer = stackalloc byte[20];
+				_hash.GetHashAndReset(buffer);
+				_entry.GitHash = new Sha1Hash(buffer);
+			}
 		}
 
 		/// <summary>
@@ -583,7 +644,7 @@ namespace Horde.Build.Perforce
 
 					using ReadOnlyMemoryStream dataStream = new ReadOnlyMemoryStream(data);
 
-					FileEntry entry = await root.AddFileByPathAsync(path, FileEntryFlags.PerforceDepotPathAndRevision, writer, cancellationToken);
+					FileEntry entry = await root.AddFileByPathAsync(path, FileEntryFlags.PerforceDepotPathAndRevision, cancellationToken);
 					await entry.AppendAsync(dataStream, Options.Chunking, writer, cancellationToken);
 				}
 			}
@@ -625,11 +686,11 @@ namespace Horde.Build.Perforce
 					DirectoryToSync? directory;
 					if (!pathToDirectory.TryGetValue(directoryPath, out directory))
 					{
-						directory = new DirectoryToSync();
-						directory._path = NormalizePathSeparators(directoryPath);
+						directory = new DirectoryToSync(NormalizePathSeparators(directoryPath), clientInfo.ServerInfo.Utf8PathComparer);
 						pathToDirectory.Add(directoryPath, directory);
 					}
 
+					directory._fileNameToSize.Add(path.Slice(fileIdx).Clone(), response.Data.FileSize);
 					directory._size += response.Data.FileSize;
 				}
 
@@ -678,52 +739,81 @@ namespace Horde.Build.Perforce
 					Stopwatch processTimer = new Stopwatch();
 					Stopwatch gcTimer = new Stopwatch();
 
-					Dictionary<int, FileEntry> handles = new Dictionary<int, FileEntry>();
-					await foreach (PerforceResponse response in perforce.StreamCommandAsync("sync", Array.Empty<string>(), new string[] { syncPath }, null, typeof(SyncRecord), true, default))
+					Dictionary<int, FileWriter> handles = new Dictionary<int, FileWriter>();
+					try
 					{
-						PerforceError? error = response.Error;
-						if (error != null)
-						{
-							_logger.LogWarning("Perforce: {Message}", error.Data);
-							continue;
-						}
-
-						processTimer.Start();
-
-						PerforceIo? io = response.Io;
-						if (io != null)
+						await foreach (PerforceResponse response in perforce.StreamCommandAsync("sync", Array.Empty<string>(), new string[] { syncPath }, null, typeof(SyncRecord), true, default))
 						{
-							if (io.Command == PerforceIoCommand.Open)
-							{
-								Utf8String file = GetClientRelativePath(io.Payload, clientInfo.Client.Root);
-								handles[io.File] = await root.AddFileByPathAsync(file, FileEntryFlags.None, writer, cancellationToken);
-							}
-							else if (io.Command == PerforceIoCommand.Write)
+							PerforceError? error = response.Error;
+							if (error != null)
 							{
-								FileEntry entry = handles[io.File];
-								await entry.AppendAsync(io.Payload, Options.Chunking, writer, cancellationToken);
-								dataSize += io.Payload.Length;
+								throw new ReplicationException($"Perforce error while replicating content: {error.Data}");
 							}
-							else if (io.Command == PerforceIoCommand.Close)
+
+							processTimer.Start();
+
+							PerforceIo? io = response.Io;
+							if (io != null)
 							{
-								FileEntry? entry;
-								if (handles.Remove(io.File, out entry))
+								if (io.Command == PerforceIoCommand.Open)
 								{
-									await entry.CollapseAsync(writer, cancellationToken);
+									Utf8String file = GetClientRelativePath(io.Payload, clientInfo.Client.Root);
+									int offset = GetFileOffset(file);
+
+									long fileSize = 0;
+									if (!pathToDirectory.TryGetValue(file.Substring(0, offset), out DirectoryToSync? directory))
+									{
+										throw new ReplicationException($"Unable to find directory for {file}");
+									}
+									if (!directory._fileNameToSize.TryGetValue(file.Substring(offset), out fileSize))
+									{
+										throw new ReplicationException($"Unable to find file entry for {file}");
+									}
+
+									FileEntry entry = await root.AddFileByPathAsync(file, FileEntryFlags.None, cancellationToken);
+									handles[io.File] = new FileWriter(entry, fileSize);
+								}
+								else if (io.Command == PerforceIoCommand.Write)
+								{
+									FileWriter file = handles[io.File];
+									await file.AppendAsync(io.Payload, Options.Chunking, writer, cancellationToken);
+									dataSize += io.Payload.Length;
+								}
+								else if (io.Command == PerforceIoCommand.Close)
+								{
+									FileWriter? file = null;
+									try
+									{
+										if (handles.Remove(io.File, out file))
+										{
+											await file.FinishAsync(writer, cancellationToken);
+										}
+									}
+									finally
+									{
+										file?.Dispose();
+									}
+								}
+								else if (io.Command == PerforceIoCommand.Unlink)
+								{
+									Utf8String file = GetClientRelativePath(io.Payload, clientInfo.Client.Root);
+									await root.DeleteFileByPathAsync(file, cancellationToken);
+								}
+								else
+								{
+									_logger.LogWarning("Unhandled command code {Code}", io.Command);
 								}
 							}
-							else if (io.Command == PerforceIoCommand.Unlink)
-							{
-								Utf8String file = GetClientRelativePath(io.Payload, clientInfo.Client.Root);
-								await root.DeleteFileByPathAsync(file, cancellationToken);
-							}
-							else
-							{
-								_logger.LogWarning("Unhandled command code {Code}", io.Command);
-							}
-						}
 
-						processTimer.Stop();
+							processTimer.Stop();
+						}
+					}
+					finally
+					{
+						foreach (FileWriter handle in handles.Values)
+						{
+							handle.Dispose();
+						}
 					}
 					_logger.LogInformation("Completed batch in {TimeSeconds:n1}s ({ProcessTimeSeconds:n1}s processing)", syncTimer.Elapsed.TotalSeconds, processTimer.Elapsed.TotalSeconds);
 
@@ -794,6 +884,16 @@ namespace Horde.Build.Perforce
 			return syncNode;
 		}
 
+		static int GetFileOffset(Utf8String path)
+		{
+			int fileIdx = path.Length;
+			while (fileIdx > 0 && path[fileIdx - 1] != '/' && path[fileIdx - 1] != '\\')
+			{
+				fileIdx--;
+			}
+			return fileIdx;
+		}
+
 		static Utf8String NormalizePathSeparators(Utf8String path)
 		{
 			byte[] newPath = new byte[path.Length];
diff --git a/Engine/Source/Programs/Horde/Horde.Build/Startup.cs b/Engine/Source/Programs/Horde/Horde.Build/Startup.cs
--- a/Engine/Source/Programs/Horde/Horde.Build/Startup.cs
+++ b/Engine/Source/Programs/Horde/Horde.Build/Startup.cs
@@ -285,7 +285,7 @@ namespace Horde.Build
 		static ITreeStore CreateTreeStore(IServiceProvider sp, TreeStoreOptions options)
 		{
 			IBlobStore store = CreateBlobStore(sp, options);
-			return new BundleStore(store, options.Bundle, sp.GetRequiredService<IMemoryCache>());
+			return new BundleStore(store, options.Bundle);
 		}
 
 		public Startup(IConfiguration configuration)
diff --git a/Engine/Source/Programs/Shared/EpicGames.Horde.Tests/BundleStoreTests.cs b/Engine/Source/Programs/Shared/EpicGames.Horde.Tests/BundleStoreTests.cs
--- a/Engine/Source/Programs/Shared/EpicGames.Horde.Tests/BundleStoreTests.cs
+++ b/Engine/Source/Programs/Shared/EpicGames.Horde.Tests/BundleStoreTests.cs
@@ -44,83 +44,82 @@ namespace EpicGames.Horde.Tests
 
 		static async Task TestTreeAsync(IBlobStore blobStore, BundleOptions options)
 		{
-			using IMemoryCache cache = new MemoryCache(new MemoryCacheOptions());
-
 			// Generate a tree
-			using (BundleStore collection = new BundleStore(blobStore, options, cache))
+			using (BundleStore store = new BundleStore(blobStore, options))
 			{
-				ITreeBlob node1 = await collection.WriteBlobAsync(new ReadOnlySequence<byte>(new byte[] { 1 }), Array.Empty<ITreeBlob>(), CancellationToken.None);
-				ITreeBlob node2 = await collection.WriteBlobAsync(new ReadOnlySequence<byte>(new byte[] { 2 }), new[] { node1 }, CancellationToken.None);
-				ITreeBlob node3 = await collection.WriteBlobAsync(new ReadOnlySequence<byte>(new byte[] { 3 }), new[] { node2 }, CancellationToken.None);
-				ITreeBlob node4 = await collection.WriteBlobAsync(new ReadOnlySequence<byte>(new byte[] { 4 }), Array.Empty<ITreeBlob>(), CancellationToken.None);
-				ITreeBlob node5 = await collection.WriteBlobAsync(new ReadOnlySequence<byte>(new byte[] { 5 }), new[] { node4, node3 }, CancellationToken.None);
+				ITreeWriter writer = store.CreateTreeWriter(new RefName("test"));
+
+				ITreeBlobRef node1 = await writer.WriteNodeAsync(new ReadOnlySequence<byte>(new byte[] { 1 }), Array.Empty<ITreeBlobRef>(), CancellationToken.None);
+				ITreeBlobRef node2 = await writer.WriteNodeAsync(new ReadOnlySequence<byte>(new byte[] { 2 }), new[] { node1 }, CancellationToken.None);
+				ITreeBlobRef node3 = await writer.WriteNodeAsync(new ReadOnlySequence<byte>(new byte[] { 3 }), new[] { node2 }, CancellationToken.None);
+				ITreeBlobRef node4 = await writer.WriteNodeAsync(new ReadOnlySequence<byte>(new byte[] { 4 }), Array.Empty<ITreeBlobRef>(), CancellationToken.None);
 
-				await collection.WriteTreeAsync(new RefName("test"), node5);
+				ITreeBlob root = TreeBlob.Create(new ReadOnlySequence<byte>(new byte[] { 5 }), new[] { node4, node3 });
+				await writer.FlushAsync(root, CancellationToken.None);
 
-				ITreeBlob root = node5;
 				await CheckTree(root);
 			}
 
 			// Check we can read it back in
-			using (BundleStore collection = new BundleStore(blobStore, options, cache))
+			using (BundleStore collection = new BundleStore(blobStore, options))
 			{
-				ITreeBlob root = await collection.ReadTreeAsync(new RefName("test"));
-				await CheckTree(root);
+				ITreeBlob? root = await collection.TryReadTreeAsync(new RefName("test"));
+				Assert.IsNotNull(root);
+				await CheckTree(root!);
 			}
 		}
 
 		static async Task CheckTree(ITreeBlob root)
 		{
 			ITreeBlob node5 = root;
-			byte[] data5 = (await node5.GetDataAsync()).ToArray();
+			byte[] data5 = node5.Data.ToArray();
 			Assert.IsTrue(data5.SequenceEqual(new byte[] { 5 }));
-			IReadOnlyList<ITreeBlob> refs5 = await node5.GetReferencesAsync();
+			IReadOnlyList<ITreeBlobRef> refs5 = node5.Refs;
 			Assert.AreEqual(2, refs5.Count);
 
-			ITreeBlob node4 = refs5[0];
-			byte[] data4 = (await node4.GetDataAsync()).ToArray();
+			ITreeBlob node4 = await refs5[0].GetTargetAsync();
+			byte[] data4 = node4.Data.ToArray();
 			Assert.IsTrue(data4.SequenceEqual(new byte[] { 4 }));
-			IReadOnlyList<ITreeBlob> refs4 = await node4.GetReferencesAsync();
+			IReadOnlyList<ITreeBlobRef> refs4 = node4.Refs;
 			Assert.AreEqual(0, refs4.Count);
 
-			ITreeBlob node3 = refs5[1];
-			byte[] data3 = (await node3.GetDataAsync()).ToArray();
+			ITreeBlob node3 = await refs5[1].GetTargetAsync();
+			byte[] data3 = node3.Data.ToArray();
 			Assert.IsTrue(data3.SequenceEqual(new byte[] { 3 }));
-			IReadOnlyList<ITreeBlob> refs3 = await node3.GetReferencesAsync();
+			IReadOnlyList<ITreeBlobRef> refs3 = node3.Refs;
 			Assert.AreEqual(1, refs3.Count);
 
-			ITreeBlob node2 = refs3[0];
-			byte[] data2 = (await node2.GetDataAsync()).ToArray();
+			ITreeBlob node2 = await refs3[0].GetTargetAsync();
+			byte[] data2 = node2.Data.ToArray();
 			Assert.IsTrue(data2.SequenceEqual(new byte[] { 2 }));
-			IReadOnlyList<ITreeBlob> refs2 = await node2.GetReferencesAsync();
+			IReadOnlyList<ITreeBlobRef> refs2 = node2.Refs;
 			Assert.AreEqual(1, refs2.Count);
 
-			ITreeBlob node1 = refs2[0];
-			byte[] data1 = (await node1.GetDataAsync()).ToArray();
+			ITreeBlob node1 = await refs2[0].GetTargetAsync();
+			byte[] data1 = node1.Data.ToArray();
 			Assert.IsTrue(data1.SequenceEqual(new byte[] { 1 }));
-			IReadOnlyList<ITreeBlob> refs1 = await node1.GetReferencesAsync();
+			IReadOnlyList<ITreeBlobRef> refs1 = node1.Refs;
 			Assert.AreEqual(0, refs1.Count);
 		}
 
 		[TestMethod]
 		public async Task DirectoryNodesAsync()
 		{
-			using IMemoryCache cache = new MemoryCache(new MemoryCacheOptions());
 			InMemoryBlobStore blobStore = new InMemoryBlobStore();
 
 			// Generate a tree
-			using (BundleStore store = new BundleStore(blobStore, new BundleOptions(), cache))
+			using (BundleStore store = new BundleStore(blobStore, new BundleOptions()))
 			{
-				DirectoryNode root = new DirectoryNode();
+				DirectoryNode root = new DirectoryNode(DirectoryFlags.None);
 				DirectoryNode hello = root.AddDirectory("hello");
 				DirectoryNode world = hello.AddDirectory("world");
-				await store.WriteTreeAsync(new RefName("test"), root, true, CancellationToken.None);
+				await store.WriteTreeAsync(new RefName("test"), root, CancellationToken.None);
 
 				await CheckDirectoryTreeAsync(root);
 			}
 
 			// Check we can read it back in
-			using (BundleStore store = new BundleStore(blobStore, new BundleOptions(), cache))
+			using (BundleStore store = new BundleStore(blobStore, new BundleOptions()))
 			{
 				DirectoryNode root = await store.ReadTreeAsync<DirectoryNode>(new RefName("test"));
 				await CheckDirectoryTreeAsync(root);
@@ -143,25 +142,26 @@ namespace EpicGames.Horde.Tests
 		[TestMethod]
 		public async Task FileNodesAsync()
 		{
-			using IMemoryCache cache = new MemoryCache(new MemoryCacheOptions());
 			InMemoryBlobStore blobStore = new InMemoryBlobStore();
 
 			// Generate a tree
-			using (BundleStore store = new BundleStore(blobStore, new BundleOptions(), cache))
+			using (BundleStore store = new BundleStore(blobStore, new BundleOptions()))
 			{
-				DirectoryNode root = new DirectoryNode();
+				ITreeWriter writer = store.CreateTreeWriter(new RefName("test"));
+
+				DirectoryNode root = new DirectoryNode(DirectoryFlags.None);
 				DirectoryNode hello = root.AddDirectory("hello");
 
 				FileEntry world = hello.AddFile("world", FileEntryFlags.None);
-				await world.AppendAsync(Encoding.UTF8.GetBytes("world"), new ChunkingOptions(), CancellationToken.None);
-				
-				await store.WriteTreeAsync(new RefName("test"), root, true, CancellationToken.None);
+				await world.AppendAsync(Encoding.UTF8.GetBytes("world"), new ChunkingOptions(), writer, CancellationToken.None);
+
+				await writer.FlushAsync(root);
 
 				await CheckFileTreeAsync(root);
 			}
 
 			// Check we can read it back in
-			using (BundleStore store = new BundleStore(blobStore, new BundleOptions(), cache))
+			using (BundleStore store = new BundleStore(blobStore, new BundleOptions()))
 			{
 				DirectoryNode root = await store.ReadTreeAsync<DirectoryNode>(new RefName("test"));
 				await CheckFileTreeAsync(root);
@@ -185,8 +185,11 @@ namespace EpicGames.Horde.Tests
 		}
 
 		[TestMethod]
-		public void ChunkingTests()
+		public async Task ChunkingTests()
 		{
+			using BundleStore store = new BundleStore(new InMemoryBlobStore(), new BundleOptions());
+			ITreeWriter writer = store.CreateTreeWriter(new RefName("test"));
+
 			byte[] chunk = RandomNumberGenerator.GetBytes(256);
 
 			byte[] data = new byte[chunk.Length * 1024];
@@ -208,7 +211,7 @@ namespace EpicGames.Horde.Tests
 			while (remaining.Length > 0)
 			{
 				LeafFileNode node = new LeafFileNode();
-				remaining = node.AppendData(remaining, options);
+				remaining = await node.AppendDataAsync(remaining, options, writer, CancellationToken.None);
 				nodes.Add(node);
 
 				IoHash hash = IoHash.Compute(node.Data);
@@ -222,7 +225,6 @@ namespace EpicGames.Horde.Tests
 		[TestMethod]
 		public async Task LargeFileTestAsync()
 		{
-			using IMemoryCache cache = new MemoryCache(new MemoryCacheOptions());
 			InMemoryBlobStore blobStore = new InMemoryBlobStore();
 
 			const int length = 1024;
@@ -238,9 +240,11 @@ namespace EpicGames.Horde.Tests
 			}
 
 			// Generate a tree
-			using (BundleStore store = new BundleStore(blobStore, new BundleOptions { MaxInlineBlobSize = 1024, MaxBlobSize = 1024 }, cache))
+			using (BundleStore store = new BundleStore(blobStore, new BundleOptions { MaxInlineBlobSize = 1024, MaxBlobSize = 1024 }))
 			{
-				DirectoryNode root = new DirectoryNode();
+				ITreeWriter writer = store.CreateTreeWriter(new RefName("test"));
+
+				DirectoryNode root = new DirectoryNode(DirectoryFlags.None);
 
 				ChunkingOptions options = new ChunkingOptions();
 				options.LeafOptions.MinSize = 1;
@@ -248,24 +252,23 @@ namespace EpicGames.Horde.Tests
 				options.LeafOptions.MaxSize = 64 * 1024;
 
 				FileEntry file = root.AddFile("test", FileEntryFlags.None);
-				await file.AppendAsync(data, options, CancellationToken.None);
+				await file.AppendAsync(data, options, writer, CancellationToken.None);
 
-				await store.WriteTreeAsync(new RefName("test"), root, true, CancellationToken.None);
+				await writer.FlushAsync(root);
 
 				await CheckLargeFileTreeAsync(root, data);
 			}
 
 			// Check we can read it back in
-			using (BundleStore store = new BundleStore(blobStore, new BundleOptions(), cache))
+			using (BundleStore store = new BundleStore(blobStore, new BundleOptions()))
 			{
 				DirectoryNode root = await store.ReadTreeAsync<DirectoryNode>(new RefName("test"));
 				await CheckLargeFileTreeAsync(root, data);
 
 				TreeNodeRef<FileNode> file = root.GetFileEntry("test");
-				ITreeBlob blob = await file.CollapseAsync(store, CancellationToken.None);
 
-				Dictionary<ITreeBlob, long> uniqueBlobs = new Dictionary<ITreeBlob, long>();
-				await FindUniqueBlobs(blob, uniqueBlobs);
+				Dictionary<IoHash, long> uniqueBlobs = new Dictionary<IoHash, long>();
+				await FindUniqueBlobs(file.Target!, uniqueBlobs);
 
 				long uniqueSize = uniqueBlobs.Sum(x => x.Value);
 				Assert.IsTrue(uniqueSize < data.Length / 3); // random fraction meaning "lots of dedupe happened"
@@ -284,13 +287,15 @@ namespace EpicGames.Horde.Tests
 			Assert.IsTrue(worldData.SequenceEqual(data));
 		}
 
-		static async Task FindUniqueBlobs(ITreeBlob blob, Dictionary<ITreeBlob, long> uniqueBlobs)
+		static async Task FindUniqueBlobs(ITreeBlobRef blobRef, Dictionary<IoHash, long> uniqueBlobs)
 		{
-			ReadOnlySequence<byte> data = await blob.GetDataAsync();
-			uniqueBlobs[blob] = data.Length;
+			ITreeBlob blob = await blobRef.GetTargetAsync();
+
+			ReadOnlySequence<byte> data = blob.Data;
+			uniqueBlobs[blobRef.Hash] = data.Length;
 
-			IReadOnlyList<ITreeBlob> references = await blob.GetReferencesAsync();
-			foreach (ITreeBlob reference in references)
+			IReadOnlyList<ITreeBlobRef> references = blob.Refs;
+			foreach (ITreeBlobRef reference in references)
 			{
 				await FindUniqueBlobs(reference, uniqueBlobs);
 			}
diff --git a/Engine/Source/Programs/Shared/EpicGames.Horde.Tests/BundleTests.cs b/Engine/Source/Programs/Shared/EpicGames.Horde.Tests/BundleTests.cs
--- a/Engine/Source/Programs/Shared/EpicGames.Horde.Tests/BundleTests.cs
+++ b/Engine/Source/Programs/Shared/EpicGames.Horde.Tests/BundleTests.cs
@@ -23,21 +23,17 @@ namespace EpicGames.Horde.Tests
 	[TestClass]
 	public sealed class BundleTests : IDisposable
 	{
-		readonly MemoryCache _cache;
 		readonly InMemoryBlobStore _blobStore;
 		readonly ITreeStore _treeStore;
 
 		public BundleTests()
 		{
-			_cache = new MemoryCache(new MemoryCacheOptions { SizeLimit = 50 * 1024 * 1024 });
-
 			_blobStore = new InMemoryBlobStore();
-			_treeStore = new BundleStore(_blobStore, new BundleOptions(), _cache);
+			_treeStore = new BundleStore(_blobStore, new BundleOptions());
 		}
 
 		public void Dispose()
 		{
-			_cache.Dispose();
 			_treeStore.Dispose();
 		}
 
@@ -71,32 +67,34 @@ namespace EpicGames.Horde.Tests
 		}
 
 		[TestMethod]
-		public void BasicChunkingTests()
+		public async Task BasicChunkingTests()
 		{
+			ITreeWriter writer = _treeStore.CreateTreeWriter(new RefName("test"));
+
 			ChunkingOptions options = new ChunkingOptions();
 			options.LeafOptions = new ChunkingOptionsForNodeType(8, 8, 8);
 
 			FileNode node = new LeafFileNode();
-			node = node.Append(new byte[7], options);
+			node = await node.AppendAsync(new byte[7], options, writer, CancellationToken.None);
 			Assert.IsTrue(node is LeafFileNode);
 			Assert.AreEqual(7, ((LeafFileNode)node).Data.Length);
 
 			node = new LeafFileNode();
-			node = node.Append(new byte[8], options);
+			node = await node.AppendAsync(new byte[8], options, writer, CancellationToken.None);
 			Assert.IsTrue(node is LeafFileNode);
 			Assert.AreEqual(8, ((LeafFileNode)node).Data.Length);
 
 			node = new LeafFileNode();
-			node = node.Append(new byte[9], options);
+			node = await node.AppendAsync(new byte[9], options, writer, CancellationToken.None);
 			Assert.IsTrue(node is InteriorFileNode);
 			Assert.AreEqual(2, node.GetReferences().Count);
 
-			FileNode? childNode1 = (FileNode?)node.GetReferences()[0].Node;
+			FileNode? childNode1 = await ((TreeNodeRef<FileNode>)node.GetReferences()[0]).ExpandAsync();
 			Assert.IsNotNull(childNode1);
 			Assert.IsTrue(childNode1 is LeafFileNode);
 			Assert.AreEqual(8, ((LeafFileNode)childNode1!).Data.Length);
 
-			FileNode? childNode2 = (FileNode?)node.GetReferences()[1].Node;
+			FileNode? childNode2 = await ((TreeNodeRef<FileNode>)node.GetReferences()[1]).ExpandAsync();
 			Assert.IsNotNull(childNode2);
 			Assert.IsTrue(childNode2 is LeafFileNode);
 			Assert.AreEqual(1, ((LeafFileNode)childNode2!).Data.Length);
@@ -173,14 +171,14 @@ namespace EpicGames.Horde.Tests
 		public async Task BasicTestDirectory()
 		{
 			InMemoryBlobStore blobStore = new InMemoryBlobStore();
-			using ITreeStore treeStore = new BundleStore(blobStore, new BundleOptions(), _cache);
+			using ITreeStore treeStore = new BundleStore(blobStore, new BundleOptions());
 
-			DirectoryNode root = new DirectoryNode();
+			DirectoryNode root = new DirectoryNode(DirectoryFlags.None);
 			DirectoryNode node = root.AddDirectory("hello");
 			DirectoryNode node2 = node.AddDirectory("world");
 
 			RefName refName = new RefName("testref");
-			await treeStore.WriteTreeAsync(refName, root, true);
+			await treeStore.WriteTreeAsync(refName, root);
 
 			// Should be stored inline
 			Assert.AreEqual(1, blobStore.Refs.Count);
@@ -217,9 +215,9 @@ namespace EpicGames.Horde.Tests
 			options.MaxInlineBlobSize = 1;
 
 			InMemoryBlobStore blobStore = new InMemoryBlobStore();
-			using ITreeStore treeStore = new BundleStore(blobStore, options, _cache);
+			using ITreeStore treeStore = new BundleStore(blobStore, options);
 
-			DirectoryNode root = new DirectoryNode();
+			DirectoryNode root = new DirectoryNode(DirectoryFlags.None);
 			root.AddDirectory("node1");
 			root.AddDirectory("node2");
 			root.AddDirectory("node3");
@@ -228,7 +226,7 @@ namespace EpicGames.Horde.Tests
 			Assert.AreEqual(0, blobStore.Blobs.Count);
 
 			RefName refName = new RefName("ref");
-			await treeStore.WriteTreeAsync(refName, root, true);
+			await treeStore.WriteTreeAsync(refName, root);
 
 			Assert.AreEqual(1, blobStore.Refs.Count);
 			Assert.AreEqual(1, blobStore.Blobs.Count);
@@ -245,22 +243,22 @@ namespace EpicGames.Horde.Tests
 
 			RefName refName = new RefName("ref");
 
-			using (ITreeStore oldBundle = new BundleStore(blobStore, options, _cache))
+			using (ITreeStore oldBundle = new BundleStore(blobStore, options))
 			{
-				DirectoryNode root = new DirectoryNode();
+				DirectoryNode root = new DirectoryNode(DirectoryFlags.None);
 
 				DirectoryNode node1 = root.AddDirectory("node1");
 				DirectoryNode node2 = node1.AddDirectory("node2");
 				DirectoryNode node3 = node2.AddDirectory("node3");
 				DirectoryNode node4 = node3.AddDirectory("node4");
 
-				await oldBundle.WriteTreeAsync(refName, root, true);
+				await oldBundle.WriteTreeAsync(refName, root);
 
 				Assert.AreEqual(1, blobStore.Refs.Count);
 				Assert.AreEqual(4, blobStore.Blobs.Count);
 			}
 
-			using (ITreeStore newBundle = new BundleStore(blobStore, options, _cache))
+			using (ITreeStore newBundle = new BundleStore(blobStore, options))
 			{
 				DirectoryNode root = await newBundle.ReadTreeAsync<DirectoryNode>(refName);
 
@@ -350,6 +348,8 @@ namespace EpicGames.Horde.Tests
 		[TestMethod]
 		public async Task CoreAppendTest()
 		{
+			ITreeWriter writer = _treeStore.CreateTreeWriter(new RefName("test"));
+
 			byte[] data = new byte[4096];
 			new Random(0).NextBytes(data);
 
@@ -359,7 +359,7 @@ namespace EpicGames.Horde.Tests
 			FileNode node = new LeafFileNode();
 			for (int idx = 0; idx < data.Length; idx++)
 			{
-				node = node.Append(data.AsMemory(idx, 1), options);
+				node = await node.AppendAsync(data.AsMemory(idx, 1), options, writer, CancellationToken.None);
 
 				byte[] outputData = await node.ToByteArrayAsync(CancellationToken.None);
 				Assert.IsTrue(data.AsMemory(0, idx + 1).Span.SequenceEqual(outputData.AsSpan(0, idx + 1)));
@@ -388,6 +388,8 @@ namespace EpicGames.Horde.Tests
 
 		async Task ChunkingTests(ChunkingOptions options)
 		{
+			ITreeWriter writer = _treeStore.CreateTreeWriter(new RefName("test"));
+
 			byte[] data = new byte[4096];
 			new Random(0).NextBytes(data);
 
@@ -401,7 +403,7 @@ namespace EpicGames.Horde.Tests
 			const int NumIterations = 100;
 			for (int idx = 0; idx < NumIterations; idx++)
 			{
-				root = root.Append(data, options);
+				root = await root.AppendAsync(data, options, writer, CancellationToken.None);
 			}
 
 			byte[] result = await root.ToByteArrayAsync(CancellationToken.None);
@@ -445,10 +447,12 @@ namespace EpicGames.Horde.Tests
 			BundleOptions options = new BundleOptions();
 			options.MaxBlobSize = 1;
 
-			using BundleStore treeStore = new BundleStore(_blobStore, options, _cache);
+			using BundleStore treeStore = new BundleStore(_blobStore, options);
 
 			RefName refName = new RefName("ref");
-			DirectoryNode root = new DirectoryNode();
+			ITreeWriter writer = treeStore.CreateTreeWriter(refName);
+
+			DirectoryNode root = new DirectoryNode(DirectoryFlags.None);
 
 			long totalLength = 0;
 			for (int idxA = 0; idxA < 10; idxA++)
@@ -465,12 +469,12 @@ namespace EpicGames.Horde.Tests
 							FileEntry file = nodeC.AddFile($"{idxD}", FileEntryFlags.None);
 							byte[] data = Encoding.UTF8.GetBytes($"This is file {idxA}/{idxB}/{idxC}/{idxD}");
 							totalLength += data.Length;
-							await file.AppendAsync(data, new ChunkingOptions(), CancellationToken.None);
+							await file.AppendAsync(data, new ChunkingOptions(), writer, CancellationToken.None);
 						}
 					}
 
 					int oldWorkingSetSize = GetWorkingSetSize(root);
-					await treeStore.WriteTreeAsync(refName, root, false);
+					await treeStore.WriteTreeAsync(refName, root);
 					int newWorkingSetSize = GetWorkingSetSize(root);
 					Assert.IsTrue(newWorkingSetSize <= oldWorkingSetSize);
 					Assert.IsTrue(newWorkingSetSize <= 20);
@@ -478,9 +482,9 @@ namespace EpicGames.Horde.Tests
 			}
 
 			Assert.IsTrue(_blobStore.Blobs.Count > 0);
-			Assert.IsTrue(_blobStore.Refs.Count == 0);
+			Assert.IsTrue(_blobStore.Refs.Count == 1);
 
-			await treeStore.WriteTreeAsync(refName, root, true, CancellationToken.None);
+			await treeStore.WriteTreeAsync(refName, root, CancellationToken.None);
 
 			Assert.AreEqual(totalLength, root.Length);
 
diff --git a/Engine/Source/Programs/Shared/EpicGames.Horde.Tests/ChunkingTests.cs b/Engine/Source/Programs/Shared/EpicGames.Horde.Tests/ChunkingTests.cs
--- a/Engine/Source/Programs/Shared/EpicGames.Horde.Tests/ChunkingTests.cs
+++ b/Engine/Source/Programs/Shared/EpicGames.Horde.Tests/ChunkingTests.cs
@@ -8,6 +8,9 @@ using System.Linq;
 using System.Threading.Tasks;
 using System.IO;
 using System.Threading;
+using EpicGames.Horde.Storage;
+using EpicGames.Horde.Storage.Backends;
+using EpicGames.Horde.Storage.Bundles;
 
 namespace EpicGames.Horde.Tests
 {
@@ -46,32 +49,35 @@ namespace EpicGames.Horde.Tests
 		[TestMethod]
 		public async Task BasicChunkingTests()
 		{
+			using BundleStore store = new BundleStore(new InMemoryBlobStore(), new BundleOptions());
+			ITreeWriter writer = store.CreateTreeWriter(new RefName("test"));
+
 			ChunkingOptions options = new ChunkingOptions();
 			options.LeafOptions = new ChunkingOptionsForNodeType(8, 8, 8);
 
 			FileNode node = new LeafFileNode();
-			node.Append(new byte[7], options);
+			await node.AppendAsync(new byte[7], options, writer, CancellationToken.None);
 			Assert.AreEqual(7, node.Length);
 			Assert.AreEqual(0, node.GetReferences().Count);
 			Assert.AreEqual(7, (await node.ToByteArrayAsync(CancellationToken.None)).Length);
 
 			node = new LeafFileNode();
-			node.Append(new byte[8], options);
+			await node.AppendAsync(new byte[8], options, writer, CancellationToken.None);
 			Assert.AreEqual(8, node.Length);
 			Assert.AreEqual(0, node.GetReferences().Count);
 			Assert.AreEqual(8, (await node.ToByteArrayAsync(CancellationToken.None)).Length);
 
 			node = new LeafFileNode();
-			node = node.Append(new byte[9], options);
+			node = await node.AppendAsync(new byte[9], options, writer, CancellationToken.None);
 			Assert.AreEqual(9, node.Length);
 			Assert.AreEqual(2, node.GetReferences().Count);
 
-			FileNode? childNode1 = (FileNode?)node.GetReferences()[0].Node;
+			FileNode? childNode1 = await ((TreeNodeRef<FileNode>)node.GetReferences()[0]).ExpandAsync();
 			Assert.IsNotNull(childNode1);
 			Assert.AreEqual(0, childNode1!.GetReferences().Count);
 			Assert.AreEqual(8, (await childNode1!.ToByteArrayAsync(CancellationToken.None)).Length);
 
-			FileNode? childNode2 = (FileNode?)node.GetReferences()[1].Node;
+			FileNode? childNode2 = await ((TreeNodeRef<FileNode>)node.GetReferences()[1]).ExpandAsync();
 			Assert.IsNotNull(childNode2);
 			Assert.AreEqual(0, childNode2!.GetReferences().Count);
 			Assert.AreEqual(1, (await childNode2!.ToByteArrayAsync(CancellationToken.None)).Length);
@@ -99,6 +105,9 @@ namespace EpicGames.Horde.Tests
 
 		static async Task TestChunkingAsync(ChunkingOptions options)
 		{
+			using BundleStore store = new BundleStore(new InMemoryBlobStore(), new BundleOptions());
+			ITreeWriter writer = store.CreateTreeWriter(new RefName("test"));
+
 			byte[] data = new byte[4096];
 			new Random(0).NextBytes(data);
 
@@ -112,7 +121,7 @@ namespace EpicGames.Horde.Tests
 			const int NumIterations = 100;
 			for (int idx = 0; idx < NumIterations; idx++)
 			{
-				root = root.Append(data, options);
+				root = await root.AppendAsync(data, options, writer, CancellationToken.None);
 			}
 
 			byte[] result;
diff --git a/Engine/Source/Programs/Shared/EpicGames.Horde/Storage/Bundles/BundleStore.cs b/Engine/Source/Programs/Shared/EpicGames.Horde/Storage/Bundles/BundleStore.cs
--- a/Engine/Source/Programs/Shared/EpicGames.Horde/Storage/Bundles/BundleStore.cs
+++ b/Engine/Source/Programs/Shared/EpicGames.Horde/Storage/Bundles/BundleStore.cs
@@ -66,6 +66,8 @@ namespace EpicGames.Horde.Storage.Bundles
 			public readonly IoHash Hash;
 			public NodeState State { get; private set; }
 
+			IoHash ITreeBlobRef.Hash => Hash;
+
 			public NodeInfo(BundleStore owner, IoHash hash, NodeState state)
 			{
 				Owner = owner;
@@ -223,7 +225,7 @@ namespace EpicGames.Horde.Storage.Bundles
 			}
 
 			/// <inheritdoc/>
-			public Task FlushAsync(ITreeBlob root, CancellationToken cancellationToken = default) => _owner.FlushAsync(this, root, cancellationToken);
+			public Task FlushAsync(ReadOnlySequence<byte> data, IReadOnlyList<ITreeBlobRef> refs, CancellationToken cancellationToken = default) => _owner.FlushAsync(this, data, refs, cancellationToken);
 
 			/// <inheritdoc/>
 			public Task<ITreeBlobRef> WriteNodeAsync(ReadOnlySequence<byte> data, IReadOnlyList<ITreeBlobRef> refs, CancellationToken cancellationToken = default) => _owner.WriteNodeAsync(this, data, refs, cancellationToken);
@@ -245,7 +247,7 @@ namespace EpicGames.Horde.Storage.Bundles
 		/// <summary>
 		/// Constructor
 		/// </summary>
-		public BundleStore(IBlobStore blobStore, BundleOptions options, IMemoryCache cache)
+		public BundleStore(IBlobStore blobStore, BundleOptions options)
 		{
 			_blobStore = blobStore;
 			_options = options;
@@ -608,7 +610,7 @@ namespace EpicGames.Horde.Storage.Bundles
 
 					// If this node pushes the current blob over the max size, queue it to be written now
 					long length = inMemoryState.Data.Length;
-					if (context.WriteLength + length > _options.MaxBlobSize)
+					if (context.WriteNodes.Count > 0 && context.WriteLength + length > _options.MaxBlobSize)
 					{
 						// Start the write, but don't wait for it to complete. Once we've copied the list of nodes they can be written asynchronously.
 						NodeInfo[] writeNodes = context.WriteNodes.ToArray();
@@ -686,11 +688,12 @@ namespace EpicGames.Horde.Storage.Bundles
 			}
 		}
 
-		async Task FlushAsync(BundleWriteContext context, ITreeBlob root, CancellationToken cancellationToken = default)
+		async Task FlushAsync(BundleWriteContext context, ReadOnlySequence<byte> data, IReadOnlyList<ITreeBlobRef> refs, CancellationToken cancellationToken = default)
 		{
 			// Create the root node, and add everything it references to the write buffer
-			IoHash rootHash = ComputeHash(root.Data, root.Refs.ConvertAll(x => (NodeInfo)x));
-			NodeInfo rootNode = new NodeInfo(this, rootHash, (InMemoryNodeState)root);
+			List<NodeInfo> typedRefs = refs.ConvertAll(x => (NodeInfo)x);
+			IoHash rootHash = ComputeHash(data, typedRefs);
+			NodeInfo rootNode = new NodeInfo(this, rootHash, new InMemoryNodeState(data, typedRefs));
 
 			// Flush all the nodes to bundles, and snapshot the remaining ref
 			Task task;
diff --git a/Engine/Source/Programs/Shared/EpicGames.Horde/Storage/ITreeStore.cs b/Engine/Source/Programs/Shared/EpicGames.Horde/Storage/ITreeStore.cs
--- a/Engine/Source/Programs/Shared/EpicGames.Horde/Storage/ITreeStore.cs
+++ b/Engine/Source/Programs/Shared/EpicGames.Horde/Storage/ITreeStore.cs
@@ -1,10 +1,11 @@
-﻿// Copyright Epic Games, Inc. All Rights Reserved.
+// Copyright Epic Games, Inc. All Rights Reserved.
 
 using System;
 using System.Buffers;
 using System.Collections.Generic;
 using System.Threading;
 using System.Threading.Tasks;
+using EpicGames.Core;
 
 namespace EpicGames.Horde.Storage
 {
@@ -29,6 +30,11 @@ namespace EpicGames.Horde.Storage
 	/// </summary>
 	public interface ITreeBlobRef
 	{
+		/// <summary>
+		/// Hash of the target
+		/// </summary>
+		IoHash Hash { get; }
+
 		/// <summary>
 		/// Gets the target of a reference
 		/// </summary>
@@ -94,11 +100,39 @@ namespace EpicGames.Horde.Storage
 		Task<ITreeBlobRef> WriteNodeAsync(ReadOnlySequence<byte> data, IReadOnlyList<ITreeBlobRef> refs, CancellationToken cancellationToken = default);
 
 		/// <summary>
-		/// Flush the tree with the given root.
+		/// Flush the tree with the given data at the root.
 		/// </summary>
-		/// <param name="root">The root node</param>
+		/// <param name="data">Data for the node</param>
+		/// <param name="refs">References to other nodes</param>
 		/// <param name="cancellationToken">Cancellation token for the operation</param>
-		Task FlushAsync(ITreeBlob root, CancellationToken cancellationToken = default);
+		/// <returns>Reference to the node that was added</returns>
+		Task FlushAsync(ReadOnlySequence<byte> data, IReadOnlyList<ITreeBlobRef> refs, CancellationToken cancellationToken = default);
+	}
+
+	/// <summary>
+	/// Utility class
+	/// </summary>
+	public static class TreeBlob
+	{
+		class TreeBlobImpl : ITreeBlob
+		{
+			public ReadOnlySequence<byte> Data { get; }
+			public IReadOnlyList<ITreeBlobRef> Refs { get; }
+
+			public TreeBlobImpl(ReadOnlySequence<byte> data, IReadOnlyList<ITreeBlobRef> refs)
+			{
+				Data = data;
+				Refs = refs;
+			}
+		}
+
+		/// <summary>
+		/// Create a blob from the given parameters
+		/// </summary>
+		public static ITreeBlob Create(ReadOnlySequence<byte> Data, IReadOnlyList<ITreeBlobRef> Refs)
+		{
+			return new TreeBlobImpl(Data, Refs);
+		}
 	}
 
 	/// <summary>
@@ -141,5 +175,16 @@ namespace EpicGames.Horde.Storage
 		{
 			return new TypedTreeStore<T>(store);
 		}
+
+		/// <summary>
+		/// Flush the root of a tree to the store
+		/// </summary>
+		/// <param name="writer">Writer for the nodes</param>
+		/// <param name="root">Root blob</param>
+		/// <param name="cancellationToken">Cancellation token for the operation</param>
+		public static Task FlushAsync(this ITreeWriter writer, ITreeBlob root, CancellationToken cancellationToken = default)
+		{
+			return writer.FlushAsync(root.Data, root.Refs, cancellationToken);
+		}
 	}
 }
diff --git a/Engine/Source/Programs/Shared/EpicGames.Horde/Storage/Nodes/DirectoryNode.cs b/Engine/Source/Programs/Shared/EpicGames.Horde/Storage/Nodes/DirectoryNode.cs
--- a/Engine/Source/Programs/Shared/EpicGames.Horde/Storage/Nodes/DirectoryNode.cs
+++ b/Engine/Source/Programs/Shared/EpicGames.Horde/Storage/Nodes/DirectoryNode.cs
@@ -1,6 +1,7 @@
 // Copyright Epic Games, Inc. All Rights Reserved.
 
 using EpicGames.Core;
+using EpicGames.Horde.Storage.Git;
 using EpicGames.Serialization;
 using Microsoft.Extensions.Logging;
 using System;
@@ -10,13 +11,14 @@ using System.Diagnostics;
 using System.Diagnostics.CodeAnalysis;
 using System.IO;
 using System.Linq;
+using System.Security.Cryptography;
 using System.Threading;
 using System.Threading.Tasks;
 
 namespace EpicGames.Horde.Storage.Nodes
 {
 	/// <summary>
-	/// Flags for a directory entry
+	/// Flags for a file entry
 	/// </summary>
 	[Flags]
 	public enum FileEntryFlags
@@ -41,15 +43,10 @@ namespace EpicGames.Horde.Storage.Nodes
 		/// </summary>
 		Text = 16,
 
-		/// <summary>
-		/// The attached entry includes a Git SHA1 of the corresponding blob/tree contents.
-		/// </summary>
-		HasGitSha1 = 32,
-
 		/// <summary>
 		/// The data for this entry is a Perforce depot path and revision rather than the actual file contents.
 		/// </summary>
-		PerforceDepotPathAndRevision = 64,
+		PerforceDepotPathAndRevision = 32,
 	}
 
 	/// <summary>
@@ -65,13 +62,18 @@ namespace EpicGames.Horde.Storage.Nodes
 		/// <summary>
 		/// Flags for this file
 		/// </summary>
-		public FileEntryFlags Flags { get; }
+		public FileEntryFlags Flags { get; set; }
 
 		/// <summary>
 		/// Length of this entry
 		/// </summary>
 		public long Length => (Node == null) ? _cachedLength : Node.Length;
 
+		/// <summary>
+		/// SHA1 hash of this file, with Git prefix
+		/// </summary>
+		public Sha1Hash GitHash { get; set; }
+
 		/// <summary>
 		/// Cached length of this node
 		/// </summary>
@@ -159,6 +161,11 @@ namespace EpicGames.Horde.Storage.Nodes
 		/// </summary>
 		public long Length => (Node == null) ? _cachedLength : Node.Length;
 
+		/// <summary>
+		/// Optional Git hash for the referenced directory object.
+		/// </summary>
+		public Sha1Hash GitHash { get; set; }
+
 		/// <summary>
 		/// Cached value for the length of this tree
 		/// </summary>
@@ -188,12 +195,32 @@ namespace EpicGames.Horde.Storage.Nodes
 		protected override void OnCollapse()
 		{
 			_cachedLength = Node!.Length;
+			if (_owner is DirectoryNode directoryNode && (directoryNode.Flags & DirectoryFlags.WithGitHashes) != 0)
+			{
+				GitHash = Node!.AsGitTree().GetHash();
+			}
 		}
 
 		/// <inheritdoc/>
 		public override string ToString() => IsDirty() ? $"{Name} (*)" : Name.ToString();
 	}
 
+	/// <summary>
+	/// Flags for a directory node
+	/// </summary>
+	public enum DirectoryFlags
+	{
+		/// <summary>
+		/// No flags specified
+		/// </summary>
+		None = 0,
+
+		/// <summary>
+		/// Includes hashes for all entries
+		/// </summary>
+		WithGitHashes = 1,
+	}
+
 	/// <summary>
 	/// A directory node
 	/// </summary>
@@ -210,6 +237,11 @@ namespace EpicGames.Horde.Storage.Nodes
 		/// </summary>
 		public long Length => Files.Sum(x => x.Length) + Directories.Sum(x => x.Length);
 
+		/// <summary>
+		/// Flags for this directory 
+		/// </summary>
+		public DirectoryFlags Flags { get; }
+
 		/// <summary>
 		/// All the files within this directory
 		/// </summary>
@@ -220,6 +252,15 @@ namespace EpicGames.Horde.Storage.Nodes
 		/// </summary>
 		public IReadOnlyCollection<DirectoryEntry> Directories => _nameToDirectoryEntry.Values;
 
+		/// <summary>
+		/// Constructor
+		/// </summary>
+		/// <param name="flags"></param>
+		public DirectoryNode(DirectoryFlags flags)
+		{
+			Flags = flags;
+		}
+
 		/// <summary>
 		/// Clear the contents of this directory
 		/// </summary>
@@ -256,7 +297,7 @@ namespace EpicGames.Horde.Storage.Nodes
 		/// <returns>The new directory object</returns>
 		public FileEntry AddFile(Utf8String name, FileEntryFlags flags)
 		{
-			if(TryGetDirectoryEntry(name, out _))
+			if (TryGetDirectoryEntry(name, out _))
 			{
 				throw new ArgumentException($"A directory with the name {name} already exists");
 			}
@@ -292,10 +333,9 @@ namespace EpicGames.Horde.Storage.Nodes
 		/// </summary>
 		/// <param name="path">Path to the file</param>
 		/// <param name="flags">Flags for the new file</param>
-		/// <param name="writer">Writer for new node data</param>
 		/// <param name="cancellationToken">Cancellation token for the operation</param>
 		/// <returns>The new directory object</returns>
-		public async ValueTask<FileEntry> AddFileByPathAsync(Utf8String path, FileEntryFlags flags, ITreeWriter writer, CancellationToken cancellationToken = default)
+		public async ValueTask<FileEntry> AddFileByPathAsync(Utf8String path, FileEntryFlags flags, CancellationToken cancellationToken = default)
 		{
 			DirectoryNode directory = this;
 
@@ -369,7 +409,7 @@ namespace EpicGames.Horde.Storage.Nodes
 		public async ValueTask<bool> DeleteFileByPathAsync(Utf8String path, CancellationToken cancellationToken)
 		{
 			Utf8String remainingPath = path;
-			for (DirectoryNode? directory = this; directory != null; )
+			for (DirectoryNode? directory = this; directory != null;)
 			{
 				int length = remainingPath.IndexOf('/');
 				if (length == -1)
@@ -401,7 +441,7 @@ namespace EpicGames.Horde.Storage.Nodes
 				throw new ArgumentException($"A file with the name '{name}' already exists in this directory", nameof(name));
 			}
 
-			DirectoryNode node = new DirectoryNode();
+			DirectoryNode node = new DirectoryNode(Flags);
 
 			DirectoryEntry entry = new DirectoryEntry(this, name, node);
 			_nameToDirectoryEntry.Add(name, entry);
@@ -466,7 +506,7 @@ namespace EpicGames.Horde.Storage.Nodes
 		/// <returns>True if the entry was found, false otherwise</returns>
 		public bool DeleteDirectory(Utf8String name)
 		{
-			if(_nameToDirectoryEntry.Remove(name))
+			if (_nameToDirectoryEntry.Remove(name))
 			{
 				MarkAsDirty();
 				return true;
@@ -493,13 +533,14 @@ namespace EpicGames.Horde.Storage.Nodes
 				refs.Add(directoryEntry);
 			}
 
-			return new NewTreeBlob(SerializeData(fileEntries, directoryEntries), refs);
+			return new NewTreeBlob(SerializeData(Flags, fileEntries, directoryEntries), refs);
 		}
 
-		static ReadOnlySequence<byte> SerializeData(List<FileEntry> fileEntries, List<DirectoryEntry> directoryEntries)
+		static ReadOnlySequence<byte> SerializeData(DirectoryFlags flags, List<FileEntry> fileEntries, List<DirectoryEntry> directoryEntries)
 		{
 			// Measure the required size of the write buffer
 			int size = 1;
+			size += VarInt.MeasureUnsigned((ulong)flags);
 			size += VarInt.MeasureUnsigned(fileEntries.Count);
 			foreach (FileEntry fileEntry in fileEntries)
 			{
@@ -518,6 +559,9 @@ namespace EpicGames.Horde.Storage.Nodes
 			span[0] = TypeId;
 			span = span.Slice(1);
 
+			int directoryFlagBytes = VarInt.WriteUnsigned(span, (ulong)flags);
+			span = span.Slice(directoryFlagBytes);
+
 			int fileCountBytes = VarInt.WriteUnsigned(span, fileEntries.Count);
 			span = span.Slice(fileCountBytes);
 
@@ -531,6 +575,12 @@ namespace EpicGames.Horde.Storage.Nodes
 
 				int lengthBytes = VarInt.WriteUnsigned(span, fileEntry.Length);
 				span = span.Slice(lengthBytes);
+
+				if ((flags & DirectoryFlags.WithGitHashes) != 0)
+				{
+					fileEntry.GitHash.CopyTo(span);
+					span = span.Slice(Sha1Hash.NumBytes);
+				}
 			}
 
 			int directoryCountBytes = VarInt.WriteUnsigned(span, directoryEntries.Count);
@@ -543,6 +593,12 @@ namespace EpicGames.Horde.Storage.Nodes
 
 				int lengthBytes = VarInt.WriteUnsigned(span, directoryEntry.Length);
 				span = span.Slice(lengthBytes);
+
+				if ((flags & DirectoryFlags.WithGitHashes) != 0)
+				{
+					directoryEntry.GitHash.CopyTo(span);
+					span = span.Slice(Sha1Hash.NumBytes);
+				}
 			}
 
 			Debug.Assert(span.Length == 0);
@@ -557,14 +613,17 @@ namespace EpicGames.Horde.Storage.Nodes
 		public static DirectoryNode Deserialize(ITreeBlob source)
 		{
 			ReadOnlySpan<byte> span = source.Data.AsSingleSegment().Span;
-			if(span[0] != TypeId)
+			if (span[0] != TypeId)
 			{
 				throw new InvalidOperationException("Invalid signature byte for directory");
 			}
 
 			span = span.Slice(1);
 
-			DirectoryNode node = new DirectoryNode();
+			DirectoryFlags directoryFlags = (DirectoryFlags)VarInt.ReadUnsigned(span, out int directoryFlagBytes);
+			span = span.Slice(directoryFlagBytes);
+
+			DirectoryNode node = new DirectoryNode(directoryFlags);
 
 			int fileCount = (int)VarInt.ReadUnsigned(span, out int fileCountBytes);
 			span = span.Slice(fileCountBytes);
@@ -578,14 +637,22 @@ namespace EpicGames.Horde.Storage.Nodes
 				Utf8String name = new Utf8String(span.Slice(0, nameLen).ToArray());
 				span = span.Slice(nameLen + 1);
 
-				FileEntryFlags flags = (FileEntryFlags)span[0];
-				span = span.Slice(1);
+				FileEntryFlags flags = (FileEntryFlags)VarInt.ReadUnsigned(span, out int flagBytes);
+				span = span.Slice(flagBytes);
 
 				long length = (long)VarInt.ReadUnsigned(span, out int lengthBytes);
 				span = span.Slice(lengthBytes);
 
 				ITreeBlobRef child = source.Refs[childIdx++];
-				node._nameToFileEntry[name] = new FileEntry(node, name, flags, length, child);
+				FileEntry entry = new FileEntry(node, name, flags, length, child);
+
+				if ((directoryFlags & DirectoryFlags.WithGitHashes) != 0)
+				{
+					entry.GitHash = new Sha1Hash(span);
+					span = span.Slice(Sha1Hash.NumBytes);
+				}
+
+				node._nameToFileEntry[name] = entry;
 			}
 
 			int directoryCount = (int)VarInt.ReadUnsigned(span, out int directoryCountBytes);
@@ -602,7 +669,15 @@ namespace EpicGames.Horde.Storage.Nodes
 				span = span.Slice(lengthBytes);
 
 				ITreeBlobRef child = source.Refs[childIdx++];
-				node._nameToDirectoryEntry[name] = new DirectoryEntry(node, name, length, child);
+				DirectoryEntry entry = new DirectoryEntry(node, name, length, child);
+
+				if ((directoryFlags & DirectoryFlags.WithGitHashes) != 0)
+				{
+					entry.GitHash = new Sha1Hash(span);
+					span = span.Slice(Sha1Hash.NumBytes);
+				}
+
+				node._nameToDirectoryEntry[name] = entry;
 			}
 
 			Debug.Assert(span.Length == 0);
@@ -663,6 +738,25 @@ namespace EpicGames.Horde.Storage.Nodes
 
 			await Task.WhenAll(tasks);
 		}
+
+		/// <summary>
+		/// Creates a Git tree object for this directory
+		/// </summary>
+		/// <returns>Tree object</returns>
+		public GitTree AsGitTree()
+		{
+			GitTree tree = new GitTree();
+			foreach (FileEntry fileEntry in _nameToFileEntry.Values)
+			{
+				tree.Entries.Add(new GitTreeEntry(GitFileMode.File, fileEntry.Name, fileEntry.GitHash));
+			}
+			foreach (DirectoryEntry directoryEntry in _nameToDirectoryEntry.Values)
+			{
+				tree.Entries.Add(new GitTreeEntry(GitFileMode.Tree, directoryEntry.Name, directoryEntry.GitHash));
+			}
+			tree.Entries.SortBy(x => x.Name);
+			return tree;
+		}
 	}
 
 	/// <summary>
diff --git a/Engine/Source/Programs/Shared/EpicGames.Horde/Storage/TreeNode.cs b/Engine/Source/Programs/Shared/EpicGames.Horde/Storage/TreeNode.cs
--- a/Engine/Source/Programs/Shared/EpicGames.Horde/Storage/TreeNode.cs
+++ b/Engine/Source/Programs/Shared/EpicGames.Horde/Storage/TreeNode.cs
@@ -1,4 +1,4 @@
-﻿// Copyright Epic Games, Inc. All Rights Reserved.
+// Copyright Epic Games, Inc. All Rights Reserved.
 
 using System;
 using System.Buffers;
@@ -148,10 +148,8 @@ namespace EpicGames.Horde.Storage
 		/// <param name="cancellationToken">Cancellation token for the operation</param>
 		public static async Task FlushAsync(this ITreeWriter writer, TreeNode root, CancellationToken cancellationToken = default)
 		{
-			ITreeBlobRef rootBlobRef = await writer.WriteTreeAsync(root, cancellationToken);
-			ITreeBlob rootBlob = await rootBlobRef.GetTargetAsync(cancellationToken);
-
-			await writer.FlushAsync(rootBlob, cancellationToken);
+			ITreeBlob rootBlob = await writer.WriteTreeAsync(root, cancellationToken);
+			await writer.FlushAsync(rootBlob.Data, rootBlob.Refs, cancellationToken);
 		}
 
 		/// <summary>
@@ -161,7 +159,7 @@ namespace EpicGames.Horde.Storage
 		/// <param name="node">Root node to serialize</param>
 		/// <param name="cancellationToken">Cancellation token for the operation</param>
 		/// <returns>Reference to the new tree of nodes</returns>
-		public static async Task<ITreeBlobRef> WriteTreeAsync(this ITreeWriter writer, TreeNode node, CancellationToken cancellationToken = default)
+		public static async Task<ITreeBlob> WriteTreeAsync(this ITreeWriter writer, TreeNode node, CancellationToken cancellationToken = default)
 		{
 			NewTreeBlob blob = node.Serialize();
 
@@ -171,7 +169,7 @@ namespace EpicGames.Horde.Storage
 				targetRefs.Add(await typedRef.CollapseAsync(writer, cancellationToken));
 			}
 
-			return await writer.WriteNodeAsync(blob.Data, targetRefs, cancellationToken);
+			return TreeBlob.Create(blob.Data, targetRefs);
 		}
 
 		/// <inheritdoc/>
