Horde: Missing files.

diff --git a/Engine/Source/Programs/Shared/EpicGames.Horde/Storage/Git/GitObject.cs b/Engine/Source/Programs/Shared/EpicGames.Horde/Storage/Git/GitObject.cs
new file mode 100644
--- /dev/null
+++ b/Engine/Source/Programs/Shared/EpicGames.Horde/Storage/Git/GitObject.cs
@@ -0,0 +1,72 @@
+// Copyright Epic Games, Inc. All Rights Reserved.
+
+using System;
+using System.Security.Cryptography;
+using EpicGames.Core;
+
+namespace EpicGames.Horde.Storage.Git
+{
+	/// <summary>
+	/// Type fields for git objects
+	/// </summary>
+	public static class GitObjectType
+	{
+		/// <summary>
+		/// Blob objects
+		/// </summary>
+		public static Utf8String Blob { get; } = "blob";
+
+		/// <summary>
+		/// Tree objects
+		/// </summary>
+		public static Utf8String Tree { get; } = "tree";
+	}
+
+	/// <summary>
+	/// Utility methods for manipulating Git objects
+	/// </summary>
+	public static class GitObject
+	{
+		/// <summary>
+		/// Appends data for a header to the hash
+		/// </summary>
+		/// <param name="type">Type of the object</param>
+		/// <param name="size">Size of the object</param>
+		/// <param name="hash">Hash for the header data</param>
+		public static void WriteHeader(Utf8String type, long size, IncrementalHash hash)
+		{
+			Span<byte> header = WriteHeader(type, size, stackalloc byte[32]);
+			hash.AppendData(header);
+		}
+
+		/// <summary>
+		/// Writes a header for an object
+		/// </summary>
+		/// <param name="type">Type of the object</param>
+		/// <param name="size">Size of the object</param>
+		/// <param name="buffer">Buffer to receive the data</param>
+		/// <returns></returns>
+		public static Span<byte> WriteHeader(Utf8String type, long size, Span<byte> buffer)
+		{
+			type.Span.CopyTo(buffer);
+			buffer[type.Length] = (byte)' ';
+
+			int start = type.Length + 1;
+			int end = start;
+			for (; ; )
+			{
+				buffer[end++] = (byte)('0' + (size % 10));
+				size /= 10;
+				if (size == 0)
+				{
+					break;
+				}
+			}
+
+			buffer.Slice(start, end - start).Reverse();
+			buffer[end] = 0;
+
+			return buffer.Slice(0, end + 1);
+		}
+	}
+}
diff --git a/Engine/Source/Programs/Shared/EpicGames.Horde/Storage/Git/GitTree.cs b/Engine/Source/Programs/Shared/EpicGames.Horde/Storage/Git/GitTree.cs
new file mode 100644
--- /dev/null
+++ b/Engine/Source/Programs/Shared/EpicGames.Horde/Storage/Git/GitTree.cs
@@ -0,0 +1,144 @@
+// Copyright Epic Games, Inc. All Rights Reserved.
+
+using System;
+using System.Collections.Generic;
+using System.Security.Cryptography;
+using EpicGames.Core;
+
+namespace EpicGames.Horde.Storage.Git
+{
+	/// <summary>
+	/// Valid file modes for tree entries
+	/// </summary>
+	public static class GitFileMode
+	{
+		/// <summary>
+		/// Regular file
+		/// </summary>
+		public static Utf8String File { get; } = "100644";
+
+		/// <summary>
+		/// Executable file
+		/// </summary>
+		public static Utf8String ExecutableFile { get; } = "100755";
+
+		/// <summary>
+		/// Child tree
+		/// </summary>
+		public static Utf8String Tree { get; } = "40000";
+	}
+
+	/// <summary>
+	/// Representation of a Git tree object
+	/// </summary>
+	public class GitTree
+	{
+		/// <summary>
+		/// Entries for this tree. Sort be sorted for consistency.
+		/// </summary>
+		public List<GitTreeEntry> Entries { get; } = new List<GitTreeEntry>();
+
+		/// <summary>
+		/// Gets the hash of this tree object.
+		/// </summary>
+		/// <returns>Hash of the object</returns>
+		public Sha1Hash GetHash()
+		{
+			Span<byte> buffer = stackalloc byte[32];
+			using (IncrementalHash hash = IncrementalHash.CreateHash(HashAlgorithmName.SHA1))
+			{
+				int size = GetPayloadSize();
+				GitObject.WriteHeader(GitObjectType.Tree, size, hash);
+
+				foreach (GitTreeEntry entry in Entries)
+				{
+					hash.AppendData(entry.Mode.Span);
+
+					buffer[0] = (byte)' ';
+					hash.AppendData(buffer.Slice(0, 1));
+
+					hash.AppendData(entry.Name.Span);
+
+					buffer[0] = 0;
+					entry.Hash.CopyTo(buffer.Slice(1));
+					hash.AppendData(buffer.Slice(0, Sha1Hash.NumBytes + 1));
+				}
+
+				hash.GetHashAndReset(buffer);
+				return new Sha1Hash(buffer);
+			}
+		}
+
+		/// <summary>
+		/// Serializes this object
+		/// </summary>
+		/// <returns></returns>
+		public byte[] Serialize()
+		{
+			int size = GetPayloadSize();
+
+			Span<byte> header = GitObject.WriteHeader(GitObjectType.Tree, size, stackalloc byte[32]);
+
+			byte[] data = new byte[header.Length + size];
+			header.CopyTo(data);
+
+			Span<byte> output = data.AsSpan(header.Length);
+			foreach (GitTreeEntry entry in Entries)
+			{
+				entry.Mode.Span.CopyTo(output);
+				output[entry.Mode.Length] = (byte)' ';
+				output = output.Slice(entry.Mode.Length + 1);
+
+				entry.Name.Span.CopyTo(output);
+				output[entry.Name.Length] = 0;
+				output = output.Slice(entry.Name.Length + 1);
+
+				entry.Hash.CopyTo(output);
+				output = output.Slice(Sha1Hash.NumBytes);
+			}
+
+			return data;
+		}
+
+		int GetPayloadSize()
+		{
+			int size = 0;
+			foreach (GitTreeEntry entry in Entries)
+			{
+				size += entry.Mode.Length + 1 + entry.Name.Length + 1 + Sha1Hash.NumBytes;
+			}
+			return size;
+		}
+	}
+
+	/// <summary>
+	/// Entry for a Git tree
+	/// </summary>
+	public class GitTreeEntry
+	{
+		/// <summary>
+		/// Mode for this entry. Can be any values from <see cref="GitFileMode"/>
+		/// </summary>
+		public Utf8String Mode { get; }
+
+		/// <summary>
+		/// Name of this entry
+		/// </summary>
+		public Utf8String Name { get; }
+
+		/// <summary>
+		/// Hash of the object for this entry
+		/// </summary>
+		public Sha1Hash Hash { get; }
+
+		/// <summary>
+		/// Constructor
+		/// </summary>
+		public GitTreeEntry(Utf8String mode, Utf8String name, Sha1Hash hash)
+		{
+			Mode = mode;
+			Name = name;
+			Hash = hash;
+		}
+	}
+}
