Horde: Missing file.

diff --git a/Engine/Source/Programs/Shared/EpicGames.Horde/Storage/TreeNodeRef.cs b/Engine/Source/Programs/Shared/EpicGames.Horde/Storage/TreeNodeRef.cs
--- a/Engine/Source/Programs/Shared/EpicGames.Horde/Storage/TreeNodeRef.cs
+++ b/Engine/Source/Programs/Shared/EpicGames.Horde/Storage/TreeNodeRef.cs
@@ -40,6 +40,11 @@ namespace EpicGames.Horde.Storage
 		/// </summary>
 		internal ITreeBlobRef? _target;
 
+		/// <summary>
+		/// Reference to the node in storage, if unmodified.
+		/// </summary>
+		public ITreeBlobRef? Target => ReferenceEquals(_strongRef, null) ? null : _target;
+
 		/// <summary>
 		/// Reference to the node.
 		/// </summary>
@@ -176,7 +181,15 @@ namespace EpicGames.Horde.Storage
 			ITreeBlobRef? target = _target;
 			if (target == null)
 			{
-				target = await writer.WriteTreeAsync(Node!, cancellationToken);
+				NewTreeBlob blob = Node!.Serialize();
+
+				List<ITreeBlobRef> targetRefs = new List<ITreeBlobRef>();
+				foreach (TreeNodeRef typedRef in blob.Refs)
+				{
+					targetRefs.Add(await typedRef.CollapseAsync(writer, cancellationToken));
+				}
+
+				target = await writer.WriteNodeAsync(blob.Data, targetRefs, cancellationToken);
 				MarkAsClean(target);
 			}
 			return target;
