[UBT/EpicGames.Build]
* Removed debug code accidentally submitted (submitted from active change than shelf and ended up submitting these lines)

diff --git a/Engine/Source/Programs/Shared/EpicGames.Build/System/Rules.cs b/Engine/Source/Programs/Shared/EpicGames.Build/System/Rules.cs
--- a/Engine/Source/Programs/Shared/EpicGames.Build/System/Rules.cs
+++ b/Engine/Source/Programs/Shared/EpicGames.Build/System/Rules.cs
@@ -137,10 +137,6 @@ namespace UnrealBuildBase
 				}
 			}
 
-			foreach (DirectoryReference f in Folders)
-				if (f.FullName[0] == 'e')
-					System.Console.WriteLine("!LOWERCASE!!!!!!!!!!");
-
 			return FindAllRulesFiles(Folders, RulesFileType);
 		}
 
