Fix missing documentation comment.

diff --git a/Engine/Source/Programs/Shared/EpicGames.Core/BuzHash.cs b/Engine/Source/Programs/Shared/EpicGames.Core/BuzHash.cs
--- a/Engine/Source/Programs/Shared/EpicGames.Core/BuzHash.cs
+++ b/Engine/Source/Programs/Shared/EpicGames.Core/BuzHash.cs
@@ -154,6 +154,7 @@ namespace EpicGames.Core
 		/// <param name="prev">Previous data added to the hash</param>
 		/// <param name="next">Next data to be added to the hash</param>
 		/// <param name="threshold">Threshold for hash below which to terminate the scan</param>
+		/// <param name="state">The current hash value</param>
 		/// <returns>Number of bytes consumed from next</returns>
 		public static unsafe int Update(ReadOnlySpan<byte> prev, ReadOnlySpan<byte> next, uint threshold, ref uint state)
 		{
