CoreNet.h will now be automatically included when needed.

diff --git a/Engine/Source/Programs/Shared/EpicGames.UHT/Exporters/CodeGen/UhtHeaderCodeGeneratorCppFile.cs b/Engine/Source/Programs/Shared/EpicGames.UHT/Exporters/CodeGen/UhtHeaderCodeGeneratorCppFile.cs
--- a/Engine/Source/Programs/Shared/EpicGames.UHT/Exporters/CodeGen/UhtHeaderCodeGeneratorCppFile.cs
+++ b/Engine/Source/Programs/Shared/EpicGames.UHT/Exporters/CodeGen/UhtHeaderCodeGeneratorCppFile.cs
@@ -41,22 +41,31 @@ namespace EpicGames.UHT.Exporters.CodeGen
 
 				bool addedStructuredArchiveFromArchiveHeader = false;
 				bool addedArchiveUObjectFromStructuredArchiveHeader = false;
+				bool addedCoreNetHeader = false;
 				HashSet<UhtHeaderFile> addedIncludes = new();
 				List<string> includesToAdd = new();
-				addedIncludes.Add(this.HeaderFile);
+				addedIncludes.Add(HeaderFile);
 
 				if (headerInfo.NeedsFastArrayHeaders)
 				{
 					includesToAdd.Add("Net/Serialization/FastArraySerializerImplementation.h");
 				}
 
-				foreach (UhtType type in this.HeaderFile.Children)
+				foreach (UhtType type in HeaderFile.Children)
 				{
 					if (type is UhtStruct structObj)
 					{
 						// Functions
 						foreach (UhtFunction function in structObj.Functions)
 						{
+							if (!function.FunctionExportFlags.HasAnyFlags(UhtFunctionExportFlags.CppStatic) && function.FunctionFlags.HasAnyFlags(EFunctionFlags.NetValidate))
+							{
+								if (!addedCoreNetHeader)
+								{
+									includesToAdd.Add("UObject/CoreNet.h");
+									addedCoreNetHeader = true;
+								}
+							}
 							foreach (UhtProperty property in function.Properties)
 							{
 								AddIncludeForProperty(property, addedIncludes, includesToAdd);
@@ -106,7 +115,7 @@ namespace EpicGames.UHT.Exporters.CodeGen
 					{
 						if (!type.HeaderFile.IsNoExportTypes && addedIncludes.Add(type.HeaderFile))
 						{
-							includesToAdd.Add(this.HeaderInfos[type.HeaderFile.HeaderFileTypeIndex].IncludePath);
+							includesToAdd.Add(HeaderInfos[type.HeaderFile.HeaderFileTypeIndex].IncludePath);
 						}
 					}
 				}
diff --git a/Engine/Source/Programs/UnrealHeaderTool/Private/CodeGenerator.cpp b/Engine/Source/Programs/UnrealHeaderTool/Private/CodeGenerator.cpp
--- a/Engine/Source/Programs/UnrealHeaderTool/Private/CodeGenerator.cpp
+++ b/Engine/Source/Programs/UnrealHeaderTool/Private/CodeGenerator.cpp
@@ -5501,6 +5501,7 @@ bool FNativeClassHeaderGenerator::WriteSource(const FManifestModule& Module, FGe
 
 		bool bAddedStructuredArchiveFromArchiveHeader = false;
 		bool bAddedArchiveUObjectFromStructuredArchiveHeader = false;
+		bool bAddedCoreNetHeader = false;
 
 		// We need to include the headers that declare any types that aren't pointers.
 		// We can't rely on the original header that generates this cpp to include the 
@@ -5512,6 +5513,14 @@ bool FNativeClassHeaderGenerator::WriteSource(const FManifestModule& Module, FGe
 				// Functions
 				for (const TSharedRef<FUnrealFunctionDefinitionInfo>& FunctionDef : Struct->GetFunctions())
 				{
+					if (!(FunctionDef->GetFunctionData().FunctionExportFlags & FUNCEXPORT_CppStatic) && FunctionDef->HasAnyFunctionFlags(FUNC_NetValidate))
+					{
+						if (!bAddedCoreNetHeader)
+						{
+							RelativeIncludes.AddUnique(TEXT("UObject/CoreNet.h"));
+							bAddedCoreNetHeader = true;
+						}
+					}
 					for (const TSharedRef<FUnrealPropertyDefinitionInfo>& PropertyDef : FunctionDef->GetProperties())
 					{
 						AddIncludeForProperty(PropertyDef, RelativeIncludes);
