Horde: Fix warnings in EpicGames.Slack

Also changed to run analyzers during build. When referencing this library from a project that does run analyzers, it failed.
Likely is for the better to enforce that here as well.

diff --git a/Engine/Source/Programs/Shared/EpicGames.Slack/Blocks/SectionBlock.cs b/Engine/Source/Programs/Shared/EpicGames.Slack/Blocks/SectionBlock.cs
--- a/Engine/Source/Programs/Shared/EpicGames.Slack/Blocks/SectionBlock.cs
+++ b/Engine/Source/Programs/Shared/EpicGames.Slack/Blocks/SectionBlock.cs
@@ -24,7 +24,7 @@ namespace EpicGames.Slack.Blocks
 		/// Maximum number of items is 10. Maximum length for the text in each item is 2000 characters.
 		/// </summary>
 		[JsonPropertyName("fields")]
-		public List<TextObject>? Fields { get; set; }
+		public IReadOnlyList<TextObject>? Fields { get; set; }
 
 		/// <summary>
 		/// One of the available element objects.
diff --git a/Engine/Source/Programs/Shared/EpicGames.Slack/Elements/CheckboxGroupElement.cs b/Engine/Source/Programs/Shared/EpicGames.Slack/Elements/CheckboxGroupElement.cs
--- a/Engine/Source/Programs/Shared/EpicGames.Slack/Elements/CheckboxGroupElement.cs
+++ b/Engine/Source/Programs/Shared/EpicGames.Slack/Elements/CheckboxGroupElement.cs
@@ -25,13 +25,13 @@ namespace EpicGames.Slack.Elements
 		/// An array of option objects. A maximum of 10 options are allowed.
 		/// </summary>
 		[JsonPropertyName("options")]
-		public List<Option> Options { get; } = new List<Option>();
+		public List<SlackOption> Options { get; } = new List<SlackOption>();
 
 		/// <summary>
 		/// An array of option objects that exactly matches one or more of the options within options. These options will be selected when the checkbox group initially loads.
 		/// </summary>
 		[JsonPropertyName("initial_options")]
-		public List<Option> InitialOptions { get; } = new List<Option>();
+		public List<SlackOption> InitialOptions { get; } = new List<SlackOption>();
 
 		/// <summary>
 		/// A confirm object that defines an optional confirmation dialog that appears after clicking one of the checkboxes in this element.
@@ -48,7 +48,7 @@ namespace EpicGames.Slack.Elements
 		/// <summary>
 		/// Construct a new Button action element.
 		/// </summary>
-		public CheckboxGroupElement(string actionId, List<Option> options)
+		public CheckboxGroupElement(string actionId, List<SlackOption> options)
 			: base("checkboxes")
 		{
 			ActionId = actionId;
diff --git a/Engine/Source/Programs/Shared/EpicGames.Slack/Elements/OverflowMenuElement.cs b/Engine/Source/Programs/Shared/EpicGames.Slack/Elements/OverflowMenuElement.cs
--- a/Engine/Source/Programs/Shared/EpicGames.Slack/Elements/OverflowMenuElement.cs
+++ b/Engine/Source/Programs/Shared/EpicGames.Slack/Elements/OverflowMenuElement.cs
@@ -29,7 +29,7 @@ namespace EpicGames.Slack.Elements
 		/// An array of up to five option objects to display in the menu.
 		/// </summary>
 		[JsonPropertyName("options")]
-		public List<Option> Options { get; } = new List<Option>();
+		public List<SlackOption> Options { get; } = new List<SlackOption>();
 
 		/// <summary>
 		/// A confirm object that defines an optional confirmation dialog that appears after clicking one of the checkboxes in this element.
@@ -40,7 +40,7 @@ namespace EpicGames.Slack.Elements
 		/// <summary>
 		/// Construct a new Button action element.
 		/// </summary>
-		public OverflowMenuElement(string actionId, List<Option> options, SlackConfirm? confirm = null)
+		public OverflowMenuElement(string actionId, List<SlackOption> options, SlackConfirm? confirm = null)
 			: base("overflow")
 		{
 			ActionId = actionId;
diff --git a/Engine/Source/Programs/Shared/EpicGames.Slack/Elements/RadioButtonGroupElement.cs b/Engine/Source/Programs/Shared/EpicGames.Slack/Elements/RadioButtonGroupElement.cs
--- a/Engine/Source/Programs/Shared/EpicGames.Slack/Elements/RadioButtonGroupElement.cs
+++ b/Engine/Source/Programs/Shared/EpicGames.Slack/Elements/RadioButtonGroupElement.cs
@@ -25,13 +25,13 @@ namespace EpicGames.Slack.Elements
 		/// An array of option objects. A maximum of 10 options are allowed.
 		/// </summary>
 		[JsonPropertyName("options")]
-		public List<Option> Options { get; } = new List<Option>();
+		public List<SlackOption> Options { get; } = new List<SlackOption>();
 
 		/// <summary>
 		/// An array of option objects that exactly matches one or more of the options within options. These options will be selected when the checkbox group initially loads.
 		/// </summary>
 		[JsonPropertyName("initial_options")]
-		public List<Option> InitialOptions { get; } = new List<Option>();
+		public List<SlackOption> InitialOptions { get; } = new List<SlackOption>();
 
 		/// <summary>
 		/// A confirm object that defines an optional confirmation dialog that appears after clicking one of the checkboxes in this element.
@@ -48,7 +48,7 @@ namespace EpicGames.Slack.Elements
 		/// <summary>
 		/// Construct a new Button action element.
 		/// </summary>
-		public RadioButtonGroupElement(string actionId, List<Option> options)
+		public RadioButtonGroupElement(string actionId, List<SlackOption> options)
 			: base("radio_buttons")
 		{
 			ActionId = actionId;
diff --git a/Engine/Source/Programs/Shared/EpicGames.Slack/Elements/SelectMenuElement.cs b/Engine/Source/Programs/Shared/EpicGames.Slack/Elements/SelectMenuElement.cs
--- a/Engine/Source/Programs/Shared/EpicGames.Slack/Elements/SelectMenuElement.cs
+++ b/Engine/Source/Programs/Shared/EpicGames.Slack/Elements/SelectMenuElement.cs
@@ -45,19 +45,19 @@ namespace EpicGames.Slack.Elements
 		/// An array of option objects. A maximum of 10 options are allowed.
 		/// </summary>
 		[JsonPropertyName("options")]
-		public List<Option>? Options { get; set; }
+		public IReadOnlyList<SlackOption>? Options { get; set; }
 
 		/// <summary>
 		/// An array of option objects. A maximum of 10 options are allowed.
 		/// </summary>
 		[JsonPropertyName("option_groups")]
-		public List<OptionGroup>? OptionGroups { get; set; }
+		public IReadOnlyList<SlackOptionGroup>? OptionGroups { get; set; }
 
 		/// <summary>
 		/// An array of option objects that exactly matches one or more of the options within options. These options will be selected when the checkbox group initially loads.
 		/// </summary>
 		[JsonPropertyName("initial_option")]
-		public Option? InitialOption { get; }
+		public SlackOption? InitialOption { get; }
 
 		/// <summary>
 		/// A confirm object that defines an optional confirmation dialog that appears after clicking one of the checkboxes in this element.
@@ -74,21 +74,21 @@ namespace EpicGames.Slack.Elements
 		/// <summary>
 		/// Construct a new Button action element.
 		/// </summary>
-		public StaticSelectMenuElement(string actionId, PlainTextObject placeholder, List<Option> options)
+		public StaticSelectMenuElement(string actionId, PlainTextObject placeholder, List<SlackOption> options)
 			: base("static_select", actionId)
 		{
 			Placeholder = placeholder;
-			Options = new List<Option>(options);
+			Options = new List<SlackOption>(options);
 		}
 
 		/// <summary>
 		/// Construct a new Button action element.
 		/// </summary>
-		public StaticSelectMenuElement(string actionId, PlainTextObject placeholder, List<OptionGroup> optionGroups)
+		public StaticSelectMenuElement(string actionId, PlainTextObject placeholder, List<SlackOptionGroup> optionGroups)
 			: base("static_select", actionId)
 		{
 			Placeholder = placeholder;
-			OptionGroups = new List<OptionGroup>(optionGroups);
+			OptionGroups = new List<SlackOptionGroup>(optionGroups);
 		}
 	}
 }
diff --git a/Engine/Source/Programs/Shared/EpicGames.Slack/EpicGames.Slack.csproj b/Engine/Source/Programs/Shared/EpicGames.Slack/EpicGames.Slack.csproj
--- a/Engine/Source/Programs/Shared/EpicGames.Slack/EpicGames.Slack.csproj
+++ b/Engine/Source/Programs/Shared/EpicGames.Slack/EpicGames.Slack.csproj
@@ -10,7 +10,7 @@
     <Configurations>Debug;Release;Development;Analyze</Configurations>
     <Nullable>enable</Nullable>
     <DebugType>pdbonly</DebugType>
-    <RunAnalyzersDuringBuild>false</RunAnalyzersDuringBuild>
+    <RunAnalyzersDuringBuild>true</RunAnalyzersDuringBuild>
   </PropertyGroup>
   <PropertyGroup Condition=" '$(Configuration)|$(Platform)' == 'Debug|AnyCPU' ">
     <DocumentationFile></DocumentationFile>
diff --git a/Engine/Source/Programs/Shared/EpicGames.Slack/Option.cs b/Engine/Source/Programs/Shared/EpicGames.Slack/Option.cs
--- a/Engine/Source/Programs/Shared/EpicGames.Slack/Option.cs
+++ b/Engine/Source/Programs/Shared/EpicGames.Slack/Option.cs
@@ -11,7 +11,7 @@ namespace EpicGames.Slack
 	/// <summary>
 	/// An object that represents a single selectable item in a select menu, multi-select menu, checkbox group, radio button group, or overflow menu.
 	/// </summary>
-	public class Option
+	public class SlackOption
 	{
 		/// <summary>
 		/// A text object that defines the text shown in the option on the menu. Overflow, select, and multi-select menus can only use plain_text objects, while radio buttons and checkboxes can use mrkdwn text objects.
@@ -41,7 +41,7 @@ namespace EpicGames.Slack
 		/// <summary>
 		/// Constructor
 		/// </summary>
-		public Option(TextObject text, string value)
+		public SlackOption(TextObject text, string value)
 		{
 			Text = text;
 			Value = value;
@@ -50,7 +50,7 @@ namespace EpicGames.Slack
 		/// <summary>
 		/// Constructor
 		/// </summary>
-		public Option(string text, string value)
+		public SlackOption(string text, string value)
 			: this(new PlainTextObject(text), value)
 		{
 		}
@@ -59,7 +59,7 @@ namespace EpicGames.Slack
 	/// <summary>
 	/// Provides a way to group options in a select menu or multi-select menu.
 	/// </summary>
-	public class OptionGroup
+	public class SlackOptionGroup
 	{
 		/// <summary>
 		/// A plain_text only text object that defines the label shown above this group of options. Maximum length for the text in this field is 75 characters.
@@ -69,12 +69,12 @@ namespace EpicGames.Slack
 		/// <summary>
 		/// An array of option objects that belong to this specific group. Maximum of 100 items.
 		/// </summary>
-		public List<Option> Options { get; } = new List<Option>();
+		public List<SlackOption> Options { get; } = new List<SlackOption>();
 
 		/// <summary>
 		/// Constructor
 		/// </summary>
-		public OptionGroup(PlainTextObject label, List<Option> options)
+		public SlackOptionGroup(PlainTextObject label, List<SlackOption> options)
 		{
 			Label = label;
 			Options.AddRange(options);
