Horde: Fix static analyzer warnings in EpicGames.Slack.

diff --git a/Engine/Source/Programs/Shared/EpicGames.Slack/Blocks/DividerBlock.cs b/Engine/Source/Programs/Shared/EpicGames.Slack/Blocks/DividerBlock.cs
--- a/Engine/Source/Programs/Shared/EpicGames.Slack/Blocks/DividerBlock.cs
+++ b/Engine/Source/Programs/Shared/EpicGames.Slack/Blocks/DividerBlock.cs
@@ -5,7 +5,7 @@ using System.Collections.Generic;
 namespace EpicGames.Slack.Blocks
 {
 	/// <summary>
-	/// A content divider, like an <hr>, to split up different blocks inside of a message. The divider block is nice and neat, requiring only a type.
+	/// A content divider, like an &lt;hr&gt;, to split up different blocks inside of a message. The divider block is nice and neat, requiring only a type.
 	/// </summary>
 	public class DividerBlock : Block
 	{
diff --git a/Engine/Source/Programs/Shared/EpicGames.Slack/Blocks/HeaderBlock.cs b/Engine/Source/Programs/Shared/EpicGames.Slack/Blocks/HeaderBlock.cs
--- a/Engine/Source/Programs/Shared/EpicGames.Slack/Blocks/HeaderBlock.cs
+++ b/Engine/Source/Programs/Shared/EpicGames.Slack/Blocks/HeaderBlock.cs
@@ -43,6 +43,7 @@ namespace EpicGames.Slack.Blocks
 		/// Add an <see cref="HeaderBlock"/> to the list of blocks
 		/// </summary>
 		/// <param name="container">List of blocks</param>
+		/// <param name="text">Text for the header</param>
 		public static void AddHeader(this ISlackBlockContainer container, PlainTextObject text)
 		{
 			HeaderBlock block = new HeaderBlock(text);
diff --git a/Engine/Source/Programs/Shared/EpicGames.Slack/Blocks/ImageBlock.cs b/Engine/Source/Programs/Shared/EpicGames.Slack/Blocks/ImageBlock.cs
--- a/Engine/Source/Programs/Shared/EpicGames.Slack/Blocks/ImageBlock.cs
+++ b/Engine/Source/Programs/Shared/EpicGames.Slack/Blocks/ImageBlock.cs
@@ -50,7 +50,6 @@ namespace EpicGames.Slack.Blocks
 		/// <summary>
 		/// Add an <see cref="ImageBlock"/> to the list of blocks
 		/// </summary>
-		/// <param name="container">Block container</param>
 		public static void AddImage(this ISlackBlockContainer container, Uri imageUrl, string altText, PlainTextObject? title = null)
 		{
 			ImageBlock block = new ImageBlock(imageUrl, altText, title);
diff --git a/Engine/Source/Programs/Shared/EpicGames.Slack/Blocks/InputBlock.cs b/Engine/Source/Programs/Shared/EpicGames.Slack/Blocks/InputBlock.cs
--- a/Engine/Source/Programs/Shared/EpicGames.Slack/Blocks/InputBlock.cs
+++ b/Engine/Source/Programs/Shared/EpicGames.Slack/Blocks/InputBlock.cs
@@ -61,7 +61,6 @@ namespace EpicGames.Slack.Blocks
 		/// <summary>
 		/// Add an <see cref="ImageBlock"/> to the list of blocks
 		/// </summary>
-		/// <param name="container">Block container</param>
 		public static InputBlock AddInput(this ISlackBlockContainer container, PlainTextObject label, Element element)
 		{
 			InputBlock block = new InputBlock(label, element);
diff --git a/Engine/Source/Programs/Shared/EpicGames.Slack/Blocks/SectionBlock.cs b/Engine/Source/Programs/Shared/EpicGames.Slack/Blocks/SectionBlock.cs
--- a/Engine/Source/Programs/Shared/EpicGames.Slack/Blocks/SectionBlock.cs
+++ b/Engine/Source/Programs/Shared/EpicGames.Slack/Blocks/SectionBlock.cs
@@ -63,7 +63,6 @@ namespace EpicGames.Slack.Blocks
 		/// <summary>
 		/// Add an <see cref="SectionBlock"/> to the list of blocks
 		/// </summary>
-		/// <param name="container">Container for the blocks</param>
 		public static void AddSection(this ISlackBlockContainer container, TextObject text)
 		{
 			SectionBlock block = new SectionBlock(text);
diff --git a/Engine/Source/Programs/Shared/EpicGames.Slack/Elements/ButtonElement.cs b/Engine/Source/Programs/Shared/EpicGames.Slack/Elements/ButtonElement.cs
--- a/Engine/Source/Programs/Shared/EpicGames.Slack/Elements/ButtonElement.cs
+++ b/Engine/Source/Programs/Shared/EpicGames.Slack/Elements/ButtonElement.cs
@@ -53,7 +53,7 @@ namespace EpicGames.Slack.Elements
 	}
 
 	/// <summary>
-	/// An button element that can be added to an <see cref="ActionsBlock"/>.
+	/// An button element that can be added to an <see cref="Blocks.ActionsBlock"/>.
 	/// </summary>
 	public class ButtonElement : Element
 	{
diff --git a/Engine/Source/Programs/Shared/EpicGames.Slack/Elements/PlainTextInputElement.cs b/Engine/Source/Programs/Shared/EpicGames.Slack/Elements/PlainTextInputElement.cs
--- a/Engine/Source/Programs/Shared/EpicGames.Slack/Elements/PlainTextInputElement.cs
+++ b/Engine/Source/Programs/Shared/EpicGames.Slack/Elements/PlainTextInputElement.cs
@@ -9,7 +9,7 @@ using System.Text.Json.Serialization;
 namespace EpicGames.Slack.Elements
 {
 	/// <summary>
-	/// A plain-text input, similar to the HTML <input> tag, creates a field where a user can enter freeform data. It can appear as a single-line field or a larger text area using the multiline flag.
+	/// A plain-text input, similar to the HTML &lt;input&gt; tag, creates a field where a user can enter freeform data. It can appear as a single-line field or a larger text area using the multiline flag.
 	/// </summary>
 	public class PlainTextInputElement : Element
 	{
diff --git a/Engine/Source/Programs/Shared/EpicGames.Slack/Elements/SelectMenuElement.cs b/Engine/Source/Programs/Shared/EpicGames.Slack/Elements/SelectMenuElement.cs
--- a/Engine/Source/Programs/Shared/EpicGames.Slack/Elements/SelectMenuElement.cs
+++ b/Engine/Source/Programs/Shared/EpicGames.Slack/Elements/SelectMenuElement.cs
@@ -2,6 +2,7 @@
 
 using System;
 using System.Collections.Generic;
+using System.Diagnostics.CodeAnalysis;
 using System.Globalization;
 using System.Runtime.Serialization;
 using System.Text.Json;
@@ -45,13 +46,15 @@ namespace EpicGames.Slack.Elements
 		/// An array of option objects. A maximum of 10 options are allowed.
 		/// </summary>
 		[JsonPropertyName("options")]
-		public IReadOnlyList<SlackOption>? Options { get; set; }
+		[SuppressMessage("Usage", "CA2227:Collection properties should be read only")]
+		public List<SlackOption>? Options { get; set; }
 
 		/// <summary>
 		/// An array of option objects. A maximum of 10 options are allowed.
 		/// </summary>
 		[JsonPropertyName("option_groups")]
-		public IReadOnlyList<SlackOptionGroup>? OptionGroups { get; set; }
+		[SuppressMessage("Usage", "CA2227:Collection properties should be read only")]
+		public List<SlackOptionGroup>? OptionGroups { get; set; }
 
 		/// <summary>
 		/// An array of option objects that exactly matches one or more of the options within options. These options will be selected when the checkbox group initially loads.
diff --git a/Engine/Source/Programs/Shared/EpicGames.Slack/Option.cs b/Engine/Source/Programs/Shared/EpicGames.Slack/Option.cs
deleted file mode 100644
diff --git a/Engine/Source/Programs/Shared/EpicGames.Slack/SlackClient.cs b/Engine/Source/Programs/Shared/EpicGames.Slack/SlackClient.cs
--- a/Engine/Source/Programs/Shared/EpicGames.Slack/SlackClient.cs
+++ b/Engine/Source/Programs/Shared/EpicGames.Slack/SlackClient.cs
@@ -13,16 +13,28 @@ using Microsoft.Extensions.Logging;
 
 namespace EpicGames.Slack
 {
+	/// <summary>
+	/// Exception thrown due to an error response from Slack
+	/// </summary>
 	public class SlackException : Exception
 	{
+		/// <summary>
+		/// Slack error code
+		/// </summary>
 		public string Code { get; }
 
+		/// <summary>
+		/// Constructor
+		/// </summary>
 		public SlackException(string code) : base(code)
 		{
 			Code = code;
 		}
 	}
 
+	/// <summary>
+	/// Wrapper around Slack client functionality
+	/// </summary>
 	public class SlackClient
 	{
 		class SlackResponse
@@ -38,6 +50,11 @@ namespace EpicGames.Slack
 		readonly JsonSerializerOptions _serializerOptions;
 		readonly ILogger _logger;
 
+		/// <summary>
+		/// Constructor
+		/// </summary>
+		/// <param name="httpClient">Http client for connecting to slack. Should have the necessary authorization headers.</param>
+		/// <param name="logger">Logger interface</param>
 		public SlackClient(HttpClient httpClient, ILogger logger)
 		{
 			_httpClient = httpClient;
@@ -132,21 +149,47 @@ namespace EpicGames.Slack
 			public string? Ts { get; set; }
 		}
 
+		/// <summary>
+		/// Posts a message to a recipient
+		/// </summary>
+		/// <param name="recipient">Recipient of the message. May be a channel or Slack user id.</param>
+		/// <param name="message">New message to post</param>
 		public async Task<string> PostMessageAsync(string recipient, SlackMessage message)
 		{
 			return await PostOrUpdateMessageAsync(recipient, null, null, message, false);
 		}
 
+		/// <summary>
+		/// Posts a message to a recipient
+		/// </summary>
+		/// <param name="recipient">Recipient of the message. May be a channel or Slack user id.</param>
+		/// <param name="threadTs">Timestamp of the thread to post the message to</param>
+		/// <param name="message">New message to post</param>
+		/// <param name="replyBroadcast">Whether to broadcast the message to the channel</param>
 		public async Task<string> PostMessageAsync(string recipient, string threadTs, SlackMessage message, bool replyBroadcast = false)
 		{
 			return await PostOrUpdateMessageAsync(recipient, null, threadTs, message, replyBroadcast);
 		}
 
+		/// <summary>
+		/// Updates an existing message
+		/// </summary>
+		/// <param name="recipient">Recipient of the message. May be a channel or Slack user id.</param>
+		/// <param name="ts">The message timestamp</param>
+		/// <param name="message">New message to post</param>
 		public async Task UpdateMessageAsync(string recipient, string ts, SlackMessage message)
 		{
 			await PostOrUpdateMessageAsync(recipient, ts, null, message, false);
 		}
 
+		/// <summary>
+		/// Updates an existing message
+		/// </summary>
+		/// <param name="recipient">Recipient of the message. May be a channel or Slack user id.</param>
+		/// <param name="ts">The message timestamp</param>
+		/// <param name="threadTs">Timestamp of the thread to post the message to</param>
+		/// <param name="message">New message to post</param>
+		/// <param name="replyBroadcast">Whether to broadcast the message to the channel</param>
 		public async Task UpdateMessageAsync(string recipient, string ts, string threadTs, SlackMessage message, bool replyBroadcast = false)
 		{
 			await PostOrUpdateMessageAsync(recipient, ts, threadTs, message, replyBroadcast);
@@ -203,6 +246,12 @@ namespace EpicGames.Slack
 			public string? Permalink { get; set; }
 		}
 
+		/// <summary>
+		/// Gets a permalink for a message
+		/// </summary>
+		/// <param name="channel">Channel containing the message</param>
+		/// <param name="ts">Message timestamp</param>
+		/// <returns>Link to the message</returns>
 		public async Task<string> GetPermalinkAsync(string channel, string ts)
 		{
 			string requestUrl = $"{GetPermalinkUrl}?channel={channel}&message_ts={ts}";
@@ -236,6 +285,12 @@ namespace EpicGames.Slack
 			public string? Name { get; set; }
 		}
 
+		/// <summary>
+		/// Adds a reaction to a posted message
+		/// </summary>
+		/// <param name="channel">Channel containing the message</param>
+		/// <param name="ts">Message timestamp</param>
+		/// <param name="name">Name of the reaction to post</param>
 		public async Task AddReactionAsync(string channel, string ts, string name)
 		{
 			ReactionMessage message = new ReactionMessage();
@@ -248,6 +303,12 @@ namespace EpicGames.Slack
 			await SendRequestAsync<SlackResponse>(AddReactionUrl, message, ShouldLogError);
 		}
 
+		/// <summary>
+		/// Removes a reaction from a posted message
+		/// </summary>
+		/// <param name="channel">Channel containing the message</param>
+		/// <param name="ts">Message timestamp</param>
+		/// <param name="name">Name of the reaction to post</param>
 		public async Task RemoveReactionAsync(string channel, string ts, string name)
 		{
 			ReactionMessage message = new ReactionMessage();
@@ -275,8 +336,18 @@ namespace EpicGames.Slack
 			public string? Users { get; set; } // Comma separated list of ids
 		}
 
+		/// <summary>
+		/// Invite a user to a channel
+		/// </summary>
+		/// <param name="channel">Channel identifier to invite the user to</param>
+		/// <param name="userId">The user id</param>
 		public Task InviteUserAsync(string channel, string userId) => InviteUsersAsync(channel, new[] { userId });
 
+		/// <summary>
+		/// Invite a set of users to a channel
+		/// </summary>
+		/// <param name="channel">Channel identifier to invite the user to</param>
+		/// <param name="userIds">The user id</param>
 		public async Task InviteUsersAsync(string channel, IEnumerable<string> userIds)
 		{
 			InviteMessage message = new InviteMessage();
@@ -298,6 +369,11 @@ namespace EpicGames.Slack
 			public SlackUser? User { get; set; }
 		}
 
+		/// <summary>
+		/// Gets a user's profile
+		/// </summary>
+		/// <param name="userId">The user id</param>
+		/// <returns>User profile</returns>
 		public async Task<SlackUser> GetUserAsync(string userId)
 		{
 			using (HttpRequestMessage request = new HttpRequestMessage(HttpMethod.Post, $"{UsersInfoUrl}?user={userId}"))
@@ -307,6 +383,11 @@ namespace EpicGames.Slack
 			}
 		}
 
+		/// <summary>
+		/// Finds a user by email address
+		/// </summary>
+		/// <param name="email">The user's email address</param>
+		/// <returns>User profile</returns>
 		public async Task<SlackUser?> FindUserByEmailAsync(string email)
 		{
 			using (HttpRequestMessage request = new HttpRequestMessage(HttpMethod.Post, $"{UsersLookupByEmailUrl}?email={email}"))
@@ -333,6 +414,11 @@ namespace EpicGames.Slack
 			public SlackView? View { get; set; }
 		}
 
+		/// <summary>
+		/// Open a new modal view, in response to a trigger
+		/// </summary>
+		/// <param name="triggerId">The trigger id, as returned as part of an interaction payload</param>
+		/// <param name="view">Definition for the view</param>
 		public async Task OpenViewAsync(string triggerId, SlackView view)
 		{
 			ViewsOpenRequest request = new ViewsOpenRequest();
@@ -342,118 +428,5 @@ namespace EpicGames.Slack
 		}
 
 		#endregion
-/*
-		#region Interactions
-
-		const string AppConnectionOpenUrl = "https://slack.com/api/apps.connections.open";
-
-		class OpenSocketResponse : SlackResponse
-		{
-			[JsonPropertyName("url")]
-			public Uri? Url { get; set; }
-		}
-
-		class EventMessage
-		{
-			[JsonPropertyName("type")]
-			public string? Type { get; set; }
-
-			[JsonPropertyName("envelope_id")]
-			public string? EnvelopeId { get; set; }
-
-			[JsonPropertyName("payload")]
-			public EventPayload? Payload { get; set; }
-		}
-
-		class EventPayload
-		{
-			[JsonPropertyName("type")]
-			public string? Type { get; set; }
-
-			[JsonPropertyName("trigger_id")]
-			public string? TriggerId { get; set; }
-
-			[JsonPropertyName("user")]
-			public UserInfo? User { get; set; }
-
-			[JsonPropertyName("response_url")]
-			public string? ResponseUrl { get; set; }
-
-			[JsonPropertyName("actions")]
-			public List<ActionInfo> Actions { get; set; } = new List<ActionInfo>();
-		}
-
-		/// <inheritdoc/>
-		public async Task HandleInteractionsAsync(CancellationToken stoppingToken)
-		{
-			using HttpRequestMessage request = new HttpRequestMessage(HttpMethod.Post, AppConnectionOpenUrl);
-			request.Content = new FormUrlEncodedContent(Array.Empty<KeyValuePair<string?, string?>>());
-
-			OpenSocketResponse response = await SendRequestAsync<OpenSocketResponse>(request, "");
-		}
-
-		private async Task HandleInteractionsInternalAsync(Uri socketUrl, CancellationToken stoppingToken)
-		{
-			using ClientWebSocket socket = new ClientWebSocket();
-			await socket.ConnectAsync(socketUrl, stoppingToken);
-
-			byte[] buffer = new byte[2048];
-			while (!stoppingToken.IsCancellationRequested)
-			{
-				// Read the next message
-				int length = 0;
-				for (; ; )
-				{
-					if (length == buffer.Length)
-					{
-						Array.Resize(ref buffer, buffer.Length + 2048);
-					}
-
-					WebSocketReceiveResult result = await socket.ReceiveAsync(new ArraySegment<byte>(buffer, length, buffer.Length - length), stoppingToken);
-					if (result.MessageType == WebSocketMessageType.Close)
-					{
-						return;
-					}
-					length += result.Count;
-
-					if (result.EndOfMessage)
-					{
-						break;
-					}
-				}
-
-				// Get the message data
-				_logger.LogInformation("Slack event: {Message}", Encoding.UTF8.GetString(buffer, 0, length));
-				EventMessage eventMessage = JsonSerializer.Deserialize<EventMessage>(buffer.AsSpan(0, length))!;
-
-				// Acknowledge the message
-				if (eventMessage.EnvelopeId != null)
-				{
-					object response = new { eventMessage.EnvelopeId };
-					await socket.SendAsync(JsonSerializer.SerializeToUtf8Bytes(response), WebSocketMessageType.Text, true, stoppingToken);
-				}
-
-				// Handle the message type
-				if (eventMessage.Type != null)
-				{
-					string type = eventMessage.Type;
-					if (type.Equals("disconnect", StringComparison.Ordinal))
-					{
-						break;
-					}
-					else if (type.Equals("interactive", StringComparison.Ordinal))
-					{
-						await HandleInteractionMessage(eventMessage);
-					}
-					else
-					{
-						_logger.LogDebug("Unhandled event type ({Type})", type);
-					}
-				}
-			}
-		}
-
-		#endregion
-*/
 	}
 }
diff --git a/Engine/Source/Programs/Shared/EpicGames.Slack/SlackOption.cs b/Engine/Source/Programs/Shared/EpicGames.Slack/SlackOption.cs
new file mode 100644
--- /dev/null
+++ b/Engine/Source/Programs/Shared/EpicGames.Slack/SlackOption.cs
@@ -0,0 +1,83 @@
+// Copyright Epic Games, Inc. All Rights Reserved.
+
+using System;
+using System.Collections.Generic;
+using System.Text.Json;
+using System.Text.Json.Serialization;
+using EpicGames.Slack.Elements;
+
+namespace EpicGames.Slack
+{
+	/// <summary>
+	/// An object that represents a single selectable item in a select menu, multi-select menu, checkbox group, radio button group, or overflow menu.
+	/// </summary>
+	public class SlackOption
+	{
+		/// <summary>
+		/// A text object that defines the text shown in the option on the menu. Overflow, select, and multi-select menus can only use plain_text objects, while radio buttons and checkboxes can use mrkdwn text objects.
+		/// Maximum length for the text in this field is 75 characters.
+		/// </summary>
+		[JsonPropertyName("text")]
+		public TextObject Text { get; set; }
+
+		/// <summary>
+		/// A unique string value that will be passed to your app when this option is chosen. Maximum length for this field is 75 characters.
+		/// </summary>
+		[JsonPropertyName("value")]
+		public string Value { get; set; }
+
+		/// <summary>
+		/// A plain_text only text object that defines a line of descriptive text shown below the text field beside the radio button. Maximum length for the text object within this field is 75 characters.
+		/// </summary>
+		[JsonPropertyName("description")]
+		public PlainTextObject? Description { get; set; }
+
+		/// <summary>
+		/// A URL to load in the user's browser when the option is clicked. The url attribute is only available in overflow menus. Maximum length for this field is 3000 characters. If you're using url, you'll still receive an interaction payload and will need to send an acknowledgement response.
+		/// </summary>
+		[JsonPropertyName("url")]
+		public Uri? Url { get; set; }
+
+		/// <summary>
+		/// Constructor
+		/// </summary>
+		public SlackOption(TextObject text, string value)
+		{
+			Text = text;
+			Value = value;
+		}
+
+		/// <summary>
+		/// Constructor
+		/// </summary>
+		public SlackOption(string text, string value)
+			: this(new PlainTextObject(text), value)
+		{
+		}
+	}
+
+	/// <summary>
+	/// Provides a way to group options in a select menu or multi-select menu.
+	/// </summary>
+	public class SlackOptionGroup
+	{
+		/// <summary>
+		/// A plain_text only text object that defines the label shown above this group of options. Maximum length for the text in this field is 75 characters.
+		/// </summary>
+		public PlainTextObject Label { get; }
+
+		/// <summary>
+		/// An array of option objects that belong to this specific group. Maximum of 100 items.
+		/// </summary>
+		public List<SlackOption> Options { get; } = new List<SlackOption>();
+
+		/// <summary>
+		/// Constructor
+		/// </summary>
+		public SlackOptionGroup(PlainTextObject label, List<SlackOption> options)
+		{
+			Label = label;
+			Options.AddRange(options);
+		}
+	}
+}
