Horde: Include callstacks in new Slack client errors.

diff --git a/Engine/Source/Programs/Horde/Horde.Build/Notifications/Sinks/SlackNotificationSink.cs b/Engine/Source/Programs/Horde/Horde.Build/Notifications/Sinks/SlackNotificationSink.cs
--- a/Engine/Source/Programs/Horde/Horde.Build/Notifications/Sinks/SlackNotificationSink.cs
+++ b/Engine/Source/Programs/Horde/Horde.Build/Notifications/Sinks/SlackNotificationSink.cs
@@ -52,7 +52,7 @@ namespace Horde.Build.Notifications.Sinks
 	/// <summary>
 	/// Maintains a connection to Slack, in order to receive socket-mode notifications of user interactions
 	/// </summary>
-	public class SlackNotificationSink : BackgroundService, INotificationSink, IAvatarService
+	public sealed class SlackNotificationSink : BackgroundService, INotificationSink, IAvatarService
 	{
 		const bool defaultAllowMentions = true;
 
@@ -260,8 +260,7 @@ namespace Horde.Build.Notifications.Sinks
 			base.Dispose();
 
 			_userCache.Dispose();
-
-			GC.SuppressFinalize(this);
+			_httpClient.Dispose();
 		}
 
 		#region Avatars
diff --git a/Engine/Source/Programs/Shared/EpicGames.Slack/SlackClient.cs b/Engine/Source/Programs/Shared/EpicGames.Slack/SlackClient.cs
--- a/Engine/Source/Programs/Shared/EpicGames.Slack/SlackClient.cs
+++ b/Engine/Source/Programs/Shared/EpicGames.Slack/SlackClient.cs
@@ -98,8 +98,15 @@ namespace EpicGames.Slack
 			TResponse responseObject = JsonSerializer.Deserialize<TResponse>(responseBytes)!;
 			if (shouldLogError(responseObject))
 			{
-				_logger.LogError("Failed to send Slack message ({Error}). Request: {Request}. Response: {Response}", responseObject.Error, requestJson, Encoding.UTF8.GetString(responseBytes));
-				throw new SlackException(responseObject.Error ?? "unknown");
+				try
+				{
+					throw new SlackException(responseObject.Error ?? "unknown");
+				}
+				catch (Exception ex)
+				{
+					_logger.LogError(ex, "Failed to send Slack message ({Error}). Request: {Request}. Response: {Response}", responseObject.Error, requestJson, Encoding.UTF8.GetString(responseBytes));
+					throw;
+				}
 			}
 			return responseObject;
 		}
