Updated documentation on PackageFlag for C#

#rnx

diff --git a/Engine/Source/Programs/Shared/EpicGames.Core/UnrealEngineTypes.cs b/Engine/Source/Programs/Shared/EpicGames.Core/UnrealEngineTypes.cs
--- a/Engine/Source/Programs/Shared/EpicGames.Core/UnrealEngineTypes.cs
+++ b/Engine/Source/Programs/Shared/EpicGames.Core/UnrealEngineTypes.cs
@@ -795,7 +795,7 @@ namespace EpicGames.Core
 		ContainsNoAsset = 0x00000400,
 
 		/// <summary>
-		/// (Not Implemented) Objects in this package cannot be referenced in a different plugin or mount point (i.e /Game -> /Engine)
+		/// Objects in this package cannot be referenced in a different plugin or mount point (i.e /Game -> /Engine)
 		/// </summary>
 		NotExternallyReferenceable = 0x00000800,
 
