Horde: Support for including instructions in Triage threads.

diff --git a/Engine/Source/Programs/Horde/Horde.Build/Notifications/Sinks/SlackNotificationSink.cs b/Engine/Source/Programs/Horde/Horde.Build/Notifications/Sinks/SlackNotificationSink.cs
--- a/Engine/Source/Programs/Horde/Horde.Build/Notifications/Sinks/SlackNotificationSink.cs
+++ b/Engine/Source/Programs/Horde/Horde.Build/Notifications/Sinks/SlackNotificationSink.cs
@@ -822,6 +822,11 @@ namespace Horde.Build.Notifications.Sinks
 					actions.AddButton("Not Me", value: $"issue_{issue.Id}_decline", style: ButtonStyle.Danger);
 					actions.AddButton("Mark Fixed", value: $"issue_{issue.Id}_markfixed");
 
+					if (workflow.TriageInstructions != null)
+					{
+						message.AddSection(workflow.TriageInstructions);
+					}
+
 					string? summaryTs = await _slackClient.PostMessageAsync(triageChannel, state.Ts, message);
 
 					// Permalink to the summary text so we link inside the thread rather than just to the original message
diff --git a/Engine/Source/Programs/Horde/Horde.Build/Server/Config.cs b/Engine/Source/Programs/Horde/Horde.Build/Server/Config.cs
--- a/Engine/Source/Programs/Horde/Horde.Build/Server/Config.cs
+++ b/Engine/Source/Programs/Horde/Horde.Build/Server/Config.cs
@@ -373,6 +373,11 @@ namespace Horde.Build.Server
 		/// Suffix for all triage messages
 		/// </summary>
 		public string? TriageSuffix { get; set; }
+		
+		/// <summary>
+		/// Instructions posted to triage threads
+		/// </summary>
+		public string? TriageInstructions { get; set; }
 
 		/// <summary>
 		/// Maximum number of people to mention on a triage thread
diff --git a/Engine/Source/Programs/Shared/EpicGames.Slack/Blocks/ContextBlock.cs b/Engine/Source/Programs/Shared/EpicGames.Slack/Blocks/ContextBlock.cs
--- a/Engine/Source/Programs/Shared/EpicGames.Slack/Blocks/ContextBlock.cs
+++ b/Engine/Source/Programs/Shared/EpicGames.Slack/Blocks/ContextBlock.cs
@@ -4,6 +4,7 @@ using System;
 using System.Collections.Generic;
 using System.Text.Json;
 using System.Text.Json.Serialization;
+using EpicGames.Slack.Elements;
 
 namespace EpicGames.Slack.Blocks
 {
@@ -42,5 +43,18 @@ namespace EpicGames.Slack.Blocks
 			configure(block);
 			container.Blocks.Add(block);
 		}
+
+		/// <summary>
+		/// Add an <see cref="ContextBlock"/> to the list of blocks
+		/// </summary>
+		/// <param name="container">Block container</param>
+		/// <param name="text">Configuration callback for the block</param>
+		public static ContextBlock AddContext(this ISlackBlockContainer container, TextObject text)
+		{
+			ContextBlock block = new ContextBlock();
+			block.Elements.Add(text);
+			container.Blocks.Add(block);
+			return block;
+		}
 	}
 }
