EpicGames.Build: FindCorrectCase for overridden RootDirectory

#rnx

diff --git a/Engine/Source/Programs/Shared/EpicGames.Build/Unreal.cs b/Engine/Source/Programs/Shared/EpicGames.Build/Unreal.cs
--- a/Engine/Source/Programs/Shared/EpicGames.Build/Unreal.cs
+++ b/Engine/Source/Programs/Shared/EpicGames.Build/Unreal.cs
@@ -17,19 +17,19 @@ namespace UnrealBuildBase
 		{
 			if (LocationOverride.RootDirectory != null)
 			{
-				return LocationOverride.RootDirectory;
+				return DirectoryReference.FindCorrectCase(LocationOverride.RootDirectory);
 			}
 
 			// This base library may be used - and so be launched - from more than one location (at time of writing, UnrealBuildTool and AutomationTool)
 			// Programs that use this assembly must be located under "Engine/Binaries/DotNET" and so we look for that sequence of directories in that path of the executing assembly
-			
+
 			// Use the EntryAssembly (the application path), rather than the ExecutingAssembly (the library path)
 			string AssemblyLocation = Assembly.GetEntryAssembly()!.GetOriginalLocation();
 
 			DirectoryReference? FoundRootDirectory = DirectoryReference.FindCorrectCase(DirectoryReference.FromString(AssemblyLocation)!);
 
 			// Search up through the directory tree for the deepest instance of the sub-path "Engine/Binaries/DotNET"
-			while(FoundRootDirectory != null)
+			while (FoundRootDirectory != null)
 			{
 				if (String.Equals("DotNET", FoundRootDirectory.GetDirectoryName()))
 				{
@@ -80,7 +80,7 @@ namespace UnrealBuildBase
 		}
 
 		static private string DotnetVersionDirectory = "6.0.302";
-			
+
 		static private string FindRelativeDotnetDirectory(RuntimePlatform.Type HostPlatform)
 		{
 			string HostDotNetDirectoryName;
@@ -88,15 +88,15 @@ namespace UnrealBuildBase
 			{
 				case RuntimePlatform.Type.Windows: HostDotNetDirectoryName = "windows"; break;
 				case RuntimePlatform.Type.Mac:
-				{
-					HostDotNetDirectoryName = "mac-x64";
-					if (RuntimeInformation.ProcessArchitecture == Architecture.Arm64)
 					{
-						HostDotNetDirectoryName = "mac-arm64";
+						HostDotNetDirectoryName = "mac-x64";
+						if (RuntimeInformation.ProcessArchitecture == Architecture.Arm64)
+						{
+							HostDotNetDirectoryName = "mac-arm64";
+						}
+						break;
 					}
-					break;
-				}
-				case RuntimePlatform.Type.Linux:   HostDotNetDirectoryName = "linux"; break;
+				case RuntimePlatform.Type.Linux: HostDotNetDirectoryName = "linux"; break;
 				default: throw new Exception("Unknown host platform");
 			}
 
@@ -177,7 +177,7 @@ namespace UnrealBuildBase
 			return bIsEngineInstalled.Value;
 		}
 
-		public static class LocationOverride 
+		public static class LocationOverride
 		{
 			/// <summary>
 			/// If set, this value will be used to populate Unreal.RootDirectory
@@ -237,7 +237,7 @@ namespace UnrealBuildBase
 		/// <param name="bIncludeRestrictedDirectories">If true, restricted (NotForLicensees, NoRedist) subdirectories are included</param>
 		/// <param name="bIncludeBaseDirectory">If true, BaseDir is included</param>
 		/// <returns>List of extension directories, including the given base directory</returns>
-		public static List<DirectoryReference> GetExtensionDirs(DirectoryReference BaseDir, bool bIncludePlatformDirectories=true, bool bIncludeRestrictedDirectories=true, bool bIncludeBaseDirectory=true)
+		public static List<DirectoryReference> GetExtensionDirs(DirectoryReference BaseDir, bool bIncludePlatformDirectories = true, bool bIncludeRestrictedDirectories = true, bool bIncludeBaseDirectory = true)
 		{
 			Tuple<List<DirectoryReference>, List<DirectoryReference>>? CachedDirs;
 			if (!CachedExtensionDirectories.TryGetValue(BaseDir, out CachedDirs))
@@ -303,7 +303,7 @@ namespace UnrealBuildBase
 		/// <param name="bIncludeRestrictedDirectories">If true, restricted (NotForLicensees, NoRedist) subdirectories are included</param>
 		/// <param name="bIncludeBaseDirectory">If true, BaseDir is included</param>
 		/// <returns>List of extension directories, including the given base directory</returns>
-		public static List<DirectoryReference> GetExtensionDirs(DirectoryReference BaseDir, string SubDir, bool bIncludePlatformDirectories=true, bool bIncludeRestrictedDirectories=true, bool bIncludeBaseDirectory=true)
+		public static List<DirectoryReference> GetExtensionDirs(DirectoryReference BaseDir, string SubDir, bool bIncludePlatformDirectories = true, bool bIncludeRestrictedDirectories = true, bool bIncludeBaseDirectory = true)
 		{
 			return GetExtensionDirs(BaseDir, bIncludePlatformDirectories, bIncludeRestrictedDirectories, bIncludeBaseDirectory).Select(x => DirectoryReference.Combine(x, SubDir)).Where(x => DirectoryReference.Exists(x)).ToList();
 		}
