OidcToken - Move the shared state into the same folder structure as other UE data (DDC cache for instance)
Also set the executable bit for the executables on mac and linux.

diff --git a/Engine/Source/Programs/Shared/EpicGames.OIDC/TokenStore.cs b/Engine/Source/Programs/Shared/EpicGames.OIDC/TokenStore.cs
--- a/Engine/Source/Programs/Shared/EpicGames.OIDC/TokenStore.cs
+++ b/Engine/Source/Programs/Shared/EpicGames.OIDC/TokenStore.cs
@@ -48,8 +48,7 @@ namespace EpicGames.OIDC
 
 		private static FileInfo GetStorePath()
 		{
-			return new FileInfo(Path.Combine(Environment.GetFolderPath(Environment.SpecialFolder.LocalApplicationData),
-				"oidc-token", "tokenStore.dat"));
+			return new FileInfo(Path.Combine(Environment.GetFolderPath(Environment.SpecialFolder.ApplicationData), "Epic", "UnrealEngine", "Common", "OidcToken", "tokenStore.dat"));
 		}
 
 		private void ReadStoreFromDisk()
diff --git a/Engine/Source/Programs/Shared/EpicGames.OIDC/WindowsTokenStore.cs b/Engine/Source/Programs/Shared/EpicGames.OIDC/WindowsTokenStore.cs
--- a/Engine/Source/Programs/Shared/EpicGames.OIDC/WindowsTokenStore.cs
+++ b/Engine/Source/Programs/Shared/EpicGames.OIDC/WindowsTokenStore.cs
@@ -41,7 +41,7 @@ namespace EpicGames.OIDC
 
 		private static FileInfo GetStorePath()
 		{
-			return new FileInfo(Path.Combine(Environment.GetFolderPath(Environment.SpecialFolder.LocalApplicationData), "oidc-token", "oidcTokenStore.dat"));
+			return new FileInfo(Path.Combine(Environment.GetFolderPath(Environment.SpecialFolder.LocalApplicationData), "UnrealEngine", "Common", "OidcToken", "oidcTokenStore.dat"));
 		}
 
 		private void ReadStoreFromDisk()
