Horde: Include request URL in Slack errors for additional context.

diff --git a/Engine/Source/Programs/Shared/EpicGames.Slack/SlackClient.cs b/Engine/Source/Programs/Shared/EpicGames.Slack/SlackClient.cs
--- a/Engine/Source/Programs/Shared/EpicGames.Slack/SlackClient.cs
+++ b/Engine/Source/Programs/Shared/EpicGames.Slack/SlackClient.cs
@@ -104,7 +104,7 @@ namespace EpicGames.Slack
 				}
 				catch (Exception ex)
 				{
-					_logger.LogError(ex, "Failed to send Slack message ({Error}). Request: {Request}. Response: {Response}", responseObject.Error, requestJson, Encoding.UTF8.GetString(responseBytes));
+					_logger.LogError(ex, "Failed to send Slack message to {Url} ({Error}). Request: {Request}. Response: {Response}", request.RequestUri, responseObject.Error, requestJson, Encoding.UTF8.GetString(responseBytes));
 					throw;
 				}
 			}
