Horde: Always include blocks in Slack messages. Possible fix for no_text errors when sending messages consisting only of attachments.

diff --git a/Engine/Source/Programs/Shared/EpicGames.Slack/SlackClient.cs b/Engine/Source/Programs/Shared/EpicGames.Slack/SlackClient.cs
--- a/Engine/Source/Programs/Shared/EpicGames.Slack/SlackClient.cs
+++ b/Engine/Source/Programs/Shared/EpicGames.Slack/SlackClient.cs
@@ -209,10 +209,8 @@ namespace EpicGames.Slack
 			request.Ts = ts;
 			request.ThreadTs = threadTs;
 			request.Text = message.Text;
-			if (message.Blocks.Count > 0)
-			{
-				request.Blocks = message.Blocks;
-			}
+			request.Blocks = message.Blocks;
+
 			if (replyBroadcast)
 			{
 				request.ReplyBroadcast = replyBroadcast;
