Horde: Fix missing attachments from Slack messages.

diff --git a/Engine/Source/Programs/Shared/EpicGames.Slack/SlackClient.cs b/Engine/Source/Programs/Shared/EpicGames.Slack/SlackClient.cs
--- a/Engine/Source/Programs/Shared/EpicGames.Slack/SlackClient.cs
+++ b/Engine/Source/Programs/Shared/EpicGames.Slack/SlackClient.cs
@@ -136,6 +136,9 @@ namespace EpicGames.Slack
 
 			[JsonPropertyName("blocks")]
 			public List<Block>? Blocks { get; set; }
+
+			[JsonPropertyName("attachments")]
+			public List<SlackAttachment>? Attachments { get; set; }
 			
 			[JsonPropertyName("reply_broadcast")]
 			public bool? ReplyBroadcast { get; set; }
@@ -215,6 +218,10 @@ namespace EpicGames.Slack
 			{
 				request.ReplyBroadcast = replyBroadcast;
 			}
+			if (message.Attachments.Count > 0)
+			{
+				request.Attachments = message.Attachments;
+			}
 			if (!message.UnfurlLinks)
 			{
 				request.UnfurlLinks = false;
