Horde: Fix incorrect property names on Slack message attachments.

diff --git a/Engine/Source/Programs/Shared/EpicGames.Slack/Blocks/SectionBlock.cs b/Engine/Source/Programs/Shared/EpicGames.Slack/Blocks/SectionBlock.cs
--- a/Engine/Source/Programs/Shared/EpicGames.Slack/Blocks/SectionBlock.cs
+++ b/Engine/Source/Programs/Shared/EpicGames.Slack/Blocks/SectionBlock.cs
@@ -16,20 +16,20 @@ namespace EpicGames.Slack.Blocks
 		/// <summary>
 		/// The text for the block, in the form of a text object. Maximum length for the text in this field is 3000 characters.
 		/// </summary>
-		[JsonPropertyName("text")]
+		[JsonPropertyName("text"), JsonPropertyOrder(1)]
 		public TextObject? Text { get; set; }
 
 		/// <summary>
 		/// Required if no text is provided. An array of text objects. Any text objects included with fields will be rendered in a compact format that allows for 2 columns of side-by-side text. 
 		/// Maximum number of items is 10. Maximum length for the text in each item is 2000 characters.
 		/// </summary>
-		[JsonPropertyName("fields")]
+		[JsonPropertyName("fields"), JsonPropertyOrder(2)]
 		public IReadOnlyList<TextObject>? Fields { get; set; }
 
 		/// <summary>
 		/// One of the available element objects.
 		/// </summary>
-		[JsonPropertyName("accessory")]
+		[JsonPropertyName("accessory"), JsonPropertyOrder(3)]
 		public Element? Accessory { get; set; }
 
 		/// <summary>
diff --git a/Engine/Source/Programs/Shared/EpicGames.Slack/SlackAttachment.cs b/Engine/Source/Programs/Shared/EpicGames.Slack/SlackAttachment.cs
--- a/Engine/Source/Programs/Shared/EpicGames.Slack/SlackAttachment.cs
+++ b/Engine/Source/Programs/Shared/EpicGames.Slack/SlackAttachment.cs
@@ -18,17 +18,20 @@ namespace EpicGames.Slack
 		/// <summary>
 		/// The color of the line down the left side of the attachment.
 		/// </summary>
+		[JsonPropertyName("color"), JsonPropertyOrder(1)]
 		public string Color { get; set; } = "#ff0000";
 
 		/// <summary>
 		/// Used for the "toast" notifications sent by Slack. If not set, and the Text property of the
 		/// base message isn't set, the notification will have a message saying it has no preview.
 		/// </summary>
+		[JsonPropertyName("fallback"), JsonPropertyOrder(2)]
 		public string? FallbackText { get; set; }
 
 		/// <summary>
 		/// A collection of BlockKit blocks the attachment is composed from.
 		/// </summary>
+		[JsonPropertyName("blocks"), JsonPropertyOrder(3)]
 		public List<Block> Blocks { get; } = new List<Block>();
 
 		/// <summary>
