Horde: Other missing property attributes for Slack messages.

diff --git a/Engine/Source/Programs/Shared/EpicGames.Slack/Blocks/HeaderBlock.cs b/Engine/Source/Programs/Shared/EpicGames.Slack/Blocks/HeaderBlock.cs
--- a/Engine/Source/Programs/Shared/EpicGames.Slack/Blocks/HeaderBlock.cs
+++ b/Engine/Source/Programs/Shared/EpicGames.Slack/Blocks/HeaderBlock.cs
@@ -2,6 +2,7 @@
 
 using System.Collections.Generic;
 using System.Text.Json;
+using System.Text.Json.Serialization;
 using EpicGames.Slack.Elements;
 
 namespace EpicGames.Slack.Blocks
@@ -14,6 +15,7 @@ namespace EpicGames.Slack.Blocks
 		/// <summary>
 		/// The text for the block, in the form of a plain_text text object. Maximum length for the text in this field is 150 characters.
 		/// </summary>
+		[JsonPropertyName("text"), JsonPropertyOrder(1)]
 		public PlainTextObject Text { get; }
 
 		/// <summary>
diff --git a/Engine/Source/Programs/Shared/EpicGames.Slack/Blocks/ImageBlock.cs b/Engine/Source/Programs/Shared/EpicGames.Slack/Blocks/ImageBlock.cs
--- a/Engine/Source/Programs/Shared/EpicGames.Slack/Blocks/ImageBlock.cs
+++ b/Engine/Source/Programs/Shared/EpicGames.Slack/Blocks/ImageBlock.cs
@@ -16,19 +16,19 @@ namespace EpicGames.Slack.Blocks
 		/// <summary>
 		/// The URL of the image to be displayed.Maximum length for this field is 3000 characters.
 		/// </summary>
-		[JsonPropertyName("image_url")]
+		[JsonPropertyName("image_url"), JsonPropertyOrder(1)]
 		public Uri ImageUrl { get; set; }
 
 		/// <summary>
 		/// A plain-text summary of the image.This should not contain any markup.Maximum length for this field is 2000 characters.
 		/// </summary>
-		[JsonPropertyName("alt_text")]
+		[JsonPropertyName("alt_text"), JsonPropertyOrder(2)]
 		public string AltText { get; set; }
 
 		/// <summary>
 		/// An optional title for the image in the form of a text object that can only be of type: plain_text. Maximum length for the text in this field is 2000 characters.
 		/// </summary>
-		[JsonPropertyName("title")]
+		[JsonPropertyName("title"), JsonPropertyOrder(3)]
 		public PlainTextObject? Title { get; set; }
 
 		/// <summary>
