Horde: Forward markdown argument on to Slack request.

diff --git a/Engine/Source/Programs/Shared/EpicGames.Slack/SlackClient.cs b/Engine/Source/Programs/Shared/EpicGames.Slack/SlackClient.cs
--- a/Engine/Source/Programs/Shared/EpicGames.Slack/SlackClient.cs
+++ b/Engine/Source/Programs/Shared/EpicGames.Slack/SlackClient.cs
@@ -213,6 +213,7 @@ namespace EpicGames.Slack
 			request.ThreadTs = threadTs;
 			request.Text = message.Text;
 			request.Blocks = message.Blocks;
+			request.Markdown = message.Markdown;
 
 			if (replyBroadcast)
 			{
