Updated C# UHT with the destructor changes and making a minor fix to the formatting.

diff --git a/Engine/Source/Programs/Shared/EpicGames.UHT/Exporters/CodeGen/UhtHeaderCodeGeneratorCppFile.cs b/Engine/Source/Programs/Shared/EpicGames.UHT/Exporters/CodeGen/UhtHeaderCodeGeneratorCppFile.cs
--- a/Engine/Source/Programs/Shared/EpicGames.UHT/Exporters/CodeGen/UhtHeaderCodeGeneratorCppFile.cs
+++ b/Engine/Source/Programs/Shared/EpicGames.UHT/Exporters/CodeGen/UhtHeaderCodeGeneratorCppFile.cs
@@ -1030,6 +1030,11 @@ namespace EpicGames.UHT.Exporters.CodeGen
 				builder.Append("\tDEFINE_VTABLE_PTR_HELPER_CTOR(").Append(classObj.SourceName).Append(");\r\n");
 			}
 
+			if (!classObj.ClassExportFlags.HasAnyFlags(UhtClassExportFlags.HasDestructor))
+			{
+				builder.Append('\t').Append(classObj.SourceName).Append("::~").Append(classObj.SourceName).Append("() {}\r\n");
+			}
+
 			// Only write out adapters if the user has provided one or the other of the Serialize overloads
 			if (classObj.SerializerArchiveType != UhtSerializerArchiveType.None && classObj.SerializerArchiveType != UhtSerializerArchiveType.All)
 			{
diff --git a/Engine/Source/Programs/Shared/EpicGames.UHT/Exporters/CodeGen/UhtHeaderCodeGeneratorHFile.cs b/Engine/Source/Programs/Shared/EpicGames.UHT/Exporters/CodeGen/UhtHeaderCodeGeneratorHFile.cs
--- a/Engine/Source/Programs/Shared/EpicGames.UHT/Exporters/CodeGen/UhtHeaderCodeGeneratorHFile.cs
+++ b/Engine/Source/Programs/Shared/EpicGames.UHT/Exporters/CodeGen/UhtHeaderCodeGeneratorHFile.cs
@@ -829,6 +829,7 @@ namespace EpicGames.UHT.Exporters.CodeGen
 
 				AppendVTableHelperCtorAndCaller(builder, classObj, api);
 				AppendCopyConstructorDefinition(builder, classObj, api);
+				AppendDestructorDefinition(builder, classObj, api);
 			}
 			return builder;
 		}
@@ -847,6 +848,7 @@ namespace EpicGames.UHT.Exporters.CodeGen
 				AppendConstructorDefinition(builder, classObj, api);
 				AppendVTableHelperCtorAndCaller(builder, classObj, api);
 				AppendDefaultConstructorCallDefinition(builder, classObj);
+				AppendDestructorDefinition(builder, classObj, api);
 			}
 			return builder;
 		}
@@ -885,6 +887,22 @@ namespace EpicGames.UHT.Exporters.CodeGen
 			return builder;
 		}
 
+		/// <summary>
+		/// Generates a destructor
+		/// </summary>
+		/// <param name="builder">Output builder</param>
+		/// <param name="classObj">Class being exported</param>
+		/// <param name="api">API text to be used</param>
+		/// <returns>Output builder</returns>
+		private static StringBuilder AppendDestructorDefinition(StringBuilder builder, UhtClass classObj, string api)
+		{
+			if (!classObj.ClassExportFlags.HasAnyFlags(UhtClassExportFlags.HasDestructor))
+			{
+				builder.Append('\t').Append(api).Append("virtual ~").Append(classObj.SourceName).Append("(); \\\r\n");
+			}
+			return builder;
+		}
+
 		/// <summary>
 		/// Generates private copy-constructor declaration.
 		/// </summary>
diff --git a/Engine/Source/Programs/Shared/EpicGames.UHT/Parsers/UhtHeaderFileParser.cs b/Engine/Source/Programs/Shared/EpicGames.UHT/Parsers/UhtHeaderFileParser.cs
--- a/Engine/Source/Programs/Shared/EpicGames.UHT/Parsers/UhtHeaderFileParser.cs
+++ b/Engine/Source/Programs/Shared/EpicGames.UHT/Parsers/UhtHeaderFileParser.cs
@@ -458,6 +458,9 @@ namespace EpicGames.UHT.Parsers
 							if (topScope.HeaderParser.CheckForConstructor(classObj, declaration))
 							{
 							}
+							else if (topScope.HeaderParser.CheckForDestructor(classObj, declaration))
+							{
+							}
 							else if (classObj.ClassType == UhtClassType.Class)
 							{
 								if (topScope.HeaderParser.CheckForSerialize(classObj, declaration))
@@ -1019,6 +1022,25 @@ namespace EpicGames.UHT.Parsers
 			return false;
 		}
 
+		private bool CheckForDestructor(UhtClass classObj, UhtDeclaration declaration)
+		{
+			if (classObj.ClassExportFlags.HasAnyFlags(UhtClassExportFlags.HasDestructor))
+			{
+				return false;
+			}
+
+			IUhtTokenReader replayReader = UhtTokenReplayReader.GetThreadInstance(this.HeaderFile, this.HeaderFile.Data.Memory, new ReadOnlyMemory<UhtToken>(declaration.Tokens), UhtTokenType.EndOfDeclaration);
+
+			SkipVirtualAndAPI(replayReader);
+
+			if (replayReader.TryOptional('~') && replayReader.TryOptional(classObj.SourceName))
+			{
+				classObj.ClassExportFlags |= UhtClassExportFlags.HasDestructor;
+				return true;
+			}
+			return false;
+		}
+
 		bool CheckForSerialize(UhtClass classObj, UhtDeclaration declaration)
 		{
 			IUhtTokenReader replayReader = UhtTokenReplayReader.GetThreadInstance(HeaderFile, HeaderFile.Data.Memory, new ReadOnlyMemory<UhtToken>(declaration.Tokens), UhtTokenType.EndOfDeclaration);
diff --git a/Engine/Source/Programs/Shared/EpicGames.UHT/Types/UhtClass.cs b/Engine/Source/Programs/Shared/EpicGames.UHT/Types/UhtClass.cs
--- a/Engine/Source/Programs/Shared/EpicGames.UHT/Types/UhtClass.cs
+++ b/Engine/Source/Programs/Shared/EpicGames.UHT/Types/UhtClass.cs
@@ -83,6 +83,11 @@ namespace EpicGames.UHT.Types
 		/// Class has a field notify
 		/// </summary>
 		HasFieldNotify = 1 << 10,
+
+		/// <summary>
+		/// Class has destructor
+		/// </summary>
+		HasDestructor = 1 << 11,
 	}
 
 	/// <summary>
