Fixed issue in C# UHT where redirected types were not being respected.

diff --git a/Engine/Source/Programs/Shared/EpicGames.UHT/Parsers/UhtPropertyParser.cs b/Engine/Source/Programs/Shared/EpicGames.UHT/Parsers/UhtPropertyParser.cs
--- a/Engine/Source/Programs/Shared/EpicGames.UHT/Parsers/UhtPropertyParser.cs
+++ b/Engine/Source/Programs/Shared/EpicGames.UHT/Parsers/UhtPropertyParser.cs
@@ -551,9 +551,11 @@ namespace EpicGames.UHT.Parsers
 						break;
 					}
 				}
-				else
+				else if (typeTokens.Span[index].IsIdentifier())
 				{
-					if (session.TryGetPropertyType(typeTokens.Span[index].Value, out UhtPropertyType propertyType))
+					UhtToken copy = typeTokens.Span[index];
+					session.Config!.RedirectTypeIdentifier(ref copy);
+					if (session.TryGetPropertyType(copy.Value, out UhtPropertyType propertyType))
 					{
 						if (resolvePhase == UhtPropertyResolvePhase.Resolving || propertyType.Options.HasAnyFlags(UhtPropertyTypeOptions.Immediate))
 						{
@@ -656,7 +658,14 @@ namespace EpicGames.UHT.Parsers
 						break;
 					}
 
-					if (session.TryGetPropertyType(token.Value, out propertyType))
+					if (!token.IsIdentifier())
+					{
+						continue;
+					}
+
+					UhtToken copy = token;
+					session.Config!.RedirectTypeIdentifier(ref copy);
+					if (session.TryGetPropertyType(copy.Value, out propertyType))
 					{
 						if (resolvePhase == UhtPropertyResolvePhase.Resolving || propertyType.Options.HasAnyFlags(UhtPropertyTypeOptions.Immediate))
 						{
