Fixed issue in C# UHT where tooltips were not parsing exactly the same as the C++ UHT.

diff --git a/Engine/Source/Programs/Shared/EpicGames.UHT/Parsers/UhtParsingScope.cs b/Engine/Source/Programs/Shared/EpicGames.UHT/Parsers/UhtParsingScope.cs
--- a/Engine/Source/Programs/Shared/EpicGames.UHT/Parsers/UhtParsingScope.cs
+++ b/Engine/Source/Programs/Shared/EpicGames.UHT/Parsers/UhtParsingScope.cs
@@ -275,7 +275,7 @@ namespace EpicGames.UHT.Parsers
 			bool javaDocStyle = span.Contains("/**", StringComparison.Ordinal);
 			bool cStyle = javaDocStyle || span.Contains("/*", StringComparison.Ordinal);
 			bool cppStyle = span.StartsWith("//", StringComparison.Ordinal);
-			commentsLength = javaDocStyle || cStyle ? RemoveBlockCommentMarkers(scratchChars, commentsLength) : commentsLength;
+			commentsLength = javaDocStyle || cStyle ? RemoveBlockCommentMarkers(scratchChars, commentsLength, javaDocStyle) : commentsLength;
 			commentsLength = cppStyle ? RemoveLineCommentMarkers(scratchChars, commentsLength) : commentsLength;
 
 			//wx widgets has a hard coded tab size of 8
@@ -491,8 +491,9 @@ namespace EpicGames.UHT.Parsers
 		/// </summary>
 		/// <param name="comments">Buffer containing comments to be processed.  Comments are removed inline</param>
 		/// <param name="inLength">Length of the comments</param>
+		/// <param name="javaDocStyle">If true, we are parsing both java and c style.  This is a strange hack for //***__ comments which end up as __</param>
 		/// <returns>New length of the comments</returns>
-		private static int RemoveBlockCommentMarkers(char[] comments, int inLength)
+		private static int RemoveBlockCommentMarkers(char[] comments, int inLength, bool javaDocStyle)
 		{
 			int outPos = 0;
 			int inPos = 0;
@@ -504,7 +505,13 @@ namespace EpicGames.UHT.Parsers
 						inPos++;
 						break;
 					case '/':
-						if (inPos + 2 < inLength && comments[inPos + 1] == '*' && comments[inPos + 2] == '*')
+						// This block of code is mimicking the old pattern of replacing "/**" with "" followed by "/*" with "".  
+						// Thus "//***" -> "/*" -> ""
+						if (javaDocStyle && inPos + 4 < inLength && comments[inPos + 1] == '/' && comments[inPos + 2] == '*' && comments[inPos + 3] == '*' && comments[inPos + 4] == '*')
+						{
+							inPos += 5;
+						}
+						else if (inPos + 2 < inLength && comments[inPos + 1] == '*' && comments[inPos + 2] == '*')
 						{
 							inPos += 3;
 						}
