Added support for ": int" size for regular and namespace enums.  Only "int" is supported.

diff --git a/Engine/Source/Programs/Shared/EpicGames.UHT/Parsers/UhtEnumParser.cs b/Engine/Source/Programs/Shared/EpicGames.UHT/Parsers/UhtEnumParser.cs
--- a/Engine/Source/Programs/Shared/EpicGames.UHT/Parsers/UhtEnumParser.cs
+++ b/Engine/Source/Programs/Shared/EpicGames.UHT/Parsers/UhtEnumParser.cs
@@ -92,6 +92,18 @@ namespace EpicGames.UHT.Parsers
 					}
 					else
 					{
+						if (enumObject.CppForm == UhtEnumCppForm.Regular)
+						{
+							if (topScope.TokenReader.TryOptional(":"))
+							{
+								UhtToken enumType = topScope.TokenReader.GetIdentifier("enumeration base");
+
+								if (enumType.Value != "int")
+								{
+									topScope.TokenReader.LogError($"Regular enums only support 'int' as the value size");
+								}
+							}
+						}
 						if ((enumObject.EnumFlags & EEnumFlags.Flags) != 0)
 						{
 							topScope.TokenReader.LogError("The 'Flags' specifier can only be used on enum classes");
@@ -117,6 +129,16 @@ namespace EpicGames.UHT.Parsers
 
 							UhtToken innerEnumToken = topScope.TokenReader.GetIdentifier("enumeration type name");
 
+							if (topScope.TokenReader.TryOptional(":"))
+							{
+								UhtToken enumType = topScope.TokenReader.GetIdentifier("enumeration base");
+
+								if (enumType.Value != "int")
+								{
+									topScope.TokenReader.LogError($"Namespace enums only support 'int' as the value size");
+								}
+							}
+
 							topScope.TokenReader.Require('{');
 							enumObject.CppType = $"{enumObject.SourceName}::{innerEnumToken.Value}";
 							break;
diff --git a/Engine/Source/Programs/UnrealHeaderTool/Private/HeaderParser.cpp b/Engine/Source/Programs/UnrealHeaderTool/Private/HeaderParser.cpp
--- a/Engine/Source/Programs/UnrealHeaderTool/Private/HeaderParser.cpp
+++ b/Engine/Source/Programs/UnrealHeaderTool/Private/HeaderParser.cpp
@@ -1120,6 +1120,22 @@ FUnrealEnumDefinitionInfo& FHeaderParser::CompileEnum()
 			Throwf(TEXT("ICE: Mismatch of underlying enum type between pre-parser and parser"));
 		}
 	}
+	else if (CppForm == UEnum::ECppForm::Regular)
+	{
+		if (MatchSymbol(TEXT(':')))
+		{
+			FToken BaseToken;
+			if (!GetIdentifier(BaseToken))
+			{
+				Throwf(TEXT("Missing enum base"));
+			}
+
+			if (!BaseToken.IsValue(TEXT("int"), ESearchCase::CaseSensitive))
+			{
+				LogError(TEXT("Regular enums only support 'int' as the value size"));
+			}
+		}
+	}
 	else
 	{
 		if (EnumHasAnyFlags(Flags, EEnumFlags::Flags))
@@ -1155,6 +1171,20 @@ FUnrealEnumDefinitionInfo& FHeaderParser::CompileEnum()
 
 			EnumDef.SetCppType(FString::Printf(TEXT("%s::%s"), *EnumIdentifier, *InnerEnumToken.GetTokenValue()));
 
+			if (MatchSymbol(TEXT(':')))
+			{
+				FToken BaseToken;
+				if (!GetIdentifier(BaseToken))
+				{
+					Throwf(TEXT("Missing enum base"));
+				}
+
+				if (!BaseToken.IsValue(TEXT("int"), ESearchCase::CaseSensitive))
+				{
+					LogError(TEXT("Namespace enums only support 'int' as the value size"));
+				}
+			}
+
 			RequireSymbol( TEXT('{'), TEXT("'Enum'") );
 		}
 		break;
