Horde: Updated blob store API.

* BlobIds now have a fixed structure, consisting of a server id (a list of identifiers separated by colons, indicating topology of the hierarchy) and content id (a globally unique identifier, which may contain path elements). Expectation is that any node in the storage hierarchy can retrieve a blob by using the server id to navigate the hierarchy until reaching a host that can provide the blob.
* Refs are now explicitly mappings from names to blob ids, and can be queried separately to the blob content. Overloads for transferring blob content inline are defined, for stores that choose to store the payload inline as an optimization.
* Ref queries now specify a time parameter, indicating required coherency of any cached values that may be returned without forwarding the request to the root storage node.
* Default implementation of the blob store now uses MongoDB to store refs (ie. name -> blob id mappings).
* Updated existing use cases to use new API and added tests.

diff --git a/Engine/Source/Programs/Horde/Horde.Agent/Commands/Bundles/CreateCommand.cs b/Engine/Source/Programs/Horde/Horde.Agent/Commands/Bundles/CreateCommand.cs
--- a/Engine/Source/Programs/Horde/Horde.Agent/Commands/Bundles/CreateCommand.cs
+++ b/Engine/Source/Programs/Horde/Horde.Agent/Commands/Bundles/CreateCommand.cs
@@ -46,27 +46,23 @@ namespace Horde.Agent.Commands.Bundles
 
 			#region Blobs
 
-			public async Task<IBlob?> TryReadBlobAsync(BlobId id, CancellationToken cancellationToken = default)
+			public async Task<IBlob> ReadBlobAsync(BlobId id, CancellationToken cancellationToken = default)
 			{
 				FileReference file = GetBlobFile(id);
-				if(!FileReference.Exists(file))
-				{
-					return null;
-				}
-
 				_logger.LogInformation("Reading {File}", file);
 				byte[] bytes = await FileReference.ReadAllBytesAsync(file, cancellationToken);
-				return BlobUtils.Deserialize(bytes);
+				return Blob.Deserialize(id, bytes);
 			}
 
-			public async Task<BlobId> WriteBlobAsync(RefName refName, ReadOnlySequence<byte> data, IReadOnlyList<BlobId> references, CancellationToken cancellationToken = default)
+			public async Task<IBlob> WriteBlobAsync(ReadOnlySequence<byte> data, IReadOnlyList<BlobId> references, RefName hintRefName = default, CancellationToken cancellationToken = default)
 			{
-				BlobId id = new BlobId($"{refName}/{Guid.NewGuid()}");
+				BlobId id = BlobId.Create(ServerId.Empty, hintRefName);
 				FileReference file = GetBlobFile(id);
 				DirectoryReference.CreateDirectory(file.Directory);
 				_logger.LogInformation("Writing {File}", file);
-				await WriteAsync(file, BlobUtils.Serialize(data, references), cancellationToken);
-				return id;
+				IBlob blob = Blob.FromMemory(id, data.AsSingleSegment(), references);
+				await WriteAsync(file, Blob.Serialize(blob), cancellationToken);
+				return blob;
 			}
 
 			#endregion
@@ -78,30 +74,42 @@ namespace Horde.Agent.Commands.Bundles
 				throw new NotImplementedException();
 			}
 
-			public Task<bool> HasRefAsync(RefName name, CancellationToken cancellationToken = default)
+			public async Task<IBlob?> TryReadRefAsync(RefName name, DateTime cacheTime = default, CancellationToken cancellationToken = default)
 			{
-				throw new NotImplementedException();
+				BlobId id = await TryReadRefTargetAsync(name, cacheTime, cancellationToken);
+				if (!id.IsValid())
+				{
+					return null;
+				}
+				return await ReadBlobAsync(id, cancellationToken);
 			}
 
-			public async Task<IBlob?> TryReadRefAsync(RefName name, CancellationToken cancellationToken = default)
+			public async Task<BlobId> TryReadRefTargetAsync(RefName name, DateTime cacheTime = default, CancellationToken cancellationToken = default)
 			{
 				FileReference file = GetRefFile(name);
 				if(!FileReference.Exists(file))
 				{
-					return null;
+					return BlobId.Empty;
 				}
 
 				_logger.LogInformation("Reading {File}", file);
-				byte[] bytes = await FileReference.ReadAllBytesAsync(file, cancellationToken);
-				return BlobUtils.Deserialize(bytes);
+				string text = await FileReference.ReadAllTextAsync(file, cancellationToken);
+				return new BlobId(text.Trim());
+			}
+
+			public async Task<IBlob> WriteRefAsync(RefName name, ReadOnlySequence<byte> data, IReadOnlyList<BlobId> references, CancellationToken cancellationToken = default)
+			{
+				IBlob blob = await WriteBlobAsync(data, references, name, cancellationToken);
+				await WriteRefTargetAsync(name, blob.Id, cancellationToken);
+				return blob;
 			}
 
-			public async Task WriteRefAsync(RefName name, ReadOnlySequence<byte> data, IReadOnlyList<BlobId> references, CancellationToken cancellationToken = default)
+			public async Task WriteRefTargetAsync(RefName name, BlobId id, CancellationToken cancellationToken = default)
 			{
 				FileReference file = GetRefFile(name);
 				DirectoryReference.CreateDirectory(file.Directory);
 				_logger.LogInformation("Writing {File}", file);
-				await WriteAsync(file, BlobUtils.Serialize(data, references), cancellationToken);
+				await FileReference.WriteAllTextAsync(file, id.ToString());
 			}
 
 			#endregion
diff --git a/Engine/Source/Programs/Horde/Horde.Agent/Commands/Bundles/DumpCommand.cs b/Engine/Source/Programs/Horde/Horde.Agent/Commands/Bundles/DumpCommand.cs
--- a/Engine/Source/Programs/Horde/Horde.Agent/Commands/Bundles/DumpCommand.cs
+++ b/Engine/Source/Programs/Horde/Horde.Agent/Commands/Bundles/DumpCommand.cs
@@ -28,12 +28,12 @@ namespace Horde.Agent.Commands.Bundles
 			IBlob blob;
 			if (BlobId != null)
 			{
-				blob = await blobStore.ReadBlobAsync(BlobId.Value, CancellationToken.None);
+				blob = await blobStore.ReadBlobAsync(BlobId.Value);
 				logger.LogInformation("Summary for blob {BlobId}", BlobId.Value);
 			}
 			else
 			{
-				blob = await blobStore.ReadRefAsync(RefName, CancellationToken.None);
+				blob = await blobStore.ReadRefAsync(RefName);
 				logger.LogInformation("Summary for ref {RefId}", RefName);
 			}
 
diff --git a/Engine/Source/Programs/Horde/Horde.Build.Tests/BlobStoreTests.cs b/Engine/Source/Programs/Horde/Horde.Build.Tests/BlobStoreTests.cs
new file mode 100644
--- /dev/null
+++ b/Engine/Source/Programs/Horde/Horde.Build.Tests/BlobStoreTests.cs
@@ -0,0 +1,108 @@
+// Copyright Epic Games, Inc. All Rights Reserved.
+
+using System;
+using System.Collections.Generic;
+using System.Linq;
+using System.Runtime.InteropServices;
+using System.Text.Json;
+using System.Threading.Tasks;
+using EpicGames.Core;
+using Horde.Build.Jobs.Graphs;
+using Horde.Build.Jobs;
+using Horde.Build.Issues;
+using Horde.Build.Logs;
+using Horde.Build.Users;
+using Horde.Build.Projects;
+using Horde.Build.Streams;
+using Horde.Build.Server;
+using Horde.Build.Tests.Stubs.Services;
+using Horde.Build.Utilities;
+using HordeCommon;
+using Microsoft.Extensions.DependencyInjection;
+using Microsoft.Extensions.Logging;
+using Microsoft.VisualStudio.TestTools.UnitTesting;
+using MongoDB.Bson;
+using MongoDB.Driver;
+using Moq;
+using Horde.Build.Issues.Handlers;
+using System.Threading;
+using EpicGames.Horde.Storage;
+using Horde.Build.Storage;
+using Horde.Build.Storage.Backends;
+using Microsoft.Extensions.Caching.Memory;
+using System.Buffers;
+
+namespace Horde.Build.Tests
+{
+	[TestClass]
+	public class BlobStoreTests : TestSetup
+	{
+		IBlobStore CreateBlobStore()
+		{
+			return new BasicBlobStore(MongoService, new TransientStorageBackend(), ServiceProvider.GetRequiredService<IMemoryCache>());
+		}
+
+		static byte[] CreateTestData(int length, int seed)
+		{
+			byte[] data = new byte[length];
+			new Random(seed).NextBytes(data);
+			return data;
+		}
+
+		[TestMethod]
+		public async Task LeafTest()
+		{
+			IBlobStore store = CreateBlobStore();
+
+			byte[] input = CreateTestData(256, 0);
+
+			IBlob blob = await store.WriteBlobAsync(new ReadOnlySequence<byte>(input), Array.Empty<BlobId>());
+			Assert.IsTrue(blob.Data.Span.SequenceEqual(input));
+
+			blob = await store.ReadBlobAsync(blob.Id);
+			Assert.IsTrue(blob.Data.Span.SequenceEqual(input));
+		}
+
+		[TestMethod]
+		public async Task ReferenceTest()
+		{
+			IBlobStore store = CreateBlobStore();
+
+			byte[] input1 = CreateTestData(256, 1);
+			IBlob blob1 = await store.WriteBlobAsync(new ReadOnlySequence<byte>(input1), Array.Empty<BlobId>());
+			blob1 = await store.ReadBlobAsync(blob1.Id);
+			Assert.IsTrue(blob1.Data.Span.SequenceEqual(input1));
+			Assert.IsTrue(blob1.References.SequenceEqual(Array.Empty<BlobId>()));
+
+			byte[] input2 = CreateTestData(256, 2);
+			IBlob blob2 = await store.WriteBlobAsync(new ReadOnlySequence<byte>(input2), new BlobId[] { blob1.Id });
+			blob2 = await store.ReadBlobAsync(blob2.Id);
+			Assert.IsTrue(blob2.Data.Span.SequenceEqual(input2));
+			Assert.IsTrue(blob2.References.SequenceEqual(new BlobId[] { blob1.Id }));
+
+			byte[] input3 = CreateTestData(256, 3);
+			IBlob blob3 = await store.WriteBlobAsync(new ReadOnlySequence<byte>(input3), new BlobId[] { blob1.Id, blob2.Id, blob1.Id });
+			blob3 = await store.ReadBlobAsync(blob3.Id);
+			Assert.IsTrue(blob3.Data.Span.SequenceEqual(input3));
+			Assert.IsTrue(blob3.References.SequenceEqual(new BlobId[] { blob1.Id, blob2.Id, blob1.Id }));
+
+			for(int idx = 0; idx < 2; idx++)
+			{
+				RefName refName = new RefName("hello");
+				await store.WriteRefTargetAsync(refName, blob3.Id);
+				BlobId refTargetId = await store.ReadRefIdAsync(refName);
+				Assert.AreEqual(blob3.Id, refTargetId);
+			}
+
+			RefName refName2 = new RefName("hello2");
+
+			IBlob refTarget = await store.WriteRefAsync(refName2, new ReadOnlySequence<byte>(input3), new BlobId[] { blob1.Id, blob2.Id });
+			Assert.IsTrue(refTarget.Data.Span.SequenceEqual(input3));
+			Assert.IsTrue(refTarget.References.SequenceEqual(new BlobId[] { blob1.Id, blob2.Id }));
+
+			refTarget = await store.ReadRefAsync(refName2);
+			Assert.IsTrue(refTarget.Data.Span.SequenceEqual(input3));
+			Assert.IsTrue(refTarget.References.SequenceEqual(new BlobId[] { blob1.Id, blob2.Id }));
+		}
+	}
+}
diff --git a/Engine/Source/Programs/Horde/Horde.Build.Tests/TestSetup.cs b/Engine/Source/Programs/Horde/Horde.Build.Tests/TestSetup.cs
--- a/Engine/Source/Programs/Horde/Horde.Build.Tests/TestSetup.cs
+++ b/Engine/Source/Programs/Horde/Horde.Build.Tests/TestSetup.cs
@@ -235,7 +235,7 @@ namespace Horde.Build.Tests
 
 			services.AddSingleton<ITreeStore<CommitService>>(sp =>
 			{
-				IBlobStore blobStore = new BlobStore(new TransientStorageBackend());
+				IBlobStore blobStore = new BasicBlobStore(sp.GetRequiredService<MongoService>(), new TransientStorageBackend(), sp.GetRequiredService<IMemoryCache>());
 				return new BundleStore(blobStore, new BundleOptions()).ForType<CommitService>();
 			});
 
diff --git a/Engine/Source/Programs/Horde/Horde.Build/Perforce/ReplicationService.cs b/Engine/Source/Programs/Horde/Horde.Build/Perforce/ReplicationService.cs
--- a/Engine/Source/Programs/Horde/Horde.Build/Perforce/ReplicationService.cs
+++ b/Engine/Source/Programs/Horde/Horde.Build/Perforce/ReplicationService.cs
@@ -426,7 +426,7 @@ namespace Horde.Build.Perforce
 				{
 					RefName prevRefName = GetRefName(stream, values[0]);
 
-					ReplicationNode? contents = await _treeStore.TryReadTreeAsync<ReplicationNode>(prevRefName, cancellationToken);
+					ReplicationNode? contents = await _treeStore.TryReadTreeAsync<ReplicationNode>(prevRefName, cancellationToken: cancellationToken);
 					if (contents == null)
 					{
 						_logger.LogInformation("No content for CL {Change}; creating full snapshot", values[0]);
@@ -506,7 +506,7 @@ namespace Horde.Build.Perforce
 		public async Task<ReplicationNode> ReadCommitTreeAsync(IStream stream, int change, string? filter, bool revisionsOnly, CancellationToken cancellationToken)
 		{
 			RefName name = GetRefName(stream.Id, change, filter, revisionsOnly);
-			return await _treeStore.ReadTreeAsync<ReplicationNode>(name, cancellationToken);
+			return await _treeStore.ReadTreeAsync<ReplicationNode>(name, cancellationToken: cancellationToken);
 		}
 
 		/// <summary>
@@ -612,7 +612,7 @@ namespace Horde.Build.Perforce
 			ITreeWriter writer = _treeStore.CreateTreeWriter(incRefName);
 
 			// Get the current sync state for this change
-			ReplicationNode? syncNode = await _treeStore.TryReadTreeAsync<ReplicationNode>(incRefName, cancellationToken);
+			ReplicationNode? syncNode = await _treeStore.TryReadTreeAsync<ReplicationNode>(incRefName, cancellationToken: cancellationToken);
 			bool deleteIncRef = syncNode != null;
 			if (syncNode == null)
 			{
diff --git a/Engine/Source/Programs/Horde/Horde.Build/Startup.cs b/Engine/Source/Programs/Horde/Horde.Build/Startup.cs
--- a/Engine/Source/Programs/Horde/Horde.Build/Startup.cs
+++ b/Engine/Source/Programs/Horde/Horde.Build/Startup.cs
@@ -95,6 +95,7 @@ using Horde.Build.Server.Notices;
 using Horde.Build.Notifications.Sinks;
 using EpicGames.Horde.Storage.Bundles;
 using StatusCode = Grpc.Core.StatusCode;
+using Microsoft.Extensions.Caching.Memory;
 
 namespace Horde.Build
 {
@@ -279,7 +280,7 @@ namespace Horde.Build
 		static IBlobStore CreateBlobStore(IServiceProvider sp, BlobStoreOptions options)
 		{
 			IStorageBackend backend = CreateStorageBackend(sp, options);
-			return new BlobStore(backend);
+			return new BasicBlobStore(sp.GetRequiredService<MongoService>(), backend, sp.GetRequiredService<IMemoryCache>());
 		}
 
 		static ITreeStore CreateTreeStore(IServiceProvider sp, TreeStoreOptions options)
@@ -460,9 +461,9 @@ namespace Horde.Build
 			}
 
 			// Storage providers
-			services.AddSingleton<IStorageBackend<PersistentLogStorage>>(sp => CreateStorageBackend(sp, settings.LogStorage).ForType<PersistentLogStorage>());
-			services.AddSingleton<IStorageBackend<ArtifactCollection>>(sp => CreateStorageBackend(sp, settings.ArtifactStorage).ForType<ArtifactCollection>());
-			services.AddSingleton<ITreeStore<ReplicationService>>(sp => CreateTreeStore(sp, settings.CommitStorage).ForType<ReplicationService>());
+			services.AddSingleton(sp => CreateStorageBackend(sp, settings.LogStorage).ForType<PersistentLogStorage>());
+			services.AddSingleton(sp => CreateStorageBackend(sp, settings.ArtifactStorage).ForType<ArtifactCollection>());
+			services.AddSingleton(sp => CreateTreeStore(sp, settings.CommitStorage).ForType<ReplicationService>());
 
 			services.AddHordeStorage(settings => configSection.GetSection("Storage").Bind(settings));
 			
diff --git a/Engine/Source/Programs/Horde/Horde.Build/Storage/Backends/AwsStorageBackend.cs b/Engine/Source/Programs/Horde/Horde.Build/Storage/Backends/AwsStorageBackend.cs
--- a/Engine/Source/Programs/Horde/Horde.Build/Storage/Backends/AwsStorageBackend.cs
+++ b/Engine/Source/Programs/Horde/Horde.Build/Storage/Backends/AwsStorageBackend.cs
@@ -318,7 +318,7 @@ namespace Horde.Build.Storage.Backends
 					}
 				}
 
-				await Task.Delay(retryTimes[attempt]);
+				await Task.Delay(retryTimes[attempt], cancellationToken);
 			}
 		}
 
diff --git a/Engine/Source/Programs/Horde/Horde.Build/Storage/BasicBlobStore.cs b/Engine/Source/Programs/Horde/Horde.Build/Storage/BasicBlobStore.cs
new file mode 100644
--- /dev/null
+++ b/Engine/Source/Programs/Horde/Horde.Build/Storage/BasicBlobStore.cs
@@ -0,0 +1,170 @@
+// Copyright Epic Games, Inc. All Rights Reserved.
+
+using System;
+using System.Buffers;
+using System.Collections.Generic;
+using System.IO;
+using System.Text;
+using System.Threading;
+using System.Threading.Tasks;
+using EpicGames.Core;
+using EpicGames.Horde.Storage;
+using Horde.Build.Server;
+using Microsoft.Extensions.Caching.Memory;
+using MongoDB.Bson.Serialization.Attributes;
+using MongoDB.Driver;
+
+namespace Horde.Build.Storage
+{
+	/// <summary>
+	/// Implementation of <see cref="IBlobStore"/> using Mongo for storing refs, and an <see cref="IStorageBackend"/> implementation for bulk data.
+	/// </summary>
+	public class BasicBlobStore : IBlobStore
+	{
+		class RefDocument
+		{
+			[BsonId]
+			public RefName Name { get; set; }
+			public BlobId Value { get; set; }
+
+			public RefDocument()
+			{
+				Name = RefName.Empty;
+				Value = BlobId.Empty;
+			}
+
+			public RefDocument(RefName name, BlobId value)
+			{
+				Name = name;
+				Value = value;
+			}
+		}
+
+		class CachedRefValue
+		{
+			public RefName Name { get; }
+			public BlobId Value { get; }
+			public DateTime Time { get; }
+
+			public CachedRefValue(RefName name, BlobId value)
+			{
+				Name = name;
+				Value = value;
+				Time = DateTime.UtcNow;
+			}
+		}
+
+		readonly ServerId _serverId = ServerId.Empty;
+
+		readonly IMongoCollection<RefDocument> _refs;
+		readonly IStorageBackend _backend;
+		readonly IMemoryCache _cache;
+
+		/// <summary>
+		/// Constructor
+		/// </summary>
+		/// <param name="mongoService">The mongo service implementation</param>
+		/// <param name="backend">Backend to use for storing data</param>
+		/// <param name="cache">Cache for ref values</param>
+		public BasicBlobStore(MongoService mongoService, IStorageBackend backend, IMemoryCache cache)
+		{
+			_refs = mongoService.GetCollection<RefDocument>("Refs");
+			_backend = backend;
+			_cache = cache;
+		}
+
+		#region Blobs
+
+		static string GetBlobPath(BlobId id) => $"{id.GetContentId()}.blob";
+
+		/// <inheritdoc/>
+		public async Task<IBlob> ReadBlobAsync(BlobId id, CancellationToken cancellationToken = default)
+		{
+			string path = GetBlobPath(id);
+
+			ReadOnlyMemory<byte>? data = await _backend.ReadBytesAsync(path, cancellationToken);
+			if (data == null)
+			{
+				throw new KeyNotFoundException($"Unable to read data from {path}");
+			}
+
+			return Blob.Deserialize(id, data.Value);
+		}
+
+		/// <inheritdoc/>
+		public async Task<IBlob> WriteBlobAsync(ReadOnlySequence<byte> data, IReadOnlyList<BlobId> references, RefName hintRefName = default, CancellationToken cancellationToken = default)
+		{
+			BlobId id = BlobId.Create(_serverId, hintRefName);
+			string path = GetBlobPath(id);
+
+			ReadOnlyMemory<byte> memory = Blob.Serialize(data, references).AsSingleSegment();
+			await _backend.WriteBytesAsync(path, memory, cancellationToken);
+			return Blob.FromMemory(id, memory.Slice(memory.Length - (int)data.Length), references);
+		}
+
+		#endregion
+
+		#region Refs
+
+		CachedRefValue AddRefToCache(RefName name, BlobId value)
+		{
+			CachedRefValue cacheValue = new CachedRefValue(name, value);
+			using (ICacheEntry newEntry = _cache.CreateEntry(name))
+			{
+				newEntry.Value = cacheValue;
+				newEntry.SetSize(name.Text.Length + value.Inner.Length);
+				newEntry.SetSlidingExpiration(TimeSpan.FromMinutes(5.0));
+			}
+			return cacheValue;
+		}
+
+		/// <inheritdoc/>
+		public async Task DeleteRefAsync(RefName name, CancellationToken cancellationToken = default)
+		{
+			await _refs.DeleteOneAsync(x => x.Name == name, cancellationToken);
+			AddRefToCache(name, BlobId.Empty);
+		}
+
+		/// <inheritdoc/>
+		public async Task<IBlob?> TryReadRefAsync(RefName name, DateTime cacheTime = default, CancellationToken cancellationToken = default)
+		{
+			BlobId blobId = await TryReadRefTargetAsync(name, cacheTime, cancellationToken);
+			if (blobId.IsValid())
+			{
+				return await ReadBlobAsync(blobId, cancellationToken);
+			}
+			return null;
+		}
+
+		/// <inheritdoc/>
+		public async Task<BlobId> TryReadRefTargetAsync(RefName name, DateTime cacheTime = default, CancellationToken cancellationToken = default)
+		{
+			CachedRefValue entry;
+			if (!_cache.TryGetValue(name, out entry) || entry.Time < cacheTime)
+			{
+				RefDocument? refDocument = await _refs.Find(x => x.Name == name).FirstOrDefaultAsync(cancellationToken);
+				BlobId value = (refDocument == null) ? BlobId.Empty : refDocument.Value;
+				entry = AddRefToCache(name, value);
+			}
+			return entry.Value;
+		}
+
+		/// <inheritdoc/>
+		public async Task<IBlob> WriteRefAsync(RefName name, ReadOnlySequence<byte> data, IReadOnlyList<BlobId> references, CancellationToken cancellationToken = default)
+		{
+			IBlob blob = await WriteBlobAsync(data, references, name, cancellationToken);
+			await WriteRefTargetAsync(name, blob.Id, cancellationToken);
+			return blob;
+		}
+
+		/// <inheritdoc/>
+		public async Task WriteRefTargetAsync(RefName name, BlobId blobId, CancellationToken cancellationToken = default)
+		{
+			RefDocument refDocument = new RefDocument(name, blobId);
+			await _refs.ReplaceOneAsync(x => x.Name == name, refDocument, new ReplaceOptions { IsUpsert = true }, cancellationToken);
+			AddRefToCache(name, blobId);
+		}
+
+		#endregion
+	}
+}
diff --git a/Engine/Source/Programs/Horde/Horde.Build/Storage/Blobs/BlobStore.cs b/Engine/Source/Programs/Horde/Horde.Build/Storage/Blobs/BlobStore.cs
deleted file mode 100644
diff --git a/Engine/Source/Programs/Shared/EpicGames.Core/StringUtils.cs b/Engine/Source/Programs/Shared/EpicGames.Core/StringUtils.cs
--- a/Engine/Source/Programs/Shared/EpicGames.Core/StringUtils.cs
+++ b/Engine/Source/Programs/Shared/EpicGames.Core/StringUtils.cs
@@ -452,6 +452,50 @@ namespace EpicGames.Core
 			return new Utf8String(characters);
 		}
 
+		/// <summary>
+		/// Formats an array of bytes as a hexadecimal string
+		/// </summary>
+		/// <param name="bytes">An array of bytes</param>
+		/// <param name="characters">Buffer to receive the characters</param>
+		public static void FormatUtf8HexString(ReadOnlySpan<byte> bytes, Span<byte> characters)
+		{
+			for (int idx = 0; idx < bytes.Length; idx++)
+			{
+				characters[idx * 2 + 0] = s_hexDigitToUtf8Byte[bytes[idx] >> 4];
+				characters[idx * 2 + 1] = s_hexDigitToUtf8Byte[bytes[idx] & 15];
+			}
+		}
+
+
+		/// <summary>
+		/// Formats a 32-bit unsigned integer as a hexadecimal string
+		/// </summary>
+		/// <param name="value">Value to render</param>
+		/// <returns>Hex string</returns>
+		public static Utf8String FormatUtf8HexString(uint value)
+		{
+			byte[] buffer = new byte[8];
+			FormatUtf8HexString(value, buffer);
+			return new Utf8String(buffer);
+		}
+
+		/// <summary>
+		/// Formats a 32-bit unsigned integer as a hexadecimal string
+		/// </summary>
+		/// <param name="value">Value to render</param>
+		/// <param name="characters">Buffer to receive the characters</param>
+		public static void FormatUtf8HexString(uint value, Span<byte> characters)
+		{
+			characters[0] = s_hexDigitToUtf8Byte[(value >> 28) & 15];
+			characters[1] = s_hexDigitToUtf8Byte[(value >> 24) & 15];
+			characters[2] = s_hexDigitToUtf8Byte[(value >> 20) & 15];
+			characters[3] = s_hexDigitToUtf8Byte[(value >> 16) & 15];
+			characters[4] = s_hexDigitToUtf8Byte[(value >> 12) & 15];
+			characters[5] = s_hexDigitToUtf8Byte[(value >> 8) & 15];
+			characters[6] = s_hexDigitToUtf8Byte[(value >> 4) & 15];
+			characters[7] = s_hexDigitToUtf8Byte[value & 15];
+		}
+
 		/// <summary>
 		/// Quotes a string as a command line argument
 		/// </summary>
diff --git a/Engine/Source/Programs/Shared/EpicGames.Horde/Storage/Backends/InMemoryBlobStore.cs b/Engine/Source/Programs/Shared/EpicGames.Horde/Storage/Backends/InMemoryBlobStore.cs
--- a/Engine/Source/Programs/Shared/EpicGames.Horde/Storage/Backends/InMemoryBlobStore.cs
+++ b/Engine/Source/Programs/Shared/EpicGames.Horde/Storage/Backends/InMemoryBlobStore.cs
@@ -22,29 +22,31 @@ namespace EpicGames.Horde.Storage.Backends
 		/// <summary>
 		/// Map of ref name to ref data
 		/// </summary>
-		readonly ConcurrentDictionary<RefName, IBlob> _refs = new ConcurrentDictionary<RefName, IBlob>();
+		readonly ConcurrentDictionary<RefName, BlobId> _refs = new ConcurrentDictionary<RefName, BlobId>();
 
 		/// <inheritdoc cref="_blobs"/>
 		public IReadOnlyDictionary<BlobId, IBlob> Blobs => _blobs;
 
 		/// <inheritdoc cref="_refs"/>
-		public IReadOnlyDictionary<RefName, IBlob> Refs => _refs;
+		public IReadOnlyDictionary<RefName, BlobId> Refs => _refs;
 
 		#region Blobs
 
 		/// <inheritdoc/>
-		public Task<IBlob?> TryReadBlobAsync(BlobId blobId, CancellationToken cancellationToken = default)
+		public Task<IBlob> ReadBlobAsync(BlobId blobId, CancellationToken cancellationToken = default)
 		{
-			_blobs.TryGetValue(blobId, out IBlob? blob);
-			return Task.FromResult(blob);
+			return Task.FromResult(_blobs[blobId]);
 		}
 
 		/// <inheritdoc/>
-		public Task<BlobId> WriteBlobAsync(RefName refName, ReadOnlySequence<byte> data, IReadOnlyList<BlobId> references, CancellationToken cancellationToken = default)
+		public Task<IBlob> WriteBlobAsync(ReadOnlySequence<byte> data, IReadOnlyList<BlobId> references, RefName hintRefName = default, CancellationToken cancellationToken = default)
 		{
 			BlobId blobId = new BlobId(Guid.NewGuid().ToString());
-			_blobs[blobId] = BlobUtils.FromMemory(data.ToArray(), references);
-			return Task.FromResult(blobId);
+
+			IBlob blob = Blob.FromMemory(blobId, data.ToArray(), references);
+			_blobs[blobId] = blob;
+
+			return Task.FromResult(blob);
 		}
 
 		#endregion
@@ -58,16 +60,41 @@ namespace EpicGames.Horde.Storage.Backends
 		public Task<bool> HasRefAsync(RefName name, CancellationToken cancellationToken) => Task.FromResult(_refs.ContainsKey(name));
 
 		/// <inheritdoc/>
-		public Task<IBlob?> TryReadRefAsync(RefName name, CancellationToken cancellationToken)
+		public Task<IBlob?> TryReadRefAsync(RefName name, DateTime cacheTime = default, CancellationToken cancellationToken = default)
 		{
-			_refs.TryGetValue(name, out IBlob? blob);
+			IBlob? blob = null;
+			if (_refs.TryGetValue(name, out BlobId blobId))
+			{
+				_blobs.TryGetValue(blobId, out blob);
+			}
 			return Task.FromResult(blob);
 		}
 
-		/// <inheritdoc/>s
-		public Task WriteRefAsync(RefName name, ReadOnlySequence<byte> data, IReadOnlyList<BlobId> references, CancellationToken cancellationToken)
+		/// <inheritdoc/>
+		public Task<BlobId> TryReadRefTargetAsync(RefName name, DateTime cacheTime = default, CancellationToken cancellationToken = default)
+		{
+			if (_refs.TryGetValue(name, out BlobId blobId))
+			{
+				return Task.FromResult(blobId);
+			}
+			else
+			{
+				return Task.FromResult(BlobId.Empty);
+			}
+		}
+
+		/// <inheritdoc/>
+		public async Task<IBlob> WriteRefAsync(RefName name, ReadOnlySequence<byte> data, IReadOnlyList<BlobId> references, CancellationToken cancellationToken)
+		{
+			IBlob blob = await WriteBlobAsync(data, references, name, cancellationToken);
+			_refs[name] = blob.Id;
+			return blob;
+		}
+
+		/// <inheritdoc/>
+		public Task WriteRefTargetAsync(RefName name, BlobId blobId, CancellationToken cancellationToken = default)
 		{
-			_refs[name] = BlobUtils.FromMemory(data.ToArray(), references);
+			_refs[name] = blobId;
 			return Task.CompletedTask;
 		}
 
diff --git a/Engine/Source/Programs/Shared/EpicGames.Horde/Storage/BlobId.cs b/Engine/Source/Programs/Shared/EpicGames.Horde/Storage/BlobId.cs
--- a/Engine/Source/Programs/Shared/EpicGames.Horde/Storage/BlobId.cs
+++ b/Engine/Source/Programs/Shared/EpicGames.Horde/Storage/BlobId.cs
@@ -5,7 +5,6 @@ using EpicGames.Serialization;
 using System;
 using System.ComponentModel;
 using System.Globalization;
-using System.Text;
 using System.Text.Json;
 using System.Text.Json.Serialization;
 
@@ -20,6 +19,22 @@ namespace EpicGames.Horde.Storage
 	[CbConverter(typeof(BlobIdCbConverter))]
 	public struct BlobId : IEquatable<BlobId>
 	{
+		/// <summary>
+		/// Dummy enum to allow invoking the constructor which takes a sanitized full path
+		/// </summary>
+		public enum Sanitize
+		{
+			/// <summary>
+			/// Dummy value
+			/// </summary>
+			None
+		}
+
+		/// <summary>
+		/// Empty blob id
+		/// </summary>
+		public static BlobId Empty { get; } = default;
+
 		/// <summary>
 		/// Identifier for the ref
 		/// </summary>
@@ -32,6 +47,117 @@ namespace EpicGames.Horde.Storage
 		public BlobId(Utf8String inner)
 		{
 			Inner = inner;
+			ValidateArgument(nameof(inner), inner);
+		}
+
+		/// <summary>
+		/// Constructor
+		/// </summary>
+		/// <param name="inner"></param>
+		/// <param name="sanitize"></param>
+		public BlobId(Utf8String inner, Sanitize sanitize)
+		{
+			Inner = inner;
+			_ = sanitize;
+		}
+
+		/// <summary>
+		/// Create a unique content id, optionally including a ref name
+		/// </summary>
+		/// <param name="serverId">The server id</param>
+		/// <param name="hintRefName">Hint for the ref containing this blob</param>
+		/// <returns>New content id</returns>
+		public static BlobId Create(ServerId serverId, RefName hintRefName = default)
+		{
+			Utf8String prefix = hintRefName.Text;
+			if (prefix.Length == 0)
+			{
+				DateTime now = DateTime.UtcNow.Date;
+				prefix = $"_by_date_/{now.Year}-{now.Month:D2}/{now.Day:D2}";
+			}
+
+			byte[] buffer;
+			Span<byte> span;
+
+			if (serverId.IsValid())
+			{
+				buffer = new byte[serverId.Inner.Length + 1 + prefix.Length + 1 + 24];
+				span = buffer;
+
+				serverId.Inner.Span.CopyTo(span);
+				span = span.Slice(serverId.Inner.Length);
+
+				span[0] = (byte)':';
+				span = span.Slice(1);
+			}
+			else
+			{
+				buffer = new byte[prefix.Length + 1 + 24];
+				span = buffer;
+			}
+
+			prefix.Span.CopyTo(span);
+			span = span.Slice(prefix.Length);
+
+			span[0] = (byte)'/';
+			span = span.Slice(1);
+
+			ContentId.GenerateUniqueId(span);
+			return new BlobId(new Utf8String(buffer), Sanitize.None);
+		}
+
+		/// <summary>
+		/// Validates a given string as a content id
+		/// </summary>
+		/// <param name="name">Name of the argument</param>
+		/// <param name="text">String to validate</param>
+		public static void ValidateArgument(string name, Utf8String text)
+		{
+			if (text.Length == 0)
+			{
+				throw new ArgumentException("Blob identifiers cannot be empty", name);
+			}
+
+			int colonIdx = text.LastIndexOf((byte)':');
+			if (colonIdx != -1)
+			{
+				ServerId.ValidateArgument(name, text.Substring(0, colonIdx));
+			}
+
+			ContentId.ValidateArgument(name, text.Substring(colonIdx + 1));
+		}
+
+		/// <summary>
+		/// Checks whether this blob id is valid
+		/// </summary>
+		/// <returns>True if the identifier is valid</returns>
+		public bool IsValid() => Inner.Length > 0;
+
+		/// <summary>
+		/// Gets the server id for this blob
+		/// </summary>
+		/// <returns>Server id</returns>
+		public ServerId GetServerId()
+		{
+			int colonIdx = Inner.LastIndexOf((byte)':');
+			if (colonIdx == -1)
+			{
+				return ServerId.Empty;
+			}
+			else
+			{
+				return new ServerId(Inner.Substring(0, colonIdx), ServerId.Sanitize.None);
+			}
+		}
+
+		/// <summary>
+		/// Gets the content id for this blob
+		/// </summary>
+		/// <returns>Content id</returns>
+		public ContentId GetContentId()
+		{
+			int colonIdx = Inner.LastIndexOf((byte)':');
+			return new ContentId(Inner.Substring(colonIdx + 1), ContentId.Sanitize.None);
 		}
 
 		/// <inheritdoc/>
diff --git a/Engine/Source/Programs/Shared/EpicGames.Horde/Storage/Bundles/Bundle.cs b/Engine/Source/Programs/Shared/EpicGames.Horde/Storage/Bundles/Bundle.cs
--- a/Engine/Source/Programs/Shared/EpicGames.Horde/Storage/Bundles/Bundle.cs
+++ b/Engine/Source/Programs/Shared/EpicGames.Horde/Storage/Bundles/Bundle.cs
@@ -5,6 +5,7 @@ using EpicGames.Serialization;
 using System;
 using System.Buffers;
 using System.Collections.Generic;
+using System.Diagnostics.CodeAnalysis;
 using System.IO;
 using System.Linq;
 using System.Threading;
@@ -63,6 +64,26 @@ namespace EpicGames.Horde.Storage.Bundles
 			Payload = new ReadOnlySequence<byte>(reader.ReadFixedLengthBytes(length));
 		}
 
+		/// <summary>
+		/// Create a bundle from a blob
+		/// </summary>
+		/// <param name="blob">Blob to parse from</param>
+		/// <returns>Bundle parsed from the given blob</returns>
+		[return: NotNullIfNotNull("blob")]
+		public static Bundle? FromBlob(IBlob? blob)
+		{
+			if (blob == null)
+			{
+				return null;
+			}
+
+			ReadOnlyMemory<byte> data = blob.Data;
+			MemoryReader reader = new MemoryReader(data);
+			Bundle bundle = new Bundle(reader);
+			reader.CheckEmpty();
+			return bundle;
+		}
+
 		/// <summary>
 		/// Serializes the bundle to a sequence of bytes
 		/// </summary>
@@ -332,78 +353,4 @@ namespace EpicGames.Horde.Storage.Bundles
 			}
 		}
 	}
-
-	/// <summary>
-	/// Extension methods for serializing bundles
-	/// </summary>
-	public static class BundleExtensions
-	{
-		/// <summary>
-		/// Read a bundle from a blob store
-		/// </summary>
-		/// <param name="store">Blob store to read from</param>
-		/// <param name="id">Identifier for the blob</param>
-		/// <param name="cancellationToken">Cancellation token for the operation</param>
-		/// <returns>The bundle that was read</returns>
-		public static async Task<Bundle> ReadBundleAsync(this IBlobStore store, BlobId id, CancellationToken cancellationToken)
-		{
-			IBlob blob = await store.ReadBlobAsync(id, cancellationToken);
-			return ReadBundle(blob);
-		}
-
-		/// <summary>
-		/// Read a bundle from a blob store
-		/// </summary>
-		/// <param name="store">Blob store to read from</param>
-		/// <param name="name">Identifier for the blob</param>
-		/// <param name="cancellationToken">Cancellation token for the operation</param>
-		/// <returns>The bundle that was read</returns>
-		public static async Task<Bundle?> TryReadBundleAsync(this IBlobStore store, RefName name, CancellationToken cancellationToken)
-		{
-			IBlob? blob = await store.TryReadRefAsync(name, cancellationToken);
-			if (blob == null)
-			{
-				return null;
-			}
-			return ReadBundle(blob);
-		}
-
-		/// <summary>
-		/// Writes a bundle to a blob store
-		/// </summary>
-		/// <param name="store">Blob store to write to</param>
-		/// <param name="name">Identifier for the blob</param>
-		/// <param name="bundle">Bundle to write</param>
-		/// <param name="cancellationToken">Cancellation token for the operation</param>
-		public static Task WriteBundleRefAsync(this IBlobStore store, RefName name, Bundle bundle, CancellationToken cancellationToken)
-		{
-			ReadOnlySequence<byte> data = bundle.AsSequence();
-			List<BlobId> imports = bundle.Header.Imports.Select(x => x.BlobId).ToList();
-			return store.WriteRefAsync(name, data, imports, cancellationToken);
-		}
-
-		/// <summary>
-		/// Writes a bundle to a blob store
-		/// </summary>
-		/// <param name="store">Blob store to write to</param>
-		/// <param name="name">Name of the root ref</param>
-		/// <param name="bundle">Bundle to write</param>
-		/// <param name="cancellationToken">Cancellation token for the operation</param>
-		/// <returns>Identifier for the blob</returns>
-		public static Task<BlobId> WriteBundleBlobAsync(this IBlobStore store, RefName name, Bundle bundle, CancellationToken cancellationToken)
-		{
-			ReadOnlySequence<byte> data = bundle.AsSequence();
-			List<BlobId> imports = bundle.Header.Imports.Select(x => x.BlobId).ToList();
-			return store.WriteBlobAsync(name, data, imports, cancellationToken);
-		}
-
-		static Bundle ReadBundle(IBlob blob)
-		{
-			ReadOnlyMemory<byte> data = blob.Data;
-			MemoryReader reader = new MemoryReader(data);
-			Bundle bundle = new Bundle(reader);
-			reader.CheckEmpty();
-			return bundle;
-		}
-	}
 }
diff --git a/Engine/Source/Programs/Shared/EpicGames.Horde/Storage/Bundles/BundleStore.cs b/Engine/Source/Programs/Shared/EpicGames.Horde/Storage/Bundles/BundleStore.cs
--- a/Engine/Source/Programs/Shared/EpicGames.Horde/Storage/Bundles/BundleStore.cs
+++ b/Engine/Source/Programs/Shared/EpicGames.Horde/Storage/Bundles/BundleStore.cs
@@ -285,12 +285,12 @@ namespace EpicGames.Horde.Storage.Bundles
 		}
 
 		/// <inheritdoc/>
-		public Task<bool> HasTreeAsync(RefName name, CancellationToken cancellationToken) => _blobStore.HasRefAsync(name, cancellationToken);
+		public Task<bool> HasTreeAsync(RefName name, CancellationToken cancellationToken) => _blobStore.HasRefAsync(name, cancellationToken: cancellationToken);
 
 		/// <inheritdoc/>
-		public async Task<ITreeBlob?> TryReadTreeAsync(RefName name, CancellationToken cancellationToken = default)
+		public async Task<ITreeBlob?> TryReadTreeAsync(RefName name, TimeSpan maxAge = default, CancellationToken cancellationToken = default)
 		{
-			Bundle? bundle = await _blobStore.TryReadBundleAsync(name, cancellationToken);
+			Bundle? bundle = Bundle.FromBlob(await _blobStore.TryReadRefAsync(name, maxAge, cancellationToken));
 			if (bundle == null)
 			{
 				return null;
@@ -413,7 +413,7 @@ namespace EpicGames.Horde.Storage.Bundles
 			Bundle bundle;
 			try
 			{
-				bundle = await _blobStore.ReadBundleAsync(blobId, cancellationToken);
+				bundle = Bundle.FromBlob(await _blobStore.ReadBlobAsync(blobId, cancellationToken));
 			}
 			finally
 			{
@@ -649,12 +649,12 @@ namespace EpicGames.Horde.Storage.Bundles
 			BundleHeader header = bundle.Header;
 
 			// Write it to storage
-			BlobId blobId = await _blobStore.WriteBundleBlobAsync(context.Name, bundle, cancellationToken);
-			string cacheKey = GetBundleCacheKey(blobId);
+			IBlob blob = await _blobStore.WriteBlobAsync(bundle.AsSequence(), bundle.Header.Imports.Select(x => x.BlobId).ToList(), context.Name, cancellationToken);
+			string cacheKey = GetBundleCacheKey(blob.Id);
 			AddBundleToCache(cacheKey, bundle);
 
 			// Create a BundleInfo for it
-			BundleInfo bundleInfo = context.FindOrAddBundle(blobId, nodes.Length);
+			BundleInfo bundleInfo = context.FindOrAddBundle(blob.Id, nodes.Length);
 			for (int idx = 0; idx < nodes.Length; idx++)
 			{
 				bundleInfo.Exports[idx] = nodes[idx];
@@ -709,7 +709,7 @@ namespace EpicGames.Horde.Storage.Bundles
 		async Task FlushInternalAsync(BundleWriteContext context, NodeInfo[] writeNodes, CancellationToken cancellationToken)
 		{
 			Bundle bundle = CreateBundle(context, writeNodes);
-			await _blobStore.WriteBundleRefAsync(context.Name, bundle, cancellationToken);
+			await _blobStore.WriteRefAsync(context.Name, bundle.AsSequence(), bundle.Header.Imports.ConvertAll(x => x.BlobId), cancellationToken);
 		}
 
 		void FindNodesToWrite(NodeInfo root, List<NodeInfo> nodes, HashSet<NodeInfo> uniqueNodes)
diff --git a/Engine/Source/Programs/Shared/EpicGames.Horde/Storage/ContentId.cs b/Engine/Source/Programs/Shared/EpicGames.Horde/Storage/ContentId.cs
new file mode 100644
--- /dev/null
+++ b/Engine/Source/Programs/Shared/EpicGames.Horde/Storage/ContentId.cs
@@ -0,0 +1,242 @@
+// Copyright Epic Games, Inc. All Rights Reserved.
+
+using EpicGames.Core;
+using EpicGames.Serialization;
+using System;
+using System.Buffers.Binary;
+using System.ComponentModel;
+using System.Globalization;
+using System.Text;
+using System.Text.Json;
+using System.Text.Json.Serialization;
+using System.Threading;
+
+namespace EpicGames.Horde.Storage
+{
+	/// <summary>
+	/// Globally unique identifier for a blob, as a utf-8 string. Clients should not assume any internal structure to this identifier.
+	/// </summary>
+	[JsonConverter(typeof(ContentIdJsonConverter))]
+	[TypeConverter(typeof(ContentIdTypeConverter))]
+	[CbConverter(typeof(ContentIdCbConverter))]
+	public struct ContentId : IEquatable<ContentId>
+	{
+		/// <summary>
+		/// Dummy enum to allow invoking the constructor which takes a sanitized full path
+		/// </summary>
+		public enum Sanitize
+		{
+			/// <summary>
+			/// Dummy value
+			/// </summary>
+			None
+		}
+
+		/// <summary>
+		/// Empty blob id
+		/// </summary>
+		public static ContentId Empty { get; } = default;
+
+		/// <summary>
+		/// Identifier for the ref
+		/// </summary>
+		public Utf8String Inner { get; }
+
+		/// <summary>
+		/// Constructor
+		/// </summary>
+		/// <param name="inner"></param>
+		public ContentId(Utf8String inner)
+		{
+			Inner = inner;
+			ValidateArgument(nameof(inner), inner);
+		}
+
+		/// <summary>
+		/// Constructor
+		/// </summary>
+		/// <param name="inner"></param>
+		/// <param name="sanitize"></param>
+		public ContentId(Utf8String inner, Sanitize sanitize)
+		{
+			Inner = inner;
+			_ = sanitize;
+		}
+
+		static readonly Utf8String s_process;
+		static int s_counter;
+
+		/// <summary>
+		/// Static constructor
+		/// </summary>
+		[System.Diagnostics.CodeAnalysis.SuppressMessage("Usage", "CA2207:Initialize value type static fields inline")]
+		static ContentId()
+		{
+			Random rnd = new Random(HashCode.Combine(DateTime.UtcNow.Ticks, Environment.ProcessId, Environment.TickCount64, Environment.MachineName.GetHashCode(StringComparison.Ordinal)));
+
+			Span<byte> process = stackalloc byte[4];
+			rnd.NextBytes(process);
+			s_process = StringUtils.FormatUtf8HexString(process);
+
+			s_counter = rnd.Next();
+		}
+
+		/// <summary>
+		/// Creates a new identifier in the given buffer.
+		/// </summary>
+		/// <param name="output">Buffer to receive the output. The first 24 bytes will be written to.</param>
+		public static void GenerateUniqueId(Span<byte> output)
+		{
+			uint timestamp = (uint)((DateTime.UtcNow - DateTime.UnixEpoch).Ticks / TimeSpan.TicksPerSecond);
+			StringUtils.FormatUtf8HexString(timestamp, output);
+
+			s_process.Span.CopyTo(output.Slice(8));
+
+			uint counter = (uint)Interlocked.Increment(ref s_counter);
+			StringUtils.FormatUtf8HexString(counter, output.Slice(16));
+		}
+
+		/// <summary>
+		/// Checks whether this blob id is valid
+		/// </summary>
+		/// <returns>True if the identifier is valid</returns>
+		public bool IsValid() => Inner.Length > 0;
+
+		/// <summary>
+		/// Validates a given string as a content id
+		/// </summary>
+		/// <param name="name">Name of the argument</param>
+		/// <param name="text">String to validate</param>
+		public static void ValidateArgument(string name, Utf8String text)
+		{
+			if (text.Length == 0)
+			{
+				throw new ArgumentException("Content identifiers cannot be empty", name);
+			}
+			if (text[^1] == '/')
+			{
+				throw new ArgumentException("Content identifiers cannot start or end with a slash", name);
+			}
+
+			int lastSlashIdx = -1;
+
+			for (int idx = 0; idx < text.Length; idx++)
+			{
+				if (text[idx] == '/')
+				{
+					if (lastSlashIdx == idx - 1)
+					{
+						throw new ArgumentException("Leading and consecutive slashes are not permitted in content identifiers", name);
+					}
+					else
+					{
+						lastSlashIdx = idx;
+					}
+				}
+				else
+				{
+					if (!IsValidChar(text[idx]))
+					{
+						throw new ArgumentException($"'{(char)text[idx]} is not a valid content identifier character", name);
+					}
+				}
+			}
+		}
+
+		/// <summary>
+		/// Test if a given character is valid in a store id
+		/// </summary>
+		/// <param name="c">Character to test</param>
+		/// <returns>True if the character is valid</returns>
+		static bool IsValidChar(byte c) => (c >= 'a' && c <= 'z') || (c >= '0' && c <= '9') || (c == '_' || c == '-');
+
+		/// <inheritdoc/>
+		public override bool Equals(object? obj) => obj is ContentId ContentId && Equals(ContentId);
+
+		/// <inheritdoc/>
+		public override int GetHashCode() => Inner.GetHashCode();
+
+		/// <inheritdoc/>
+		public bool Equals(ContentId ContentId) => Inner == ContentId.Inner;
+
+		/// <inheritdoc/>
+		public override string ToString() => Inner.ToString();
+
+		/// <inheritdoc/>
+		public static bool operator ==(ContentId lhs, ContentId rhs) => lhs.Equals(rhs);
+
+		/// <inheritdoc/>
+		public static bool operator !=(ContentId lhs, ContentId rhs) => !lhs.Equals(rhs);
+	}
+
+	/// <summary>
+	/// Type converter for ContentId to and from JSON
+	/// </summary>
+	sealed class ContentIdJsonConverter : JsonConverter<ContentId>
+	{
+		/// <inheritdoc/>
+		public override ContentId Read(ref Utf8JsonReader reader, Type typeToConvert, JsonSerializerOptions options) => new ContentId(new Utf8String(reader.GetUtf8String().ToArray()));
+
+		/// <inheritdoc/>
+		public override void Write(Utf8JsonWriter writer, ContentId value, JsonSerializerOptions options) => writer.WriteStringValue(value.Inner.Span);
+	}
+
+	/// <summary>
+	/// Type converter from strings to ContentId objects
+	/// </summary>
+	sealed class ContentIdTypeConverter : TypeConverter
+	{
+		/// <inheritdoc/>
+		public override bool CanConvertFrom(ITypeDescriptorContext? context, Type sourceType)
+		{
+			return sourceType == typeof(string);
+		}
+
+		/// <inheritdoc/>
+		public override object ConvertFrom(ITypeDescriptorContext? context, CultureInfo? culture, object? value)
+		{
+			return new ContentId((string)value!);
+		}
+	}
+
+	/// <summary>
+	/// Type converter to compact binary
+	/// </summary>
+	sealed class ContentIdCbConverter : CbConverterBase<ContentId>
+	{
+		/// <inheritdoc/>
+		public override ContentId Read(CbField field) => new ContentId(field.AsUtf8String());
+
+		/// <inheritdoc/>
+		public override void Write(CbWriter writer, ContentId value) => writer.WriteUtf8StringValue(value.Inner);
+
+		/// <inheritdoc/>
+		public override void WriteNamed(CbWriter writer, Utf8String name, ContentId value) => writer.WriteUtf8String(name, value.Inner);
+	}
+
+	/// <summary>
+	/// Extension methods for blob ids
+	/// </summary>
+	public static class ContentIdExtensions
+	{
+		/// <summary>
+		/// Deserialize a blob id
+		/// </summary>
+		/// <param name="reader">Reader to deserialize from</param>
+		/// <returns>The blob id that was read</returns>
+		public static ContentId ReadContentId(this IMemoryReader reader)
+		{
+			return new ContentId(reader.ReadUtf8String());
+		}
+
+		/// <summary>
+		/// Serialize a blob id
+		/// </summary>
+		/// <param name="writer">Writer to serialize to</param>
+		/// <param name="value">Value to serialize</param>
+		public static void WriteContentId(this IMemoryWriter writer, ContentId value)
+		{
+			writer.WriteUtf8String(value.Inner);
+		}
+	}
+}
diff --git a/Engine/Source/Programs/Shared/EpicGames.Horde/Storage/IBlobStore.cs b/Engine/Source/Programs/Shared/EpicGames.Horde/Storage/IBlobStore.cs
--- a/Engine/Source/Programs/Shared/EpicGames.Horde/Storage/IBlobStore.cs
+++ b/Engine/Source/Programs/Shared/EpicGames.Horde/Storage/IBlobStore.cs
@@ -6,6 +6,7 @@ using System.Collections.Generic;
 using System.Threading;
 using System.Threading.Tasks;
 using EpicGames.Core;
+using Microsoft.CodeAnalysis.Rename;
 
 namespace EpicGames.Horde.Storage
 {
@@ -35,6 +36,11 @@ namespace EpicGames.Horde.Storage
 	/// </summary>
 	public interface IBlob
 	{
+		/// <summary>
+		/// Identifier for the blob
+		/// </summary>
+		BlobId Id { get; }
+
 		/// <summary>
 		/// Gets the data for this blob
 		/// </summary>
@@ -61,47 +67,58 @@ namespace EpicGames.Horde.Storage
 		/// <param name="id">The blob identifier</param>
 		/// <param name="cancellationToken">Cancellation token for the operation</param>
 		/// <returns></returns>
-		Task<IBlob?> TryReadBlobAsync(BlobId id, CancellationToken cancellationToken = default);
+		Task<IBlob> ReadBlobAsync(BlobId id, CancellationToken cancellationToken = default);
 
 		/// <summary>
 		/// Writes a new object to the store
 		/// </summary>
-		/// <param name="refName">Name of a ref that this blob is being written for. While the returned BlobId is guaranteed to be unique, this name can be used as a prefix to aid debugging.</param>
 		/// <param name="data">Payload for the object</param>
 		/// <param name="references">Object references</param>
 		/// <param name="cancellationToken">Cancellation token for the operation</param>
+		/// <param name="hintRefName">Name of a ref that this blob is being written for. While the returned BlobId is guaranteed to be unique, this name can be used as a prefix to aid debugging.</param>
 		/// <returns>Unique identifier for the bundle</returns>
-		Task<BlobId> WriteBlobAsync(RefName refName, ReadOnlySequence<byte> data, IReadOnlyList<BlobId> references, CancellationToken cancellationToken = default);
+		Task<IBlob> WriteBlobAsync(ReadOnlySequence<byte> data, IReadOnlyList<BlobId> references, RefName hintRefName = default, CancellationToken cancellationToken = default);
 
 		#endregion
 
 		#region Refs
 
 		/// <summary>
-		/// Checks if the given ref exists
+		/// Reads data for a ref from the store, along with the blob contents.
 		/// </summary>
-		/// <param name="name">Name of the reference to look for</param>
+		/// <param name="name">The ref name</param>
+		/// <param name="cacheTime">Minimum coherency for any cached value to be returned</param>
 		/// <param name="cancellationToken">Cancellation token for the operation</param>
-		/// <returns>True if the ref exists, false if it did not exist</returns>
-		Task<bool> HasRefAsync(RefName name, CancellationToken cancellationToken = default);
+		/// <returns></returns>
+		Task<IBlob?> TryReadRefAsync(RefName name, DateTime cacheTime = default, CancellationToken cancellationToken = default);
 
 		/// <summary>
 		/// Reads data for a ref from the store
 		/// </summary>
 		/// <param name="name">The ref name</param>
+		/// <param name="cacheTime">Minimum coherency for any cached value to be returned</param>
 		/// <param name="cancellationToken">Cancellation token for the operation</param>
 		/// <returns></returns>
-		Task<IBlob?> TryReadRefAsync(RefName name, CancellationToken cancellationToken = default);
+		Task<BlobId> TryReadRefTargetAsync(RefName name, DateTime cacheTime = default, CancellationToken cancellationToken = default);
 
 		/// <summary>
-		/// Writes a new ref to the store
+		/// Writes a new ref to the store which points to a new blob
 		/// </summary>
 		/// <param name="name">Ref to write</param>
 		/// <param name="data">Payload for the object</param>
 		/// <param name="references">Object references</param>
 		/// <param name="cancellationToken">Cancellation token for the operation</param>
 		/// <returns>Unique identifier for the blob</returns>
-		Task WriteRefAsync(RefName name, ReadOnlySequence<byte> data, IReadOnlyList<BlobId> references, CancellationToken cancellationToken = default);
+		Task<IBlob> WriteRefAsync(RefName name, ReadOnlySequence<byte> data, IReadOnlyList<BlobId> references, CancellationToken cancellationToken = default);
+
+		/// <summary>
+		/// Writes a new ref to the store
+		/// </summary>
+		/// <param name="name">Ref to write</param>
+		/// <param name="blobId">The target for the ref</param>
+		/// <param name="cancellationToken">Cancellation token for the operation</param>
+		/// <returns>Unique identifier for the blob</returns>
+		Task WriteRefTargetAsync(RefName name, BlobId blobId, CancellationToken cancellationToken = default);
 
 		/// <summary>
 		/// Reads data for a ref from the store
@@ -134,10 +151,10 @@ namespace EpicGames.Horde.Storage
 			#region Blobs
 
 			/// <inheritdoc/>
-			public Task<IBlob?> TryReadBlobAsync(BlobId id, CancellationToken cancellationToken = default) => _inner.TryReadBlobAsync(id, cancellationToken);
+			public Task<IBlob> ReadBlobAsync(BlobId id, CancellationToken cancellationToken = default) => _inner.ReadBlobAsync(id, cancellationToken);
 
 			/// <inheritdoc/>
-			public Task<BlobId> WriteBlobAsync(RefName refName, ReadOnlySequence<byte> data, IReadOnlyList<BlobId> references, CancellationToken cancellationToken = default) => _inner.WriteBlobAsync(refName, data, references, cancellationToken);
+			public Task<IBlob> WriteBlobAsync(ReadOnlySequence<byte> data, IReadOnlyList<BlobId> references, RefName hintRefName = default, CancellationToken cancellationToken = default) => _inner.WriteBlobAsync(data, references, hintRefName, cancellationToken);
 
 			#endregion
 
@@ -147,13 +164,16 @@ namespace EpicGames.Horde.Storage
 			public Task DeleteRefAsync(RefName name, CancellationToken cancellationToken = default) => _inner.DeleteRefAsync(name, cancellationToken);
 
 			/// <inheritdoc/>
-			public Task<bool> HasRefAsync(RefName name, CancellationToken cancellationToken = default) => _inner.HasRefAsync(name, cancellationToken);
+			public Task<IBlob?> TryReadRefAsync(RefName name, DateTime cacheTime = default, CancellationToken cancellationToken = default) => _inner.TryReadRefAsync(name, cacheTime, cancellationToken);
+
+			/// <inheritdoc/>
+			public Task<BlobId> TryReadRefTargetAsync(RefName name, DateTime cacheTime = default, CancellationToken cancellationToken = default) => _inner.TryReadRefTargetAsync(name, cacheTime, cancellationToken);
 
 			/// <inheritdoc/>
-			public Task<IBlob?> TryReadRefAsync(RefName name, CancellationToken cancellationToken = default) => _inner.TryReadRefAsync(name, cancellationToken);
+			public Task<IBlob> WriteRefAsync(RefName name, ReadOnlySequence<byte> data, IReadOnlyList<BlobId> references, CancellationToken cancellationToken = default) => _inner.WriteRefAsync(name, data, references, cancellationToken);
 
 			/// <inheritdoc/>
-			public Task WriteRefAsync(RefName name, ReadOnlySequence<byte> data, IReadOnlyList<BlobId> references, CancellationToken cancellationToken = default) => _inner.WriteRefAsync(name, data, references, cancellationToken);
+			public Task WriteRefTargetAsync(RefName name, BlobId blobId, CancellationToken cancellationToken = default) => _inner.WriteRefTargetAsync(name, blobId, cancellationToken);
 
 			#endregion
 		}
@@ -166,18 +186,72 @@ namespace EpicGames.Horde.Storage
 		public static IBlobStore<T> ForType<T>(this IBlobStore blobStore) => new TypedBlobStore<T>(blobStore);
 
 		/// <summary>
-		/// Reads a blob from the store, throwing an exception if it does not exist
+		/// Checks if the given ref exists
 		/// </summary>
 		/// <param name="store">The store instance to read from</param>
-		/// <param name="id">Id for the blob</param>
+		/// <param name="name">Name of the reference to look for</param>
+		/// <param name="cacheTime">Minimum coherency for any cached value to be returned</param>
+		/// <param name="cancellationToken">Cancellation token for the operation</param>
+		/// <returns>True if the ref exists, false if it did not exist</returns>
+		public static async Task<bool> HasRefAsync(this IBlobStore store, RefName name, DateTime cacheTime = default, CancellationToken cancellationToken = default)
+		{
+			BlobId blobId = await store.TryReadRefTargetAsync(name, cacheTime, cancellationToken);
+			return blobId.IsValid();
+		}
+
+		/// <summary>
+		/// Checks if the given ref exists
+		/// </summary>
+		/// <param name="store">The store instance to read from</param>
+		/// <param name="name">Name of the reference to look for</param>
+		/// <param name="maxAge">Maximum age of any cached ref</param>
+		/// <param name="cancellationToken">Cancellation token for the operation</param>
+		/// <returns>True if the ref exists, false if it did not exist</returns>
+		public static Task<bool> HasRefAsync(this IBlobStore store, RefName name, TimeSpan maxAge, CancellationToken cancellationToken = default)
+		{
+			return HasRefAsync(store, name, DateTime.UtcNow - maxAge, cancellationToken);
+		}
+
+		/// <summary>
+		/// Reads a ref from the store, throwing an exception if it does not exist
+		/// </summary>
+		/// <param name="store">The store instance to read from</param>
+		/// <param name="name">Id for the ref</param>
+		/// <param name="maxAge">Maximum age of any cached ref</param>
 		/// <param name="cancellationToken">Cancellation token for the operation</param>
 		/// <returns>The blob instance</returns>
-		public static async Task<IBlob> ReadBlobAsync(this IBlobStore store, BlobId id, CancellationToken cancellationToken = default)
+		public static Task<IBlob?> TryReadRefAsync(this IBlobStore store, RefName name, TimeSpan maxAge, CancellationToken cancellationToken = default)
 		{
-			IBlob? blob = await store.TryReadBlobAsync(id, cancellationToken);
+			return store.TryReadRefAsync(name, DateTime.UtcNow - maxAge, cancellationToken);
+		}
+
+		/// <summary>
+		/// Reads a ref from the store, throwing an exception if it does not exist
+		/// </summary>
+		/// <param name="store">The store instance to read from</param>
+		/// <param name="name">Id for the ref</param>
+		/// <param name="maxAge">Maximum age of any cached ref</param>
+		/// <param name="cancellationToken">Cancellation token for the operation</param>
+		/// <returns>The blob instance</returns>
+		public static Task<BlobId> TryReadRefIdAsync(this IBlobStore store, RefName name, TimeSpan maxAge, CancellationToken cancellationToken = default)
+		{
+			return store.TryReadRefTargetAsync(name, DateTime.UtcNow - maxAge, cancellationToken);
+		}
+
+		/// <summary>
+		/// Reads a ref from the store, throwing an exception if it does not exist
+		/// </summary>
+		/// <param name="store">The store instance to read from</param>
+		/// <param name="name">Id for the ref</param>
+		/// <param name="cacheTime">Minimum coherency of any cached result</param>
+		/// <param name="cancellationToken">Cancellation token for the operation</param>
+		/// <returns>The blob instance</returns>
+		public static async Task<IBlob> ReadRefAsync(this IBlobStore store, RefName name, DateTime cacheTime = default, CancellationToken cancellationToken = default)
+		{
+			IBlob? blob = await store.TryReadRefAsync(name, cacheTime, cancellationToken);
 			if (blob == null)
 			{
-				throw new KeyNotFoundException($"Unable to find blob {id}");
+				throw new RefNameNotFoundException(name);
 			}
 			return blob;
 		}
@@ -187,34 +261,65 @@ namespace EpicGames.Horde.Storage
 		/// </summary>
 		/// <param name="store">The store instance to read from</param>
 		/// <param name="name">Id for the ref</param>
+		/// <param name="maxAge">Maximum age for any cached result</param>
 		/// <param name="cancellationToken">Cancellation token for the operation</param>
 		/// <returns>The blob instance</returns>
-		public static async Task<IBlob> ReadRefAsync(this IBlobStore store, RefName name, CancellationToken cancellationToken = default)
+		public static Task<IBlob> ReadRefAsync(this IBlobStore store, RefName name, TimeSpan maxAge, CancellationToken cancellationToken = default)
 		{
-			IBlob? blob = await store.TryReadRefAsync(name, cancellationToken);
-			if (blob == null)
+			return ReadRefAsync(store, name, DateTime.UtcNow - maxAge, cancellationToken);
+		}
+
+		/// <summary>
+		/// Reads a ref from the store, throwing an exception if it does not exist
+		/// </summary>
+		/// <param name="store">The store instance to read from</param>
+		/// <param name="name">Id for the ref</param>
+		/// <param name="cacheTime">Minimum coherency of any cached result</param>
+		/// <param name="cancellationToken">Cancellation token for the operation</param>
+		/// <returns>The blob instance</returns>
+		public static async Task<BlobId> ReadRefIdAsync(this IBlobStore store, RefName name, DateTime cacheTime = default, CancellationToken cancellationToken = default)
+		{
+			BlobId blobId = await store.TryReadRefTargetAsync(name, cacheTime, cancellationToken);
+			if (!blobId.IsValid())
 			{
 				throw new RefNameNotFoundException(name);
 			}
-			return blob;
+			return blobId;
+		}
+
+		/// <summary>
+		/// Reads a ref from the store, throwing an exception if it does not exist
+		/// </summary>
+		/// <param name="store">The store instance to read from</param>
+		/// <param name="name">Id for the ref</param>
+		/// <param name="maxAge">Maximum age for any cached result</param>
+		/// <param name="cancellationToken">Cancellation token for the operation</param>
+		/// <returns>The blob instance</returns>
+		public static Task<BlobId> ReadRefIdAsync(this IBlobStore store, RefName name, TimeSpan maxAge, CancellationToken cancellationToken = default)
+		{
+			return ReadRefIdAsync(store, name, DateTime.UtcNow - maxAge, cancellationToken);
 		}
 	}
 
 	/// <summary>
 	/// Utility methods for blobs
 	/// </summary>
-	public static class BlobUtils
+	public static class Blob
 	{
 		class InMemoryBlob : IBlob
 		{
+			/// <inheritdoc/>
+			public BlobId Id { get; }
+
 			/// <inheritdoc/>
 			public ReadOnlyMemory<byte> Data { get; }
 
 			/// <inheritdoc/>
 			public IReadOnlyList<BlobId> References { get; }
 
-			public InMemoryBlob(ReadOnlyMemory<byte> data, IReadOnlyList<BlobId> references)
+			public InMemoryBlob(BlobId id, ReadOnlyMemory<byte> data, IReadOnlyList<BlobId> references)
 			{
+				Id = id;
 				Data = data;
 				References = references;
 			}
@@ -223,10 +328,11 @@ namespace EpicGames.Horde.Storage
 		/// <summary>
 		/// Create a blob from memory
 		/// </summary>
+		/// <param name="id">Id for the blob</param>
 		/// <param name="data">Payload for the blob</param>
 		/// <param name="references">References to other blobs</param>
 		/// <returns>Blob instance</returns>
-		public static IBlob FromMemory(ReadOnlyMemory<byte> data, IReadOnlyList<BlobId> references) => new InMemoryBlob(data, references);
+		public static IBlob FromMemory(BlobId id, ReadOnlyMemory<byte> data, IReadOnlyList<BlobId> references) => new InMemoryBlob(id, data, references);
 
 		/// <summary>
 		/// Serialize a blob into a flat memory buffer
@@ -271,23 +377,25 @@ namespace EpicGames.Horde.Storage
 		/// <summary>
 		/// Deserialize a blob from a block of memory
 		/// </summary>
+		/// <param name="id">Id for the blob</param>
 		/// <param name="memory">Memory to deserialize from</param>
 		/// <returns>Deserialized blob data. May reference the supplied memory.</returns>
-		public static IBlob Deserialize(ReadOnlyMemory<byte> memory) => Deserialize(new ReadOnlySequence<byte>(memory));
+		public static IBlob Deserialize(BlobId id, ReadOnlyMemory<byte> memory) => Deserialize(id, new ReadOnlySequence<byte>(memory));
 
 		/// <summary>
 		/// Deserialize a blob from a block of memory
 		/// </summary>
+		/// <param name="id">Id for the blob</param>
 		/// <param name="sequence">Sequence to deserialize from</param>
 		/// <returns>Deserialized blob data. May reference the supplied memory.</returns>
-		public static IBlob Deserialize(ReadOnlySequence<byte> sequence)
+		public static IBlob Deserialize(BlobId id, ReadOnlySequence<byte> sequence)
 		{
 			MemoryReader reader = new MemoryReader(sequence.First);
 			IReadOnlyList<BlobId> references = reader.ReadVariableLengthArray(() => reader.ReadBlobId());
 			long length = (long)reader.ReadUnsignedVarInt();
 
 			ReadOnlyMemory<byte> data = sequence.Slice(sequence.First.Length - reader.Memory.Length).AsSingleSegment();
-			return new InMemoryBlob(data, references);
+			return new InMemoryBlob(id, data, references);
 		}
 	}
 }
diff --git a/Engine/Source/Programs/Shared/EpicGames.Horde/Storage/ITreeStore.cs b/Engine/Source/Programs/Shared/EpicGames.Horde/Storage/ITreeStore.cs
--- a/Engine/Source/Programs/Shared/EpicGames.Horde/Storage/ITreeStore.cs
+++ b/Engine/Source/Programs/Shared/EpicGames.Horde/Storage/ITreeStore.cs
@@ -73,9 +73,10 @@ namespace EpicGames.Horde.Storage
 		/// Attempts to read a tree with the given name
 		/// </summary>
 		/// <param name="name">Name of the ref used to store the tree</param>
+		/// <param name="maxAge"></param>
 		/// <param name="cancellationToken">Cancellation token for the operation</param>
 		/// <returns>Reference to the root of the tree. Null if the ref does not exist.</returns>
-		Task<ITreeBlob?> TryReadTreeAsync(RefName name, CancellationToken cancellationToken = default);
+		Task<ITreeBlob?> TryReadTreeAsync(RefName name, TimeSpan maxAge = default, CancellationToken cancellationToken = default);
 	}
 
 	/// <summary>
@@ -162,7 +163,7 @@ namespace EpicGames.Horde.Storage
 			public Task<bool> HasTreeAsync(RefName name, CancellationToken cancellationToken = default) => _inner.HasTreeAsync(name, cancellationToken);
 
 			/// <inheritdoc/>
-			public Task<ITreeBlob?> TryReadTreeAsync(RefName name, CancellationToken cancellationToken = default) => _inner.TryReadTreeAsync(name, cancellationToken);
+			public Task<ITreeBlob?> TryReadTreeAsync(RefName name, TimeSpan maxAge = default, CancellationToken cancellationToken = default) => _inner.TryReadTreeAsync(name, maxAge, cancellationToken);
 		}
 
 		/// <summary>
diff --git a/Engine/Source/Programs/Shared/EpicGames.Horde/Storage/RefName.cs b/Engine/Source/Programs/Shared/EpicGames.Horde/Storage/RefName.cs
--- a/Engine/Source/Programs/Shared/EpicGames.Horde/Storage/RefName.cs
+++ b/Engine/Source/Programs/Shared/EpicGames.Horde/Storage/RefName.cs
@@ -19,6 +19,11 @@ namespace EpicGames.Horde.Storage
 	[CbConverter(typeof(RefNameCbConverter))]
 	public struct RefName : IEquatable<RefName>, IComparable<RefName>
 	{
+		/// <summary>
+		/// Empty ref name
+		/// </summary>
+		public static RefName Empty { get; } = default;
+
 		/// <summary>
 		/// String for the ref name
 		/// </summary>
@@ -31,51 +36,7 @@ namespace EpicGames.Horde.Storage
 		public RefName(Utf8String text)
 		{
 			Text = text;
-
-			for (int idx = 0; idx < text.Length; idx++)
-			{
-				if (!IsValidCharacter(text[idx]))
-				{
-					throw new ArgumentException($"Character '{(char)text[idx]}' is not valid in ref name ('{text}')", nameof(text));
-				}
-			}
-
-			if (text.Length == 0)
-			{
-				throw new ArgumentException($"Ref name cannot be empty", nameof(text));
-			}
-		}
-
-		/// <summary>
-		/// Constructor
-		/// </summary>
-		/// <param name="text"></param>
-		public RefName(string text)
-			: this(new Utf8String(text))
-		{
-		}
-
-		/// <summary>
-		/// Test whether a character is valid in a ref name
-		/// </summary>
-		/// <param name="c">Character to test</param>
-		/// <returns></returns>
-		public static bool IsValidCharacter(byte c)
-		{
-			switch (c)
-			{
-				case (byte)'<':
-				case (byte)'>':
-				case (byte)':':
-				case (byte)'"':
-				case (byte)'\\':
-				case (byte)'|':
-				case (byte)'*':
-				case (byte)'?':
-					return false;
-				default:
-					return c > 0x20 && c < 0x7f;
-			}
+			ContentId.ValidateArgument(nameof(text), text);
 		}
 
 		/// <inheritdoc/>
diff --git a/Engine/Source/Programs/Shared/EpicGames.Horde/Storage/ServerId.cs b/Engine/Source/Programs/Shared/EpicGames.Horde/Storage/ServerId.cs
new file mode 100644
--- /dev/null
+++ b/Engine/Source/Programs/Shared/EpicGames.Horde/Storage/ServerId.cs
@@ -0,0 +1,234 @@
+// Copyright Epic Games, Inc. All Rights Reserved.
+
+using EpicGames.Core;
+using EpicGames.Serialization;
+using System;
+using System.ComponentModel;
+using System.Globalization;
+using System.Text;
+using System.Text.Json;
+using System.Text.Json.Serialization;
+
+namespace EpicGames.Horde.Storage
+{
+	/// <summary>
+	/// Identifies a server in the storage hierarchy, via a string of identifiers separated with colons.
+	/// </summary>
+	[JsonConverter(typeof(ServerIdJsonConverter))]
+	[TypeConverter(typeof(ServerIdTypeConverter))]
+	[CbConverter(typeof(ServerIdCbConverter))]
+	public struct ServerId : IEquatable<ServerId>
+	{
+		/// <summary>
+		/// Dummy enum to allow invoking the constructor which takes a sanitized full path
+		/// </summary>
+		public enum Sanitize
+		{
+			/// <summary>
+			/// Dummy value
+			/// </summary>
+			None
+		}
+
+		/// <summary>
+		/// Empty server id
+		/// </summary>
+		public static ServerId Empty { get; } = default;
+
+		/// <summary>
+		/// Identifier for the ref
+		/// </summary>
+		public Utf8String Inner { get; }
+
+		/// <summary>
+		/// Constructor
+		/// </summary>
+		/// <param name="inner"></param>
+		public ServerId(Utf8String inner)
+		{
+			Inner = inner;
+			ValidateArgument(nameof(inner), inner);
+		}
+
+		/// <summary>
+		/// Constructor
+		/// </summary>
+		/// <param name="inner"></param>
+		/// <param name="sanitize"></param>
+		public ServerId(Utf8String inner, Sanitize sanitize)
+		{
+			Inner = inner;
+			_ = sanitize;
+		}
+
+		/// <summary>
+		/// Appends another identifier to this server id
+		/// </summary>
+		/// <param name="identifier"></param>
+		/// <returns></returns>
+		public ServerId Append(Utf8String identifier)
+		{
+			ValidateArgument(nameof(identifier), identifier);
+
+			byte[] buffer = new byte[Inner.Length + 1 + identifier.Length];
+			Inner.Span.CopyTo(buffer);
+			buffer[Inner.Length] = (byte)':';
+			identifier.Span.CopyTo(buffer.AsSpan(Inner.Length + 1));
+
+			return new ServerId(new Utf8String(buffer), Sanitize.None);
+		}
+
+		/// <summary>
+		/// Validates a given string as a server id
+		/// </summary>
+		/// <param name="name">Name of the argument</param>
+		/// <param name="text">String to validate</param>
+		public static void ValidateArgument(string name, Utf8String text)
+		{
+			if (text.Length == 0)
+			{
+				throw new ArgumentException("Server identifiers cannot be empty", name);
+			}
+			if (text[^1] == ':')
+			{
+				throw new ArgumentException("Server identifiers cannot end with a colon", name);
+			}
+
+			int lastColonIdx = -1;
+
+			for (int idx = 0; idx < text.Length; idx++)
+			{
+				if (text[idx] == ':')
+				{
+					if (lastColonIdx == idx - 1)
+					{
+						throw new ArgumentException("Leading and consecutive colons are not permitted in server identifiers", name);
+					}
+					else
+					{
+						lastColonIdx = idx;
+					}
+				}
+				else
+				{
+					if (!IsValidChar(text[idx]))
+					{
+						throw new ArgumentException($"'{(char)text[idx]} is not a valid server identifier character", name);
+					}
+				}
+			}
+		}
+
+		/// <summary>
+		/// Test if a given character is valid in a store id
+		/// </summary>
+		/// <param name="c">Character to test</param>
+		/// <returns>True if the character is valid</returns>
+		static bool IsValidChar(byte c) => (c >= 'a' && c <= 'z') || (c >= '0' && c <= '9') || (c == '_' || c == '-');
+
+		/// <summary>
+		/// Checks whether this blob id is valid
+		/// </summary>
+		/// <returns>True if the identifier is valid</returns>
+		public bool IsValid() => Inner.Length > 0;
+
+		/// <summary>
+		/// Gets the upstream server id
+		/// </summary>
+		/// <returns>The upstream server id</returns>
+		public ServerId GetUpstream()
+		{
+			int idx = Inner.LastIndexOf(':');
+			return (idx == -1) ? Empty : new ServerId(Inner.Substring(0, idx), Sanitize.None);
+		}
+
+		/// <inheritdoc/>
+		public override bool Equals(object? obj) => obj is ServerId ServerId && Equals(ServerId);
+
+		/// <inheritdoc/>
+		public override int GetHashCode() => Inner.GetHashCode();
+
+		/// <inheritdoc/>
+		public bool Equals(ServerId ServerId) => Inner == ServerId.Inner;
+
+		/// <inheritdoc/>
+		public override string ToString() => Inner.ToString();
+
+		/// <inheritdoc/>
+		public static bool operator ==(ServerId lhs, ServerId rhs) => lhs.Equals(rhs);
+
+		/// <inheritdoc/>
+		public static bool operator !=(ServerId lhs, ServerId rhs) => !lhs.Equals(rhs);
+	}
+
+	/// <summary>
+	/// Type converter for ServerId to and from JSON
+	/// </summary>
+	sealed class ServerIdJsonConverter : JsonConverter<ServerId>
+	{
+		/// <inheritdoc/>
+		public override ServerId Read(ref Utf8JsonReader reader, Type typeToConvert, JsonSerializerOptions options) => new ServerId(new Utf8String(reader.GetUtf8String().ToArray()));
+
+		/// <inheritdoc/>
+		public override void Write(Utf8JsonWriter writer, ServerId value, JsonSerializerOptions options) => writer.WriteStringValue(value.Inner.Span);
+	}
+
+	/// <summary>
+	/// Type converter from strings to ServerId objects
+	/// </summary>
+	sealed class ServerIdTypeConverter : TypeConverter
+	{
+		/// <inheritdoc/>
+		public override bool CanConvertFrom(ITypeDescriptorContext? context, Type sourceType)
+		{
+			return sourceType == typeof(string);
+		}
+
+		/// <inheritdoc/>
+		public override object ConvertFrom(ITypeDescriptorContext? context, CultureInfo? culture, object? value)
+		{
+			return new ServerId((string)value!);
+		}
+	}
+
+	/// <summary>
+	/// Type converter to compact binary
+	/// </summary>
+	sealed class ServerIdCbConverter : CbConverterBase<ServerId>
+	{
+		/// <inheritdoc/>
+		public override ServerId Read(CbField field) => new ServerId(field.AsUtf8String());
+
+		/// <inheritdoc/>
+		public override void Write(CbWriter writer, ServerId value) => writer.WriteUtf8StringValue(value.Inner);
+
+		/// <inheritdoc/>
+		public override void WriteNamed(CbWriter writer, Utf8String name, ServerId value) => writer.WriteUtf8String(name, value.Inner);
+	}
+
+	/// <summary>
+	/// Extension methods for blob ids
+	/// </summary>
+	public static class ServerIdExtensions
+	{
+		/// <summary>
+		/// Deserialize a blob id
+		/// </summary>
+		/// <param name="reader">Reader to deserialize from</param>
+		/// <returns>The blob id that was read</returns>
+		public static ServerId ReadServerId(this IMemoryReader reader)
+		{
+			return new ServerId(reader.ReadUtf8String());
+		}
+
+		/// <summary>
+		/// Serialize a blob id
+		/// </summary>
+		/// <param name="writer">Writer to serialize to</param>
+		/// <param name="value">Value to serialize</param>
+		public static void WriteServerId(this IMemoryWriter writer, ServerId value)
+		{
+			writer.WriteUtf8String(value.Inner);
+		}
+	}
+}
diff --git a/Engine/Source/Programs/Shared/EpicGames.Horde/Storage/TreeNode.cs b/Engine/Source/Programs/Shared/EpicGames.Horde/Storage/TreeNode.cs
--- a/Engine/Source/Programs/Shared/EpicGames.Horde/Storage/TreeNode.cs
+++ b/Engine/Source/Programs/Shared/EpicGames.Horde/Storage/TreeNode.cs
@@ -153,9 +153,9 @@ namespace EpicGames.Horde.Storage
 		}
 
 		/// <inheritdoc/>
-		public static async Task<T?> TryReadTreeAsync<T>(this ITreeStore store, RefName name, CancellationToken cancellationToken = default) where T : TreeNode
+		public static async Task<T?> TryReadTreeAsync<T>(this ITreeStore store, RefName name, TimeSpan maxAge = default, CancellationToken cancellationToken = default) where T : TreeNode
 		{
-			ITreeBlob? root = await store.TryReadTreeAsync(name, cancellationToken);
+			ITreeBlob? root = await store.TryReadTreeAsync(name, maxAge, cancellationToken);
 			if (root == null)
 			{
 				return null;
@@ -164,9 +164,9 @@ namespace EpicGames.Horde.Storage
 		}
 
 		/// <inheritdoc/>
-		public static async Task<T> ReadTreeAsync<T>(this ITreeStore store, RefName name, CancellationToken cancellationToken = default) where T : TreeNode
+		public static async Task<T> ReadTreeAsync<T>(this ITreeStore store, RefName name, TimeSpan maxAge = default, CancellationToken cancellationToken = default) where T : TreeNode
 		{
-			T? result = await store.TryReadTreeAsync<T>(name, cancellationToken);
+			T? result = await store.TryReadTreeAsync<T>(name, maxAge, cancellationToken);
 			if (result == null)
 			{
 				throw new RefNameNotFoundException(name);
