Horde: Add support for different compressors on bundle data.

diff --git a/Engine/Source/Programs/Shared/EpicGames.Horde/Storage/Bundles/Bundle.cs b/Engine/Source/Programs/Shared/EpicGames.Horde/Storage/Bundles/Bundle.cs
--- a/Engine/Source/Programs/Shared/EpicGames.Horde/Storage/Bundles/Bundle.cs
+++ b/Engine/Source/Programs/Shared/EpicGames.Horde/Storage/Bundles/Bundle.cs
@@ -24,10 +24,15 @@ namespace EpicGames.Horde.Storage.Bundles
 		/// </summary>
 		Initial = 0,
 
+		/// <summary>
+		/// Added compression format enum
+		/// </summary>
+		CompressionFormat = 1,
+
 		/// <summary>
 		/// The current version number
 		/// </summary>
-		Current = Initial,
+		Current = CompressionFormat,
 	}
 
 	/// <summary>
@@ -102,11 +107,32 @@ namespace EpicGames.Horde.Storage.Bundles
 		}
 	}
 
+	/// <summary>
+	/// Indicates the compression format in the bundle
+	/// </summary>
+	public enum BundleCompressionFormat
+	{
+		/// <summary>
+		/// Packets are uncompressed
+		/// </summary>
+		None = 0,
+
+		/// <summary>
+		/// LZ4 compression
+		/// </summary>
+		LZ4 = 1,
+	}
+
 	/// <summary>
 	/// Header for the contents of a bundle.
 	/// </summary>
 	public class BundleHeader
 	{
+		/// <summary>
+		/// Compression format for the bundle
+		/// </summary>
+		public BundleCompressionFormat CompressionFormat { get; }
+
 		/// <summary>
 		/// References to nodes in other bundles
 		/// </summary>
@@ -125,11 +151,13 @@ namespace EpicGames.Horde.Storage.Bundles
 		/// <summary>
 		/// Constructs a new bundle header
 		/// </summary>
+		/// <param name="compressionFormat">Compression format for bundle packets</param>
 		/// <param name="imports">Imports from other bundles</param>
 		/// <param name="exports">Exports for nodes</param>
 		/// <param name="packets">Compression packets within the bundle</param>
-		public BundleHeader(IReadOnlyList<BundleImport> imports, IReadOnlyList<BundleExport> exports, IReadOnlyList<BundlePacket> packets)
+		public BundleHeader(BundleCompressionFormat compressionFormat, IReadOnlyList<BundleImport> imports, IReadOnlyList<BundleExport> exports, IReadOnlyList<BundlePacket> packets)
 		{
+			CompressionFormat = compressionFormat;
 			Imports = imports;
 			Exports = exports;
 			Packets = packets;
@@ -143,9 +171,18 @@ namespace EpicGames.Horde.Storage.Bundles
 		public BundleHeader(IMemoryReader reader)
 		{
 			BundleVersion version = (BundleVersion)reader.ReadUnsignedVarInt();
-			if (version != BundleVersion.Current)
+			if (version > BundleVersion.Current)
+			{
+				throw new InvalidDataException($"Unknown bundle version {(int)version}. Max supported is {(int)BundleVersion.Current}.");
+			}
+
+			if (version >= BundleVersion.CompressionFormat)
+			{
+				CompressionFormat = (BundleCompressionFormat)reader.ReadUnsignedVarInt();
+			}
+			else
 			{
-				throw new InvalidDataException();
+				CompressionFormat = BundleCompressionFormat.LZ4;
 			}
 
 			Imports = reader.ReadVariableLengthArray(() => new BundleImport(reader));
@@ -161,6 +198,7 @@ namespace EpicGames.Horde.Storage.Bundles
 		{
 			writer.WriteUnsignedVarInt((ulong)BundleVersion.Current);
 
+			writer.WriteUnsignedVarInt((ulong)CompressionFormat);
 			writer.WriteVariableLengthArray(Imports, x => x.Write(writer));
 			writer.WriteVariableLengthArray(Exports, x => x.Write(writer));
 			writer.WriteVariableLengthArray(Packets, x => x.Write(writer));
diff --git a/Engine/Source/Programs/Shared/EpicGames.Horde/Storage/Bundles/BundleStore.cs b/Engine/Source/Programs/Shared/EpicGames.Horde/Storage/Bundles/BundleStore.cs
--- a/Engine/Source/Programs/Shared/EpicGames.Horde/Storage/Bundles/BundleStore.cs
+++ b/Engine/Source/Programs/Shared/EpicGames.Horde/Storage/Bundles/BundleStore.cs
@@ -5,6 +5,7 @@ using System.Buffers;
 using System.Collections.Concurrent;
 using System.Collections.Generic;
 using System.Diagnostics;
+using System.IO;
 using System.Linq;
 using System.Threading;
 using System.Threading.Tasks;
@@ -29,6 +30,11 @@ namespace EpicGames.Horde.Storage.Bundles
 		/// </summary>
 		public int MaxInlineBlobSize { get; set; } = 1024 * 64;
 
+		/// <summary>
+		/// Compression format to use
+		/// </summary>
+		public BundleCompressionFormat CompressionFormat { get; set; } = BundleCompressionFormat.LZ4;
+
 		/// <summary>
 		/// Minimum size of a block to be compressed
 		/// </summary>
@@ -311,7 +317,7 @@ namespace EpicGames.Horde.Storage.Bundles
 				ReadOnlyMemory<byte> encodedPacket = bundle.Payload.Slice(packetOffset, packet.EncodedLength).AsSingleSegment();
 
 				byte[] decodedPacket = new byte[packet.DecodedLength];
-				LZ4Codec.Decode(encodedPacket.Span, decodedPacket);
+				Decode(header.CompressionFormat, encodedPacket.Span, decodedPacket);
 
 				int nodeOffset = 0;
 				for (; exportIdx < header.Exports.Count && nodeOffset + header.Exports[exportIdx].Length <= packet.DecodedLength; exportIdx++)
@@ -818,7 +824,7 @@ namespace EpicGames.Horde.Storage.Bundles
 				// If we can't fit this data into the current block, flush the contents of it first
 				if (blockSize > 0 && blockSize + nodeData.Length > _options.MinCompressionPacketSize)
 				{
-					FlushPacket(context.PacketBuffer.AsMemory(0, blockSize), builder, packets);
+					FlushPacket(_options.CompressionFormat, context.PacketBuffer.AsMemory(0, blockSize), builder, packets);
 					blockSize = 0;
 				}
 
@@ -842,30 +848,63 @@ namespace EpicGames.Horde.Storage.Bundles
 				}
 				else
 				{
-					FlushPacket(nodeData.First, builder, packets);
+					FlushPacket(_options.CompressionFormat, nodeData.First, builder, packets);
 				}
 			}
-			FlushPacket(context.PacketBuffer.AsMemory(0, blockSize), builder, packets);
+			FlushPacket(_options.CompressionFormat, context.PacketBuffer.AsMemory(0, blockSize), builder, packets);
 
 			// Flush the data
-			BundleHeader header = new BundleHeader(imports, exports, packets);
+			BundleHeader header = new BundleHeader(_options.CompressionFormat, imports, exports, packets);
 			return new Bundle(header, builder.AsSequence());
 		}
 
-		static void FlushPacket(ReadOnlyMemory<byte> inputData, ByteArrayBuilder builder, List<BundlePacket> packets)
+		static void FlushPacket(BundleCompressionFormat format, ReadOnlyMemory<byte> inputData, ByteArrayBuilder builder, List<BundlePacket> packets)
 		{
 			if (inputData.Length > 0)
 			{
-				int maxSize = LZ4Codec.MaximumOutputSize(inputData.Length);
+				int encodedLength = Encode(format, inputData.Span, builder);
+				Debug.Assert(encodedLength >= 0);
+				packets.Add(new BundlePacket(encodedLength, inputData.Length));
+			}
+		}
+
+		static int Encode(BundleCompressionFormat format, ReadOnlySpan<byte> inputSpan, ByteArrayBuilder builder)
+		{
+			switch (format)
+			{
+				case BundleCompressionFormat.None:
+					{
+						builder.Append(inputSpan);
+						builder.Advance(inputSpan.Length);
+						return inputSpan.Length;
+					}
+				case BundleCompressionFormat.LZ4:
+					{
+						int maxSize = LZ4Codec.MaximumOutputSize(inputSpan.Length);
 
-				ReadOnlySpan<byte> inputSpan = inputData.Span;
-				Span<byte> outputSpan = builder.GetMemory(maxSize).Span;
+						Span<byte> outputSpan = builder.GetMemory(maxSize).Span;
+						int encodedLength = LZ4Codec.Encode(inputSpan, outputSpan);
+						builder.Advance(encodedLength);
 
-				int encodedLength = LZ4Codec.Encode(inputSpan, outputSpan);
-				packets.Add(new BundlePacket(encodedLength, inputData.Length));
+						return encodedLength;
+					}
+				default:
+					throw new InvalidDataException($"Invalid compression format '{(int)format}'");
+			}
+		}
 
-				Debug.Assert(encodedLength >= 0);
-				builder.Advance(encodedLength);
+		static void Decode(BundleCompressionFormat format, ReadOnlySpan<byte> inputSpan, Span<byte> outputSpan)
+		{
+			switch (format)
+			{
+				case BundleCompressionFormat.None:
+					inputSpan.CopyTo(outputSpan);
+					break;
+				case BundleCompressionFormat.LZ4:
+					LZ4Codec.Decode(inputSpan, outputSpan);
+					break;
+				default:
+					throw new InvalidDataException($"Invalid compression format '{(int)format}'");
 			}
 		}
 
