Fixed issue where in release mode the EpicGames.MsBuild assembly was always being loaded.  This causes issues with GetTypes when invoked on all loaded assemblies.

diff --git a/Engine/Source/Programs/Shared/EpicGames.Build/System/CompileScriptModules.cs b/Engine/Source/Programs/Shared/EpicGames.Build/System/CompileScriptModules.cs
--- a/Engine/Source/Programs/Shared/EpicGames.Build/System/CompileScriptModules.cs
+++ b/Engine/Source/Programs/Shared/EpicGames.Build/System/CompileScriptModules.cs
@@ -10,6 +10,7 @@ using System.Text.Json;
 using Microsoft.Build.Shared;
 using EpicGames.MsBuild;
 using Microsoft.Extensions.Logging;
+using System.Runtime.CompilerServices;
 
 namespace UnrealBuildBase
 {
@@ -262,6 +263,7 @@ namespace UnrealBuildBase
 		/// If it is placed in the caller directly, then when the caller is invoked, the assembly will be loaded resulting
 		/// in the possible Microsoft.Build.Framework load issue later on is this method isn't invoked.
 		/// </summary>
+		[MethodImpl(MethodImplOptions.NoInlining)]
 		static Dictionary<FileReference, CsProjBuildRecordEntry> Build(HashSet<FileReference> FoundProjects,
 			bool bForceCompile, out bool bBuildSuccess, CsProjBuildHook Hook, List<DirectoryReference> BaseDirectories,
 			List<string>? DefineConstants, Action<int> OnBuildingProjects, ILogger Logger)
