The code generator no longer modifies the class definitions.

#rnx

diff --git a/Engine/Source/Programs/Shared/EpicGames.UHT/Exporters/CodeGen/UhtHeaderCodeGeneratorHFile.cs b/Engine/Source/Programs/Shared/EpicGames.UHT/Exporters/CodeGen/UhtHeaderCodeGeneratorHFile.cs
--- a/Engine/Source/Programs/Shared/EpicGames.UHT/Exporters/CodeGen/UhtHeaderCodeGeneratorHFile.cs
+++ b/Engine/Source/Programs/Shared/EpicGames.UHT/Exporters/CodeGen/UhtHeaderCodeGeneratorHFile.cs
@@ -904,47 +904,86 @@ namespace EpicGames.UHT.Exporters.CodeGen
 		}
 
 		/// <summary>
-		/// Generates private copy-constructor declaration.
+		/// Type of constructor on the class regardless of explicit or generated.
 		/// </summary>
-		/// <param name="builder">Output builder</param>
-		/// <param name="classObj">Class being exported</param>
-		/// <param name="api">API text to be used</param>
-		/// <returns>Output builder</returns>
-		private static StringBuilder AppendConstructorDefinition(StringBuilder builder, UhtClass classObj, string api)
+		enum ConstructorType
+		{
+			ObjectInitializer,
+			Default,
+			ForbiddenDefault,
+		}
+
+		/// <summary>
+		/// Return the type of constructor the class will have regardless of if one has been explicitly 
+		/// declared or will be generated.
+		/// </summary>
+		/// <param name="classObj">Class in question</param>
+		/// <returns>Constructor type</returns>
+		private static ConstructorType GetConstructorType(UhtClass classObj)
 		{
 			if (!classObj.ClassExportFlags.HasAnyFlags(UhtClassExportFlags.HasConstructor))
 			{
-				builder.Append("\t/** Standard constructor, called after all reflected properties have been initialized */ \\\r\n");
 
-				// Assume super class has OI constructor, this may not always be true but we should always be able to check 
-				// In any case, it will default to old behavior before we even checked 
+				// Assume super class has OI constructor, this may not always be true but we should always be able to check this.
+				// In any case, it will default to old behavior before we even checked this.
 				bool superClassObjectInitializerConstructorDeclared = true;
 				UhtClass? superClass = classObj.SuperClass;
 				if (superClass != null)
 				{
 					if (!superClass.HeaderFile.IsNoExportTypes)
 					{
-						superClassObjectInitializerConstructorDeclared = superClass.ClassExportFlags.HasAnyFlags(UhtClassExportFlags.HasObjectInitializerConstructor);
+						superClassObjectInitializerConstructorDeclared = GetConstructorType(superClass) == ConstructorType.ObjectInitializer;
 					}
 				}
 
 				if (superClassObjectInitializerConstructorDeclared)
 				{
-					builder.Append('\t').Append(api).Append(classObj.SourceName).Append("(const FObjectInitializer& ObjectInitializer = FObjectInitializer::Get()) : Super(ObjectInitializer) { }; \\\r\n");
-					//ETSTODO - Modification of class
-					classObj.ClassExportFlags |= UhtClassExportFlags.HasObjectInitializerConstructor;
+					return ConstructorType.ObjectInitializer;
 				}
 				else
 				{
-					builder.Append('\t').Append(api).Append(classObj.SourceName).Append("() { }; \\\r\n");
-					//ETSTODO - Modification of class
-					classObj.ClassExportFlags |= UhtClassExportFlags.HasDefaultConstructor;
+					return ConstructorType.Default;
+				}
+			}
+			else
+			{
+				if (classObj.ClassExportFlags.HasAnyFlags(UhtClassExportFlags.HasObjectInitializerConstructor))
+				{
+					return ConstructorType.ObjectInitializer;
+				}
+				else if (classObj.ClassExportFlags.HasAnyFlags(UhtClassExportFlags.HasDefaultConstructor))
+				{
+					return ConstructorType.Default;
 				}
+				else
+				{
+					return ConstructorType.ForbiddenDefault;
+				}
+			}
+		}
+
+		/// <summary>
+		/// Generates private copy-constructor declaration.
+		/// </summary>
+		/// <param name="builder">Output builder</param>
+		/// <param name="classObj">Class being exported</param>
+		/// <param name="api">API text to be used</param>
+		/// <returns>Output builder</returns>
+		private static StringBuilder AppendConstructorDefinition(StringBuilder builder, UhtClass classObj, string api)
+		{
+			if (!classObj.ClassExportFlags.HasAnyFlags(UhtClassExportFlags.HasConstructor))
+			{
+				builder.Append("\t/** Standard constructor, called after all reflected properties have been initialized */ \\\r\n");
+				switch (GetConstructorType(classObj))
+				{
+					case ConstructorType.ObjectInitializer:
+						builder.Append('\t').Append(api).Append(classObj.SourceName).Append("(const FObjectInitializer& ObjectInitializer = FObjectInitializer::Get()) : Super(ObjectInitializer) { }; \\\r\n");
+						break;
 
-				// The original code would mark this true at this point.  However,
-				// we are no longer allowed to modify the data.
-				//ETSTODO - Modification of class
-				classObj.ClassExportFlags |= UhtClassExportFlags.HasConstructor;
+					case ConstructorType.Default:
+						builder.Append('\t').Append(api).Append(classObj.SourceName).Append("() { }; \\\r\n");
+						break;
+				}
 			}
 			AppendCopyConstructorDefinition(builder, classObj, api);
 			return builder;
@@ -958,31 +997,36 @@ namespace EpicGames.UHT.Exporters.CodeGen
 		/// <returns>Output builder</returns>
 		private static StringBuilder AppendDefaultConstructorCallDefinition(StringBuilder builder, UhtClass classObj)
 		{
-			if (classObj.ClassExportFlags.HasAnyFlags(UhtClassExportFlags.HasObjectInitializerConstructor))
-			{
-				if (classObj.ClassFlags.HasAnyFlags(EClassFlags.Abstract))
-				{
-					builder.Append("\tDEFINE_ABSTRACT_DEFAULT_OBJECT_INITIALIZER_CONSTRUCTOR_CALL(").Append(classObj.SourceName).Append(") \\\r\n");
-				}
-				else
-				{
-					builder.Append("\tDEFINE_DEFAULT_OBJECT_INITIALIZER_CONSTRUCTOR_CALL(").Append(classObj.SourceName).Append(") \\\r\n");
-				}
-			}
-			else if (classObj.ClassExportFlags.HasAnyFlags(UhtClassExportFlags.HasDefaultConstructor))
+			switch (GetConstructorType(classObj))
 			{
-				if (classObj.ClassFlags.HasAnyFlags(EClassFlags.Abstract))
-				{
-					builder.Append("\tDEFINE_ABSTRACT_DEFAULT_CONSTRUCTOR_CALL(").Append(classObj.SourceName).Append(") \\\r\n");
-				}
-				else
-				{
-					builder.Append("\tDEFINE_DEFAULT_CONSTRUCTOR_CALL(").Append(classObj.SourceName).Append(") \\\r\n");
-				}
-			}
-			else
-			{
-				builder.Append("\tDEFINE_FORBIDDEN_DEFAULT_CONSTRUCTOR_CALL(").Append(classObj.SourceName).Append(") \\\r\n");
+				case ConstructorType.ObjectInitializer:
+					if (classObj.ClassFlags.HasAnyFlags(EClassFlags.Abstract))
+					{
+						builder.Append("\tDEFINE_ABSTRACT_DEFAULT_OBJECT_INITIALIZER_CONSTRUCTOR_CALL(").Append(classObj.SourceName).Append(") \\\r\n");
+					}
+					else
+					{
+						builder.Append("\tDEFINE_DEFAULT_OBJECT_INITIALIZER_CONSTRUCTOR_CALL(").Append(classObj.SourceName).Append(") \\\r\n");
+					}
+					break;
+
+				case ConstructorType.Default:
+					if (classObj.ClassFlags.HasAnyFlags(EClassFlags.Abstract))
+					{
+						builder.Append("\tDEFINE_ABSTRACT_DEFAULT_CONSTRUCTOR_CALL(").Append(classObj.SourceName).Append(") \\\r\n");
+					}
+					else
+					{
+						builder.Append("\tDEFINE_DEFAULT_CONSTRUCTOR_CALL(").Append(classObj.SourceName).Append(") \\\r\n");
+					}
+					break;
+
+				case ConstructorType.ForbiddenDefault:
+					builder.Append("\tDEFINE_FORBIDDEN_DEFAULT_CONSTRUCTOR_CALL(").Append(classObj.SourceName).Append(") \\\r\n");
+					break;
+
+				default:
+					throw new UhtIceException("Unexpected constructor type");
 			}
 			return builder;
 		}
