Horde: Support for concurrent writers to bundle store.

Each writer can create any number of child writers. Data in the parent writer will not be flushed until all child writers it depends on have been flushed.

Each writer will output data to different blobs, allowing better control of data locality in output bundles.

diff --git a/Engine/Source/Programs/Shared/EpicGames.Horde.Tests/BundleStoreTests.cs b/Engine/Source/Programs/Shared/EpicGames.Horde.Tests/BundleStoreTests.cs
--- a/Engine/Source/Programs/Shared/EpicGames.Horde.Tests/BundleStoreTests.cs
+++ b/Engine/Source/Programs/Shared/EpicGames.Horde.Tests/BundleStoreTests.cs
@@ -42,7 +42,7 @@ namespace EpicGames.Horde.Tests
 			Assert.AreEqual(1, blobStore.Refs.Count);
 		}
 
-		static async Task TestTreeAsync(IBlobStore blobStore, BundleOptions options)
+		static async Task TestTreeAsync(InMemoryBlobStore blobStore, BundleOptions options)
 		{
 			// Generate a tree
 			using (BundleStore store = new BundleStore(blobStore, options))
@@ -55,7 +55,8 @@ namespace EpicGames.Horde.Tests
 				ITreeBlobRef node4 = await writer.WriteNodeAsync(new ReadOnlySequence<byte>(new byte[] { 4 }), Array.Empty<ITreeBlobRef>(), CancellationToken.None);
 
 				ITreeBlob root = TreeBlob.Create(new ReadOnlySequence<byte>(new byte[] { 5 }), new[] { node4, node3 });
-				await writer.FlushAsync(root, CancellationToken.None);
+				await writer.WriteNodeAsync(root.Data, root.Refs);
+				await writer.FlushAsync();
 
 				await CheckTree(root);
 			}
@@ -155,7 +156,8 @@ namespace EpicGames.Horde.Tests
 				FileEntry world = hello.AddFile("world", FileEntryFlags.None);
 				await world.AppendAsync(Encoding.UTF8.GetBytes("world"), new ChunkingOptions(), writer, CancellationToken.None);
 
-				await writer.FlushAsync(root);
+				await writer.WriteNodeAsync(root);
+				await writer.FlushAsync();
 
 				await CheckFileTreeAsync(root);
 			}
@@ -222,6 +224,43 @@ namespace EpicGames.Horde.Tests
 			Assert.IsTrue(uniqueSize < data.Length / 3);
 		}
 
+		[TestMethod]
+		public async Task ChildWriterAsync()
+		{
+			InMemoryBlobStore blobStore = new InMemoryBlobStore();
+
+			using (BundleStore store = new BundleStore(blobStore, new BundleOptions()))
+			{
+				RefName refName = new RefName("test");
+				ITreeWriter writer = store.CreateTreeWriter(refName);
+
+				ITreeWriter childWriter1 = writer.CreateChildWriter();
+				ITreeBlobRef childNode1 = await childWriter1.WriteNodeAsync(new ReadOnlySequence<byte>(new byte[] { 4, 5, 6 }), Array.Empty<ITreeBlobRef>());
+
+				ITreeWriter childWriter2 = writer.CreateChildWriter();
+				ITreeBlobRef childNode2 = await childWriter2.WriteNodeAsync(new ReadOnlySequence<byte>(new byte[] { 7, 8, 9 }), Array.Empty<ITreeBlobRef>());
+
+				await writer.WriteNodeAsync(new ReadOnlySequence<byte>(new byte[] { 1, 2, 3 }), new[] { childNode1, childNode2 });
+				await writer.FlushAsync();
+
+				Assert.AreEqual(3, blobStore.Blobs.Count);
+				Assert.AreEqual(1, blobStore.Refs.Count);
+
+				ITreeBlob? rootBlob = await store.TryReadTreeAsync(refName);
+				Assert.IsNotNull(rootBlob);
+				Assert.IsTrue(rootBlob.Data.AsSingleSegment().Span.SequenceEqual(new byte[] { 1, 2, 3 }));
+				Assert.AreEqual(2, rootBlob.Refs.Count);
+
+				ITreeBlob? childBlob1 = await rootBlob.Refs[0].GetTargetAsync();
+				Assert.IsTrue(childBlob1.Data.AsSingleSegment().Span.SequenceEqual(new byte[] { 4, 5, 6 }));
+				Assert.AreEqual(0, childBlob1.Refs.Count);
+
+				ITreeBlob? childBlob2 = await rootBlob.Refs[1].GetTargetAsync();
+				Assert.IsTrue(childBlob2.Data.AsSingleSegment().Span.SequenceEqual(new byte[] { 7, 8, 9 }));
+				Assert.AreEqual(0, childBlob2.Refs.Count);
+			}
+		}
+
 		[TestMethod]
 		public async Task LargeFileTestAsync()
 		{
@@ -240,7 +279,7 @@ namespace EpicGames.Horde.Tests
 			}
 
 			// Generate a tree
-			using (BundleStore store = new BundleStore(blobStore, new BundleOptions { MaxInlineBlobSize = 1024, MaxBlobSize = 1024 }))
+			using (BundleStore store = new BundleStore(blobStore, new BundleOptions { MaxBlobSize = 1024 }))
 			{
 				ITreeWriter writer = store.CreateTreeWriter(new RefName("test"));
 
@@ -254,7 +293,8 @@ namespace EpicGames.Horde.Tests
 				FileEntry file = root.AddFile("test", FileEntryFlags.None);
 				await file.AppendAsync(data, options, writer, CancellationToken.None);
 
-				await writer.FlushAsync(root);
+				await writer.WriteNodeAsync(root);
+				await writer.FlushAsync();
 
 				await CheckLargeFileTreeAsync(root, data);
 			}
diff --git a/Engine/Source/Programs/Shared/EpicGames.Horde.Tests/BundleTests.cs b/Engine/Source/Programs/Shared/EpicGames.Horde.Tests/BundleTests.cs
--- a/Engine/Source/Programs/Shared/EpicGames.Horde.Tests/BundleTests.cs
+++ b/Engine/Source/Programs/Shared/EpicGames.Horde.Tests/BundleTests.cs
@@ -212,7 +212,6 @@ namespace EpicGames.Horde.Tests
 		{
 			BundleOptions options = new BundleOptions();
 			options.MaxBlobSize = 1;
-			options.MaxInlineBlobSize = 1;
 
 			InMemoryBlobStore blobStore = new InMemoryBlobStore();
 			using ITreeStore treeStore = new BundleStore(blobStore, options);
@@ -237,7 +236,6 @@ namespace EpicGames.Horde.Tests
 		{
 			BundleOptions options = new BundleOptions();
 			options.MaxBlobSize = 1;
-			options.MaxInlineBlobSize = 1;
 
 			InMemoryBlobStore blobStore = new InMemoryBlobStore();
 
diff --git a/Engine/Source/Programs/Shared/EpicGames.Horde/Storage/Backends/InMemoryBlobStore.cs b/Engine/Source/Programs/Shared/EpicGames.Horde/Storage/Backends/InMemoryBlobStore.cs
--- a/Engine/Source/Programs/Shared/EpicGames.Horde/Storage/Backends/InMemoryBlobStore.cs
+++ b/Engine/Source/Programs/Shared/EpicGames.Horde/Storage/Backends/InMemoryBlobStore.cs
@@ -41,7 +41,7 @@ namespace EpicGames.Horde.Storage.Backends
 		/// <inheritdoc/>
 		public Task<IBlob> WriteBlobAsync(ReadOnlySequence<byte> data, IReadOnlyList<BlobId> references, RefName hintRefName = default, CancellationToken cancellationToken = default)
 		{
-			BlobId blobId = new BlobId(Guid.NewGuid().ToString());
+			BlobId blobId = BlobId.Create(ServerId.Empty, hintRefName);
 
 			IBlob blob = Blob.FromMemory(blobId, data.ToArray(), references);
 			_blobs[blobId] = blob;
diff --git a/Engine/Source/Programs/Shared/EpicGames.Horde/Storage/Bundles/BundleStore.cs b/Engine/Source/Programs/Shared/EpicGames.Horde/Storage/Bundles/BundleStore.cs
--- a/Engine/Source/Programs/Shared/EpicGames.Horde/Storage/Bundles/BundleStore.cs
+++ b/Engine/Source/Programs/Shared/EpicGames.Horde/Storage/Bundles/BundleStore.cs
@@ -9,6 +9,7 @@ using System.IO;
 using System.Linq;
 using System.Threading;
 using System.Threading.Tasks;
+using System.Xml.Linq;
 using EpicGames.Core;
 using K4os.Compression.LZ4;
 using Microsoft.Extensions.Caching.Memory;
@@ -25,11 +26,6 @@ namespace EpicGames.Horde.Storage.Bundles
 		/// </summary>
 		public int MaxBlobSize { get; set; } = 10 * 1024 * 1024;
 
-		/// <summary>
-		/// Maximum payload size for the root blob
-		/// </summary>
-		public int MaxInlineBlobSize { get; set; } = 1024 * 64;
-
 		/// <summary>
 		/// Compression format to use
 		/// </summary>
@@ -164,6 +160,9 @@ namespace EpicGames.Horde.Storage.Bundles
 			public readonly ReadOnlySequence<byte> Data;
 			public readonly IReadOnlyList<NodeInfo> References;
 
+			public BundleWriter? _writer;
+			public bool _writing; // This node is currently being written asynchronously
+
 			ReadOnlySequence<byte> ITreeBlob.Data => Data;
 			IReadOnlyList<ITreeBlobRef> ITreeBlob.Refs => References;
 
@@ -210,31 +209,448 @@ namespace EpicGames.Horde.Storage.Bundles
 			}
 		}
 
-		class BundleWriteContext : BundleContext, ITreeWriter
+		class BundleWriter : ITreeWriter
 		{
+			readonly object _lockObject = new object();
+
 			readonly BundleStore _owner;
+			readonly BundleContext _context;
+			readonly BundleWriter? _parent;
+			readonly RefName _refName;
+
+			public readonly ConcurrentDictionary<IoHash, NodeInfo> HashToNode = new ConcurrentDictionary<IoHash, NodeInfo>(); // TODO: this needs to include some additional state from external sources.
 
-			public readonly RefName Name;
-			public readonly ConcurrentDictionary<IoHash, NodeInfo> HashToNode = new ConcurrentDictionary<IoHash, NodeInfo>();
+			// List of child writers that need to be flushed
+			readonly List<BundleWriter> _children = new List<BundleWriter>();
 
-			public Task WriteTask = Task.CompletedTask;
-			public readonly List<NodeInfo> WriteNodes = new List<NodeInfo>();
-			public readonly HashSet<NodeInfo> WriteNodesSet = new HashSet<NodeInfo>();
-			public long WriteLength;
+			// Task which includes all active writes, and returns the last blob
+			Task<BlobId> _flushTask = Task.FromResult(BlobId.Empty);
 
-			public byte[] PacketBuffer = Array.Empty<byte>();
+			// Nodes which have been queued to be written, but which are not yet part of any active write
+			readonly List<NodeInfo> _queueNodes = new List<NodeInfo>();
 
-			public BundleWriteContext(BundleStore owner, RefName name)
+			// Number of nodes in _queueNodes that are ready to be written (ie. all their dependencies have been written)
+			int _readyCount;
+
+			// Sum of lengths for nodes in _queueNodes up to _readyCount
+			long _readyLength;
+
+			/// <summary>
+			/// Constructor
+			/// </summary>
+			public BundleWriter(BundleStore owner, BundleContext context, BundleWriter? parent, RefName refName)
 			{
 				_owner = owner;
-				Name = name;
+				_context = context;
+				_parent = parent;
+				_refName = refName;
 			}
 
 			/// <inheritdoc/>
-			public Task FlushAsync(ReadOnlySequence<byte> data, IReadOnlyList<ITreeBlobRef> refs, CancellationToken cancellationToken = default) => _owner.FlushAsync(this, data, refs, cancellationToken);
+			public ITreeWriter CreateChildWriter()
+			{
+				BundleWriter writer = new BundleWriter(_owner, _context, this, _refName);
+				lock (_lockObject)
+				{
+					_children.Add(writer);
+				}
+				return writer;
+			}
 
 			/// <inheritdoc/>
-			public Task<ITreeBlobRef> WriteNodeAsync(ReadOnlySequence<byte> data, IReadOnlyList<ITreeBlobRef> refs, CancellationToken cancellationToken = default) => _owner.WriteNodeAsync(this, data, refs, cancellationToken);
+			public Task<ITreeBlobRef> WriteNodeAsync(ReadOnlySequence<byte> data, IReadOnlyList<ITreeBlobRef> refs, CancellationToken cancellationToken = default)
+			{
+				IReadOnlyList<NodeInfo> typedRefs = refs.Select(x => (NodeInfo)x).ToList();
+
+				IoHash hash = ComputeHash(data, typedRefs);
+
+				NodeInfo? node;
+				if (!HashToNode.TryGetValue(hash, out node))
+				{
+					InMemoryNodeState state = new InMemoryNodeState(data, typedRefs);
+					NodeInfo newNode = new NodeInfo(_owner, hash, state);
+
+					// Try to add the node again. If we succeed, check if we need to flush the current bundle being built.
+					node = HashToNode.GetOrAdd(hash, newNode);
+					if (node == newNode)
+					{
+						lock (_lockObject)
+						{
+							// Add these nodes to the queue for writing
+							AddToQueue(node);
+
+							// Update the list of nodes which are ready to be written. We need to make sure that any child writers have flushed before
+							// we can reference the blobs containing their nodes.
+							UpdateReady();
+						}
+					}
+				}
+
+				return Task.FromResult<ITreeBlobRef>(node);
+			}
+
+			void AddToQueue(NodeInfo root)
+			{
+				if (root.State is InMemoryNodeState state && state._writer == null)
+				{
+					lock (state)
+					{
+						// Check again to avoid race
+						if (state._writer == null) 
+						{
+							// Write all the dependencies first
+							foreach (NodeInfo reference in state.References)
+							{
+								AddToQueue(reference);
+							}
+
+							state._writer = this;
+							_queueNodes.Add(root);
+						}
+					}
+				}
+			}
+
+			void UpdateReady()
+			{
+				// Update the number of nodes which can be written
+				while (_readyCount < _queueNodes.Count)
+				{
+					NodeInfo nextNode = _queueNodes[_readyCount];
+					InMemoryNodeState nextState = (InMemoryNodeState)nextNode.State;
+
+					if (_readyCount > 0 && _readyLength + nextState.Data.Length > _owner._options.MaxBlobSize)
+					{
+						WriteReady();
+					}
+					else
+					{
+						foreach (NodeInfo other in nextState.References)
+						{
+							if (other.State is InMemoryNodeState otherState)
+							{
+								// Need to wait for nodes in another writer to be flushed first
+								if (otherState._writer != this)
+								{
+									return;
+								}
+
+								// Need to wait for previous bundles from the current writer to complete before we can reference them
+								if (otherState._writing)
+								{
+									return;
+								}
+							}
+						}
+
+						_readyCount++;
+						_readyLength += nextState.Data.Length;
+					}
+				}
+			}
+
+			void WriteReady()
+			{
+				if (_readyCount > 0)
+				{
+					NodeInfo[] writeNodes = _queueNodes.Slice(0, _readyCount).ToArray();
+					_queueNodes.RemoveRange(0, _readyCount);
+
+					// Mark the nodes as writing so nothing else will be flushed until they're ready
+					foreach (NodeInfo writeNode in writeNodes)
+					{
+						InMemoryNodeState writeState = (InMemoryNodeState)writeNode.State;
+						writeState._writing = true;
+					}
+
+					_readyCount = 0;
+					_readyLength = 0;
+
+					Task prevFlushTask = _flushTask ?? Task.CompletedTask;
+					_flushTask = Task.Run(() => WriteNodesAsync(writeNodes, prevFlushTask, CancellationToken.None));
+				}
+			}
+
+			/// <inheritdoc/>
+			public async Task<BlobId> FlushAsync(CancellationToken cancellationToken)
+			{
+				if (_refName == RefName.Empty)
+				{
+					throw new InvalidOperationException("Flushing a child writer is not permitted.");
+				}
+
+				BlobId blobId = await FlushInternalAsync(cancellationToken);
+				await _owner._blobStore.WriteRefTargetAsync(_refName, blobId, cancellationToken);
+				return blobId;
+			}
+
+			async Task<BlobId> FlushInternalAsync(CancellationToken cancellationToken)
+			{
+				// Get all the child writers and flush them
+				BundleWriter[] children;
+				lock (_lockObject)
+				{
+					children = _children.ToArray();
+				}
+
+				Task[] tasks = new Task[children.Length];
+				for (int idx = 0; idx < children.Length; idx++)
+				{
+					tasks[idx] = children[idx].FlushInternalAsync(cancellationToken);
+				}
+				await Task.WhenAll(tasks);
+
+				// Wait for any current writes to finish. _flushTask may be updated through calls to UpdateReady() during its execution, so loop until it's complete.
+				while (!_flushTask.IsCompleted)
+				{
+					await Task.WhenAll(_flushTask, Task.Delay(0, cancellationToken));
+				}
+
+				// Write the last batch
+				if (_queueNodes.Count > 0)
+				{
+					lock (_lockObject)
+					{
+						WriteReady();
+					}
+					await Task.WhenAll(_flushTask, Task.Delay(0, cancellationToken));
+				}
+
+				// Return the identifier of the blob containing the root node
+				return await _flushTask;
+			}
+
+			async Task<BlobId> WriteNodesAsync(NodeInfo[] writeNodes, Task prevWriteTask, CancellationToken cancellationToken)
+			{
+				// Create the bundle
+				Bundle bundle = CreateBundle(writeNodes);
+				BundleHeader header = bundle.Header;
+
+				// Write it to storage
+				IBlob blob = await _owner._blobStore.WriteBlobAsync(bundle.AsSequence(), bundle.Header.Imports.Select(x => x.BlobId).ToList(), _refName, cancellationToken);
+				string cacheKey = GetBundleCacheKey(blob.Id);
+				_owner.AddBundleToCache(cacheKey, bundle);
+
+				// Create a BundleInfo for it
+				BundleInfo bundleInfo = _context.FindOrAddBundle(blob.Id, writeNodes.Length);
+				for (int idx = 0; idx < writeNodes.Length; idx++)
+				{
+					bundleInfo.Exports[idx] = writeNodes[idx];
+				}
+
+				// Update the node states to reference the written bundle
+				int exportIdx = 0;
+				int packetOffset = 0;
+				foreach (BundlePacket packet in bundle.Header.Packets)
+				{
+					PacketInfo packetInfo = new PacketInfo(packetOffset, packet);
+
+					int nodeOffset = 0;
+					for (; exportIdx < header.Exports.Count && nodeOffset + bundle.Header.Exports[exportIdx].Length <= packet.DecodedLength; exportIdx++)
+					{
+						InMemoryNodeState inMemoryState = (InMemoryNodeState)writeNodes[exportIdx].State;
+
+						BundleExport export = header.Exports[exportIdx];
+						writeNodes[exportIdx].SetState(new ExportedNodeState(bundleInfo, exportIdx, packetInfo, nodeOffset, export.Length, inMemoryState.References));
+
+						nodeOffset += header.Exports[exportIdx].Length;
+					}
+
+					packetOffset += packet.EncodedLength;
+				}
+
+				// Now that we've written some nodes to storage, update any parent writers that may be dependent on us
+				lock (_lockObject)
+				{
+					UpdateReady();
+				}
+
+				// Also update the parent write queue
+				if (_parent != null)
+				{
+					lock (_parent._lockObject)
+					{
+						_parent.UpdateReady();
+					}
+				}
+
+				// Wait for other writes to finish
+				await prevWriteTask;
+				return blob.Id;
+			}
+
+			/// <summary>
+			/// Creates a Bundle containing a set of nodes. 
+			/// </summary>
+			Bundle CreateBundle(NodeInfo[] nodes)
+			{
+				BundleOptions options = _owner._options;
+
+				// Create a set from the nodes to be written. We use this to determine references that are imported.
+				HashSet<NodeInfo> nodeSet = new HashSet<NodeInfo>(nodes);
+
+				// Find all the imported nodes by bundle
+				Dictionary<BundleInfo, List<NodeInfo>> bundleToImportedNodes = new Dictionary<BundleInfo, List<NodeInfo>>();
+				foreach (NodeInfo node in nodes)
+				{
+					InMemoryNodeState state = (InMemoryNodeState)node.State;
+					foreach (NodeInfo reference in state.References)
+					{
+						if (nodeSet.Add(reference))
+						{
+							// Get the persisted node info
+							ImportedNodeState importedState = (ImportedNodeState)reference.State;
+							BundleInfo bundleInfo = importedState.BundleInfo;
+
+							// Get the list of nodes within it
+							List<NodeInfo>? importedNodes;
+							if (!bundleToImportedNodes.TryGetValue(bundleInfo, out importedNodes))
+							{
+								importedNodes = new List<NodeInfo>();
+								bundleToImportedNodes.Add(bundleInfo, importedNodes);
+							}
+							importedNodes.Add(reference);
+						}
+					}
+				}
+
+				// Create the import list
+				List<BundleImport> imports = new List<BundleImport>();
+
+				// Add all the imports and assign them identifiers
+				Dictionary<NodeInfo, int> nodeToIndex = new Dictionary<NodeInfo, int>();
+				foreach ((BundleInfo bundleInfo, List<NodeInfo> importedNodes) in bundleToImportedNodes)
+				{
+					(int, IoHash)[] exportInfos = new (int, IoHash)[importedNodes.Count];
+					for (int idx = 0; idx < importedNodes.Count; idx++)
+					{
+						NodeInfo importedNode = importedNodes[idx];
+						nodeToIndex.Add(importedNode, nodeToIndex.Count);
+
+						ImportedNodeState importedNodeState = (ImportedNodeState)importedNode.State;
+						exportInfos[idx] = (importedNodeState.ExportIdx, importedNode.Hash);
+					}
+					imports.Add(new BundleImport(bundleInfo.BlobId, bundleInfo.Exports.Length, exportInfos));
+				}
+
+				// Preallocate data in the encoded buffer to reduce fragmentation if we have to resize
+				ByteArrayBuilder builder = new ByteArrayBuilder(options.MinCompressionPacketSize * 2);
+
+				// Create the export list
+				List<BundleExport> exports = new List<BundleExport>();
+				List<BundlePacket> packets = new List<BundlePacket>();
+
+				// Size of data currently stored in the block buffer
+				int blockSize = 0;
+
+				// Segments of data in the current block
+				List<ReadOnlyMemory<byte>> blockSegments = new List<ReadOnlyMemory<byte>>();
+
+				// Compress all the nodes into the output buffer buffer
+				foreach (NodeInfo node in nodes)
+				{
+					InMemoryNodeState nodeState = (InMemoryNodeState)node.State;
+					ReadOnlySequence<byte> nodeData = nodeState.Data;
+
+					// If we can't fit this data into the current block, flush the contents of it first
+					if (blockSize > 0 && blockSize + nodeData.Length > options.MinCompressionPacketSize)
+					{
+						FlushPacket(_owner._options.CompressionFormat, blockSegments, blockSize, builder, packets);
+						blockSize = 0;
+					}
+
+					// Create the export for this node
+					int[] references = nodeState.References.Select(x => nodeToIndex[x]).ToArray();
+					BundleExport export = new BundleExport(node.Hash, (int)nodeData.Length, references);
+					exports.Add(export);
+					nodeToIndex[node] = nodeToIndex.Count;
+
+					// Write out the new block
+					if (nodeData.Length < options.MinCompressionPacketSize || !nodeData.IsSingleSegment)
+					{
+						blockSize += AddSegments(nodeData, blockSegments);
+					}
+					else
+					{
+						FlushPacket(options.CompressionFormat, nodeData.First, builder, packets);
+					}
+				}
+				FlushPacket(options.CompressionFormat, blockSegments, blockSize, builder, packets);
+
+				// Flush the data
+				BundleHeader header = new BundleHeader(options.CompressionFormat, imports, exports, packets);
+				return new Bundle(header, builder.AsSequence());
+			}
+
+			static int AddSegments(ReadOnlySequence<byte> sequence, List<ReadOnlyMemory<byte>> segments)
+			{
+				int size = 0;
+				foreach (ReadOnlyMemory<byte> segment in sequence)
+				{
+					segments.Add(segment);
+					size += segment.Length;
+				}
+				return size;
+			}
+
+			static void FlushPacket(BundleCompressionFormat format, List<ReadOnlyMemory<byte>> blockSegments, int blockSize, ByteArrayBuilder builder, List<BundlePacket> packets)
+			{
+				if (blockSize > 0)
+				{
+					if (blockSegments.Count == 1)
+					{
+						FlushPacket(format, blockSegments[0], builder, packets);
+					}
+					else
+					{
+						using IMemoryOwner<byte> buffer = MemoryPool<byte>.Shared.Rent(blockSize);
+
+						Memory<byte> output = buffer.Memory;
+						foreach (ReadOnlyMemory<byte> blockSegment in blockSegments)
+						{
+							blockSegment.CopyTo(output);
+							output = output.Slice(blockSegment.Length);
+						}
+
+						FlushPacket(format, buffer.Memory, builder, packets);
+					}
+				}
+			}
+
+			static void FlushPacket(BundleCompressionFormat format, ReadOnlyMemory<byte> inputData, ByteArrayBuilder builder, List<BundlePacket> packets)
+			{
+				if (inputData.Length > 0)
+				{
+					int encodedLength = Encode(format, inputData.Span, builder);
+					Debug.Assert(encodedLength >= 0);
+					packets.Add(new BundlePacket(encodedLength, inputData.Length));
+				}
+			}
+
+			static int Encode(BundleCompressionFormat format, ReadOnlySpan<byte> inputSpan, ByteArrayBuilder builder)
+			{
+				switch (format)
+				{
+					case BundleCompressionFormat.None:
+						{
+							builder.Append(inputSpan);
+							builder.Advance(inputSpan.Length);
+							return inputSpan.Length;
+						}
+					case BundleCompressionFormat.LZ4:
+						{
+							int maxSize = LZ4Codec.MaximumOutputSize(inputSpan.Length);
+
+							Span<byte> outputSpan = builder.GetMemory(maxSize).Span;
+							int encodedLength = LZ4Codec.Encode(inputSpan, outputSpan);
+							builder.Advance(encodedLength);
+
+							return encodedLength;
+						}
+					default:
+						throw new InvalidDataException($"Invalid compression format '{(int)format}'");
+				}
+			}
 		}
 
 		readonly IBlobStore _blobStore;
@@ -265,7 +681,7 @@ namespace EpicGames.Horde.Storage.Bundles
 		}
 
 		/// <inheritdoc/>
-		public ITreeWriter CreateTreeWriter(RefName name) => new BundleWriteContext(this, name);
+		public ITreeWriter CreateTreeWriter(RefName name) => new BundleWriter(this, new BundleContext(), null, name);
 
 		/// <inheritdoc/>
 		public Task DeleteTreeAsync(RefName name, CancellationToken cancellationToken) => _blobStore.DeleteRefAsync(name, cancellationToken);
@@ -576,63 +992,6 @@ namespace EpicGames.Horde.Storage.Bundles
 
 		#region Writing bundles
 
-		Task<ITreeBlobRef> WriteNodeAsync(BundleWriteContext context, ReadOnlySequence<byte> data, IReadOnlyList<ITreeBlobRef> refs, CancellationToken cancellationToken = default)
-		{
-			IReadOnlyList<NodeInfo> typedRefs = refs.Select(x => (NodeInfo)x).ToList();
-
-			IoHash hash = ComputeHash(data, typedRefs);
-
-			NodeInfo? node;
-			if (!context.HashToNode.TryGetValue(hash, out node))
-			{
-				InMemoryNodeState state = new InMemoryNodeState(data, typedRefs);
-				NodeInfo newNode = new NodeInfo(this, hash, state);
-
-				// Try to add the node again. If we succeed, check if we need to flush the current bundle being built.
-				node = context.HashToNode.GetOrAdd(hash, newNode);
-				if (node == newNode)
-				{
-					lock (context)
-					{
-						WriteInMemoryNodes(context, node);
-					}
-				}
-			}
-
-			return Task.FromResult<ITreeBlobRef>(node);
-		}
-
-		void WriteInMemoryNodes(BundleWriteContext context, NodeInfo root)
-		{
-			if (root.State is InMemoryNodeState inMemoryState)
-			{
-				if (context.WriteNodesSet.Add(root))
-				{
-					// Write all the references
-					foreach (NodeInfo reference in inMemoryState.References)
-					{
-						WriteInMemoryNodes(context, reference);
-					}
-
-					// If this node pushes the current blob over the max size, queue it to be written now
-					long length = inMemoryState.Data.Length;
-					if (context.WriteNodes.Count > 0 && context.WriteLength + length > _options.MaxBlobSize)
-					{
-						// Start the write, but don't wait for it to complete. Once we've copied the list of nodes they can be written asynchronously.
-						NodeInfo[] writeNodes = context.WriteNodes.ToArray();
-						_ = SequenceWriteAsync(context, () => WriteBundleAsync(context, writeNodes, CancellationToken.None), CancellationToken.None);
-
-						context.WriteNodes.Clear();
-						context.WriteLength = 0;
-					}
-
-					// Add this node to the list of nodes to be written
-					context.WriteNodes.Add(root);
-					context.WriteLength += length;
-				}
-			}
-		}
-
 		static IoHash ComputeHash(ReadOnlySequence<byte> data, IReadOnlyList<NodeInfo> references)
 		{
 			byte[] buffer = new byte[IoHash.NumBytes * (references.Count + 1)];
@@ -648,251 +1007,6 @@ namespace EpicGames.Horde.Storage.Bundles
 			return IoHash.Compute(buffer);
 		}
 
-		async Task WriteBundleAsync(BundleWriteContext context, NodeInfo[] nodes, CancellationToken cancellationToken)
-		{
-			// Create the bundle
-			Bundle bundle = CreateBundle(context, nodes);
-			BundleHeader header = bundle.Header;
-
-			// Write it to storage
-			IBlob blob = await _blobStore.WriteBlobAsync(bundle.AsSequence(), bundle.Header.Imports.Select(x => x.BlobId).ToList(), context.Name, cancellationToken);
-			string cacheKey = GetBundleCacheKey(blob.Id);
-			AddBundleToCache(cacheKey, bundle);
-
-			// Create a BundleInfo for it
-			BundleInfo bundleInfo = context.FindOrAddBundle(blob.Id, nodes.Length);
-			for (int idx = 0; idx < nodes.Length; idx++)
-			{
-				bundleInfo.Exports[idx] = nodes[idx];
-			}
-
-			// Update the node states to reference the written bundle
-			int exportIdx = 0;
-			int packetOffset = 0;
-			foreach (BundlePacket packet in bundle.Header.Packets)
-			{
-				PacketInfo packetInfo = new PacketInfo(packetOffset, packet);
-
-				int nodeOffset = 0;
-				for (; exportIdx < header.Exports.Count && nodeOffset + bundle.Header.Exports[exportIdx].Length <= packet.DecodedLength; exportIdx++)
-				{
-					InMemoryNodeState inMemoryState = (InMemoryNodeState)nodes[exportIdx].State;
-
-					BundleExport export = header.Exports[exportIdx];
-					nodes[exportIdx].SetState(new ExportedNodeState(bundleInfo, exportIdx, packetInfo, nodeOffset, export.Length, inMemoryState.References));
-
-					nodeOffset += header.Exports[exportIdx].Length;
-				}
-
-				packetOffset += packet.EncodedLength;
-			}
-
-			// Remove all the newly written nodes from the dirty set
-			lock (context)
-			{
-				context.WriteNodesSet.ExceptWith(nodes);
-			}
-		}
-
-		async Task FlushAsync(BundleWriteContext context, ReadOnlySequence<byte> data, IReadOnlyList<ITreeBlobRef> refs, CancellationToken cancellationToken = default)
-		{
-			// Create the root node, and add everything it references to the write buffer
-			List<NodeInfo> typedRefs = refs.ConvertAll(x => (NodeInfo)x);
-			IoHash rootHash = ComputeHash(data, typedRefs);
-			NodeInfo rootNode = new NodeInfo(this, rootHash, new InMemoryNodeState(data, typedRefs));
-
-			// Flush all the nodes to bundles, and snapshot the remaining ref
-			Task task;
-			lock (context)
-			{
-				WriteInMemoryNodes(context, rootNode);
-				NodeInfo[] writeNodes = context.WriteNodes.ToArray();
-				task = SequenceWriteAsync(context, () => FlushInternalAsync(context, writeNodes, cancellationToken), cancellationToken);
-			}
-			await task;
-		}
-
-		async Task FlushInternalAsync(BundleWriteContext context, NodeInfo[] writeNodes, CancellationToken cancellationToken)
-		{
-			Bundle bundle = CreateBundle(context, writeNodes);
-			await _blobStore.WriteRefAsync(context.Name, bundle.AsSequence(), bundle.Header.Imports.ConvertAll(x => x.BlobId), cancellationToken);
-		}
-
-		void FindNodesToWrite(NodeInfo root, List<NodeInfo> nodes, HashSet<NodeInfo> uniqueNodes)
-		{
-			if (root.State is InMemoryNodeState inMemoryState)
-			{
-				if (uniqueNodes.Add(root))
-				{
-					foreach (NodeInfo reference in inMemoryState.References)
-					{
-						FindNodesToWrite(reference, nodes, uniqueNodes);
-					}
-					nodes.Add(root);
-				}
-			}
-		}
-
-		static async Task SequenceWriteAsync(BundleWriteContext context, Func<Task> writeFunc, CancellationToken cancellationToken)
-		{
-			Task task;
-			lock (context)
-			{
-				Task prevTask = context.WriteTask;
-
-				// Wait for the previous write to complete first, ignoring any exceptions
-				Func<Task> wrappedWriteFunc = async () =>
-				{
-					await prevTask.ContinueWith(x => { }, cancellationToken, TaskContinuationOptions.None, TaskScheduler.Default);
-					await writeFunc();
-				};
-				task = Task.Run(wrappedWriteFunc, cancellationToken);
-
-				context.WriteTask = task;
-			}
-			await task;
-		}
-
-		/// <summary>
-		/// Creates a Bundle containing a set of nodes. 
-		/// </summary>
-		Bundle CreateBundle(BundleWriteContext context, NodeInfo[] nodes)
-		{
-			// Create a set from the nodes to be written. We use this to determine references that are imported.
-			HashSet<NodeInfo> nodeSet = new HashSet<NodeInfo>(nodes);
-
-			// Find all the imported nodes by bundle
-			Dictionary<BundleInfo, List<NodeInfo>> bundleToImportedNodes = new Dictionary<BundleInfo, List<NodeInfo>>();
-			foreach (NodeInfo node in nodes)
-			{
-				InMemoryNodeState state = (InMemoryNodeState)node.State;
-				foreach (NodeInfo reference in state.References)
-				{
-					if (nodeSet.Add(reference))
-					{
-						// Get the persisted node info
-						ImportedNodeState importedState = (ImportedNodeState)reference.State;
-						BundleInfo bundleInfo = importedState.BundleInfo;
-
-						// Get the list of nodes within it
-						List<NodeInfo>? importedNodes;
-						if (!bundleToImportedNodes.TryGetValue(bundleInfo, out importedNodes))
-						{
-							importedNodes = new List<NodeInfo>();
-							bundleToImportedNodes.Add(bundleInfo, importedNodes);
-						}
-						importedNodes.Add(reference);
-					}
-				}
-			}
-
-			// Create the import list
-			List<BundleImport> imports = new List<BundleImport>();
-
-			// Add all the imports and assign them identifiers
-			Dictionary<NodeInfo, int> nodeToIndex = new Dictionary<NodeInfo, int>();
-			foreach ((BundleInfo bundleInfo, List<NodeInfo> importedNodes) in bundleToImportedNodes)
-			{
-				(int, IoHash)[] exportInfos = new (int, IoHash)[importedNodes.Count];
-				for (int idx = 0; idx < importedNodes.Count; idx++)
-				{
-					NodeInfo importedNode = importedNodes[idx];
-					nodeToIndex.Add(importedNode, nodeToIndex.Count);
-
-					ImportedNodeState importedNodeState = (ImportedNodeState)importedNode.State;
-					exportInfos[idx] = (importedNodeState.ExportIdx, importedNode.Hash);
-				}
-				imports.Add(new BundleImport(bundleInfo.BlobId, bundleInfo.Exports.Length, exportInfos));
-			}
-
-			// Preallocate data in the encoded buffer to reduce fragmentation if we have to resize
-			ByteArrayBuilder builder = new ByteArrayBuilder(_options.MinCompressionPacketSize * 2);
-
-			// Create the export list
-			List<BundleExport> exports = new List<BundleExport>();
-			List<BundlePacket> packets = new List<BundlePacket>();
-
-			// Size of data currently stored in the block buffer
-			int blockSize = 0;
-
-			// Compress all the nodes into the output buffer buffer
-			foreach (NodeInfo node in nodes)
-			{
-				InMemoryNodeState nodeState = (InMemoryNodeState)node.State;
-				ReadOnlySequence<byte> nodeData = nodeState.Data;
-
-				// If we can't fit this data into the current block, flush the contents of it first
-				if (blockSize > 0 && blockSize + nodeData.Length > _options.MinCompressionPacketSize)
-				{
-					FlushPacket(_options.CompressionFormat, context.PacketBuffer.AsMemory(0, blockSize), builder, packets);
-					blockSize = 0;
-				}
-
-				// Create the export for this node
-				int[] references = nodeState.References.Select(x => nodeToIndex[x]).ToArray();
-				BundleExport export = new BundleExport(node.Hash, (int)nodeData.Length, references);
-				exports.Add(export);
-				nodeToIndex[node] = nodeToIndex.Count;
-
-				// Write out the new block
-				if (nodeData.Length < _options.MinCompressionPacketSize || !nodeData.IsSingleSegment)
-				{
-					int requiredSize = Math.Max(blockSize + (int)nodeData.Length, (int)(_options.MaxBlobSize * 1.2));
-					CreateFreeSpace(ref context.PacketBuffer, blockSize, requiredSize);
-
-					foreach (ReadOnlyMemory<byte> nodeSegment in nodeData)
-					{
-						nodeSegment.CopyTo(context.PacketBuffer.AsMemory(blockSize));
-						blockSize += nodeSegment.Length;
-					}
-				}
-				else
-				{
-					FlushPacket(_options.CompressionFormat, nodeData.First, builder, packets);
-				}
-			}
-			FlushPacket(_options.CompressionFormat, context.PacketBuffer.AsMemory(0, blockSize), builder, packets);
-
-			// Flush the data
-			BundleHeader header = new BundleHeader(_options.CompressionFormat, imports, exports, packets);
-			return new Bundle(header, builder.AsSequence());
-		}
-
-		static void FlushPacket(BundleCompressionFormat format, ReadOnlyMemory<byte> inputData, ByteArrayBuilder builder, List<BundlePacket> packets)
-		{
-			if (inputData.Length > 0)
-			{
-				int encodedLength = Encode(format, inputData.Span, builder);
-				Debug.Assert(encodedLength >= 0);
-				packets.Add(new BundlePacket(encodedLength, inputData.Length));
-			}
-		}
-
-		static int Encode(BundleCompressionFormat format, ReadOnlySpan<byte> inputSpan, ByteArrayBuilder builder)
-		{
-			switch (format)
-			{
-				case BundleCompressionFormat.None:
-					{
-						builder.Append(inputSpan);
-						builder.Advance(inputSpan.Length);
-						return inputSpan.Length;
-					}
-				case BundleCompressionFormat.LZ4:
-					{
-						int maxSize = LZ4Codec.MaximumOutputSize(inputSpan.Length);
-
-						Span<byte> outputSpan = builder.GetMemory(maxSize).Span;
-						int encodedLength = LZ4Codec.Encode(inputSpan, outputSpan);
-						builder.Advance(encodedLength);
-
-						return encodedLength;
-					}
-				default:
-					throw new InvalidDataException($"Invalid compression format '{(int)format}'");
-			}
-		}
-
 		static void Decode(BundleCompressionFormat format, ReadOnlySpan<byte> inputSpan, Span<byte> outputSpan)
 		{
 			switch (format)
diff --git a/Engine/Source/Programs/Shared/EpicGames.Horde/Storage/ITreeStore.cs b/Engine/Source/Programs/Shared/EpicGames.Horde/Storage/ITreeStore.cs
--- a/Engine/Source/Programs/Shared/EpicGames.Horde/Storage/ITreeStore.cs
+++ b/Engine/Source/Programs/Shared/EpicGames.Horde/Storage/ITreeStore.cs
@@ -91,6 +91,13 @@ namespace EpicGames.Horde.Storage
 	/// </summary>
 	public interface ITreeWriter
 	{
+		/// <summary>
+		/// Creates a new tree writer parented to this. Any nodes added to this writer can reference nodes written to any child writers, and bundle writes will 
+		/// be sequenced to allow references between them.
+		/// </summary>
+		/// <returns>New writer instance</returns>
+		ITreeWriter CreateChildWriter();
+
 		/// <summary>
 		/// Adds a new node to this tree
 		/// </summary>
@@ -101,13 +108,11 @@ namespace EpicGames.Horde.Storage
 		Task<ITreeBlobRef> WriteNodeAsync(ReadOnlySequence<byte> data, IReadOnlyList<ITreeBlobRef> refs, CancellationToken cancellationToken = default);
 
 		/// <summary>
-		/// Flush the tree with the given data at the root.
+		/// Flushes the current state and writes the ref, with the last written node as the tree root. Throws an exception if called for a child writer.
 		/// </summary>
-		/// <param name="data">Data for the node</param>
-		/// <param name="refs">References to other nodes</param>
-		/// <param name="cancellationToken">Cancellation token for the operation</param>
-		/// <returns>Reference to the node that was added</returns>
-		Task FlushAsync(ReadOnlySequence<byte> data, IReadOnlyList<ITreeBlobRef> refs, CancellationToken cancellationToken = default);
+		/// <param name="cancellationToken"></param>
+		/// <returns></returns>
+		Task<BlobId> FlushAsync(CancellationToken cancellationToken = default);
 	}
 
 	/// <summary>
@@ -176,16 +181,5 @@ namespace EpicGames.Horde.Storage
 		{
 			return new TypedTreeStore<T>(store);
 		}
-
-		/// <summary>
-		/// Flush the root of a tree to the store
-		/// </summary>
-		/// <param name="writer">Writer for the nodes</param>
-		/// <param name="root">Root blob</param>
-		/// <param name="cancellationToken">Cancellation token for the operation</param>
-		public static Task FlushAsync(this ITreeWriter writer, ITreeBlob root, CancellationToken cancellationToken = default)
-		{
-			return writer.FlushAsync(root.Data, root.Refs, cancellationToken);
-		}
 	}
 }
diff --git a/Engine/Source/Programs/Shared/EpicGames.Horde/Storage/TreeNode.cs b/Engine/Source/Programs/Shared/EpicGames.Horde/Storage/TreeNode.cs
--- a/Engine/Source/Programs/Shared/EpicGames.Horde/Storage/TreeNode.cs
+++ b/Engine/Source/Programs/Shared/EpicGames.Horde/Storage/TreeNode.cs
@@ -5,6 +5,7 @@ using System.Buffers;
 using System.Collections.Generic;
 using System.Diagnostics;
 using System.Reflection;
+using System.Reflection.Metadata;
 using System.Threading;
 using System.Threading.Tasks;
 using System.Xml.Linq;
@@ -141,15 +142,15 @@ namespace EpicGames.Horde.Storage
 	public static class TreeNodeExtensions
 	{
 		/// <summary>
-		/// Flush the tree with the given root node to storage
+		/// Writes a node to storage
 		/// </summary>
 		/// <param name="writer">Writer to output the nodes to</param>
-		/// <param name="root">Root node to serialize</param>
+		/// <param name="node">Root node to serialize</param>
 		/// <param name="cancellationToken">Cancellation token for the operation</param>
-		public static async Task FlushAsync(this ITreeWriter writer, TreeNode root, CancellationToken cancellationToken = default)
+		public static async Task WriteNodeAsync(this ITreeWriter writer, TreeNode node, CancellationToken cancellationToken = default)
 		{
-			ITreeBlob rootBlob = await root.SerializeAsync(writer, cancellationToken);
-			await writer.FlushAsync(rootBlob.Data, rootBlob.Refs, cancellationToken);
+			ITreeBlob blob = await node.SerializeAsync(writer, cancellationToken);
+			await writer.WriteNodeAsync(blob.Data, blob.Refs, cancellationToken);
 		}
 
 		/// <inheritdoc/>
@@ -178,7 +179,8 @@ namespace EpicGames.Horde.Storage
 		public static async Task WriteTreeAsync(this ITreeStore store, RefName name, TreeNode root, CancellationToken cancellationToken = default)
 		{
 			ITreeWriter writer = store.CreateTreeWriter(name);
-			await writer.FlushAsync(root, cancellationToken);
+			await writer.WriteNodeAsync(root, cancellationToken);
+			await writer.FlushAsync(cancellationToken);
 		}
 	}
 }
