Horde: Write memory sequences to storage without copying to a separate buffer first.

diff --git a/Engine/Source/Programs/Horde/Horde.Build/Perforce/ReplicationService.cs b/Engine/Source/Programs/Horde/Horde.Build/Perforce/ReplicationService.cs
--- a/Engine/Source/Programs/Horde/Horde.Build/Perforce/ReplicationService.cs
+++ b/Engine/Source/Programs/Horde/Horde.Build/Perforce/ReplicationService.cs
@@ -713,7 +713,8 @@ namespace Horde.Build.Perforce
 					if (syncedSize > 0)
 					{
 						Stopwatch flushTimer = Stopwatch.StartNew();
-						await writer.FlushAsync(syncNode, cancellationToken);
+						await writer.WriteNodeAsync(syncNode, cancellationToken);
+						await writer.FlushAsync(cancellationToken);
 						flushTimer.Stop();
 						deleteIncRef = true;
 					}
diff --git a/Engine/Source/Programs/Horde/Horde.Build/Storage/BasicBlobStore.cs b/Engine/Source/Programs/Horde/Horde.Build/Storage/BasicBlobStore.cs
--- a/Engine/Source/Programs/Horde/Horde.Build/Storage/BasicBlobStore.cs
+++ b/Engine/Source/Programs/Horde/Horde.Build/Storage/BasicBlobStore.cs
@@ -92,14 +92,17 @@ namespace Horde.Build.Storage
 		}
 
 		/// <inheritdoc/>
-		public async Task<IBlob> WriteBlobAsync(ReadOnlySequence<byte> data, IReadOnlyList<BlobId> references, RefName hintRefName = default, CancellationToken cancellationToken = default)
+		public async Task<BlobId> WriteBlobAsync(ReadOnlySequence<byte> data, IReadOnlyList<BlobId> references, RefName hintRefName = default, CancellationToken cancellationToken = default)
 		{
 			BlobId id = BlobId.Create(_serverId, hintRefName);
 			string path = GetBlobPath(id);
 
-			ReadOnlyMemory<byte> memory = Blob.Serialize(data, references).AsSingleSegment();
-			await _backend.WriteBytesAsync(path, memory, cancellationToken);
-			return Blob.FromMemory(id, memory.Slice(memory.Length - (int)data.Length), references);
+			ReadOnlySequence<byte> sequence = Blob.Serialize(data, references);
+			using (ReadOnlySequenceStream stream = new ReadOnlySequenceStream(sequence))
+			{
+				await _backend.WriteAsync(path, stream, cancellationToken);
+			}
+			return id;
 		}
 
 		#endregion
@@ -150,11 +153,11 @@ namespace Horde.Build.Storage
 		}
 
 		/// <inheritdoc/>
-		public async Task<IBlob> WriteRefAsync(RefName name, ReadOnlySequence<byte> data, IReadOnlyList<BlobId> references, CancellationToken cancellationToken = default)
+		public async Task<BlobId> WriteRefAsync(RefName name, ReadOnlySequence<byte> data, IReadOnlyList<BlobId> references, CancellationToken cancellationToken = default)
 		{
-			IBlob blob = await WriteBlobAsync(data, references, name, cancellationToken);
-			await WriteRefTargetAsync(name, blob.Id, cancellationToken);
-			return blob;
+			BlobId blobId = await WriteBlobAsync(data, references, name, cancellationToken);
+			await WriteRefTargetAsync(name, blobId, cancellationToken);
+			return blobId;
 		}
 
 		/// <inheritdoc/>
diff --git a/Engine/Source/Programs/Shared/EpicGames.Core/ReadOnlySequenceStream.cs b/Engine/Source/Programs/Shared/EpicGames.Core/ReadOnlySequenceStream.cs
new file mode 100644
--- /dev/null
+++ b/Engine/Source/Programs/Shared/EpicGames.Core/ReadOnlySequenceStream.cs
@@ -0,0 +1,97 @@
+// Copyright Epic Games, Inc. All Rights Reserved.
+
+using System;
+using System.Buffers;
+using System.IO;
+
+namespace EpicGames.Core
+{
+	/// <summary>
+	/// Stream which reads from a <see cref="ReadOnlyMemory{Byte}"/>
+	/// </summary>
+	public class ReadOnlySequenceStream : Stream
+	{
+		/// <summary>
+		/// The buffer to read from
+		/// </summary>
+		readonly ReadOnlySequence<byte> _source;
+
+		/// <summary>
+		/// Remaining sequence to read
+		/// </summary>
+		ReadOnlySequence<byte> _remaining;
+
+		/// <summary>
+		/// Constructor
+		/// </summary>
+		/// <param name="sequence">The memory to read from</param>
+		public ReadOnlySequenceStream(ReadOnlySequence<byte> sequence)
+		{
+			_source = sequence;
+			_remaining = _source;
+		}
+
+		/// <inheritdoc/>
+		public override bool CanRead => true;
+
+		/// <inheritdoc/>
+		public override bool CanSeek => true;
+
+		/// <inheritdoc/>
+		public override bool CanWrite => false;
+
+		/// <inheritdoc/>
+		public override long Length => _source.Length;
+
+		/// <inheritdoc/>
+		public override long Position
+		{
+			get => _source.Length - _remaining.Length;
+			set => _remaining = _source.Slice(value);
+		}
+
+		/// <inheritdoc/>
+		public override void Flush()
+		{
+		}
+
+		/// <inheritdoc/>
+		public override int Read(Span<byte> buffer)
+		{
+			ReadOnlySpan<byte> firstSpan = _remaining.FirstSpan;
+			int copyLength = Math.Min(firstSpan.Length, buffer.Length);
+			firstSpan.Slice(0, copyLength).CopyTo(buffer);
+			_remaining = _remaining.Slice(copyLength);
+			return copyLength;
+		}
+
+		/// <inheritdoc/>
+		public override int Read(byte[] buffer, int offset, int count) => Read(buffer.AsSpan(offset, count));
+
+		/// <inheritdoc/>
+		public override long Seek(long offset, SeekOrigin origin)
+		{
+			switch (origin)
+			{
+				case SeekOrigin.Begin:
+					Position = offset;
+					break;
+				case SeekOrigin.Current:
+					Position += offset;
+					break;
+				case SeekOrigin.End:
+					Position = Length + offset;
+					break;
+				default:
+					throw new ArgumentException(null, nameof(origin));
+			}
+			return Position;
+		}
+
+		/// <inheritdoc/>
+		public override void SetLength(long value) => throw new InvalidOperationException();
+
+		/// <inheritdoc/>
+		public override void Write(byte[] buffer, int offset, int count) => throw new InvalidOperationException();
+	}
+}
diff --git a/Engine/Source/Programs/Shared/EpicGames.Horde/Storage/Backends/InMemoryBlobStore.cs b/Engine/Source/Programs/Shared/EpicGames.Horde/Storage/Backends/InMemoryBlobStore.cs
--- a/Engine/Source/Programs/Shared/EpicGames.Horde/Storage/Backends/InMemoryBlobStore.cs
+++ b/Engine/Source/Programs/Shared/EpicGames.Horde/Storage/Backends/InMemoryBlobStore.cs
@@ -39,14 +39,14 @@ namespace EpicGames.Horde.Storage.Backends
 		}
 
 		/// <inheritdoc/>
-		public Task<IBlob> WriteBlobAsync(ReadOnlySequence<byte> data, IReadOnlyList<BlobId> references, RefName hintRefName = default, CancellationToken cancellationToken = default)
+		public Task<BlobId> WriteBlobAsync(ReadOnlySequence<byte> data, IReadOnlyList<BlobId> references, RefName hintRefName = default, CancellationToken cancellationToken = default)
 		{
 			BlobId blobId = BlobId.Create(ServerId.Empty, hintRefName);
 
 			IBlob blob = Blob.FromMemory(blobId, data.ToArray(), references);
 			_blobs[blobId] = blob;
 
-			return Task.FromResult(blob);
+			return Task.FromResult(blobId);
 		}
 
 		#endregion
@@ -84,11 +84,11 @@ namespace EpicGames.Horde.Storage.Backends
 		}
 
 		/// <inheritdoc/>
-		public async Task<IBlob> WriteRefAsync(RefName name, ReadOnlySequence<byte> data, IReadOnlyList<BlobId> references, CancellationToken cancellationToken)
+		public async Task<BlobId> WriteRefAsync(RefName name, ReadOnlySequence<byte> data, IReadOnlyList<BlobId> references, CancellationToken cancellationToken)
 		{
-			IBlob blob = await WriteBlobAsync(data, references, name, cancellationToken);
-			_refs[name] = blob.Id;
-			return blob;
+			BlobId blobId = await WriteBlobAsync(data, references, name, cancellationToken);
+			_refs[name] = blobId;
+			return blobId;
 		}
 
 		/// <inheritdoc/>
diff --git a/Engine/Source/Programs/Shared/EpicGames.Horde/Storage/Bundles/BundleStore.cs b/Engine/Source/Programs/Shared/EpicGames.Horde/Storage/Bundles/BundleStore.cs
--- a/Engine/Source/Programs/Shared/EpicGames.Horde/Storage/Bundles/BundleStore.cs
+++ b/Engine/Source/Programs/Shared/EpicGames.Horde/Storage/Bundles/BundleStore.cs
@@ -427,12 +427,12 @@ namespace EpicGames.Horde.Storage.Bundles
 				BundleHeader header = bundle.Header;
 
 				// Write it to storage
-				IBlob blob = await _owner._blobStore.WriteBlobAsync(bundle.AsSequence(), bundle.Header.Imports.Select(x => x.BlobId).ToList(), _refName, cancellationToken);
-				string cacheKey = GetBundleCacheKey(blob.Id);
+				BlobId blobId = await _owner._blobStore.WriteBlobAsync(bundle.AsSequence(), bundle.Header.Imports.Select(x => x.BlobId).ToList(), _refName, cancellationToken);
+				string cacheKey = GetBundleCacheKey(blobId);
 				_owner.AddBundleToCache(cacheKey, bundle);
 
 				// Create a BundleInfo for it
-				BundleInfo bundleInfo = _context.FindOrAddBundle(blob.Id, writeNodes.Length);
+				BundleInfo bundleInfo = _context.FindOrAddBundle(blobId, writeNodes.Length);
 				for (int idx = 0; idx < writeNodes.Length; idx++)
 				{
 					bundleInfo.Exports[idx] = writeNodes[idx];
@@ -476,7 +476,7 @@ namespace EpicGames.Horde.Storage.Bundles
 
 				// Wait for other writes to finish
 				await prevWriteTask;
-				return blob.Id;
+				return blobId;
 			}
 
 			/// <summary>
diff --git a/Engine/Source/Programs/Shared/EpicGames.Horde/Storage/IBlobStore.cs b/Engine/Source/Programs/Shared/EpicGames.Horde/Storage/IBlobStore.cs
--- a/Engine/Source/Programs/Shared/EpicGames.Horde/Storage/IBlobStore.cs
+++ b/Engine/Source/Programs/Shared/EpicGames.Horde/Storage/IBlobStore.cs
@@ -76,8 +76,8 @@ namespace EpicGames.Horde.Storage
 		/// <param name="references">Object references</param>
 		/// <param name="cancellationToken">Cancellation token for the operation</param>
 		/// <param name="hintRefName">Name of a ref that this blob is being written for. While the returned BlobId is guaranteed to be unique, this name can be used as a prefix to aid debugging.</param>
-		/// <returns>Unique identifier for the bundle</returns>
-		Task<IBlob> WriteBlobAsync(ReadOnlySequence<byte> data, IReadOnlyList<BlobId> references, RefName hintRefName = default, CancellationToken cancellationToken = default);
+		/// <returns>Unique identifier for the blob</returns>
+		Task<BlobId> WriteBlobAsync(ReadOnlySequence<byte> data, IReadOnlyList<BlobId> references, RefName hintRefName = default, CancellationToken cancellationToken = default);
 
 		#endregion
 
@@ -109,7 +109,7 @@ namespace EpicGames.Horde.Storage
 		/// <param name="references">Object references</param>
 		/// <param name="cancellationToken">Cancellation token for the operation</param>
 		/// <returns>Unique identifier for the blob</returns>
-		Task<IBlob> WriteRefAsync(RefName name, ReadOnlySequence<byte> data, IReadOnlyList<BlobId> references, CancellationToken cancellationToken = default);
+		Task<BlobId> WriteRefAsync(RefName name, ReadOnlySequence<byte> data, IReadOnlyList<BlobId> references, CancellationToken cancellationToken = default);
 
 		/// <summary>
 		/// Writes a new ref to the store
@@ -154,7 +154,7 @@ namespace EpicGames.Horde.Storage
 			public Task<IBlob> ReadBlobAsync(BlobId id, CancellationToken cancellationToken = default) => _inner.ReadBlobAsync(id, cancellationToken);
 
 			/// <inheritdoc/>
-			public Task<IBlob> WriteBlobAsync(ReadOnlySequence<byte> data, IReadOnlyList<BlobId> references, RefName hintRefName = default, CancellationToken cancellationToken = default) => _inner.WriteBlobAsync(data, references, hintRefName, cancellationToken);
+			public Task<BlobId> WriteBlobAsync(ReadOnlySequence<byte> data, IReadOnlyList<BlobId> references, RefName hintRefName = default, CancellationToken cancellationToken = default) => _inner.WriteBlobAsync(data, references, hintRefName, cancellationToken);
 
 			#endregion
 
@@ -170,7 +170,7 @@ namespace EpicGames.Horde.Storage
 			public Task<BlobId> TryReadRefTargetAsync(RefName name, DateTime cacheTime = default, CancellationToken cancellationToken = default) => _inner.TryReadRefTargetAsync(name, cacheTime, cancellationToken);
 
 			/// <inheritdoc/>
-			public Task<IBlob> WriteRefAsync(RefName name, ReadOnlySequence<byte> data, IReadOnlyList<BlobId> references, CancellationToken cancellationToken = default) => _inner.WriteRefAsync(name, data, references, cancellationToken);
+			public Task<BlobId> WriteRefAsync(RefName name, ReadOnlySequence<byte> data, IReadOnlyList<BlobId> references, CancellationToken cancellationToken = default) => _inner.WriteRefAsync(name, data, references, cancellationToken);
 
 			/// <inheritdoc/>
 			public Task WriteRefTargetAsync(RefName name, BlobId blobId, CancellationToken cancellationToken = default) => _inner.WriteRefTargetAsync(name, blobId, cancellationToken);
