Horde: Write file content to a separate blob stream than directory content when replicating Perforce data.

Other changes:

* Allow passing any string prefix for blob names, not just ref names.
* Fix deserialization of blob ids from MongoDB.
* Arbitrary nodes can now be passed in as the root for any refs written using a writer. This allows writing multiple refs using the same data, and makes semantics clearer for serialization.

diff --git a/Engine/Source/Programs/Horde/Horde.Build/Commands/Bundles/CreateCommand.cs b/Engine/Source/Programs/Horde/Horde.Build/Commands/Bundles/CreateCommand.cs
--- a/Engine/Source/Programs/Horde/Horde.Build/Commands/Bundles/CreateCommand.cs
+++ b/Engine/Source/Programs/Horde/Horde.Build/Commands/Bundles/CreateCommand.cs
@@ -41,7 +41,7 @@ namespace Horde.Build.Commands.Bundles
 			using ServiceProvider serviceProvider = Startup.CreateServiceProvider(_configuration, _loggerProvider);
 
 			ITreeStore store = serviceProvider.GetRequiredService<ITreeStore<CommitService>>();
-			ITreeWriter writer = store.CreateTreeWriter(RefName);
+			ITreeWriter writer = store.CreateTreeWriter(RefName.Text);
 
 			DirectoryNode node = new DirectoryNode(DirectoryFlags.None);
 			await node.CopyFromDirectoryAsync(InputDir.ToDirectoryInfo(), new ChunkingOptions(), writer, logger, CancellationToken.None);
diff --git a/Engine/Source/Programs/Horde/Horde.Build/Perforce/ReplicationService.cs b/Engine/Source/Programs/Horde/Horde.Build/Perforce/ReplicationService.cs
--- a/Engine/Source/Programs/Horde/Horde.Build/Perforce/ReplicationService.cs
+++ b/Engine/Source/Programs/Horde/Horde.Build/Perforce/ReplicationService.cs
@@ -609,7 +609,8 @@ namespace Horde.Build.Perforce
 
 			RefName refName = GetRefName(stream.Id, change, filter, revisionsOnly);
 			RefName incRefName = new RefName($"{refName}-inc");
-			ITreeWriter writer = _treeStore.CreateTreeWriter(incRefName);
+			ITreeWriter directoryWriter = _treeStore.CreateTreeWriter(refName.Text);
+			ITreeWriter fileWriter = directoryWriter.CreateChildWriter();
 
 			// Get the current sync state for this change
 			ReplicationNode? syncNode = await _treeStore.TryReadTreeAsync<ReplicationNode>(incRefName, cancellationToken: cancellationToken);
@@ -645,7 +646,7 @@ namespace Horde.Build.Perforce
 					using ReadOnlyMemoryStream dataStream = new ReadOnlyMemoryStream(data);
 
 					FileEntry entry = await root.AddFileByPathAsync(path, FileEntryFlags.PerforceDepotPathAndRevision, cancellationToken);
-					await entry.AppendAsync(dataStream, Options.Chunking, writer, cancellationToken);
+					await entry.AppendAsync(dataStream, Options.Chunking, directoryWriter, cancellationToken);
 				}
 			}
 			else
@@ -702,9 +703,6 @@ namespace Horde.Build.Perforce
 				long totalSize = directories.Sum(x => x._size);
 				_logger.LogInformation("Total sync size: {Size:n1}mb", totalSize / (1024.0 * 1024.0));
 
-				// Previously scheduled flush task
-				Task flushTask = Task.CompletedTask;
-
 				// Sync incrementally
 				long syncedSize = 0;
 				while (directories.Count > 0)
@@ -713,8 +711,7 @@ namespace Horde.Build.Perforce
 					if (syncedSize > 0)
 					{
 						Stopwatch flushTimer = Stopwatch.StartNew();
-						await writer.WriteNodeAsync(syncNode, cancellationToken);
-						await writer.FlushAsync(cancellationToken);
+						await directoryWriter.WriteRefAsync(incRefName, syncNode, cancellationToken);
 						flushTimer.Stop();
 						deleteIncRef = true;
 					}
@@ -804,7 +801,7 @@ namespace Horde.Build.Perforce
 								else if (io.Command == PerforceIoCommand.Write)
 								{
 									FileWriter file = handles[io.File];
-									await file.AppendAsync(io.Payload, Options.Chunking, writer, cancellationToken);
+									await file.AppendAsync(io.Payload, Options.Chunking, fileWriter, cancellationToken);
 								}
 								else if (io.Command == PerforceIoCommand.Close)
 								{
@@ -813,7 +810,7 @@ namespace Horde.Build.Perforce
 									{
 										if (handles.Remove(io.File, out file))
 										{
-											await file.FinishAsync(writer, cancellationToken);
+											await file.FinishAsync(fileWriter, cancellationToken);
 										}
 									}
 									finally
@@ -877,14 +874,14 @@ namespace Horde.Build.Perforce
 						directories.RemoveAt(directories.Count - 1);
 					}
 				}
-				await flushTask;
 			}
+			await directoryWriter.WriteNodeAsync(syncNode, cancellationToken);
 
 			clientInfo.Change = change;
 
 			// Return the new root object
 			_logger.LogInformation("Writing ref {RefId} for {StreamId} change {Change}", refName, stream.Id, change);
-			await _treeStore.WriteTreeAsync(refName, syncNode, cancellationToken);
+			await directoryWriter.WriteRefAsync(refName, syncNode, cancellationToken);
 
 			// Delete the incremental state
 			if (deleteIncRef)
diff --git a/Engine/Source/Programs/Horde/Horde.Build/Startup.cs b/Engine/Source/Programs/Horde/Horde.Build/Startup.cs
--- a/Engine/Source/Programs/Horde/Horde.Build/Startup.cs
+++ b/Engine/Source/Programs/Horde/Horde.Build/Startup.cs
@@ -280,7 +280,7 @@ namespace Horde.Build
 		static IBlobStore CreateBlobStore(IServiceProvider sp, BlobStoreOptions options)
 		{
 			IStorageBackend backend = CreateStorageBackend(sp, options);
-			return new BasicBlobStore(sp.GetRequiredService<MongoService>(), backend, sp.GetRequiredService<IMemoryCache>());
+			return new BasicBlobStore(sp.GetRequiredService<MongoService>(), backend, sp.GetRequiredService<IMemoryCache>(), sp.GetRequiredService<ILogger<BasicBlobStore>>());
 		}
 
 		static ITreeStore CreateTreeStore(IServiceProvider sp, TreeStoreOptions options)
@@ -803,6 +803,21 @@ namespace Horde.Build
 				}
 		*/
 
+		public sealed class BlobIdBsonSerializer : SerializerBase<BlobId>
+		{
+			/// <inheritdoc/>
+			public override BlobId Deserialize(BsonDeserializationContext context, BsonDeserializationArgs args)
+			{
+				return new BlobId(context.Reader.ReadString());
+			}
+
+			/// <inheritdoc/>
+			public override void Serialize(BsonSerializationContext context, BsonSerializationArgs args, BlobId value)
+			{
+				context.Writer.WriteString(value.ToString());
+			}
+		}
+
 		public sealed class RefIdBsonSerializer : SerializerBase<RefId>
 		{
 			/// <inheritdoc/>
@@ -846,6 +861,7 @@ namespace Horde.Build
 				ConventionRegistry.Register("Horde", conventionPack, type => true);
 
 				// Register the custom serializers
+				BsonSerializer.RegisterSerializer(new BlobIdBsonSerializer());
 				BsonSerializer.RegisterSerializer(new RefIdBsonSerializer());
 				BsonSerializer.RegisterSerializer(new RefNameBsonSerializer());
 				BsonSerializer.RegisterSerializer(new ConditionSerializer());
diff --git a/Engine/Source/Programs/Horde/Horde.Build/Storage/BasicBlobStore.cs b/Engine/Source/Programs/Horde/Horde.Build/Storage/BasicBlobStore.cs
--- a/Engine/Source/Programs/Horde/Horde.Build/Storage/BasicBlobStore.cs
+++ b/Engine/Source/Programs/Horde/Horde.Build/Storage/BasicBlobStore.cs
@@ -11,6 +11,7 @@ using EpicGames.Core;
 using EpicGames.Horde.Storage;
 using Horde.Build.Server;
 using Microsoft.Extensions.Caching.Memory;
+using Microsoft.Extensions.Logging;
 using MongoDB.Bson.Serialization.Attributes;
 using MongoDB.Driver;
 
@@ -60,18 +61,22 @@ namespace Horde.Build.Storage
 		readonly IStorageBackend _backend;
 		readonly IMemoryCache _cache;
 
+		readonly ILogger _logger;
+
 		/// <summary>
 		/// Constructor
 		/// </summary>
 		/// <param name="mongoService">The mongo service implementation</param>
 		/// <param name="backend">Backend to use for storing data</param>
 		/// <param name="cache">Cache for ref values</param>
-		public BasicBlobStore(MongoService mongoService, IStorageBackend backend, IMemoryCache cache)
+		/// <param name="logger">Logger for this instance</param>
+		public BasicBlobStore(MongoService mongoService, IStorageBackend backend, IMemoryCache cache, ILogger<BasicBlobStore> logger)
 		{
 			_refs = mongoService.GetCollection<RefDocument>("Refs");
 			_backend = backend;
 			_cache = cache;
-		}
+			_logger = logger;
+ 		}
 
 		#region Blobs
 
@@ -92,9 +97,9 @@ namespace Horde.Build.Storage
 		}
 
 		/// <inheritdoc/>
-		public async Task<BlobId> WriteBlobAsync(ReadOnlySequence<byte> data, IReadOnlyList<BlobId> references, RefName hintRefName = default, CancellationToken cancellationToken = default)
+		public async Task<BlobId> WriteBlobAsync(ReadOnlySequence<byte> data, IReadOnlyList<BlobId> references, Utf8String prefix = default, CancellationToken cancellationToken = default)
 		{
-			BlobId id = BlobId.Create(_serverId, hintRefName);
+			BlobId id = BlobId.Create(_serverId, prefix);
 			string path = GetBlobPath(id);
 
 			ReadOnlySequence<byte> sequence = Blob.Serialize(data, references);
@@ -134,7 +139,14 @@ namespace Horde.Build.Storage
 			BlobId blobId = await TryReadRefTargetAsync(name, cacheTime, cancellationToken);
 			if (blobId.IsValid())
 			{
-				return await ReadBlobAsync(blobId, cancellationToken);
+				try
+				{
+					return await ReadBlobAsync(blobId, cancellationToken);
+				}
+				catch (Exception ex)
+				{
+					_logger.LogWarning(ex, "Unable to read blob {BlobId} referenced by ref {RefName}", blobId, name);
+				}
 			}
 			return null;
 		}
@@ -155,7 +167,7 @@ namespace Horde.Build.Storage
 		/// <inheritdoc/>
 		public async Task<BlobId> WriteRefAsync(RefName name, ReadOnlySequence<byte> data, IReadOnlyList<BlobId> references, CancellationToken cancellationToken = default)
 		{
-			BlobId blobId = await WriteBlobAsync(data, references, name, cancellationToken);
+			BlobId blobId = await WriteBlobAsync(data, references, name.Text, cancellationToken);
 			await WriteRefTargetAsync(name, blobId, cancellationToken);
 			return blobId;
 		}
diff --git a/Engine/Source/Programs/Shared/EpicGames.Horde/Storage/Backends/InMemoryBlobStore.cs b/Engine/Source/Programs/Shared/EpicGames.Horde/Storage/Backends/InMemoryBlobStore.cs
--- a/Engine/Source/Programs/Shared/EpicGames.Horde/Storage/Backends/InMemoryBlobStore.cs
+++ b/Engine/Source/Programs/Shared/EpicGames.Horde/Storage/Backends/InMemoryBlobStore.cs
@@ -6,6 +6,7 @@ using System.Collections.Concurrent;
 using System.Collections.Generic;
 using System.Threading;
 using System.Threading.Tasks;
+using EpicGames.Core;
 
 namespace EpicGames.Horde.Storage.Backends
 {
@@ -39,9 +40,9 @@ namespace EpicGames.Horde.Storage.Backends
 		}
 
 		/// <inheritdoc/>
-		public Task<BlobId> WriteBlobAsync(ReadOnlySequence<byte> data, IReadOnlyList<BlobId> references, RefName hintRefName = default, CancellationToken cancellationToken = default)
+		public Task<BlobId> WriteBlobAsync(ReadOnlySequence<byte> data, IReadOnlyList<BlobId> references, Utf8String prefix = default, CancellationToken cancellationToken = default)
 		{
-			BlobId blobId = BlobId.Create(ServerId.Empty, hintRefName);
+			BlobId blobId = BlobId.Create(ServerId.Empty, prefix);
 
 			IBlob blob = Blob.FromMemory(blobId, data.ToArray(), references);
 			_blobs[blobId] = blob;
@@ -86,7 +87,7 @@ namespace EpicGames.Horde.Storage.Backends
 		/// <inheritdoc/>
 		public async Task<BlobId> WriteRefAsync(RefName name, ReadOnlySequence<byte> data, IReadOnlyList<BlobId> references, CancellationToken cancellationToken)
 		{
-			BlobId blobId = await WriteBlobAsync(data, references, name, cancellationToken);
+			BlobId blobId = await WriteBlobAsync(data, references, name.Text, cancellationToken);
 			_refs[name] = blobId;
 			return blobId;
 		}
diff --git a/Engine/Source/Programs/Shared/EpicGames.Horde/Storage/BlobId.cs b/Engine/Source/Programs/Shared/EpicGames.Horde/Storage/BlobId.cs
--- a/Engine/Source/Programs/Shared/EpicGames.Horde/Storage/BlobId.cs
+++ b/Engine/Source/Programs/Shared/EpicGames.Horde/Storage/BlobId.cs
@@ -65,16 +65,19 @@ namespace EpicGames.Horde.Storage
 		/// Create a unique content id, optionally including a ref name
 		/// </summary>
 		/// <param name="serverId">The server id</param>
-		/// <param name="hintRefName">Hint for the ref containing this blob</param>
+		/// <param name="prefix">Prefix for blob names. Follows the same restrictions as for content ids.</param>
 		/// <returns>New content id</returns>
-		public static BlobId Create(ServerId serverId, RefName hintRefName = default)
+		public static BlobId Create(ServerId serverId, Utf8String prefix = default)
 		{
-			Utf8String prefix = hintRefName.Text;
 			if (prefix.Length == 0)
 			{
 				DateTime now = DateTime.UtcNow.Date;
 				prefix = $"_by_date_/{now.Year}-{now.Month:D2}/{now.Day:D2}";
 			}
+			else
+			{
+				ContentId.ValidateArgument(nameof(prefix), prefix);
+			}
 
 			byte[] buffer;
 			Span<byte> span;
diff --git a/Engine/Source/Programs/Shared/EpicGames.Horde/Storage/Bundles/BundleStore.cs b/Engine/Source/Programs/Shared/EpicGames.Horde/Storage/Bundles/BundleStore.cs
--- a/Engine/Source/Programs/Shared/EpicGames.Horde/Storage/Bundles/BundleStore.cs
+++ b/Engine/Source/Programs/Shared/EpicGames.Horde/Storage/Bundles/BundleStore.cs
@@ -10,6 +10,7 @@ using System.Linq;
 using System.Threading;
 using System.Threading.Tasks;
 using System.Xml.Linq;
+using Blake3;
 using EpicGames.Core;
 using K4os.Compression.LZ4;
 using Microsoft.Extensions.Caching.Memory;
@@ -216,7 +217,7 @@ namespace EpicGames.Horde.Storage.Bundles
 			readonly BundleStore _owner;
 			readonly BundleContext _context;
 			readonly BundleWriter? _parent;
-			readonly RefName _refName;
+			readonly Utf8String _prefix;
 
 			public readonly ConcurrentDictionary<IoHash, NodeInfo> HashToNode = new ConcurrentDictionary<IoHash, NodeInfo>(); // TODO: this needs to include some additional state from external sources.
 
@@ -224,7 +225,7 @@ namespace EpicGames.Horde.Storage.Bundles
 			readonly List<BundleWriter> _children = new List<BundleWriter>();
 
 			// Task which includes all active writes, and returns the last blob
-			Task<BlobId> _flushTask = Task.FromResult(BlobId.Empty);
+			Task<NodeInfo> _flushTask = Task.FromResult<NodeInfo>(null!);
 
 			// Nodes which have been queued to be written, but which are not yet part of any active write
 			readonly List<NodeInfo> _queueNodes = new List<NodeInfo>();
@@ -238,18 +239,18 @@ namespace EpicGames.Horde.Storage.Bundles
 			/// <summary>
 			/// Constructor
 			/// </summary>
-			public BundleWriter(BundleStore owner, BundleContext context, BundleWriter? parent, RefName refName)
+			public BundleWriter(BundleStore owner, BundleContext context, BundleWriter? parent, Utf8String prefix)
 			{
 				_owner = owner;
 				_context = context;
 				_parent = parent;
-				_refName = refName;
+				_prefix = prefix;
 			}
 
 			/// <inheritdoc/>
 			public ITreeWriter CreateChildWriter()
 			{
-				BundleWriter writer = new BundleWriter(_owner, _context, this, _refName);
+				BundleWriter writer = new BundleWriter(_owner, _context, this, _prefix);
 				lock (_lockObject)
 				{
 					_children.Add(writer);
@@ -274,21 +275,26 @@ namespace EpicGames.Horde.Storage.Bundles
 					node = HashToNode.GetOrAdd(hash, newNode);
 					if (node == newNode)
 					{
-						lock (_lockObject)
-						{
-							// Add these nodes to the queue for writing
-							AddToQueue(node);
-
-							// Update the list of nodes which are ready to be written. We need to make sure that any child writers have flushed before
-							// we can reference the blobs containing their nodes.
-							UpdateReady();
-						}
+						WriteNode(node);
 					}
 				}
 
 				return Task.FromResult<ITreeBlobRef>(node);
 			}
 
+			void WriteNode(NodeInfo node)
+			{
+				lock (_lockObject)
+				{
+					// Add these nodes to the queue for writing
+					AddToQueue(node);
+
+					// Update the list of nodes which are ready to be written. We need to make sure that any child writers have flushed before
+					// we can reference the blobs containing their nodes.
+					UpdateReady();
+				}
+			}
+
 			void AddToQueue(NodeInfo root)
 			{
 				if (root.State is InMemoryNodeState state && state._writer == null)
@@ -372,19 +378,32 @@ namespace EpicGames.Horde.Storage.Bundles
 			}
 
 			/// <inheritdoc/>
-			public async Task<BlobId> FlushAsync(CancellationToken cancellationToken)
+			public async Task WriteRefAsync(RefName refName, ITreeBlobRef root, CancellationToken cancellationToken)
 			{
-				if (_refName == RefName.Empty)
+				if (_parent != null)
 				{
 					throw new InvalidOperationException("Flushing a child writer is not permitted.");
 				}
 
-				BlobId blobId = await FlushInternalAsync(cancellationToken);
-				await _owner._blobStore.WriteRefTargetAsync(_refName, blobId, cancellationToken);
-				return blobId;
+				// Make sure the last written node is the desired root. If not, we'll need to write it to a new blob so it can be last.
+				NodeInfo? last = await _flushTask;
+				if (last != root)
+				{
+					if (_queueNodes.Count == 0 || _queueNodes[^1] != root)
+					{
+						ITreeBlob blob = await root.GetTargetAsync(cancellationToken);
+						InMemoryNodeState state = new InMemoryNodeState(blob.Data, blob.Refs.ConvertAll(x => (NodeInfo)x));
+						WriteNode(new NodeInfo(_owner, root.Hash, state));
+					}
+					last = await FlushInternalAsync(cancellationToken);
+				}
+
+				// Write a reference to the blob containing this node
+				ImportedNodeState importedState = (ImportedNodeState)last.State;
+				await _owner._blobStore.WriteRefTargetAsync(refName, importedState.BundleInfo.BlobId, cancellationToken);
 			}
 
-			async Task<BlobId> FlushInternalAsync(CancellationToken cancellationToken)
+			async Task<NodeInfo> FlushInternalAsync(CancellationToken cancellationToken)
 			{
 				// Get all the child writers and flush them
 				BundleWriter[] children;
@@ -420,14 +439,14 @@ namespace EpicGames.Horde.Storage.Bundles
 				return await _flushTask;
 			}
 
-			async Task<BlobId> WriteNodesAsync(NodeInfo[] writeNodes, Task prevWriteTask, CancellationToken cancellationToken)
+			async Task<NodeInfo> WriteNodesAsync(NodeInfo[] writeNodes, Task prevWriteTask, CancellationToken cancellationToken)
 			{
 				// Create the bundle
 				Bundle bundle = CreateBundle(writeNodes);
 				BundleHeader header = bundle.Header;
 
 				// Write it to storage
-				BlobId blobId = await _owner._blobStore.WriteBlobAsync(bundle.AsSequence(), bundle.Header.Imports.Select(x => x.BlobId).ToList(), _refName, cancellationToken);
+				BlobId blobId = await _owner._blobStore.WriteBlobAsync(bundle.AsSequence(), bundle.Header.Imports.Select(x => x.BlobId).ToList(), _prefix, cancellationToken);
 				string cacheKey = GetBundleCacheKey(blobId);
 				_owner.AddBundleToCache(cacheKey, bundle);
 
@@ -476,7 +495,7 @@ namespace EpicGames.Horde.Storage.Bundles
 
 				// Wait for other writes to finish
 				await prevWriteTask;
-				return blobId;
+				return writeNodes[^1];
 			}
 
 			/// <summary>
@@ -682,7 +701,7 @@ namespace EpicGames.Horde.Storage.Bundles
 		}
 
 		/// <inheritdoc/>
-		public ITreeWriter CreateTreeWriter(RefName name) => new BundleWriter(this, new BundleContext(), null, name);
+		public ITreeWriter CreateTreeWriter(Utf8String prefix = default) => new BundleWriter(this, new BundleContext(), null, prefix);
 
 		/// <inheritdoc/>
 		public Task DeleteTreeAsync(RefName name, CancellationToken cancellationToken) => _blobStore.DeleteRefAsync(name, cancellationToken);
diff --git a/Engine/Source/Programs/Shared/EpicGames.Horde/Storage/IBlobStore.cs b/Engine/Source/Programs/Shared/EpicGames.Horde/Storage/IBlobStore.cs
--- a/Engine/Source/Programs/Shared/EpicGames.Horde/Storage/IBlobStore.cs
+++ b/Engine/Source/Programs/Shared/EpicGames.Horde/Storage/IBlobStore.cs
@@ -75,9 +75,9 @@ namespace EpicGames.Horde.Storage
 		/// <param name="data">Payload for the object</param>
 		/// <param name="references">Object references</param>
 		/// <param name="cancellationToken">Cancellation token for the operation</param>
-		/// <param name="hintRefName">Name of a ref that this blob is being written for. While the returned BlobId is guaranteed to be unique, this name can be used as a prefix to aid debugging.</param>
+		/// <param name="prefix">Prefix for blob names. While the returned BlobId is guaranteed to be unique, this name can be used as a prefix to aid debugging.</param>
 		/// <returns>Unique identifier for the blob</returns>
-		Task<BlobId> WriteBlobAsync(ReadOnlySequence<byte> data, IReadOnlyList<BlobId> references, RefName hintRefName = default, CancellationToken cancellationToken = default);
+		Task<BlobId> WriteBlobAsync(ReadOnlySequence<byte> data, IReadOnlyList<BlobId> references, Utf8String prefix = default, CancellationToken cancellationToken = default);
 
 		#endregion
 
@@ -154,7 +154,7 @@ namespace EpicGames.Horde.Storage
 			public Task<IBlob> ReadBlobAsync(BlobId id, CancellationToken cancellationToken = default) => _inner.ReadBlobAsync(id, cancellationToken);
 
 			/// <inheritdoc/>
-			public Task<BlobId> WriteBlobAsync(ReadOnlySequence<byte> data, IReadOnlyList<BlobId> references, RefName hintRefName = default, CancellationToken cancellationToken = default) => _inner.WriteBlobAsync(data, references, hintRefName, cancellationToken);
+			public Task<BlobId> WriteBlobAsync(ReadOnlySequence<byte> data, IReadOnlyList<BlobId> references, Utf8String prefix = default, CancellationToken cancellationToken = default) => _inner.WriteBlobAsync(data, references, prefix, cancellationToken);
 
 			#endregion
 
diff --git a/Engine/Source/Programs/Shared/EpicGames.Horde/Storage/ITreeStore.cs b/Engine/Source/Programs/Shared/EpicGames.Horde/Storage/ITreeStore.cs
--- a/Engine/Source/Programs/Shared/EpicGames.Horde/Storage/ITreeStore.cs
+++ b/Engine/Source/Programs/Shared/EpicGames.Horde/Storage/ITreeStore.cs
@@ -51,8 +51,9 @@ namespace EpicGames.Horde.Storage
 		/// <summary>
 		/// Creates a new context for reading trees.
 		/// </summary>
+		/// <param name="prefix">Prefix for blobs. See <see cref="IBlobStore.WriteBlobAsync(ReadOnlySequence{Byte}, IReadOnlyList{BlobId}, Utf8String, CancellationToken)"/></param>
 		/// <returns>New context instance</returns>
-		ITreeWriter CreateTreeWriter(RefName name);
+		ITreeWriter CreateTreeWriter(Utf8String prefix = default);
 
 		/// <summary>
 		/// Deletes a tree with the given name
@@ -108,11 +109,13 @@ namespace EpicGames.Horde.Storage
 		Task<ITreeBlobRef> WriteNodeAsync(ReadOnlySequence<byte> data, IReadOnlyList<ITreeBlobRef> refs, CancellationToken cancellationToken = default);
 
 		/// <summary>
-		/// Flushes the current state and writes the ref, with the last written node as the tree root. Throws an exception if called for a child writer.
+		/// Flushes the current state using a new ref name.
 		/// </summary>
-		/// <param name="cancellationToken"></param>
+		/// <param name="name">Name of the ref</param>
+		/// <param name="root">Root reference</param>
+		/// <param name="cancellationToken">Cancellation token for the operation</param>
 		/// <returns></returns>
-		Task<BlobId> FlushAsync(CancellationToken cancellationToken = default);
+		Task WriteRefAsync(RefName name, ITreeBlobRef root, CancellationToken cancellationToken = default);
 	}
 
 	/// <summary>
@@ -156,7 +159,7 @@ namespace EpicGames.Horde.Storage
 			}
 
 			/// <inheritdoc/>
-			public ITreeWriter CreateTreeWriter(RefName name) => _inner.CreateTreeWriter(name);
+			public ITreeWriter CreateTreeWriter(Utf8String prefix) => _inner.CreateTreeWriter(prefix);
 
 			/// <inheritdoc/>
 			public Task DeleteTreeAsync(RefName name, CancellationToken cancellationToken = default) => _inner.DeleteTreeAsync(name, cancellationToken);
@@ -181,5 +184,20 @@ namespace EpicGames.Horde.Storage
 		{
 			return new TypedTreeStore<T>(store);
 		}
+
+		/// <summary>
+		/// Writes a ref using the given root data
+		/// </summary>
+		/// <param name="writer">Writer to modify</param>
+		/// <param name="name">Name of the ref</param>
+		/// <param name="data">Data for the node</param>
+		/// <param name="refs">References to other nodes</param>
+		/// <param name="cancellationToken">Cancellation token for the operation</param>
+		/// <returns>Reference to the node that was added</returns>
+		public static async Task WriteRefAsync(this ITreeWriter writer, RefName name, ReadOnlySequence<byte> data, IReadOnlyList<ITreeBlobRef> refs, CancellationToken cancellationToken = default)
+		{
+			ITreeBlobRef root = await writer.WriteNodeAsync(data, refs, cancellationToken);
+			await writer.WriteRefAsync(name, root, cancellationToken);
+		}
 	}
 }
diff --git a/Engine/Source/Programs/Shared/EpicGames.Horde/Storage/TreeNode.cs b/Engine/Source/Programs/Shared/EpicGames.Horde/Storage/TreeNode.cs
--- a/Engine/Source/Programs/Shared/EpicGames.Horde/Storage/TreeNode.cs
+++ b/Engine/Source/Programs/Shared/EpicGames.Horde/Storage/TreeNode.cs
@@ -147,10 +147,23 @@ namespace EpicGames.Horde.Storage
 		/// <param name="writer">Writer to output the nodes to</param>
 		/// <param name="node">Root node to serialize</param>
 		/// <param name="cancellationToken">Cancellation token for the operation</param>
-		public static async Task WriteNodeAsync(this ITreeWriter writer, TreeNode node, CancellationToken cancellationToken = default)
+		public static async Task<ITreeBlobRef> WriteNodeAsync(this ITreeWriter writer, TreeNode node, CancellationToken cancellationToken = default)
 		{
 			ITreeBlob blob = await node.SerializeAsync(writer, cancellationToken);
-			await writer.WriteNodeAsync(blob.Data, blob.Refs, cancellationToken);
+			return await writer.WriteNodeAsync(blob.Data, blob.Refs, cancellationToken);
+		}
+
+		/// <summary>
+		/// Flushes a tree to storage using the given root node
+		/// </summary>
+		/// <param name="writer">Writer to output the nodes to</param>
+		/// <param name="name">Name of the ref to write</param>
+		/// <param name="node">Root node to serialize</param>
+		/// <param name="cancellationToken">Cancellation token for the operation</param>
+		public static async Task WriteRefAsync(this ITreeWriter writer, RefName name, TreeNode node, CancellationToken cancellationToken = default)
+		{
+			ITreeBlobRef root = await writer.WriteNodeAsync(node, cancellationToken);
+			await writer.WriteRefAsync(name, root, cancellationToken);
 		}
 
 		/// <inheritdoc/>
@@ -178,9 +191,8 @@ namespace EpicGames.Horde.Storage
 		/// <inheritdoc/>
 		public static async Task WriteTreeAsync(this ITreeStore store, RefName name, TreeNode root, CancellationToken cancellationToken = default)
 		{
-			ITreeWriter writer = store.CreateTreeWriter(name);
-			await writer.WriteNodeAsync(root, cancellationToken);
-			await writer.FlushAsync(cancellationToken);
+			ITreeWriter writer = store.CreateTreeWriter(name.Text);
+			await writer.WriteRefAsync(name, root, cancellationToken);
 		}
 	}
 }
