Horde: Fix tests.

diff --git a/Engine/Source/Programs/Horde/Horde.Agent/Commands/Bundles/CreateCommand.cs b/Engine/Source/Programs/Horde/Horde.Agent/Commands/Bundles/CreateCommand.cs
--- a/Engine/Source/Programs/Horde/Horde.Agent/Commands/Bundles/CreateCommand.cs
+++ b/Engine/Source/Programs/Horde/Horde.Agent/Commands/Bundles/CreateCommand.cs
@@ -54,9 +54,9 @@ namespace Horde.Agent.Commands.Bundles
 				return Blob.Deserialize(id, bytes);
 			}
 
-			public async Task<BlobId> WriteBlobAsync(ReadOnlySequence<byte> data, IReadOnlyList<BlobId> references, RefName hintRefName = default, CancellationToken cancellationToken = default)
+			public async Task<BlobId> WriteBlobAsync(ReadOnlySequence<byte> data, IReadOnlyList<BlobId> references, Utf8String prefix = default, CancellationToken cancellationToken = default)
 			{
-				BlobId id = BlobId.Create(ServerId.Empty, hintRefName);
+				BlobId id = BlobId.Create(ServerId.Empty, prefix);
 				FileReference file = GetBlobFile(id);
 				DirectoryReference.CreateDirectory(file.Directory);
 				_logger.LogInformation("Writing {File}", file);
@@ -99,7 +99,7 @@ namespace Horde.Agent.Commands.Bundles
 
 			public async Task<BlobId> WriteRefAsync(RefName name, ReadOnlySequence<byte> data, IReadOnlyList<BlobId> references, CancellationToken cancellationToken = default)
 			{
-				BlobId blobId = await WriteBlobAsync(data, references, name, cancellationToken);
+				BlobId blobId = await WriteBlobAsync(data, references, name.Text, cancellationToken);
 				await WriteRefTargetAsync(name, blobId, cancellationToken);
 				return blobId;
 			}
@@ -147,13 +147,12 @@ namespace Horde.Agent.Commands.Bundles
 		public override async Task<int> ExecuteAsync(ILogger logger)
 		{
 			using ITreeStore store = CreateTreeStore(logger);
-			ITreeWriter writer = store.CreateTreeWriter(RefName);
+			ITreeWriter writer = store.CreateTreeWriter(RefName.Text);
 
 			DirectoryNode node = new DirectoryNode(DirectoryFlags.None);
 			await node.CopyFromDirectoryAsync(InputDir.ToDirectoryInfo(), new ChunkingOptions(), writer, logger, CancellationToken.None);
 
-			await writer.WriteNodeAsync(node);
-			await writer.FlushAsync();
+			await writer.WriteRefAsync(RefName, node);
 			return 0;
 		}
 	}
diff --git a/Engine/Source/Programs/Horde/Horde.Build.Tests/BlobStoreTests.cs b/Engine/Source/Programs/Horde/Horde.Build.Tests/BlobStoreTests.cs
--- a/Engine/Source/Programs/Horde/Horde.Build.Tests/BlobStoreTests.cs
+++ b/Engine/Source/Programs/Horde/Horde.Build.Tests/BlobStoreTests.cs
@@ -39,7 +39,7 @@ namespace Horde.Build.Tests
 	{
 		IBlobStore CreateBlobStore()
 		{
-			return new BasicBlobStore(MongoService, new TransientStorageBackend(), ServiceProvider.GetRequiredService<IMemoryCache>());
+			return new BasicBlobStore(MongoService, new TransientStorageBackend(), ServiceProvider.GetRequiredService<IMemoryCache>(), ServiceProvider.GetRequiredService<ILogger<BasicBlobStore>>());
 		}
 
 		static byte[] CreateTestData(int length, int seed)
diff --git a/Engine/Source/Programs/Horde/Horde.Build.Tests/TestSetup.cs b/Engine/Source/Programs/Horde/Horde.Build.Tests/TestSetup.cs
--- a/Engine/Source/Programs/Horde/Horde.Build.Tests/TestSetup.cs
+++ b/Engine/Source/Programs/Horde/Horde.Build.Tests/TestSetup.cs
@@ -237,7 +237,7 @@ namespace Horde.Build.Tests
 
 			services.AddSingleton<ITreeStore<CommitService>>(sp =>
 			{
-				IBlobStore blobStore = new BasicBlobStore(sp.GetRequiredService<MongoService>(), new TransientStorageBackend(), sp.GetRequiredService<IMemoryCache>());
+				IBlobStore blobStore = new BasicBlobStore(sp.GetRequiredService<MongoService>(), new TransientStorageBackend(), sp.GetRequiredService<IMemoryCache>(), sp.GetRequiredService<ILogger<BasicBlobStore>>());
 				return new BundleStore(blobStore, new BundleOptions()).ForType<CommitService>();
 			});
 
diff --git a/Engine/Source/Programs/Shared/EpicGames.Horde.Tests/BundleStoreTests.cs b/Engine/Source/Programs/Shared/EpicGames.Horde.Tests/BundleStoreTests.cs
--- a/Engine/Source/Programs/Shared/EpicGames.Horde.Tests/BundleStoreTests.cs
+++ b/Engine/Source/Programs/Shared/EpicGames.Horde.Tests/BundleStoreTests.cs
@@ -47,7 +47,7 @@ namespace EpicGames.Horde.Tests
 			// Generate a tree
 			using (BundleStore store = new BundleStore(blobStore, options))
 			{
-				ITreeWriter writer = store.CreateTreeWriter(new RefName("test"));
+				ITreeWriter writer = store.CreateTreeWriter("test");
 
 				ITreeBlobRef node1 = await writer.WriteNodeAsync(new ReadOnlySequence<byte>(new byte[] { 1 }), Array.Empty<ITreeBlobRef>(), CancellationToken.None);
 				ITreeBlobRef node2 = await writer.WriteNodeAsync(new ReadOnlySequence<byte>(new byte[] { 2 }), new[] { node1 }, CancellationToken.None);
@@ -55,8 +55,7 @@ namespace EpicGames.Horde.Tests
 				ITreeBlobRef node4 = await writer.WriteNodeAsync(new ReadOnlySequence<byte>(new byte[] { 4 }), Array.Empty<ITreeBlobRef>(), CancellationToken.None);
 
 				ITreeBlob root = TreeBlob.Create(new ReadOnlySequence<byte>(new byte[] { 5 }), new[] { node4, node3 });
-				await writer.WriteNodeAsync(root.Data, root.Refs);
-				await writer.FlushAsync();
+				await writer.WriteRefAsync(new RefName("test"), root.Data, root.Refs);
 
 				await CheckTree(root);
 			}
@@ -148,7 +147,7 @@ namespace EpicGames.Horde.Tests
 			// Generate a tree
 			using (BundleStore store = new BundleStore(blobStore, new BundleOptions()))
 			{
-				ITreeWriter writer = store.CreateTreeWriter(new RefName("test"));
+				ITreeWriter writer = store.CreateTreeWriter();
 
 				DirectoryNode root = new DirectoryNode(DirectoryFlags.None);
 				DirectoryNode hello = root.AddDirectory("hello");
@@ -156,8 +155,7 @@ namespace EpicGames.Horde.Tests
 				FileEntry world = hello.AddFile("world", FileEntryFlags.None);
 				await world.AppendAsync(Encoding.UTF8.GetBytes("world"), new ChunkingOptions(), writer, CancellationToken.None);
 
-				await writer.WriteNodeAsync(root);
-				await writer.FlushAsync();
+				await writer.WriteRefAsync(new RefName("test"), root);
 
 				await CheckFileTreeAsync(root);
 			}
@@ -190,7 +188,7 @@ namespace EpicGames.Horde.Tests
 		public async Task ChunkingTests()
 		{
 			using BundleStore store = new BundleStore(new InMemoryBlobStore(), new BundleOptions());
-			ITreeWriter writer = store.CreateTreeWriter(new RefName("test"));
+			ITreeWriter writer = store.CreateTreeWriter();
 
 			byte[] chunk = RandomNumberGenerator.GetBytes(256);
 
@@ -231,8 +229,7 @@ namespace EpicGames.Horde.Tests
 
 			using (BundleStore store = new BundleStore(blobStore, new BundleOptions()))
 			{
-				RefName refName = new RefName("test");
-				ITreeWriter writer = store.CreateTreeWriter(refName);
+				ITreeWriter writer = store.CreateTreeWriter();
 
 				ITreeWriter childWriter1 = writer.CreateChildWriter();
 				ITreeBlobRef childNode1 = await childWriter1.WriteNodeAsync(new ReadOnlySequence<byte>(new byte[] { 4, 5, 6 }), Array.Empty<ITreeBlobRef>());
@@ -240,8 +237,8 @@ namespace EpicGames.Horde.Tests
 				ITreeWriter childWriter2 = writer.CreateChildWriter();
 				ITreeBlobRef childNode2 = await childWriter2.WriteNodeAsync(new ReadOnlySequence<byte>(new byte[] { 7, 8, 9 }), Array.Empty<ITreeBlobRef>());
 
-				await writer.WriteNodeAsync(new ReadOnlySequence<byte>(new byte[] { 1, 2, 3 }), new[] { childNode1, childNode2 });
-				await writer.FlushAsync();
+				RefName refName = new RefName("test");
+				await writer.WriteRefAsync(refName, new ReadOnlySequence<byte>(new byte[] { 1, 2, 3 }), new[] { childNode1, childNode2 });
 
 				Assert.AreEqual(3, blobStore.Blobs.Count);
 				Assert.AreEqual(1, blobStore.Refs.Count);
@@ -281,7 +278,7 @@ namespace EpicGames.Horde.Tests
 			// Generate a tree
 			using (BundleStore store = new BundleStore(blobStore, new BundleOptions { MaxBlobSize = 1024 }))
 			{
-				ITreeWriter writer = store.CreateTreeWriter(new RefName("test"));
+				ITreeWriter writer = store.CreateTreeWriter();
 
 				DirectoryNode root = new DirectoryNode(DirectoryFlags.None);
 
@@ -293,8 +290,7 @@ namespace EpicGames.Horde.Tests
 				FileEntry file = root.AddFile("test", FileEntryFlags.None);
 				await file.AppendAsync(data, options, writer, CancellationToken.None);
 
-				await writer.WriteNodeAsync(root);
-				await writer.FlushAsync();
+				await writer.WriteRefAsync(new RefName("test"), root);
 
 				await CheckLargeFileTreeAsync(root, data);
 			}
diff --git a/Engine/Source/Programs/Shared/EpicGames.Horde.Tests/BundleTests.cs b/Engine/Source/Programs/Shared/EpicGames.Horde.Tests/BundleTests.cs
--- a/Engine/Source/Programs/Shared/EpicGames.Horde.Tests/BundleTests.cs
+++ b/Engine/Source/Programs/Shared/EpicGames.Horde.Tests/BundleTests.cs
@@ -69,7 +69,8 @@ namespace EpicGames.Horde.Tests
 		[TestMethod]
 		public async Task BasicChunkingTests()
 		{
-			ITreeWriter writer = _treeStore.CreateTreeWriter(new RefName("test"));
+			RefName refName = new RefName("test");
+			ITreeWriter writer = _treeStore.CreateTreeWriter(refName.Text);
 
 			ChunkingOptions options = new ChunkingOptions();
 			options.LeafOptions = new ChunkingOptionsForNodeType(8, 8, 8);
@@ -346,7 +347,7 @@ namespace EpicGames.Horde.Tests
 		[TestMethod]
 		public async Task CoreAppendTest()
 		{
-			ITreeWriter writer = _treeStore.CreateTreeWriter(new RefName("test"));
+			ITreeWriter writer = _treeStore.CreateTreeWriter();
 
 			byte[] data = new byte[4096];
 			new Random(0).NextBytes(data);
@@ -386,7 +387,7 @@ namespace EpicGames.Horde.Tests
 
 		async Task ChunkingTests(ChunkingOptions options)
 		{
-			ITreeWriter writer = _treeStore.CreateTreeWriter(new RefName("test"));
+			ITreeWriter writer = _treeStore.CreateTreeWriter();
 
 			byte[] data = new byte[4096];
 			new Random(0).NextBytes(data);
@@ -448,7 +449,7 @@ namespace EpicGames.Horde.Tests
 			using BundleStore treeStore = new BundleStore(_blobStore, options);
 
 			RefName refName = new RefName("ref");
-			ITreeWriter writer = treeStore.CreateTreeWriter(refName);
+			ITreeWriter writer = treeStore.CreateTreeWriter(refName.Text);
 
 			DirectoryNode root = new DirectoryNode(DirectoryFlags.None);
 
diff --git a/Engine/Source/Programs/Shared/EpicGames.Horde.Tests/ChunkingTests.cs b/Engine/Source/Programs/Shared/EpicGames.Horde.Tests/ChunkingTests.cs
--- a/Engine/Source/Programs/Shared/EpicGames.Horde.Tests/ChunkingTests.cs
+++ b/Engine/Source/Programs/Shared/EpicGames.Horde.Tests/ChunkingTests.cs
@@ -50,7 +50,7 @@ namespace EpicGames.Horde.Tests
 		public async Task BasicChunkingTests()
 		{
 			using BundleStore store = new BundleStore(new InMemoryBlobStore(), new BundleOptions());
-			ITreeWriter writer = store.CreateTreeWriter(new RefName("test"));
+			ITreeWriter writer = store.CreateTreeWriter("test");
 
 			ChunkingOptions options = new ChunkingOptions();
 			options.LeafOptions = new ChunkingOptionsForNodeType(8, 8, 8);
@@ -106,7 +106,7 @@ namespace EpicGames.Horde.Tests
 		static async Task TestChunkingAsync(ChunkingOptions options)
 		{
 			using BundleStore store = new BundleStore(new InMemoryBlobStore(), new BundleOptions());
-			ITreeWriter writer = store.CreateTreeWriter(new RefName("test"));
+			ITreeWriter writer = store.CreateTreeWriter();
 
 			byte[] data = new byte[4096];
 			new Random(0).NextBytes(data);
