Optimize Horde printing a shelved CL for Preflight by issuing one call to p4 print

Using the wildcard support for p4 print we can print the entire contents of a shelf with one call instead of a call for each file.

Splitting TryPrint*Async into a method that assumes one file in and one result out, and another method that can return multiple results.

#rnx

diff --git a/Engine/Source/Programs/Shared/EpicGames.Perforce.Managed/ManagedWorkspace.cs b/Engine/Source/Programs/Shared/EpicGames.Perforce.Managed/ManagedWorkspace.cs
--- a/Engine/Source/Programs/Shared/EpicGames.Perforce.Managed/ManagedWorkspace.cs
+++ b/Engine/Source/Programs/Shared/EpicGames.Perforce.Managed/ManagedWorkspace.cs
@@ -860,15 +860,15 @@ namespace EpicGames.Perforce.Managed
 
 			// query the location of each file
 			List<PerforceResponse<WhereRecord>> whereResponseList = await perforce.TryWhereAsync(lastRecord.Files.Select(x => x.DepotFile).ToArray(), cancellationToken).ToListAsync(cancellationToken);
-			List<WhereRecord> whereRecords = whereResponseList.Where(x => x.Succeeded).Select(x => x.Data).ToList();
+			Dictionary<string, WhereRecord> whereRecords = whereResponseList.Where(x => x.Succeeded).Select(x => x.Data).ToDictionary(x => x.DepotFile, x => x, StringComparer.OrdinalIgnoreCase);
 
 			// parse out all the list of deleted and modified files
 			List<WhereRecord> deleteFiles = new List<WhereRecord>();
 			List<WhereRecord> writeFiles = new List<WhereRecord>();
 			foreach (DescribeFileRecord fileRecord in lastRecord.Files)
 			{
-				WhereRecord? whereRecord = whereRecords.FirstOrDefault(x => x.DepotFile.Equals(fileRecord.DepotFile, StringComparison.OrdinalIgnoreCase));
-				if (whereRecord == null)
+				WhereRecord? whereRecord = null;
+				if (whereRecords.TryGetValue(fileRecord.DepotFile, out whereRecord) == false)
 				{
 					_logger.LogInformation("Unable to get location of {File} in current workspace; ignoring.", fileRecord.DepotFile);
 					continue;
@@ -892,6 +892,11 @@ namespace EpicGames.Perforce.Managed
 				}
 			}
 
+			if (!_createdClients.TryGetValue(perforce.Settings.ClientName!, out ClientRecord? perforceClient))
+			{
+				throw new Exception($"Unknown client {perforce.Settings.ClientName}");
+			}
+
 			// Add all the files to be written to the workspace with invalid metadata. This will ensure they're removed on next clean.
 			if (writeFiles.Count > 0)
 			{
@@ -915,19 +920,12 @@ namespace EpicGames.Perforce.Managed
 				}
 			}
 
-			// Add all the new files
-			foreach (WhereRecord writeFile in writeFiles)
+			// Use common paths with wild cards speed up the print operation with one call instead of many calls to print.
+			_logger.LogInformation("Writing files from shelved changelist", unshelveChangelist);
+			PerforceResponseList<PrintRecord> printResponse = await perforce.TryPrintAsync($"{perforceClient.Root}{Path.DirectorySeparatorChar}...", $"//{perforceClient.Name}/...@={unshelveChangelist}", cancellationToken);
+			if (!printResponse.Succeeded)
 			{
-				string localPath = writeFile.Path;
-				_logger.LogInformation("  Writing {LocalPath}", localPath);
-
-				Directory.CreateDirectory(Path.GetDirectoryName(localPath));
-
-				PerforceResponse printResponse = await perforce.TryPrintAsync(localPath, $"{writeFile}@={unshelveChangelist}", cancellationToken);
-				if (!printResponse.Succeeded)
-				{
-					_logger.LogWarning("Unable to print {LocalPath}: {Error}", localPath, printResponse.ToString());
-				}
+				_logger.LogWarning("Unable to print shelved changelist: {Error}", printResponse.ToString());
 			}
 
 			_logger.LogInformation("Completed in {TimeSeconds}s", $"{timer.Elapsed.TotalSeconds:0.0}");
diff --git a/Engine/Source/Programs/Shared/EpicGames.Perforce/PerforceConnection.cs b/Engine/Source/Programs/Shared/EpicGames.Perforce/PerforceConnection.cs
--- a/Engine/Source/Programs/Shared/EpicGames.Perforce/PerforceConnection.cs
+++ b/Engine/Source/Programs/Shared/EpicGames.Perforce/PerforceConnection.cs
@@ -2595,7 +2595,7 @@ namespace EpicGames.Perforce
 		/// <returns>Response from the server</returns>
 		public static async Task<PrintRecord> PrintAsync(this IPerforceConnection connection, string outputFile, string fileSpec, CancellationToken cancellationToken = default)
 		{
-			return (await TryPrintAsync(connection, outputFile, fileSpec, cancellationToken)).Data;
+			return (await TryPrintSingleInternalAsync<PrintRecord>(connection, outputFile, fileSpec, cancellationToken)).Data;
 		}
 
 		/// <summary>
@@ -2606,9 +2606,9 @@ namespace EpicGames.Perforce
 		/// <param name="fileSpec">Specification for the files to print</param>
 		/// <param name="cancellationToken">Token used to cancel the operation</param>
 		/// <returns>Response from the server</returns>
-		public static Task<PerforceResponse<PrintRecord>> TryPrintAsync(this IPerforceConnection connection, string outputFile, string fileSpec, CancellationToken cancellationToken = default)
+		public static Task<PerforceResponseList<PrintRecord>> TryPrintAsync(this IPerforceConnection connection, string outputFile, string fileSpec, CancellationToken cancellationToken = default)
 		{
-			return TryPrintInternalAsync<PrintRecord>(connection, outputFile, fileSpec, cancellationToken);
+			return TryPrintInternalAsync(connection, outputFile, fileSpec, cancellationToken);
 		}
 
 		/// <summary>
@@ -2623,7 +2623,7 @@ namespace EpicGames.Perforce
 			string tempFile = Path.GetTempFileName();
 			try
 			{
-				PerforceResponse<PrintRecord<byte[]>> record = await TryPrintInternalAsync<PrintRecord<byte[]>>(connection, tempFile, fileSpec, cancellationToken);
+				PerforceResponse<PrintRecord<byte[]>> record = await TryPrintSingleInternalAsync<PrintRecord<byte[]>>(connection, tempFile, fileSpec, cancellationToken);
 				if (record.Succeeded)
 				{
 					record.Data.Contents = await File.ReadAllBytesAsync(tempFile, cancellationToken);
@@ -2648,7 +2648,7 @@ namespace EpicGames.Perforce
 			string tempFile = Path.GetTempFileName();
 			try
 			{
-				PerforceResponse<PrintRecord<string[]>> record = await TryPrintInternalAsync<PrintRecord<string[]>>(connection, tempFile, fileSpec, cancellationToken);
+				PerforceResponse<PrintRecord<string[]>> record = await TryPrintSingleInternalAsync<PrintRecord<string[]>>(connection, tempFile, fileSpec, cancellationToken);
 				if (record.Succeeded)
 				{
 					record.Data.Contents = await File.ReadAllLinesAsync(tempFile, cancellationToken);
@@ -2676,7 +2676,28 @@ namespace EpicGames.Perforce
 		/// <param name="fileSpec">Specification for the files to print</param>
 		/// <param name="cancellationToken">Token used to cancel the operation</param>
 		/// <returns>Response from the server</returns>
-		static Task<PerforceResponse<T>> TryPrintInternalAsync<T>(this IPerforceConnection connection, string outputFile, string fileSpec, CancellationToken cancellationToken = default) where T : class
+		static async Task<PerforceResponseList<PrintRecord>> TryPrintInternalAsync(this IPerforceConnection connection, string outputFile, string fileSpec, CancellationToken cancellationToken = default)
+		{
+			// Build the argument list
+			List<string> arguments = new List<string>();
+			arguments.Add("-o");
+			arguments.Add(outputFile);
+			arguments.Add(fileSpec);
+
+			PerforceResponseList<PrintRecord> records = await CommandAsync<PrintRecord>(connection, "print", arguments, null, cancellationToken);
+			records.RemoveAll(x => x.Error != null && x.Error.Generic == PerforceGenericCode.Empty);
+			return records;
+		}
+
+		/// <summary>
+		/// Execute the 'print' command for a single file
+		/// </summary>
+		/// <param name="connection">Connection to the Perforce server</param>
+		/// <param name="outputFile">Output file to redirect output to</param>
+		/// <param name="fileSpec">Specification for the files to print</param>
+		/// <param name="cancellationToken">Token used to cancel the operation</param>
+		/// <returns>Response from the server</returns>
+		static Task<PerforceResponse<T>> TryPrintSingleInternalAsync<T>(this IPerforceConnection connection, string outputFile, string fileSpec, CancellationToken cancellationToken = default) where T : class
 		{
 			// Build the argument list
 			List<string> arguments = new List<string>();
