EpicGames.Perforce: Allow passing a stdin prompt response directly into the regular command method, rather than having a special implementation for login commands.

diff --git a/Engine/Source/Programs/Shared/EpicGames.Perforce.Native/Client.cpp b/Engine/Source/Programs/Shared/EpicGames.Perforce.Native/Client.cpp
--- a/Engine/Source/Programs/Shared/EpicGames.Perforce.Native/Client.cpp
+++ b/Engine/Source/Programs/Shared/EpicGames.Perforce.Native/Client.cpp
@@ -717,24 +717,17 @@ extern "C" NATIVE_API FClient* Client_Create(const FSettings* Settings, FWriteBu
 	return Client;
 }
 
-extern "C" NATIVE_API void Client_Login(FClient * Client, const char* Password)
-{
-	Client->User.PromptResponse = Password;
-	Client->User.SetInputBuffer(nullptr, 0);
-	Client->ClientApi.SetArgv(0, nullptr);
-	Client->ClientApi.Run("login", &Client->User);
-	Client->User.Flush();
-	Client->User.PromptResponse = nullptr;
-}
-
-extern "C" NATIVE_API void Client_Command(FClient* Client, const char* Func, int ArgCount, const char** Args, const char* InputData, int InputLength, bool InterceptIo)
+extern "C" NATIVE_API void Client_Command(FClient* Client, const char* Func, int ArgCount, const char** Args, const char* InputData, int InputLength, const char* promptResponse, bool InterceptIo)
 {
 	Client->User.InterceptIo = InterceptIo;
 	Client->User.Func = Func;
 	Client->User.SetInputBuffer(InputData, InputLength);
+	Client->User.PromptResponse = promptResponse;
 	Client->ClientApi.SetArgv(ArgCount, (char* const*)Args);
 	Client->ClientApi.Run(Func, &Client->User);
 	Client->User.Flush();
+	Client->User.PromptResponse = nullptr;
+	Client->User.SetInputBuffer(nullptr, 0);
 	Client->User.Func = nullptr;
 	Client->User.InterceptIo = false;
 }
