EpicGames.Perforce: Use regular command execution flow for logging in.

diff --git a/Engine/Source/Programs/Shared/EpicGames.Perforce/IPerforceConnection.cs b/Engine/Source/Programs/Shared/EpicGames.Perforce/IPerforceConnection.cs
--- a/Engine/Source/Programs/Shared/EpicGames.Perforce/IPerforceConnection.cs
+++ b/Engine/Source/Programs/Shared/EpicGames.Perforce/IPerforceConnection.cs
@@ -24,23 +24,16 @@ namespace EpicGames.Perforce
 		ILogger Logger { get; }
 
 		/// <summary>
-		/// Runs a Perforce command
+		/// Queues a Perforce command to be executed
 		/// </summary>
 		/// <param name="command">The command name</param>
 		/// <param name="arguments">Arguments for the command</param>
 		/// <param name="fileArguments">File arguments (may be put into a response file)</param>
 		/// <param name="inputData">Input data to be passed to the command</param>
+		/// <param name="promptResponse">Response for login prompts</param>
 		/// <param name="interceptIo">Whether to intercept file I/O and return it in the reponse stream. Only supported by the native client.</param>
 		/// <returns>Response object</returns>
-		Task<IPerforceOutput> CommandAsync(string command, IReadOnlyList<string> arguments, IReadOnlyList<string>? fileArguments, byte[]? inputData, bool interceptIo);
-
-		/// <summary>
-		/// Execute the 'login' command
-		/// </summary>
-		/// <param name="password">Password to use to login</param>
-		/// <param name="cancellationToken">Token used to cancel the operation</param>
-		/// <returns>Response from the server</returns>
-		Task<IPerforceOutput> LoginCommandAsync(string password, CancellationToken cancellationToken = default);
+		IPerforceOutput Command(string command, IReadOnlyList<string> arguments, IReadOnlyList<string>? fileArguments, byte[]? inputData, string? promptResponse, bool interceptIo);
 
 		/// <summary>
 		/// Creates a record from a set of input fields
diff --git a/Engine/Source/Programs/Shared/EpicGames.Perforce/NativePerforceConnection.cs b/Engine/Source/Programs/Shared/EpicGames.Perforce/NativePerforceConnection.cs
--- a/Engine/Source/Programs/Shared/EpicGames.Perforce/NativePerforceConnection.cs
+++ b/Engine/Source/Programs/Shared/EpicGames.Perforce/NativePerforceConnection.cs
@@ -71,10 +71,7 @@ namespace EpicGames.Perforce
 		static extern IntPtr Client_Create(NativeSettings? settings, NativeWriteBuffer writeBuffer, IntPtr onBufferReadyFnPtr);
 
 		[DllImport(NativeDll, CallingConvention = CallingConvention.Cdecl, CharSet = CharSet.Ansi, BestFitMapping = false, ThrowOnUnmappableChar = true)]
-		static extern void Client_Login(IntPtr client, [MarshalAs(UnmanagedType.LPStr)] string password);
-
-		[DllImport(NativeDll, CallingConvention = CallingConvention.Cdecl, CharSet = CharSet.Ansi, BestFitMapping = false, ThrowOnUnmappableChar = true)]
-		static extern void Client_Command(IntPtr client, [MarshalAs(UnmanagedType.LPStr)] string command, int numArgs, IntPtr[] args, byte[]? inputData, int inputLength, bool interceptIo);
+		static extern void Client_Command(IntPtr client, [MarshalAs(UnmanagedType.LPStr)] string command, int numArgs, IntPtr[] args, byte[]? inputData, int inputLength, byte[]? promptResponse, bool interceptIo);
 
 		[DllImport(NativeDll, CallingConvention = CallingConvention.Cdecl)]
 		static extern void Client_Destroy(IntPtr client);
@@ -509,7 +506,7 @@ namespace EpicGames.Perforce
 		}
 
 		/// <inheritdoc/>
-		public Task<IPerforceOutput> CommandAsync(string command, IReadOnlyList<string> arguments, IReadOnlyList<string>? fileArguments, byte[]? inputData, bool interceptIo)
+		public IPerforceOutput Command(string command, IReadOnlyList<string> arguments, IReadOnlyList<string>? fileArguments, byte[]? inputData, string? promptResponse, bool interceptIo)
 		{
 			byte[]? specData = null;
 			if (inputData != null)
@@ -524,11 +521,11 @@ namespace EpicGames.Perforce
 			}
 
 			Response response = new Response(this);
-			_requests.Add((() => ExecCommand(command, allArguments, specData, interceptIo), response));
-			return Task.FromResult<IPerforceOutput>(response);
+			_requests.Add((() => ExecCommand(command, allArguments, specData, promptResponse, interceptIo), response));
+			return response;
 		}
 
-		private void ExecCommand(string command, List<string> args, byte[]? specData, bool interceptIo)
+		private void ExecCommand(string command, List<string> args, byte[]? inputData, string? promptResponse, bool interceptIo)
 		{
 			List<IntPtr> nativeArgs = new List<IntPtr>();
 			try
@@ -542,7 +539,14 @@ namespace EpicGames.Perforce
 					nativeArgs.Add(nativeArg);
 				}
 
-				Client_Command(_client, command, nativeArgs.Count, nativeArgs.ToArray(), specData, specData?.Length ?? 0, interceptIo);
+				byte[]? promptResponseBytes = null;
+				if (promptResponse != null)
+				{
+					promptResponseBytes = new byte[Encoding.UTF8.GetByteCount(promptResponse) + 1];
+					Encoding.UTF8.GetBytes(promptResponse, promptResponseBytes);
+				}
+
+				Client_Command(_client, command, nativeArgs.Count, nativeArgs.ToArray(), inputData, inputData?.Length ?? 0, promptResponseBytes, interceptIo);
 			}
 			finally
 			{
@@ -553,14 +557,6 @@ namespace EpicGames.Perforce
 			}
 		}
 
-		/// <inheritdoc/>
-		public Task<IPerforceOutput> LoginCommandAsync(string password, CancellationToken cancellationToken = default)
-		{
-			Response response = new Response(this);
-			_requests.Add((() => Client_Login(_client, password), response), cancellationToken);
-			return Task.FromResult<IPerforceOutput>(response);
-		}
-
 		/// <summary>
 		/// Converts a Python marshalled data blob to a spec definition
 		/// </summary>
diff --git a/Engine/Source/Programs/Shared/EpicGames.Perforce/PerforceConnection.cs b/Engine/Source/Programs/Shared/EpicGames.Perforce/PerforceConnection.cs
--- a/Engine/Source/Programs/Shared/EpicGames.Perforce/PerforceConnection.cs
+++ b/Engine/Source/Programs/Shared/EpicGames.Perforce/PerforceConnection.cs
@@ -305,26 +305,17 @@ namespace EpicGames.Perforce
 		}
 
 		/// <inheritdoc/>
-		public Task<IPerforceOutput> CommandAsync(string command, IReadOnlyList<string> arguments, IReadOnlyList<string>? fileArguments, byte[]? inputData, bool interceptIo)
+		public IPerforceOutput Command(string command, IReadOnlyList<string> arguments, IReadOnlyList<string>? fileArguments, byte[]? inputData, string? promptResponse, bool interceptIo)
 		{
+			if (promptResponse != null)
+			{
+				inputData = Encoding.UTF8.GetBytes(promptResponse);
+			}
 			if (interceptIo)
 			{
 				throw new NotSupportedException($"{nameof(interceptIo)} option is not supported through legacy Perforce client");
 			}
-			return Task.FromResult<IPerforceOutput>(new PerforceChildProcess(command, arguments, fileArguments, inputData, GetGlobalArguments(), Logger));
-		}
-
-		/// <summary>
-		/// Execute the 'login' command
-		/// </summary>
-		/// <param name="password">Password to use to login</param>
-		/// <param name="cancellationToken">Token used to cancel the operation</param>
-		/// <returns>Response from the server</returns>
-		public Task<IPerforceOutput> LoginCommandAsync(string password, CancellationToken cancellationToken = default)
-		{
-			// Some versions of P4.EXE do not support marshalled output for P4 login calls, so we only support this as a basic text query.
-			byte[] passwordBytes = Encoding.UTF8.GetBytes(password);
-			return Task.FromResult<IPerforceOutput>(new PerforceChildProcess("login", new List<string>(), null, passwordBytes, GetGlobalArguments(), Logger));
+			return new PerforceChildProcess(command, arguments, fileArguments, inputData, GetGlobalArguments(), Logger);
 		}
 
 		/// <summary>
@@ -424,7 +415,7 @@ namespace EpicGames.Perforce
 		/// <returns>List of objects returned by the server</returns>
 		public static async Task<List<PerforceResponse>> CommandAsync(this IPerforceConnection perforce, string command, IReadOnlyList<string> arguments, IReadOnlyList<string>? fileArguments, byte[]? inputData, Type? statRecordType, CancellationToken cancellationToken = default)
 		{
-			await using (IPerforceOutput response = await perforce.CommandAsync(command, arguments, fileArguments, inputData, false))
+			await using (IPerforceOutput response = perforce.Command(command, arguments, fileArguments, inputData, null, false))
 			{
 				return await response.ReadResponsesAsync(statRecordType, cancellationToken);
 			}
@@ -444,7 +435,7 @@ namespace EpicGames.Perforce
 		/// <returns>List of objects returned by the server</returns>
 		public static async IAsyncEnumerable<PerforceResponse> StreamCommandAsync(this IPerforceConnection perforce, string command, IReadOnlyList<string> arguments, IReadOnlyList<string>? fileArguments, byte[]? inputData, Type? statRecordType, bool interceptIo, [EnumeratorCancellation] CancellationToken cancellationToken)
 		{
-			await using (IPerforceOutput output = await perforce.CommandAsync(command, arguments, fileArguments, inputData, interceptIo))
+			await using (IPerforceOutput output = perforce.Command(command, arguments, fileArguments, inputData, null, interceptIo))
 			{
 				await foreach (PerforceResponse response in output.ReadStreamingResponsesAsync(statRecordType, cancellationToken))
 				{
@@ -465,7 +456,7 @@ namespace EpicGames.Perforce
 		/// <returns>List of objects returned by the server</returns>
 		public static async IAsyncEnumerable<PerforceResponse<T>> StreamCommandAsync<T>(this IPerforceConnection perforce, string command, IReadOnlyList<string> arguments, IReadOnlyList<string>? fileArguments = null, byte[]? inputData = null, [EnumeratorCancellation] CancellationToken cancellationToken = default) where T : class
 		{
-			await using (IPerforceOutput output = await perforce.CommandAsync(command, arguments, fileArguments, inputData, false))
+			await using (IPerforceOutput output = perforce.Command(command, arguments, fileArguments, inputData, null, false))
 			{
 				Type statRecordType = typeof(T);
 				await foreach (PerforceResponse response in output.ReadStreamingResponsesAsync(statRecordType, cancellationToken))
@@ -487,7 +478,7 @@ namespace EpicGames.Perforce
 		/// <returns>List of objects returned by the server</returns>
 		public static async Task RecordCommandAsync(this IPerforceConnection perforce, string command, IReadOnlyList<string> arguments, byte[]? inputData, Action<PerforceRecord> handleRecord, CancellationToken cancellationToken = default)
 		{
-			await using (IPerforceOutput response = await perforce.CommandAsync(command, arguments, null, inputData, false))
+			await using (IPerforceOutput response = perforce.Command(command, arguments, null, inputData, null, false))
 			{
 				await response.ReadRecordsAsync(handleRecord, cancellationToken);
 			}
@@ -2321,7 +2312,7 @@ namespace EpicGames.Perforce
 		public static async Task<PerforceResponse<LoginRecord>> TryLoginAsync(this IPerforceConnection connection, string password, CancellationToken cancellationToken = default)
 		{
 			List<PerforceResponse> parsedResponses;
-			await using (IPerforceOutput response = await connection.LoginCommandAsync(password, cancellationToken))
+			await using (IPerforceOutput response = connection.Command("login", Array.Empty<string>(), null, null, password, false))
 			{
 				for (; ; )
 				{
