Fixed issue where C# UHT was not generating errors with unsupported blueprint member types.

diff --git a/Engine/Source/Programs/Shared/EpicGames.UHT/Types/Properties/UhtArrayProperty.cs b/Engine/Source/Programs/Shared/EpicGames.UHT/Types/Properties/UhtArrayProperty.cs
--- a/Engine/Source/Programs/Shared/EpicGames.UHT/Types/Properties/UhtArrayProperty.cs
+++ b/Engine/Source/Programs/Shared/EpicGames.UHT/Types/Properties/UhtArrayProperty.cs
@@ -39,18 +39,18 @@ namespace EpicGames.UHT.Types
 		public UhtArrayProperty(UhtPropertySettings propertySettings, UhtProperty value) : base(propertySettings, value)
 		{
 			// If the creation of the value property set more flags, then copy those flags to ourselves
-			PropertyFlags |= value.PropertyFlags & EPropertyFlags.UObjectWrapper;
+			PropertyFlags |= ValueProperty.PropertyFlags & EPropertyFlags.UObjectWrapper;
 
-			if (value.MetaData.ContainsKey(UhtNames.NativeConst))
+			if (ValueProperty.MetaData.ContainsKey(UhtNames.NativeConst))
 			{
 				MetaData.Add(UhtNames.NativeConstTemplateArg, "");
-				value.MetaData.Remove(UhtNames.NativeConst);
+				ValueProperty.MetaData.Remove(UhtNames.NativeConst);
 			}
 
 			PropertyCaps |= UhtPropertyCaps.PassCppArgsByRef | UhtPropertyCaps.SupportsRigVM | UhtPropertyCaps.IsRigVMArray;
 			PropertyCaps &= ~(UhtPropertyCaps.CanBeContainerValue | UhtPropertyCaps.CanBeContainerKey);
-			PropertyCaps = (PropertyCaps & ~(UhtPropertyCaps.IsParameterSupportedByBlueprint | UhtPropertyCaps.IsMemberSupportedByBlueprint | UhtPropertyCaps.CanExposeOnSpawn)) |
-				(ValueProperty.PropertyCaps & (UhtPropertyCaps.IsParameterSupportedByBlueprint | UhtPropertyCaps.IsMemberSupportedByBlueprint | UhtPropertyCaps.CanExposeOnSpawn));
+			UpdateCaps();
+
 			ValueProperty.SourceName = SourceName;
 			ValueProperty.EngineName = EngineName;
 			ValueProperty.PropertyFlags = PropertyFlags & EPropertyFlags.PropagateToArrayInner;
@@ -70,11 +70,22 @@ namespace EpicGames.UHT.Types
 					ValueProperty.PropertyFlags = PropertyFlags & EPropertyFlags.PropagateToArrayInner;
 					ValueProperty.MetaData.Clear();
 					PropagateFlagsFromInnerAndHandlePersistentInstanceMetadata(this, MetaData, ValueProperty);
+					UpdateCaps();
 					break;
 			}
 			return results;
 		}
 
+		private void UpdateCaps()
+		{
+			PropertyCaps &= ~(UhtPropertyCaps.IsParameterSupportedByBlueprint | UhtPropertyCaps.IsMemberSupportedByBlueprint | UhtPropertyCaps.CanExposeOnSpawn);
+			PropertyCaps |= ValueProperty.PropertyCaps & UhtPropertyCaps.CanExposeOnSpawn;
+			if (ValueProperty.PropertyCaps.HasAnyFlags(UhtPropertyCaps.IsParameterSupportedByBlueprint))
+			{
+				PropertyCaps |= UhtPropertyCaps.IsParameterSupportedByBlueprint | UhtPropertyCaps.IsMemberSupportedByBlueprint;
+			}
+		}
+
 		/// <inheritdoc/>
 		public override void CollectReferencesInternal(IUhtReferenceCollector collector, bool templateProperty)
 		{
diff --git a/Engine/Source/Programs/Shared/EpicGames.UHT/Types/Properties/UhtMapProperty.cs b/Engine/Source/Programs/Shared/EpicGames.UHT/Types/Properties/UhtMapProperty.cs
--- a/Engine/Source/Programs/Shared/EpicGames.UHT/Types/Properties/UhtMapProperty.cs
+++ b/Engine/Source/Programs/Shared/EpicGames.UHT/Types/Properties/UhtMapProperty.cs
@@ -44,20 +44,21 @@ namespace EpicGames.UHT.Types
 		/// <param name="value">Value property</param>
 		public UhtMapProperty(UhtPropertySettings propertySettings, UhtProperty key, UhtProperty value) : base(propertySettings, value)
 		{
+			KeyProperty = key;
+
 			// If the creation of the value property set more flags, then copy those flags to ourselves
-			PropertyFlags |= value.PropertyFlags & EPropertyFlags.UObjectWrapper;
+			PropertyFlags |= ValueProperty.PropertyFlags & EPropertyFlags.UObjectWrapper;
 
 			// Make sure the 'UObjectWrapper' flag is maintained so that both 'TMap<TSubclassOf<...>, ...>' and 'TMap<UClass*, TSubclassOf<...>>' works correctly
-			key.PropertyFlags = (value.PropertyFlags & ~EPropertyFlags.UObjectWrapper) | (key.PropertyFlags & EPropertyFlags.UObjectWrapper);
-			key.DisallowPropertyFlags = ~(EPropertyFlags.ContainsInstancedReference | EPropertyFlags.InstancedReference | EPropertyFlags.UObjectWrapper);
+			KeyProperty.PropertyFlags = (ValueProperty.PropertyFlags & ~EPropertyFlags.UObjectWrapper) | (KeyProperty.PropertyFlags & EPropertyFlags.UObjectWrapper);
+			KeyProperty.DisallowPropertyFlags = ~(EPropertyFlags.ContainsInstancedReference | EPropertyFlags.InstancedReference | EPropertyFlags.UObjectWrapper);
+
 
-			KeyProperty = key;
 			PropertyCaps |= UhtPropertyCaps.PassCppArgsByRef;
 			PropertyCaps &= ~(UhtPropertyCaps.CanBeContainerValue | UhtPropertyCaps.CanBeContainerKey);
-			PropertyCaps &= ~(UhtPropertyCaps.IsParameterSupportedByBlueprint | UhtPropertyCaps.IsMemberSupportedByBlueprint | UhtPropertyCaps.CanExposeOnSpawn);
-			PropertyCaps |= (ValueProperty.PropertyCaps & (UhtPropertyCaps.IsParameterSupportedByBlueprint | UhtPropertyCaps.IsMemberSupportedByBlueprint | UhtPropertyCaps.CanExposeOnSpawn)) &
-				(KeyProperty.PropertyCaps & (UhtPropertyCaps.IsParameterSupportedByBlueprint | UhtPropertyCaps.IsMemberSupportedByBlueprint));
-			PropertyFlags = value.PropertyFlags;
+			UpdateCaps();
+
+			PropertyFlags = ValueProperty.PropertyFlags;
 			ValueProperty.SourceName = SourceName;
 			ValueProperty.EngineName = EngineName;
 			ValueProperty.PropertyFlags &= EPropertyFlags.PropagateToMapValue;
@@ -92,6 +93,7 @@ namespace EpicGames.UHT.Types
 					ValueProperty.MetaData.Clear();
 					ValueProperty.PropertyFlags = PropertyFlags & EPropertyFlags.PropagateToMapValue;
 
+					UpdateCaps();
 					PropagateFlagsFromInnerAndHandlePersistentInstanceMetadata(this, null, KeyProperty);
 					PropagateFlagsFromInnerAndHandlePersistentInstanceMetadata(this, MetaData, ValueProperty);
 					break;
@@ -99,6 +101,16 @@ namespace EpicGames.UHT.Types
 			return results;
 		}
 
+		private void UpdateCaps()
+		{
+			PropertyCaps &= ~(UhtPropertyCaps.IsParameterSupportedByBlueprint | UhtPropertyCaps.IsMemberSupportedByBlueprint | UhtPropertyCaps.CanExposeOnSpawn);
+			PropertyCaps |= ValueProperty.PropertyCaps & UhtPropertyCaps.CanExposeOnSpawn;
+			if (KeyProperty.PropertyCaps.HasAnyFlags(UhtPropertyCaps.IsParameterSupportedByBlueprint) && ValueProperty.PropertyCaps.HasAnyFlags(UhtPropertyCaps.IsParameterSupportedByBlueprint))
+			{
+				PropertyCaps |= UhtPropertyCaps.IsParameterSupportedByBlueprint | UhtPropertyCaps.IsMemberSupportedByBlueprint;
+			}
+		}
+
 		/// <inheritdoc/>
 		public override void CollectReferencesInternal(IUhtReferenceCollector collector, bool templateProperty)
 		{
diff --git a/Engine/Source/Programs/Shared/EpicGames.UHT/Types/Properties/UhtSetProperty.cs b/Engine/Source/Programs/Shared/EpicGames.UHT/Types/Properties/UhtSetProperty.cs
--- a/Engine/Source/Programs/Shared/EpicGames.UHT/Types/Properties/UhtSetProperty.cs
+++ b/Engine/Source/Programs/Shared/EpicGames.UHT/Types/Properties/UhtSetProperty.cs
@@ -38,12 +38,12 @@ namespace EpicGames.UHT.Types
 		public UhtSetProperty(UhtPropertySettings propertySettings, UhtProperty value) : base(propertySettings, value)
 		{
 			// If the creation of the value property set more flags, then copy those flags to ourselves
-			PropertyFlags |= value.PropertyFlags & EPropertyFlags.UObjectWrapper;
+			PropertyFlags |= ValueProperty.PropertyFlags & EPropertyFlags.UObjectWrapper;
 
 			PropertyCaps |= UhtPropertyCaps.PassCppArgsByRef;
 			PropertyCaps &= ~(UhtPropertyCaps.CanBeContainerValue | UhtPropertyCaps.CanBeContainerKey);
-			PropertyCaps = (PropertyCaps & ~(UhtPropertyCaps.IsParameterSupportedByBlueprint | UhtPropertyCaps.IsMemberSupportedByBlueprint | UhtPropertyCaps.CanExposeOnSpawn)) |
-				(ValueProperty.PropertyCaps & (UhtPropertyCaps.IsParameterSupportedByBlueprint | UhtPropertyCaps.IsMemberSupportedByBlueprint | UhtPropertyCaps.CanExposeOnSpawn));
+			UpdateCaps();
+
 			ValueProperty.SourceName = SourceName;
 			ValueProperty.EngineName = EngineName;
 			ValueProperty.PropertyFlags = PropertyFlags & EPropertyFlags.PropagateToSetElement;
@@ -62,12 +62,23 @@ namespace EpicGames.UHT.Types
 					MetaData.Add(ValueProperty.MetaData);
 					ValueProperty.MetaData.Clear();
 					ValueProperty.PropertyFlags = PropertyFlags & EPropertyFlags.PropagateToSetElement;
+					UpdateCaps();
 					PropagateFlagsFromInnerAndHandlePersistentInstanceMetadata(this, MetaData, ValueProperty);
 					break;
 			}
 			return results;
 		}
 
+		private void UpdateCaps()
+		{
+			PropertyCaps &= ~(UhtPropertyCaps.IsParameterSupportedByBlueprint | UhtPropertyCaps.IsMemberSupportedByBlueprint | UhtPropertyCaps.CanExposeOnSpawn);
+			PropertyCaps |= ValueProperty.PropertyCaps & UhtPropertyCaps.CanExposeOnSpawn;
+			if (ValueProperty.PropertyCaps.HasAnyFlags(UhtPropertyCaps.IsParameterSupportedByBlueprint))
+			{
+				PropertyCaps |= UhtPropertyCaps.IsParameterSupportedByBlueprint | UhtPropertyCaps.IsMemberSupportedByBlueprint;
+			}
+		}
+
 		/// <inheritdoc/>
 		public override void CollectReferencesInternal(IUhtReferenceCollector collector, bool templateProperty)
 		{
diff --git a/Engine/Source/Programs/Shared/EpicGames.UHT/Types/Properties/UhtStructProperty.cs b/Engine/Source/Programs/Shared/EpicGames.UHT/Types/Properties/UhtStructProperty.cs
--- a/Engine/Source/Programs/Shared/EpicGames.UHT/Types/Properties/UhtStructProperty.cs
+++ b/Engine/Source/Programs/Shared/EpicGames.UHT/Types/Properties/UhtStructProperty.cs
@@ -47,18 +47,8 @@ namespace EpicGames.UHT.Types
 			ScriptStruct = scriptStruct;
 			HeaderFile.AddReferencedHeader(scriptStruct);
 			PropertyCaps |= UhtPropertyCaps.SupportsRigVM;
-			if (ScriptStruct.HasNoOpConstructor)
-			{
-				PropertyCaps |= UhtPropertyCaps.RequiresNullConstructorArg;
-			}
-			if (ScriptStruct.MetaData.GetBoolean(UhtNames.BlueprintType))
-			{
-				PropertyCaps |= UhtPropertyCaps.CanExposeOnSpawn | UhtPropertyCaps.IsParameterSupportedByBlueprint | UhtPropertyCaps.IsMemberSupportedByBlueprint;
-			}
-			else if (ScriptStruct.MetaData.GetBooleanHierarchical(UhtNames.BlueprintType))
-			{
-				PropertyCaps |= UhtPropertyCaps.IsParameterSupportedByBlueprint | UhtPropertyCaps.IsMemberSupportedByBlueprint;
-			}
+
+			UpdateCaps();
 		}
 
 		/// <inheritdoc/>
@@ -72,11 +62,34 @@ namespace EpicGames.UHT.Types
 					{
 						PropertyFlags |= EPropertyFlags.ContainsInstancedReference;
 					}
+					UpdateCaps();
 					break;
 			}
 			return results;
 		}
 
+		private void UpdateCaps()
+		{
+			if (ScriptStruct.HasNoOpConstructor)
+			{
+				PropertyCaps |= UhtPropertyCaps.RequiresNullConstructorArg;
+			}
+
+			// There is a good chance Blueprint type was set during property resolve phase.  Check this flag again.
+			const UhtPropertyCaps BlueprintCaps = UhtPropertyCaps.CanExposeOnSpawn | UhtPropertyCaps.IsParameterSupportedByBlueprint | UhtPropertyCaps.IsMemberSupportedByBlueprint;
+			if (!PropertyCaps.HasExactFlags(BlueprintCaps, BlueprintCaps))
+			{
+				if (ScriptStruct.MetaData.GetBoolean(UhtNames.BlueprintType))
+				{
+					PropertyCaps |= BlueprintCaps;
+				}
+				else if (ScriptStruct.MetaData.GetBooleanHierarchical(UhtNames.BlueprintType))
+				{
+					PropertyCaps |= UhtPropertyCaps.IsParameterSupportedByBlueprint | UhtPropertyCaps.IsMemberSupportedByBlueprint;
+				}
+			}
+		}
+
 		/// <inheritdoc/>
 		public override bool ScanForInstancedReferenced(bool deepScan)
 		{
diff --git a/Engine/Source/Programs/Shared/EpicGames.UHT/Types/UhtClass.cs b/Engine/Source/Programs/Shared/EpicGames.UHT/Types/UhtClass.cs
--- a/Engine/Source/Programs/Shared/EpicGames.UHT/Types/UhtClass.cs
+++ b/Engine/Source/Programs/Shared/EpicGames.UHT/Types/UhtClass.cs
@@ -1796,6 +1796,19 @@ namespace EpicGames.UHT.Types
 								ValidateBlueprintPopertySetter(property, setter, (UhtFunction?)FindType(UhtFindOptions.EngineName | UhtFindOptions.Function, setter));
 							}
 						}
+
+						if (!property.PropertyFlags.HasAnyFlags(EPropertyFlags.BlueprintAssignable | EPropertyFlags.BlueprintCallable) && property is not UhtMulticastDelegateProperty)
+						{
+							if (property.IsStaticArray)
+							{
+								property.LogError($"Static array cannot be exposed to blueprint Class: {SourceName} Property: {property.SourceName}");
+							}
+
+							if (!property.PropertyCaps.HasAnyFlags(UhtPropertyCaps.IsMemberSupportedByBlueprint))
+							{
+								property.LogError($"Type '{property.GetUserFacingDecl()}' is not supported by blueprint. Class: {SourceName} Property: {property.SourceName}");
+							}
+						}
 					}
 
 					// Validate if we are using editor only data in a class or struct definition
diff --git a/Engine/Source/Programs/Shared/EpicGames.UHT/Types/UhtScriptStruct.cs b/Engine/Source/Programs/Shared/EpicGames.UHT/Types/UhtScriptStruct.cs
--- a/Engine/Source/Programs/Shared/EpicGames.UHT/Types/UhtScriptStruct.cs
+++ b/Engine/Source/Programs/Shared/EpicGames.UHT/Types/UhtScriptStruct.cs
@@ -775,8 +775,45 @@ namespace EpicGames.UHT.Types
 				this.LogError("virtual FRigVMStructUpgradeInfo GetUpgradeInfo() const override;");
 			}
 
+			ValidateProperties();
+
 			return options |= UhtValidationOptions.Shadowing;
 		}
+
+		void ValidateProperties()
+		{
+			bool hasTestedStruct = false;
+			foreach (UhtType child in Children)
+			{
+				if (child is UhtProperty property)
+				{
+					if (property.PropertyFlags.HasAnyFlags(EPropertyFlags.BlueprintVisible))
+					{
+						if (!property.PropertyFlags.HasAnyFlags(EPropertyFlags.BlueprintAssignable | EPropertyFlags.BlueprintCallable) && property is not UhtMulticastDelegateProperty)
+						{
+							if (!hasTestedStruct)
+							{
+								hasTestedStruct = true;
+								if (!MetaData.GetBooleanHierarchical(UhtNames.BlueprintType))
+								{
+									property.LogError($"Cannot expose property to blueprints in a struct that is not a BlueprintType. Struct: {SourceName} Property: {property.SourceName}");
+								}
+							}
+
+							if (property.IsStaticArray)
+							{
+								property.LogError($"Static array cannot be exposed to blueprint Class: {SourceName} Property: {property.SourceName}");
+							}
+
+							if (!property.PropertyCaps.HasAnyFlags(UhtPropertyCaps.IsMemberSupportedByBlueprint))
+							{
+								property.LogError($"Type '{property.GetUserFacingDecl()}' is not supported by blueprint. Struct: {SourceName} Property: {property.SourceName}");
+							}
+						}
+					}
+				}
+			}
+		}
 		#endregion
 
 		/// <inheritdoc/>
