Removed NameToDisplayString port to C# since we don't need to save the generated strings and imply that the C# version will match the C++ version.

diff --git a/Engine/Source/Programs/Shared/EpicGames.UHT/Types/UhtProperty.cs b/Engine/Source/Programs/Shared/EpicGames.UHT/Types/UhtProperty.cs
--- a/Engine/Source/Programs/Shared/EpicGames.UHT/Types/UhtProperty.cs
+++ b/Engine/Source/Programs/Shared/EpicGames.UHT/Types/UhtProperty.cs
@@ -1041,12 +1041,6 @@ namespace EpicGames.UHT.Types
 			PropertyCaps = UhtPropertyCaps.CanBeContainerValue | UhtPropertyCaps.CanBeContainerKey | UhtPropertyCaps.CanHaveConfig;
 		}
 
-		/// <inheritdoc/>
-		protected override string GetDisplayStringFromEngineName()
-		{
-			return UhtFCString.NameToDisplayString(EngineName, this is UhtBoolProperty);
-		}
-
 		#region Text and code generation support
 		/// <summary>
 		/// Internal version of AppendText.  Don't append any text to the builder to get the default behavior
diff --git a/Engine/Source/Programs/Shared/EpicGames.UHT/Types/UhtType.cs b/Engine/Source/Programs/Shared/EpicGames.UHT/Types/UhtType.cs
--- a/Engine/Source/Programs/Shared/EpicGames.UHT/Types/UhtType.cs
+++ b/Engine/Source/Programs/Shared/EpicGames.UHT/Types/UhtType.cs
@@ -1208,7 +1208,10 @@ namespace EpicGames.UHT.Types
 				return nativeDisplayName;
 			}
 
-			return GetDisplayStringFromEngineName();
+			// In the engine, this is a call to NameToDisplayString.  UHT doesn't store these values
+			// so it is just important that they be reasonably unique.  This allows us to not have to 
+			// support changes to NameToDisplayString in the engine
+			return SourceName;
 		}
 
 		/// <summary>
@@ -1231,22 +1234,16 @@ namespace EpicGames.UHT.Types
 
 			if (nativeToolTip.Length == 0)
 			{
-				return GetDisplayStringFromEngineName();
+				// In the engine, this is a call to NameToDisplayString.  UHT doesn't store these values
+				// so it is just important that they be reasonably unique.  This allows us to not have to 
+				// support changes to NameToDisplayString in the engine
+				return SourceName;
 			}
 			else
 			{
 				return nativeToolTip.TrimEnd().Replace("@see", "See:", StringComparison.Ordinal);
 			}
 		}
-
-		/// <summary>
-		/// Return a generated display string from the engine name
-		/// </summary>
-		/// <returns>The generated string</returns>
-		protected virtual string GetDisplayStringFromEngineName()
-		{
-			return UhtFCString.NameToDisplayString(EngineName, false);
-		}
 		#endregion
 
 		/// <summary>
diff --git a/Engine/Source/Programs/Shared/EpicGames.UHT/Utils/UhtFCString.cs b/Engine/Source/Programs/Shared/EpicGames.UHT/Utils/UhtFCString.cs
--- a/Engine/Source/Programs/Shared/EpicGames.UHT/Utils/UhtFCString.cs
+++ b/Engine/Source/Programs/Shared/EpicGames.UHT/Utils/UhtFCString.cs
@@ -302,141 +302,6 @@ namespace EpicGames.UHT.Utils
 			return result.ToString();
 		}
 
-		/// <summary>
-		/// Words that are considered articles when parsing comment strings
-		/// Some words are always forced lowercase
-		/// </summary>
-		private static readonly string[] s_articles = new string[]
-		{
-			"In",
-			"As",
-			"To",
-			"Or",
-			"At",
-			"On",
-			"If",
-			"Be",
-			"By",
-			"The",
-			"For",
-			"And",
-			"With",
-			"When",
-			"From",
-		};
-
-		/// <summary>
-		/// Convert a name to a display string
-		/// </summary>
-		/// <param name="name">Name to convert</param>
-		/// <param name="isBool">True if the name is from a boolean property</param>
-		/// <returns>Display string</returns>
-		public static string NameToDisplayString(StringView name, bool isBool)
-		{
-			ReadOnlySpan<char> chars = name.Span;
-
-			// This is used to indicate that we are in a run of uppercase letter and/or digits.  The code attempts to keep
-			// these characters together as breaking them up often looks silly (i.e. "Draw Scale 3 D" as opposed to "Draw Scale 3D"
-			bool inARun = false;
-			bool wasSpace = false;
-			bool wasOpenParen = false;
-			bool wasNumber = false;
-			bool wasMinusSign = false;
-
-			StringBuilder builder = new();
-			for (int charIndex = 0; charIndex < chars.Length; ++charIndex)
-			{
-				char ch = chars[charIndex];
-
-				bool lowerCase = Char.IsLower(ch);
-				bool upperCase = Char.IsUpper(ch);
-				bool isDigit = Char.IsDigit(ch);
-				bool isUnderscore = ch == '_';
-
-				// Skip the first character if the property is a bool (they should all start with a lowercase 'b', which we don't want to keep)
-				if (charIndex == 0 && isBool && ch == 'b')
-				{
-					// Check if next character is uppercase as it may be a user created string that doesn't follow the rules of Unreal variables
-					if (chars.Length > 1 && Char.IsUpper(chars[1]))
-					{
-						continue;
-					}
-				}
-
-				// If the current character is upper case or a digit, and the previous character wasn't, then we need to insert a space if there wasn't one previously
-				// We don't do this for numerical expressions, for example "-1.2" should not be formatted as "- 1. 2"
-				if ((upperCase || (isDigit && !wasMinusSign)) && !inARun && !wasOpenParen && !wasNumber)
-				{
-					if (!wasSpace && builder.Length > 0)
-					{
-						builder.Append(' ');
-						wasSpace = true;
-					}
-					inARun = true;
-				}
-
-				// A lower case character will break a run of upper case letters and/or digits
-				if (lowerCase)
-				{
-					inARun = false;
-				}
-
-				// An underscore denotes a space, so replace it and continue the run
-				if (isUnderscore)
-				{
-					ch = ' ';
-					inARun = true;
-				}
-
-				// If this is the first character in the string, then it will always be upper-case
-				if (builder.Length == 0)
-				{
-					ch = Char.ToUpper(ch);
-				}
-				else if (!isDigit && (wasSpace || wasOpenParen)) // If this is first character after a space, then make sure it is case-correct
-				{
-
-					// Search for a word that needs case repaired
-					bool isArticle = false;
-					foreach (string article in s_articles)
-					{
-						// Make sure the character following the string we're testing is not lowercase (we don't want to match "in" with "instance")
-						int articleLength = article.Length;
-						if (charIndex + articleLength < chars.Length && !Char.IsLower(chars[charIndex + articleLength]))
-						{
-							// Does this match the current article?
-							if (chars.Slice(charIndex, articleLength).Equals(article.AsSpan(), StringComparison.OrdinalIgnoreCase))
-							{
-								isArticle = true;
-								break;
-							}
-						}
-					}
-
-					// Start of a keyword, force to lowercase
-					if (isArticle)
-					{
-						ch = Char.ToLower(ch);
-					}
-					else    // First character after a space that's not a reserved keyword, make sure it's uppercase
-					{
-						ch = Char.ToUpper(ch);
-					}
-				}
-
-				wasSpace = ch == ' ';
-				wasOpenParen = ch == '(';
-
-				// What could be included as part of a numerical representation.
-				// For example -1.2
-				wasMinusSign = ch == '-';
-				bool potentialNumericalChar = wasMinusSign || ch == '.';
-				wasNumber = isDigit || (wasNumber && potentialNumericalChar);
-				builder.Append(ch);
-			}
-			return builder.ToString();
-		}
-
 		/// <summary>
 		/// Replace tabs to spaces in a string containing only a single line.
 		/// </summary>
