EpicGames.Perforce: Add a host parameter for native perforce lib.

diff --git a/Engine/Source/Programs/Shared/EpicGames.Perforce.Native/Client.cpp b/Engine/Source/Programs/Shared/EpicGames.Perforce.Native/Client.cpp
--- a/Engine/Source/Programs/Shared/EpicGames.Perforce.Native/Client.cpp
+++ b/Engine/Source/Programs/Shared/EpicGames.Perforce.Native/Client.cpp
@@ -32,6 +32,7 @@ struct FSettings
 	const char* ServerAndPort;
 	const char* User;
 	const char* Password;
+	const char* Host;
 	const char* Client;
 	const char* AppName;
 	const char* AppVersion;
@@ -691,6 +692,10 @@ extern "C" NATIVE_API FClient* Client_Create(const FSettings* Settings, FWriteBu
 		{
 			Client->ClientApi.SetPassword(Settings->Password);
 		}
+		if (Settings->Host != nullptr)
+		{
+			Client->ClientApi.SetHost(Settings->Host);
+		}
 		if (Settings->Client != nullptr)
 		{
 			Client->ClientApi.SetClient(Settings->Client);
