Horde: Change to the correct host when impersonating a user.

diff --git a/Engine/Source/Programs/Horde/Horde.Build/Perforce/PerforceService.cs b/Engine/Source/Programs/Horde/Horde.Build/Perforce/PerforceService.cs
--- a/Engine/Source/Programs/Horde/Horde.Build/Perforce/PerforceService.cs
+++ b/Engine/Source/Programs/Horde/Horde.Build/Perforce/PerforceService.cs
@@ -263,6 +263,7 @@ namespace Horde.Build.Perforce
 			PerforceSettings settings = new PerforceSettings(serverAndPort, credentials.UserName);
 			settings.AppName = "Horde.Build";
 			settings.Password = credentials.Password;
+			settings.HostName = clientRecord?.Host;
 			settings.ClientName = clientRecord?.Name ?? "__DOES_NOT_EXIST__";
 			settings.PreferNativeClient = true;
 
diff --git a/Engine/Source/Programs/Shared/EpicGames.Perforce/NativePerforceConnection.cs b/Engine/Source/Programs/Shared/EpicGames.Perforce/NativePerforceConnection.cs
--- a/Engine/Source/Programs/Shared/EpicGames.Perforce/NativePerforceConnection.cs
+++ b/Engine/Source/Programs/Shared/EpicGames.Perforce/NativePerforceConnection.cs
@@ -35,6 +35,9 @@ namespace EpicGames.Perforce
 			[MarshalAs(UnmanagedType.LPStr)]
 			public string? _password;
 
+			[MarshalAs(UnmanagedType.LPStr)]
+			public string? _hostName;
+
 			[MarshalAs(UnmanagedType.LPStr)]
 			public string? _clientName;
 
@@ -433,6 +436,7 @@ namespace EpicGames.Perforce
 				nativeSettings._serverAndPort = Settings.ServerAndPort;
 				nativeSettings._userName = Settings.UserName;
 				nativeSettings._password = Settings.Password;
+				nativeSettings._hostName = Settings.HostName;
 				nativeSettings._clientName = Settings.ClientName;
 				nativeSettings._appName = Settings.AppName;
 				nativeSettings._appVersion = Settings.AppVersion;
diff --git a/Engine/Source/Programs/Shared/EpicGames.Perforce/PerforceConnection.cs b/Engine/Source/Programs/Shared/EpicGames.Perforce/PerforceConnection.cs
--- a/Engine/Source/Programs/Shared/EpicGames.Perforce/PerforceConnection.cs
+++ b/Engine/Source/Programs/Shared/EpicGames.Perforce/PerforceConnection.cs
@@ -10,7 +10,6 @@ using System.Text;
 using System.Text.RegularExpressions;
 using System.Threading;
 using System.Threading.Tasks;
-using EpicGames.Core;
 using Microsoft.Extensions.Logging;
 
 namespace EpicGames.Perforce
@@ -164,6 +163,15 @@ namespace EpicGames.Perforce
 			set;
 		}
 
+		/// <summary>
+		/// The current host name
+		/// </summary>
+		public string? HostName
+		{
+			get;
+			set;
+		}
+
 		/// <summary>
 		/// The current client name
 		/// </summary>
@@ -261,6 +269,7 @@ namespace EpicGames.Perforce
 		public PerforceConnection(IPerforceSettings settings, ILogger logger)
 			: this(settings.ServerAndPort, settings.UserName, settings.ClientName, settings.AppName, settings.AppVersion, logger)
 		{
+			HostName = settings.HostName;
 		}
 
 		/// <summary>
@@ -289,6 +298,10 @@ namespace EpicGames.Perforce
 			{
 				arguments.Add($"-u{UserName}");
 			}
+			if (HostName != null)
+			{
+				arguments.Add($"-H{HostName}");
+			}
 			if (ClientName != null)
 			{
 				arguments.Add($"-c{ClientName}");
diff --git a/Engine/Source/Programs/Shared/EpicGames.Perforce/PerforceSettings.cs b/Engine/Source/Programs/Shared/EpicGames.Perforce/PerforceSettings.cs
--- a/Engine/Source/Programs/Shared/EpicGames.Perforce/PerforceSettings.cs
+++ b/Engine/Source/Programs/Shared/EpicGames.Perforce/PerforceSettings.cs
@@ -27,6 +27,11 @@ namespace EpicGames.Perforce
 		/// </summary>
 		public string? Password { get; }
 
+		/// <summary>
+		/// Name of the current host
+		/// </summary>
+		public string? HostName { get; }
+
 		/// <summary>
 		/// Name of the client to use
 		/// </summary>
@@ -67,6 +72,9 @@ namespace EpicGames.Perforce
 		/// <inheritdoc/>
 		public string? Password { get; set; }
 
+		/// <inheritdoc/>
+		public string? HostName { get; set; }
+
 		/// <inheritdoc/>
 		public string? ClientName { get; set; }
 
@@ -87,6 +95,7 @@ namespace EpicGames.Perforce
 			ServerAndPort = environment.GetValue("P4PORT") ?? "perforce:1666";
 			UserName = environment.GetValue("P4USER") ?? System.Environment.UserName;
 			Password = environment.GetValue("P4PASSWD");
+			HostName = environment.GetValue("P4HOST");
 			ClientName = environment.GetValue("P4CLIENT");
 
 			AssemblyName entryAssemblyName = Assembly.GetEntryAssembly()!.GetName();
@@ -118,6 +127,7 @@ namespace EpicGames.Perforce
 			ServerAndPort = other.ServerAndPort;
 			UserName = other.UserName;
 			Password = other.Password;
+			HostName = other.HostName;
 			ClientName = other.ClientName;
 			AppName = other.AppName;
 			AppVersion = other.AppVersion;
