Horde: Fix errors being ignored when submitting changelists, and incorrect changelist number being returned.

diff --git a/Engine/Source/Programs/Horde/Horde.Build/Perforce/PerforceService.cs b/Engine/Source/Programs/Horde/Horde.Build/Perforce/PerforceService.cs
--- a/Engine/Source/Programs/Horde/Horde.Build/Perforce/PerforceService.cs
+++ b/Engine/Source/Programs/Horde/Horde.Build/Perforce/PerforceService.cs
@@ -972,7 +972,7 @@ namespace Horde.Build.Perforce
 						change = await perforce.CreateChangeAsync(change, cancellationToken);
 
 						SubmitRecord submit = await perforce.SubmitAsync(change.Number, SubmitOptions.SubmitUnchanged, cancellationToken);
-						return submit.ChangeNumber;
+						return submit.SubmittedChangeNumber;
 					}
 					catch (Exception ex)
 					{
@@ -1049,7 +1049,7 @@ namespace Horde.Build.Perforce
 				try
 				{
 					SubmitRecord submit = await perforce.SubmitShelvedAsync(change, cancellationToken);
-					return (submit.ChangeNumber, records[0].Description);
+					return (submit.SubmittedChangeNumber, records[0].Description);
 				}
 				catch (Exception ex)
 				{
diff --git a/Engine/Source/Programs/Shared/EpicGames.Perforce/PerforceConnection.cs b/Engine/Source/Programs/Shared/EpicGames.Perforce/PerforceConnection.cs
--- a/Engine/Source/Programs/Shared/EpicGames.Perforce/PerforceConnection.cs
+++ b/Engine/Source/Programs/Shared/EpicGames.Perforce/PerforceConnection.cs
@@ -3426,7 +3426,42 @@ namespace EpicGames.Perforce
 			List<string> arguments = new List<string>();
 			arguments.Add($"-e{changeNumber}");
 
-			return (await CommandAsync<SubmitRecord>(connection, "submit", arguments, null, cancellationToken))[0];
+			PerforceResponseList<SubmitRecord> responses = await CommandAsync<SubmitRecord>(connection, "submit", arguments, null, cancellationToken);
+
+			SubmitRecord? success = null;
+			PerforceResponse<SubmitRecord>? error = null;
+			foreach (PerforceResponse<SubmitRecord> response in responses)
+			{
+				if (response.Error != null)
+				{
+					error ??= response;
+				}
+				else if (response.Info != null)
+				{
+					connection.Logger.LogInformation("Submit: {Info}", response.Info.Data);
+				}
+				else if (success == null)
+				{
+					success = response.Data;
+				}
+				else
+				{
+					success.Merge(response.Data);
+				}
+			}
+
+			if (error != null)
+			{
+				return error;
+			}
+			else if (success != null)
+			{
+				return new PerforceResponse<SubmitRecord>(success);
+			}
+			else
+			{
+				return responses[0];
+			}
 		}
 
 		#endregion
diff --git a/Engine/Source/Programs/Shared/EpicGames.Perforce/Records/SubmitRecord.cs b/Engine/Source/Programs/Shared/EpicGames.Perforce/Records/SubmitRecord.cs
--- a/Engine/Source/Programs/Shared/EpicGames.Perforce/Records/SubmitRecord.cs
+++ b/Engine/Source/Programs/Shared/EpicGames.Perforce/Records/SubmitRecord.cs
@@ -8,10 +8,32 @@ namespace EpicGames.Perforce
 	public class SubmitRecord
 	{
 		/// <summary>
-		/// Submitted changelist number
+		/// Original changelist number
 		/// </summary>
 		[PerforceTag("change", Optional = true)]
-		public int ChangeNumber { get; set; } = -1;
+		public int OriginalChangeNumber { get; set; } = -1;
+
+		/// <summary>
+		/// Submitted changelist number
+		/// </summary>
+		[PerforceTag("submittedChange", Optional = true)]
+		public int SubmittedChangeNumber { get; set; } = -1;
+
+		/// <summary>
+		/// Merge another submit record into this one. Perforce outputs multiple records for a submit command, so we merge them together for convenience.
+		/// </summary>
+		/// <param name="other"></param>
+		public void Merge(SubmitRecord other)
+		{
+			if (other.OriginalChangeNumber != -1)
+			{
+				OriginalChangeNumber = other.OriginalChangeNumber;
+			}
+			if (other.SubmittedChangeNumber != -1)
+			{
+				SubmittedChangeNumber = other.SubmittedChangeNumber;
+			}
+		}
 
 		/// <summary>
 		/// Format this record for display in the debugger
@@ -19,7 +41,7 @@ namespace EpicGames.Perforce
 		/// <returns>Summary of this revision</returns>
 		public override string ToString()
 		{
-			return ChangeNumber.ToString();
+			return SubmittedChangeNumber.ToString();
 		}
 	}
 }
