UGS: Fix exception when querying current stream from a non-stream workspace.

diff --git a/Engine/Source/Programs/Shared/EpicGames.Perforce/PerforceConnection.cs b/Engine/Source/Programs/Shared/EpicGames.Perforce/PerforceConnection.cs
--- a/Engine/Source/Programs/Shared/EpicGames.Perforce/PerforceConnection.cs
+++ b/Engine/Source/Programs/Shared/EpicGames.Perforce/PerforceConnection.cs
@@ -4020,7 +4020,7 @@ namespace EpicGames.Perforce
 		/// <returns></returns>
 		public static async Task<string?> GetCurrentStreamAsync(this IPerforceConnection connection, CancellationToken cancellationToken = default)
 		{
-			return await TryGetCurrentStreamAsync(connection, cancellationToken).UnwrapAsync();
+			return await TryGetCurrentStreamAsync(connection, cancellationToken);
 		}
 
 		/// <summary>
@@ -4029,16 +4029,16 @@ namespace EpicGames.Perforce
 		/// <param name="connection">Connection to the Perforce server</param>
 		/// <param name="cancellationToken">Token used to cancel the operation</param>
 		/// <returns></returns>
-		public static async Task<PerforceResponse<string>> TryGetCurrentStreamAsync(this IPerforceConnection connection, CancellationToken cancellationToken = default)
+		public static async Task<string?> TryGetCurrentStreamAsync(this IPerforceConnection connection, CancellationToken cancellationToken = default)
 		{
 			PerforceResponse<ClientRecord> response = await connection.TryGetClientAsync(null, cancellationToken);
 			if (response.Succeeded)
 			{
-				return new PerforceResponse<string>(response.Data.Stream!);
+				return response.Data.Stream;
 			}
 			else
 			{
-				return new PerforceResponse<string>((string)null!);
+				return null;
 			}
 		}
 
