Horde: Fix errors posting summary messages to threads for issues.

diff --git a/Engine/Source/Programs/Horde/Horde.Build/Notifications/Sinks/SlackNotificationSink.cs b/Engine/Source/Programs/Horde/Horde.Build/Notifications/Sinks/SlackNotificationSink.cs
--- a/Engine/Source/Programs/Horde/Horde.Build/Notifications/Sinks/SlackNotificationSink.cs
+++ b/Engine/Source/Programs/Horde/Horde.Build/Notifications/Sinks/SlackNotificationSink.cs
@@ -2035,6 +2035,9 @@ namespace Horde.Build.Notifications.Sinks
 					id = await _slackClient.PostMessageToThreadAsync(threadId, message);
 				}
 
+				state.Channel = id.Channel;
+				state.Ts = id.Ts;
+
 				await UpdateMessageStateAsync(state.Id, id);
 			}
 			else if (!String.IsNullOrEmpty(state.Ts))
diff --git a/Engine/Source/Programs/Shared/EpicGames.Slack/SlackClient.cs b/Engine/Source/Programs/Shared/EpicGames.Slack/SlackClient.cs
--- a/Engine/Source/Programs/Shared/EpicGames.Slack/SlackClient.cs
+++ b/Engine/Source/Programs/Shared/EpicGames.Slack/SlackClient.cs
@@ -231,8 +231,8 @@ namespace EpicGames.Slack
 		{
 			PostMessageRequest request = new PostMessageRequest();
 			request.Channel = recipient;
-			request.Ts = ts;
-			request.ThreadTs = threadTs;
+			request.Ts = String.IsNullOrEmpty(ts)? null : ts;
+			request.ThreadTs = String.IsNullOrEmpty(threadTs)? null : threadTs;
 			request.Text = message.Text;
 			request.Blocks = message.Blocks;
 			request.Markdown = message.Markdown;
