Fixed issue where C# UHT would not declare shadowing errors in the incomplete instances where UHT would.

diff --git a/Engine/Source/Programs/Shared/EpicGames.UHT/Types/UhtProperty.cs b/Engine/Source/Programs/Shared/EpicGames.UHT/Types/UhtProperty.cs
--- a/Engine/Source/Programs/Shared/EpicGames.UHT/Types/UhtProperty.cs
+++ b/Engine/Source/Programs/Shared/EpicGames.UHT/Types/UhtProperty.cs
@@ -1549,11 +1549,7 @@ namespace EpicGames.UHT.Types
 			}
 
 			// Shadowing checks are done at this level, not in the properties themselves
-			// Note: The old UHT would check for all property type where this version only checks for member variables.
-			// In the old code if you defined the property after the function with argument with the same name, UHT would
-			// not issue an error.  However, if the property was defined PRIOR to the function with the matching argument name,
-			// UHT would generate an error.
-			if (options.HasAnyFlags(UhtValidationOptions.Shadowing) && PropertyCategory == UhtPropertyCategory.Member)
+			if (options.HasAnyFlags(UhtValidationOptions.Shadowing))
 			{
 
 				// First check for duplicate name in self and then duplicate name in parents
@@ -1571,7 +1567,15 @@ namespace EpicGames.UHT.Types
 						// fixed up first before the exception it is removed.
 						bool existingPropDeprecated = existingProperty.PropertyFlags.HasAnyFlags(EPropertyFlags.Deprecated);
 						bool newPropDeprecated = PropertyCategory == UhtPropertyCategory.Member && PropertyFlags.HasAnyFlags(EPropertyFlags.Deprecated);
-						if (!newPropDeprecated && !existingPropDeprecated)
+
+						//@TODO: Work around to the issue that C++ UHT does not recognize shadowing errors between function parameters and properties that are 
+						// in the same structure but appear later in the file
+						// Note: The old UHT would check for all property type where this version only checks for member variables.
+						// In the old code if you defined the property after the function with argument with the same name, UHT would
+						// not issue an error.  However, if the property was defined PRIOR to the function with the matching argument name,
+						// UHT would generate an error.
+						bool appearsLater = existingProperty.HeaderFile == HeaderFile && existingProperty.LineNumber > LineNumber;
+						if (!newPropDeprecated && !existingPropDeprecated && !appearsLater)
 						{
 							LogShadowingError(existingProperty);
 						}
