EpicGames.Perforce: Fix invalid changelist number being returned when submitting a non-shelved changelist.

diff --git a/Engine/Source/Programs/Shared/EpicGames.Perforce/PerforceConnection.cs b/Engine/Source/Programs/Shared/EpicGames.Perforce/PerforceConnection.cs
--- a/Engine/Source/Programs/Shared/EpicGames.Perforce/PerforceConnection.cs
+++ b/Engine/Source/Programs/Shared/EpicGames.Perforce/PerforceConnection.cs
@@ -3399,7 +3399,8 @@ namespace EpicGames.Perforce
 
 			arguments.Add($"-c{changeNumber}");
 
-			return (await CommandAsync<SubmitRecord>(connection, "submit", arguments, null, cancellationToken))[0];
+			PerforceResponseList<SubmitRecord> responses = await CommandAsync<SubmitRecord>(connection, "submit", arguments, null, cancellationToken);
+			return ParseSubmitResponses(responses, connection.Logger);
 		}
 
 		/// <summary>
@@ -3427,7 +3428,11 @@ namespace EpicGames.Perforce
 			arguments.Add($"-e{changeNumber}");
 
 			PerforceResponseList<SubmitRecord> responses = await CommandAsync<SubmitRecord>(connection, "submit", arguments, null, cancellationToken);
+			return ParseSubmitResponses(responses, connection.Logger);
+		}
 
+		static PerforceResponse<SubmitRecord> ParseSubmitResponses(PerforceResponseList<SubmitRecord> responses, ILogger logger)
+		{
 			SubmitRecord? success = null;
 			PerforceResponse<SubmitRecord>? error = null;
 			foreach (PerforceResponse<SubmitRecord> response in responses)
@@ -3438,7 +3443,7 @@ namespace EpicGames.Perforce
 				}
 				else if (response.Info != null)
 				{
-					connection.Logger.LogInformation("Submit: {Info}", response.Info.Data);
+					logger.LogInformation("Submit: {Info}", response.Info.Data);
 				}
 				else if (success == null)
 				{
