Adding missing file from previous CL # 21982447

#rnx

diff --git a/Engine/Source/Programs/Shared/EpicGames.Core/KnownLogEvents.cs b/Engine/Source/Programs/Shared/EpicGames.Core/KnownLogEvents.cs
--- a/Engine/Source/Programs/Shared/EpicGames.Core/KnownLogEvents.cs
+++ b/Engine/Source/Programs/Shared/EpicGames.Core/KnownLogEvents.cs
@@ -240,6 +240,13 @@ namespace EpicGames.Core
 		/// </summary>
 		public static EventId Systemic_HostDownIOException { get; } = new EventId(764);
 
+		/// <summary>
+		/// Microsoft SignTool cannot reach specified timestamp server 
+		/// 
+		/// Signing performs a looping retry using different servers and throws when all retries attemps have been exhausted.
+		/// </summary>
+		public static EventId Systemic_SignToolTimeStampServer { get; } = new EventId(765);
+
 		/// <summary>
 		/// Maximum code for systemic events. Add new events in the 700-799 range.
 		/// </summary>
