Adds more paths for finding the p4 binary on Mac and Linux, specifically the location as to where Homebrew puts p4.

diff --git a/Engine/Source/Programs/AutomationTool/AutomationUtils/MacHostPlatform.cs b/Engine/Source/Programs/AutomationTool/AutomationUtils/MacHostPlatform.cs
--- a/Engine/Source/Programs/AutomationTool/AutomationUtils/MacHostPlatform.cs
+++ b/Engine/Source/Programs/AutomationTool/AutomationUtils/MacHostPlatform.cs
@@ -87,10 +87,19 @@ namespace AutomationTool
 			{
 				if (P4ExePath == null)
 				{
-					P4ExePath = "/usr/bin/p4";
-					if (!File.Exists(P4ExePath))
+					string[] p4Paths = { 
+						"/usr/bin/p4", // Default path
+						"/opt/homebrew/bin/p4", // Apple Silicon Homebrew Path
+						"/usr/local/bin/p4" // Apple Intel Homebrew Path
+					};
+					
+					foreach (string path in p4Paths)
 					{
-						P4ExePath = "/usr/local/bin/p4";
+						if (File.Exists(path))
+						{
+							P4ExePath = path;
+							break;
+						}
 					}
 				}
 				return P4ExePath;
diff --git a/Engine/Source/Programs/Shared/EpicGames.Perforce/PerforceChildProcess.cs b/Engine/Source/Programs/Shared/EpicGames.Perforce/PerforceChildProcess.cs
--- a/Engine/Source/Programs/Shared/EpicGames.Perforce/PerforceChildProcess.cs
+++ b/Engine/Source/Programs/Shared/EpicGames.Perforce/PerforceChildProcess.cs
@@ -93,14 +93,22 @@ namespace EpicGames.Perforce
 		/// <returns>Path to the executable</returns>
 		public static string GetExecutable()
 		{
-			string perforceFileName;
-			if (RuntimePlatform.IsWindows)
+			string perforceFileName = "p4.exe";
+			if (!RuntimePlatform.IsWindows)
 			{
-				perforceFileName = "p4.exe";
-			}
-			else
-			{
-				perforceFileName = File.Exists("/usr/local/bin/p4") ? "/usr/local/bin/p4" : "/usr/bin/p4";
+				string[] p4Paths = { 
+					"/usr/bin/p4", // Default path
+					"/opt/homebrew/bin/p4", // Apple Silicon Homebrew Path
+					"/usr/local/bin/p4", // Apple Intel Homebrew Path
+					"[REDACTED_USER_HOME]/.linuxbrew/bin/p4" }; // Linux Homebrew path
+				foreach (string path in p4Paths)
+				{
+					if (File.Exists(path))
+					{
+						perforceFileName = path;
+						break;
+					}
+				}
 			}
 			return perforceFileName;
 		}
