Horde: Add a signature and header to bundle objects, plus a GUID type identifier for each object inside it. This will allow transcoding objects to different types when extracted from a bundle.

diff --git a/Engine/Source/Programs/Horde/Horde.Agent/Commands/Bundles/DumpCommand.cs b/Engine/Source/Programs/Horde/Horde.Agent/Commands/Bundles/DumpCommand.cs
--- a/Engine/Source/Programs/Horde/Horde.Agent/Commands/Bundles/DumpCommand.cs
+++ b/Engine/Source/Programs/Horde/Horde.Agent/Commands/Bundles/DumpCommand.cs
@@ -3,7 +3,6 @@
 using System;
 using System.Collections.Generic;
 using System.Linq;
-using System.Threading;
 using System.Threading.Tasks;
 using EpicGames.Core;
 using EpicGames.Horde.Storage;
@@ -56,10 +55,10 @@ namespace Horde.Agent.Commands.Bundles
 			int refIdx = 0;
 			foreach (BundleImport import in header.Imports)
 			{
-				logger.LogInformation("  From blob {BlobId} ({NumExports}/{TotalExports} nodes)", import.BlobId, import.Exports.Count, import.ExportCount);
+				logger.LogInformation("  From blob {BlobId} ({NumExports}/{TotalExports} nodes)", import.Id, import.Exports.Count, import.ExportCount);
 				foreach ((int exportIdx, IoHash exportHash) in import.Exports)
 				{
-					logger.LogInformation("    [{Index}] IMP {BlobId}:{ExportIdx} = {ExportHash}", refIdx, import.BlobId, exportIdx, exportHash);
+					logger.LogInformation("    [{Index}] IMP {BlobId}:{ExportIdx} = {ExportHash}", refIdx, import.Id, exportIdx, exportHash);
 					refIdx++;
 				}
 			}
diff --git a/Engine/Source/Programs/Horde/Horde.Build/Perforce/ReplicationService.cs b/Engine/Source/Programs/Horde/Horde.Build/Perforce/ReplicationService.cs
--- a/Engine/Source/Programs/Horde/Horde.Build/Perforce/ReplicationService.cs
+++ b/Engine/Source/Programs/Horde/Horde.Build/Perforce/ReplicationService.cs
@@ -62,7 +62,7 @@ namespace Horde.Build.Perforce
 	/// <summary>
 	/// Root node for a commit snapshot
 	/// </summary>
-	[TreeSerializer(typeof(ReplicationNodeSerializer))]
+	[TreeSerializer("17F251FB-6D62-48D9-B243-0F06B61BCC74", typeof(ReplicationNodeSerializer))]
 	public class ReplicationNode : TreeNode
 	{
 		/// <summary>
@@ -99,7 +99,7 @@ namespace Horde.Build.Perforce
 		}
 
 		/// <inheritdoc/>
-		public override async Task<ITreeBlob> SerializeAsync(ITreeWriter writer, CancellationToken cancellationToken)
+		public override async Task<NewTreeBlob> SerializeAsync(ITreeWriter writer, CancellationToken cancellationToken)
 		{
 			List<ITreeBlobRef> references = new List<ITreeBlobRef>();
 			references.Add(await Contents.CollapseAsync(writer, cancellationToken));
diff --git a/Engine/Source/Programs/Shared/EpicGames.Core/MemoryReader.cs b/Engine/Source/Programs/Shared/EpicGames.Core/MemoryReader.cs
--- a/Engine/Source/Programs/Shared/EpicGames.Core/MemoryReader.cs
+++ b/Engine/Source/Programs/Shared/EpicGames.Core/MemoryReader.cs
@@ -198,6 +198,17 @@ namespace EpicGames.Core
 			return new DateTime((long)encoded >> 2, (DateTimeKind)(encoded & 3UL));
 		}
 
+		/// <summary>
+		/// Reads a GUID from the memory buffer
+		/// </summary>
+		/// <param name="reader">Reader to deserialize from</param>
+		public static Guid ReadGuid(this IMemoryReader reader)
+		{
+			Guid guid = new Guid(reader.GetMemory(16).Slice(0, 16).Span);
+			reader.Advance(16);
+			return guid;
+		}
+
 		/// <summary>
 		/// Reads a sequence of bytes from the buffer
 		/// </summary>
diff --git a/Engine/Source/Programs/Shared/EpicGames.Core/MemoryWriter.cs b/Engine/Source/Programs/Shared/EpicGames.Core/MemoryWriter.cs
--- a/Engine/Source/Programs/Shared/EpicGames.Core/MemoryWriter.cs
+++ b/Engine/Source/Programs/Shared/EpicGames.Core/MemoryWriter.cs
@@ -13,6 +13,11 @@ namespace EpicGames.Core
 	/// </summary>
 	public interface IMemoryWriter
 	{
+		/// <summary>
+		/// Length of the written data
+		/// </summary>
+		int Length { get; }
+
 		/// <summary>
 		/// Gets a block of memory with at least the given size
 		/// </summary>
@@ -35,7 +40,12 @@ namespace EpicGames.Core
 		/// <summary>
 		/// The memory block to write to
 		/// </summary>
-		Memory<byte> _memory;
+		readonly Memory<byte> _memory;
+
+		/// <summary>
+		/// Length of the data that has been written
+		/// </summary>
+		int _length;
 
 		/// <summary>
 		/// Constructor
@@ -51,17 +61,20 @@ namespace EpicGames.Core
 		/// </summary>
 		public void CheckEmpty()
 		{
-			if (_memory.Length > 0)
+			if (_length < _memory.Length)
 			{
-				throw new Exception($"Serialization is not at expected offset within the output buffer ({_memory.Length} bytes unused)");
+				throw new Exception($"Serialization is not at expected offset within the output buffer ({_length}/{_memory.Length} bytes used)");
 			}
 		}
 
 		/// <inheritdoc/>
-		public Memory<byte> GetMemory(int length) => _memory;
+		public int Length => _length;
+
+		/// <inheritdoc/>
+		public Memory<byte> GetMemory(int length) => _memory.Slice(_length);
 
 		/// <inheritdoc/>
-		public void Advance(int length) => _memory = _memory.Slice(length);
+		public void Advance(int length) => _length += length;
 	}
 
 	/// <summary>
@@ -190,6 +203,21 @@ namespace EpicGames.Core
 			writer.WriteUnsignedVarInt(encoded);
 		}
 
+		/// <summary>
+		/// Writes a Guid to the memory writer
+		/// </summary>
+		/// <param name="writer">Writer to serialize to</param>
+		/// <param name="value">Value to write</param>
+		public static void WriteGuid(this IMemoryWriter writer, Guid guid)
+		{
+			Memory<byte> buffer = writer.GetMemory(16);
+			if (!guid.TryWriteBytes(buffer.Slice(0, 16).Span))
+			{
+				throw new InvalidOperationException("Unable to write guid to buffer");
+			}
+			writer.Advance(16);
+		}
+
 		/// <summary>
 		/// Appends a sequence of bytes to the buffer
 		/// </summary>
diff --git a/Engine/Source/Programs/Shared/EpicGames.Horde.Tests/BundleStoreTests.cs b/Engine/Source/Programs/Shared/EpicGames.Horde.Tests/BundleStoreTests.cs
--- a/Engine/Source/Programs/Shared/EpicGames.Horde.Tests/BundleStoreTests.cs
+++ b/Engine/Source/Programs/Shared/EpicGames.Horde.Tests/BundleStoreTests.cs
@@ -49,13 +49,13 @@ namespace EpicGames.Horde.Tests
 			{
 				ITreeWriter writer = store.CreateTreeWriter("test");
 
-				ITreeBlobRef node1 = await writer.WriteNodeAsync(new ReadOnlySequence<byte>(new byte[] { 1 }), Array.Empty<ITreeBlobRef>(), CancellationToken.None);
-				ITreeBlobRef node2 = await writer.WriteNodeAsync(new ReadOnlySequence<byte>(new byte[] { 2 }), new[] { node1 }, CancellationToken.None);
-				ITreeBlobRef node3 = await writer.WriteNodeAsync(new ReadOnlySequence<byte>(new byte[] { 3 }), new[] { node2 }, CancellationToken.None);
-				ITreeBlobRef node4 = await writer.WriteNodeAsync(new ReadOnlySequence<byte>(new byte[] { 4 }), Array.Empty<ITreeBlobRef>(), CancellationToken.None);
+				ITreeBlobRef node1 = await writer.WriteNodeAsync(Guid.Empty, new ReadOnlySequence<byte>(new byte[] { 1 }), Array.Empty<ITreeBlobRef>(), CancellationToken.None);
+				ITreeBlobRef node2 = await writer.WriteNodeAsync(Guid.Empty, new ReadOnlySequence<byte>(new byte[] { 2 }), new[] { node1 }, CancellationToken.None);
+				ITreeBlobRef node3 = await writer.WriteNodeAsync(Guid.Empty, new ReadOnlySequence<byte>(new byte[] { 3 }), new[] { node2 }, CancellationToken.None);
+				ITreeBlobRef node4 = await writer.WriteNodeAsync(Guid.Empty, new ReadOnlySequence<byte>(new byte[] { 4 }), Array.Empty<ITreeBlobRef>(), CancellationToken.None);
 
-				ITreeBlob root = TreeBlob.Create(new ReadOnlySequence<byte>(new byte[] { 5 }), new[] { node4, node3 });
-				await writer.WriteRefAsync(new RefName("test"), root.Data, root.Refs);
+				ITreeBlob root = TreeBlob.Create(Guid.Empty, new ReadOnlySequence<byte>(new byte[] { 5 }), new[] { node4, node3 });
+				await writer.WriteRefAsync(new RefName("test"), Guid.Empty, root.Data, root.Refs);
 
 				await CheckTree(root);
 			}
@@ -232,13 +232,13 @@ namespace EpicGames.Horde.Tests
 				ITreeWriter writer = store.CreateTreeWriter();
 
 				ITreeWriter childWriter1 = writer.CreateChildWriter();
-				ITreeBlobRef childNode1 = await childWriter1.WriteNodeAsync(new ReadOnlySequence<byte>(new byte[] { 4, 5, 6 }), Array.Empty<ITreeBlobRef>());
+				ITreeBlobRef childNode1 = await childWriter1.WriteNodeAsync(Guid.Empty, new ReadOnlySequence<byte>(new byte[] { 4, 5, 6 }), Array.Empty<ITreeBlobRef>());
 
 				ITreeWriter childWriter2 = writer.CreateChildWriter();
-				ITreeBlobRef childNode2 = await childWriter2.WriteNodeAsync(new ReadOnlySequence<byte>(new byte[] { 7, 8, 9 }), Array.Empty<ITreeBlobRef>());
+				ITreeBlobRef childNode2 = await childWriter2.WriteNodeAsync(Guid.Empty, new ReadOnlySequence<byte>(new byte[] { 7, 8, 9 }), Array.Empty<ITreeBlobRef>());
 
 				RefName refName = new RefName("test");
-				await writer.WriteRefAsync(refName, new ReadOnlySequence<byte>(new byte[] { 1, 2, 3 }), new[] { childNode1, childNode2 });
+				await writer.WriteRefAsync(refName, Guid.Empty, new ReadOnlySequence<byte>(new byte[] { 1, 2, 3 }), new[] { childNode1, childNode2 });
 
 				Assert.AreEqual(3, blobStore.Blobs.Count);
 				Assert.AreEqual(1, blobStore.Refs.Count);
diff --git a/Engine/Source/Programs/Shared/EpicGames.Horde.Tests/BundleTests.cs b/Engine/Source/Programs/Shared/EpicGames.Horde.Tests/BundleTests.cs
--- a/Engine/Source/Programs/Shared/EpicGames.Horde.Tests/BundleTests.cs
+++ b/Engine/Source/Programs/Shared/EpicGames.Horde.Tests/BundleTests.cs
@@ -107,18 +107,18 @@ namespace EpicGames.Horde.Tests
 			BundleHeader oldHeader;
 			{
 				List<BundleImport> imports = new List<BundleImport>();
-				imports.Add(new BundleImport(new BlobId("import1"), 10, new (int, IoHash)[] { (5, IoHash.Compute(Encoding.UTF8.GetBytes("blob1"))) }));
-				imports.Add(new BundleImport(new BlobId("import2"), 20, new (int, IoHash)[] { (6, IoHash.Compute(Encoding.UTF8.GetBytes("blob2"))) }));
+				imports.Add(new BundleImport("import1", 10, new (int, IoHash)[] { (5, IoHash.Compute(Encoding.UTF8.GetBytes("blob1"))) }));
+				imports.Add(new BundleImport("import2", 20, new (int, IoHash)[] { (6, IoHash.Compute(Encoding.UTF8.GetBytes("blob2"))) }));
 
 				List<BundleExport> exports = new List<BundleExport>();
-				exports.Add(new BundleExport(IoHash.Compute(Encoding.UTF8.GetBytes("export1")), 2, new int[] { 1, 2 }));
-				exports.Add(new BundleExport(IoHash.Compute(Encoding.UTF8.GetBytes("export2")), 3, new int[] { 3 }));
+				exports.Add(new BundleExport(IoHash.Compute(Encoding.UTF8.GetBytes("export1")), 0, 2, new int[] { 1, 2 }));
+				exports.Add(new BundleExport(IoHash.Compute(Encoding.UTF8.GetBytes("export2")), 0, 3, new int[] { 3 }));
 
 				List<BundlePacket> packets = new List<BundlePacket>();
 				packets.Add(new BundlePacket(20, 40));
 				packets.Add(new BundlePacket(10, 20));
 
-				oldHeader = new BundleHeader(BundleCompressionFormat.LZ4, imports, exports, packets);
+				oldHeader = new BundleHeader(BundleCompressionFormat.LZ4, new Guid[1], imports, exports, packets);
 			}
 
 			ByteArrayBuilder writer = new ByteArrayBuilder();
@@ -133,7 +133,7 @@ namespace EpicGames.Horde.Tests
 				BundleImport oldImport = oldHeader.Imports[idx];
 				BundleImport newImport = newHeader.Imports[idx];
 
-				Assert.AreEqual(oldImport.BlobId, newImport.BlobId);
+				Assert.AreEqual(oldImport.Id, newImport.Id);
 				Assert.AreEqual(oldImport.ExportCount, newImport.ExportCount);
 				Assert.AreEqual(oldImport.Exports.Count, newImport.Exports.Count);
 
diff --git a/Engine/Source/Programs/Shared/EpicGames.Horde/Storage/Bundles/Bundle.cs b/Engine/Source/Programs/Shared/EpicGames.Horde/Storage/Bundles/Bundle.cs
--- a/Engine/Source/Programs/Shared/EpicGames.Horde/Storage/Bundles/Bundle.cs
+++ b/Engine/Source/Programs/Shared/EpicGames.Horde/Storage/Bundles/Bundle.cs
@@ -1,13 +1,14 @@
 // Copyright Epic Games, Inc. All Rights Reserved.
 
 using EpicGames.Core;
-using EpicGames.Serialization;
 using System;
 using System.Buffers;
+using System.Buffers.Binary;
 using System.Collections.Generic;
 using System.Diagnostics.CodeAnalysis;
 using System.IO;
 using System.Linq;
+using System.Text;
 using System.Threading;
 using System.Threading.Tasks;
 
@@ -25,14 +26,16 @@ namespace EpicGames.Horde.Storage.Bundles
 		Initial = 0,
 
 		/// <summary>
-		/// Added compression format enum
+		/// Last item in the enum. Used for <see cref="Latest"/>
 		/// </summary>
-		CompressionFormat = 1,
+		LatestPlusOne,
 
+#pragma warning disable CA1069 // Enums values should not be duplicated
 		/// <summary>
 		/// The current version number
 		/// </summary>
-		Current = CompressionFormat,
+		Latest = (int)LatestPlusOne - 1,
+#pragma warning restore CA1069 // Enums values should not be duplicated
 	}
 
 	/// <summary>
@@ -97,6 +100,7 @@ namespace EpicGames.Horde.Storage.Bundles
 		{
 			ByteArrayBuilder builder = new ByteArrayBuilder();
 			Header.Write(builder);
+
 			builder.WriteUnsignedVarInt((ulong)Payload.Length);
 
 			ReadOnlySequenceBuilder<byte> sequence = new ReadOnlySequenceBuilder<byte>();
@@ -128,36 +132,48 @@ namespace EpicGames.Horde.Storage.Bundles
 	/// </summary>
 	public class BundleHeader
 	{
+		/// <summary>
+		/// Signature bytes
+		/// </summary>
+		public static ReadOnlyMemory<byte> Signature { get; } = Encoding.UTF8.GetBytes("UEBN");
+
 		/// <summary>
 		/// Compression format for the bundle
 		/// </summary>
 		public BundleCompressionFormat CompressionFormat { get; }
 
+		/// <summary>
+		/// List of object type ids in this bundle
+		/// </summary>
+		public IReadOnlyList<Guid> TypeIds { get; }
+
 		/// <summary>
 		/// References to nodes in other bundles
 		/// </summary>
-		public IReadOnlyList<BundleImport> Imports { get; } = new List<BundleImport>();
+		public IReadOnlyList<BundleImport> Imports { get; }
 
 		/// <summary>
 		/// Nodes exported from this bundle
 		/// </summary>
-		public IReadOnlyList<BundleExport> Exports { get; } = new List<BundleExport>();
+		public IReadOnlyList<BundleExport> Exports { get; }
 
 		/// <summary>
 		/// List of data packets within this bundle
 		/// </summary>
-		public IReadOnlyList<BundlePacket> Packets { get; } = new List<BundlePacket>();
+		public IReadOnlyList<BundlePacket> Packets { get; }
 
 		/// <summary>
 		/// Constructs a new bundle header
 		/// </summary>
 		/// <param name="compressionFormat">Compression format for bundle packets</param>
+		/// <param name="typeIds">List of type ids for this bundle</param>
 		/// <param name="imports">Imports from other bundles</param>
 		/// <param name="exports">Exports for nodes</param>
 		/// <param name="packets">Compression packets within the bundle</param>
-		public BundleHeader(BundleCompressionFormat compressionFormat, IReadOnlyList<BundleImport> imports, IReadOnlyList<BundleExport> exports, IReadOnlyList<BundlePacket> packets)
+		public BundleHeader(BundleCompressionFormat compressionFormat, IReadOnlyList<Guid> typeIds, IReadOnlyList<BundleImport> imports, IReadOnlyList<BundleExport> exports, IReadOnlyList<BundlePacket> packets)
 		{
 			CompressionFormat = compressionFormat;
+			TypeIds = typeIds;
 			Imports = imports;
 			Exports = exports;
 			Packets = packets;
@@ -170,24 +186,66 @@ namespace EpicGames.Horde.Storage.Bundles
 		/// <returns>New header object</returns>
 		public BundleHeader(IMemoryReader reader)
 		{
+			ReadOnlyMemory<byte> prelude = reader.GetMemory(8);
+			CheckPrelude(prelude);
+			reader.Advance(8);
+
 			BundleVersion version = (BundleVersion)reader.ReadUnsignedVarInt();
-			if (version > BundleVersion.Current)
+			if (version > BundleVersion.Latest)
 			{
-				throw new InvalidDataException($"Unknown bundle version {(int)version}. Max supported is {(int)BundleVersion.Current}.");
+				throw new InvalidDataException($"Unknown bundle version {(int)version}. Max supported is {(int)BundleVersion.Latest}.");
 			}
 
-			if (version >= BundleVersion.CompressionFormat)
+			CompressionFormat = (BundleCompressionFormat)reader.ReadUnsignedVarInt();
+			TypeIds = reader.ReadVariableLengthArray(() => reader.ReadGuid());
+			Imports = reader.ReadVariableLengthArray(() => new BundleImport(reader));
+			Exports = reader.ReadVariableLengthArray(() => new BundleExport(reader));
+			Packets = reader.ReadVariableLengthArray(() => new BundlePacket(reader));
+		}
+
+		/// <summary>
+		/// Reads a bundle header from a stream
+		/// </summary>
+		/// <param name="stream">Stream to read from</param>
+		/// <param name="cancellationToken">Cancellation token for the stream</param>
+		/// <returns>New header</returns>
+		public static async Task<BundleHeader> ReadAsync(Stream stream, CancellationToken cancellationToken)
+		{
+			byte[] prelude = new byte[8];
+
+			int length = await stream.ReadAsync(prelude, cancellationToken);
+			if (length != 8)
 			{
-				CompressionFormat = (BundleCompressionFormat)reader.ReadUnsignedVarInt();
+				throw new InvalidDataException("Unexpected end of stream while reading prelude data");
 			}
-			else
+			CheckPrelude(prelude.AsMemory(0, 8));
+
+			int headerLength = BinaryPrimitives.ReadInt32BigEndian(prelude.AsSpan(4));
+			using (IMemoryOwner<byte> owner = MemoryPool<byte>.Shared.Rent(headerLength))
 			{
-				CompressionFormat = BundleCompressionFormat.LZ4;
+				Memory<byte> header = owner.Memory.Slice(headerLength);
+				prelude.CopyTo(header);
+
+				int readLength = await stream.ReadAsync(header.Slice(8), cancellationToken);
+				if (readLength != headerLength)
+				{
+					throw new InvalidDataException("Unexpected end of file in stream");
+				}
+
+				return new BundleHeader(new MemoryReader(header));
 			}
+		}
 
-			Imports = reader.ReadVariableLengthArray(() => new BundleImport(reader));
-			Exports = reader.ReadVariableLengthArray(() => new BundleExport(reader));
-			Packets = reader.ReadVariableLengthArray(() => new BundlePacket(reader));
+		/// <summary>
+		/// Validates that the prelude bytes for a bundle header are correct
+		/// </summary>
+		/// <param name="prelude">The prelude bytes</param>
+		static void CheckPrelude(ReadOnlyMemory<byte> prelude)
+		{
+			if (!prelude.Slice(0, 4).Span.SequenceEqual(Signature.Span))
+			{
+				throw new InvalidDataException("Invalid signature bytes for bundle. Corrupt data?");
+			}
 		}
 
 		/// <summary>
@@ -196,12 +254,21 @@ namespace EpicGames.Horde.Storage.Bundles
 		/// <param name="writer">Writer to serialize to</param>
 		public void Write(IMemoryWriter writer)
 		{
-			writer.WriteUnsignedVarInt((ulong)BundleVersion.Current);
+			int initialLength = writer.Length;
+
+			Memory<byte> prelude = writer.GetMemory(8);
+			Signature.CopyTo(prelude);
+			writer.Advance(8);
+
+			writer.WriteUnsignedVarInt((ulong)BundleVersion.Latest);
 
 			writer.WriteUnsignedVarInt((ulong)CompressionFormat);
+			writer.WriteVariableLengthArray(TypeIds, x => writer.WriteGuid(x));
 			writer.WriteVariableLengthArray(Imports, x => x.Write(writer));
 			writer.WriteVariableLengthArray(Exports, x => x.Write(writer));
 			writer.WriteVariableLengthArray(Packets, x => x.Write(writer));
+
+			BinaryPrimitives.WriteInt32BigEndian(prelude.Slice(4).Span, writer.Length - initialLength);
 		}
 	}
 
@@ -211,9 +278,9 @@ namespace EpicGames.Horde.Storage.Bundles
 	public class BundleImport
 	{
 		/// <summary>
-		/// Blob containing the bundle data.
+		/// Identifier for the referenced bundle; typically a path, token, or blob identifier.
 		/// </summary>
-		public BlobId BlobId { get; }
+		public Utf8String Id { get; }
 
 		/// <summary>
 		/// Number of exports from this blob.
@@ -228,9 +295,9 @@ namespace EpicGames.Horde.Storage.Bundles
 		/// <summary>
 		/// Constructor
 		/// </summary>
-		public BundleImport(BlobId blobId, int exportCount, IReadOnlyList<(int, IoHash)> exports)
+		public BundleImport(Utf8String id, int exportCount, IReadOnlyList<(int, IoHash)> exports)
 		{
-			BlobId = blobId;
+			Id = id;
 			ExportCount = exportCount;
 			Exports = exports;
 		}
@@ -241,7 +308,7 @@ namespace EpicGames.Horde.Storage.Bundles
 		/// <param name="reader">Reader to deserialize from</param>
 		internal BundleImport(IMemoryReader reader)
 		{
-			BlobId = reader.ReadBlobId();
+			Id = reader.ReadUtf8String();
 
 			ExportCount = (int)reader.ReadUnsignedVarInt();
 
@@ -263,7 +330,7 @@ namespace EpicGames.Horde.Storage.Bundles
 		/// <param name="writer">Writer to serialize to</param>
 		public void Write(IMemoryWriter writer)
 		{
-			writer.WriteBlobId(BlobId);
+			writer.WriteUtf8String(Id);
 			writer.WriteUnsignedVarInt(ExportCount);
 
 			writer.WriteUnsignedVarInt(Exports.Count);
@@ -333,11 +400,16 @@ namespace EpicGames.Horde.Storage.Bundles
 		/// </summary>
 		public IoHash Hash { get; }
 
+		/// <summary>
+		/// Index into the type list for this node
+		/// </summary>
+		public int TypeIndex { get; }
+
 		/// <summary>
 		/// Uncompressed length of this node
 		/// </summary>
 		public int Length { get; }
-		
+
 		/// <summary>
 		/// Nodes referenced by this export. Indices in this array correspond to a lookup table consisting
         /// of the imported nodes in the order they are declared in the header, followed by nodes listed in the
@@ -348,9 +420,10 @@ namespace EpicGames.Horde.Storage.Bundles
 		/// <summary>
 		/// Constructor
 		/// </summary>
-		public BundleExport(IoHash hash, int length, IReadOnlyList<int> references)
+		public BundleExport(IoHash hash, int typeIndex, int length, IReadOnlyList<int> references)
 		{
 			Hash = hash;
+			TypeIndex = typeIndex;
 			Length = length;
 			References = references;
 		}
@@ -362,6 +435,7 @@ namespace EpicGames.Horde.Storage.Bundles
 		public BundleExport(IMemoryReader reader)
 		{
 			Hash = reader.ReadIoHash();
+			TypeIndex = (int)reader.ReadUnsignedVarInt();
 			Length = (int)reader.ReadUnsignedVarInt();
 
 			int numReferences = (int)reader.ReadUnsignedVarInt();
@@ -382,6 +456,7 @@ namespace EpicGames.Horde.Storage.Bundles
 		public void Write(IMemoryWriter writer)
 		{
 			writer.WriteIoHash(Hash);
+			writer.WriteUnsignedVarInt(TypeIndex);
 			writer.WriteUnsignedVarInt(Length);
 
 			writer.WriteUnsignedVarInt(References.Count);
diff --git a/Engine/Source/Programs/Shared/EpicGames.Horde/Storage/Bundles/BundleStore.cs b/Engine/Source/Programs/Shared/EpicGames.Horde/Storage/Bundles/BundleStore.cs
--- a/Engine/Source/Programs/Shared/EpicGames.Horde/Storage/Bundles/BundleStore.cs
+++ b/Engine/Source/Programs/Shared/EpicGames.Horde/Storage/Bundles/BundleStore.cs
@@ -137,14 +137,16 @@ namespace EpicGames.Horde.Storage.Bundles
 		/// </summary>
 		class ExportedNodeState : ImportedNodeState
 		{
+			public readonly Guid TypeId;
 			public readonly PacketInfo PacketInfo;
 			public readonly int Offset;
 			public readonly int Length;
 			public readonly IReadOnlyList<NodeInfo> References;
 
-			public ExportedNodeState(BundleInfo bundleInfo, int exportIdx, PacketInfo packetInfo, int offset, int length, IReadOnlyList<NodeInfo> references)
+			public ExportedNodeState(BundleInfo bundleInfo, int exportIdx, Guid typeId, PacketInfo packetInfo, int offset, int length, IReadOnlyList<NodeInfo> references)
 				: base(bundleInfo, exportIdx)
 			{
+				TypeId = typeId;
 				PacketInfo = packetInfo;
 				Offset = offset;
 				Length = length;
@@ -158,17 +160,20 @@ namespace EpicGames.Horde.Storage.Bundles
 		/// </summary>
 		class InMemoryNodeState : NodeState, ITreeBlob
 		{
+			public readonly Guid TypeId;
 			public readonly ReadOnlySequence<byte> Data;
 			public readonly IReadOnlyList<NodeInfo> References;
 
 			public BundleWriter? _writer;
 			public bool _writing; // This node is currently being written asynchronously
 
+			Guid ITreeBlob.TypeId => TypeId;
 			ReadOnlySequence<byte> ITreeBlob.Data => Data;
 			IReadOnlyList<ITreeBlobRef> ITreeBlob.Refs => References;
 
-			public InMemoryNodeState(ReadOnlySequence<byte> data, IReadOnlyList<NodeInfo> references)
+			public InMemoryNodeState(Guid typeId, ReadOnlySequence<byte> data, IReadOnlyList<NodeInfo> references)
 			{
+				TypeId = typeId;
 				Data = data;
 				References = references;
 			}
@@ -259,16 +264,16 @@ namespace EpicGames.Horde.Storage.Bundles
 			}
 
 			/// <inheritdoc/>
-			public Task<ITreeBlobRef> WriteNodeAsync(ReadOnlySequence<byte> data, IReadOnlyList<ITreeBlobRef> refs, CancellationToken cancellationToken = default)
+			public Task<ITreeBlobRef> WriteNodeAsync(Guid typeId, ReadOnlySequence<byte> data, IReadOnlyList<ITreeBlobRef> refs, CancellationToken cancellationToken = default)
 			{
 				IReadOnlyList<NodeInfo> typedRefs = refs.Select(x => (NodeInfo)x).ToList();
 
-				IoHash hash = ComputeHash(data, typedRefs);
+				IoHash hash = ComputeHash(typeId, data, typedRefs);
 
 				NodeInfo? node;
 				if (!HashToNode.TryGetValue(hash, out node))
 				{
-					InMemoryNodeState state = new InMemoryNodeState(data, typedRefs);
+					InMemoryNodeState state = new InMemoryNodeState(typeId, data, typedRefs);
 					NodeInfo newNode = new NodeInfo(_owner, hash, state);
 
 					// Try to add the node again. If we succeed, check if we need to flush the current bundle being built.
@@ -392,7 +397,7 @@ namespace EpicGames.Horde.Storage.Bundles
 					if (_queueNodes.Count == 0 || _queueNodes[^1] != root)
 					{
 						ITreeBlob blob = await root.GetTargetAsync(cancellationToken);
-						InMemoryNodeState state = new InMemoryNodeState(blob.Data, blob.Refs.ConvertAll(x => (NodeInfo)x));
+						InMemoryNodeState state = new InMemoryNodeState(blob.TypeId, blob.Data, blob.Refs.ConvertAll(x => (NodeInfo)x));
 						WriteNode(new NodeInfo(_owner, root.Hash, state));
 					}
 					last = await FlushInternalAsync(cancellationToken);
@@ -446,7 +451,7 @@ namespace EpicGames.Horde.Storage.Bundles
 				BundleHeader header = bundle.Header;
 
 				// Write it to storage
-				BlobId blobId = await _owner._blobStore.WriteBlobAsync(bundle.AsSequence(), bundle.Header.Imports.Select(x => x.BlobId).ToList(), _prefix, cancellationToken);
+				BlobId blobId = await _owner._blobStore.WriteBlobAsync(bundle.AsSequence(), bundle.Header.Imports.Select(x => new BlobId(x.Id)).ToList(), _prefix, cancellationToken);
 				string cacheKey = GetBundleCacheKey(blobId);
 				_owner.AddBundleToCache(cacheKey, bundle);
 
@@ -470,7 +475,8 @@ namespace EpicGames.Horde.Storage.Bundles
 						InMemoryNodeState inMemoryState = (InMemoryNodeState)writeNodes[exportIdx].State;
 
 						BundleExport export = header.Exports[exportIdx];
-						writeNodes[exportIdx].SetState(new ExportedNodeState(bundleInfo, exportIdx, packetInfo, nodeOffset, export.Length, inMemoryState.References));
+						Guid typeId = header.TypeIds[export.TypeIndex];
+						writeNodes[exportIdx].SetState(new ExportedNodeState(bundleInfo, exportIdx, typeId, packetInfo, nodeOffset, export.Length, inMemoryState.References));
 
 						nodeOffset += header.Exports[exportIdx].Length;
 					}
@@ -549,12 +555,16 @@ namespace EpicGames.Horde.Storage.Bundles
 						ImportedNodeState importedNodeState = (ImportedNodeState)importedNode.State;
 						exportInfos[idx] = (importedNodeState.ExportIdx, importedNode.Hash);
 					}
-					imports.Add(new BundleImport(bundleInfo.BlobId, bundleInfo.Exports.Length, exportInfos));
+					imports.Add(new BundleImport(bundleInfo.BlobId.Inner, bundleInfo.Exports.Length, exportInfos));
 				}
 
 				// Preallocate data in the encoded buffer to reduce fragmentation if we have to resize
 				ByteArrayBuilder builder = new ByteArrayBuilder(options.MinCompressionPacketSize * 2);
 
+				// List of type ids for the bundle
+				List<Guid> typeIds = new List<Guid>(10);
+				Dictionary<Guid, int> typeIdToIndex = new Dictionary<Guid, int>(typeIds.Capacity);
+
 				// Create the export list
 				List<BundleExport> exports = new List<BundleExport>();
 				List<BundlePacket> packets = new List<BundlePacket>();
@@ -578,9 +588,18 @@ namespace EpicGames.Horde.Storage.Bundles
 						blockSize = 0;
 					}
 
+					// Get the type index
+					int typeIndex;
+					if (!typeIdToIndex.TryGetValue(nodeState.TypeId, out typeIndex))
+					{
+						typeIndex = typeIds.Count;
+						typeIds.Add(nodeState.TypeId);
+						typeIdToIndex.Add(nodeState.TypeId, typeIndex);
+					}
+
 					// Create the export for this node
 					int[] references = nodeState.References.Select(x => nodeToIndex[x]).ToArray();
-					BundleExport export = new BundleExport(node.Hash, (int)nodeData.Length, references);
+					BundleExport export = new BundleExport(node.Hash, typeIndex, (int)nodeData.Length, references);
 					exports.Add(export);
 					nodeToIndex[node] = nodeToIndex.Count;
 
@@ -597,7 +616,7 @@ namespace EpicGames.Horde.Storage.Bundles
 				FlushPacket(options.CompressionFormat, blockSegments, blockSize, builder, packets);
 
 				// Flush the data
-				BundleHeader header = new BundleHeader(options.CompressionFormat, imports, exports, packets);
+				BundleHeader header = new BundleHeader(options.CompressionFormat, typeIds, imports, exports, packets);
 				return new Bundle(header, builder.AsSequence());
 			}
 
@@ -763,7 +782,7 @@ namespace EpicGames.Horde.Storage.Bundles
 					ReadOnlyMemory<byte> nodeData = decodedPacket.AsMemory(nodeOffset, export.Length);
 					List<NodeInfo> nodeRefs = export.References.ConvertAll(x => nodes[x]);
 
-					NodeInfo node = new NodeInfo(this, export.Hash, new InMemoryNodeState(new ReadOnlySequence<byte>(nodeData), nodeRefs));
+					NodeInfo node = new NodeInfo(this, export.Hash, new InMemoryNodeState(header.TypeIds[export.TypeIndex], new ReadOnlySequence<byte>(nodeData), nodeRefs));
 					nodes.Add(node);
 
 					nodeOffset += export.Length;
@@ -795,7 +814,7 @@ namespace EpicGames.Horde.Storage.Bundles
 			ReadOnlyMemory<byte> packetData = DecodePacket(exportedState.BundleInfo.BlobId, exportedState.PacketInfo, bundle.Payload);
 			ReadOnlyMemory<byte> nodeData = packetData.Slice(exportedState.Offset, exportedState.Length);
 
-			return new InMemoryNodeState(new ReadOnlySequence<byte>(nodeData), exportedState.References);
+			return new InMemoryNodeState(exportedState.TypeId, new ReadOnlySequence<byte>(nodeData), exportedState.References);
 		}
 
 		async ValueTask<ExportedNodeState> GetExportedStateAsync(NodeInfo node, CancellationToken cancellationToken)
@@ -930,7 +949,8 @@ namespace EpicGames.Horde.Storage.Bundles
 						BundleExport export = header.Exports[exportIdx];
 						List<NodeInfo> nodeRefs = export.References.ConvertAll(x => nodes[x]);
 
-						ExportedNodeState state = new ExportedNodeState(bundleInfo, exportIdx, packetInfo, nodeOffset, export.Length, nodeRefs);
+						Guid typeId = header.TypeIds[export.TypeIndex];
+						ExportedNodeState state = new ExportedNodeState(bundleInfo, exportIdx, typeId, packetInfo, nodeOffset, export.Length, nodeRefs);
 
 						NodeInfo? node = bundleInfo.Exports[exportIdx];
 						if (node == null)
@@ -959,7 +979,7 @@ namespace EpicGames.Horde.Storage.Bundles
 			List<NodeInfo> indexedNodes = new List<NodeInfo>(header.Imports.Sum(x => x.Exports.Count) + header.Exports.Count);
 			foreach (BundleImport import in header.Imports)
 			{
-				BundleInfo importBundle = context.FindOrAddBundle(import.BlobId, import.ExportCount);
+				BundleInfo importBundle = context.FindOrAddBundle(new BlobId(import.Id), import.ExportCount);
 				foreach ((int exportIdx, IoHash exportHash) in import.Exports)
 				{
 					NodeInfo? node = importBundle.Exports[exportIdx];
@@ -1012,18 +1032,22 @@ namespace EpicGames.Horde.Storage.Bundles
 
 		#region Writing bundles
 
-		static IoHash ComputeHash(ReadOnlySequence<byte> data, IReadOnlyList<NodeInfo> references)
+		static IoHash ComputeHash(Guid typeId, ReadOnlySequence<byte> data, IReadOnlyList<NodeInfo> references)
 		{
-			byte[] buffer = new byte[IoHash.NumBytes * (references.Count + 1)];
+			byte[] buffer = new byte[16 + IoHash.NumBytes * (references.Count + 1)];
 
 			Span<byte> span = buffer;
+
+			typeId.TryWriteBytes(buffer.AsSpan(0, 16));
+			span = span[16..];
+
 			for (int idx = 0; idx < references.Count; idx++)
 			{
 				references[idx].Hash.CopyTo(span);
 				span = span[IoHash.NumBytes..];
 			}
-			IoHash.Compute(data).CopyTo(span);
 
+			IoHash.Compute(data).CopyTo(span);
 			return IoHash.Compute(buffer);
 		}
 
diff --git a/Engine/Source/Programs/Shared/EpicGames.Horde/Storage/ITreeStore.cs b/Engine/Source/Programs/Shared/EpicGames.Horde/Storage/ITreeStore.cs
--- a/Engine/Source/Programs/Shared/EpicGames.Horde/Storage/ITreeStore.cs
+++ b/Engine/Source/Programs/Shared/EpicGames.Horde/Storage/ITreeStore.cs
@@ -14,6 +14,11 @@ namespace EpicGames.Horde.Storage
 	/// </summary>
 	public interface ITreeBlob
 	{
+		/// <summary>
+		/// Unique id for the type of this node
+		/// </summary>
+		Guid TypeId { get; }
+
 		/// <summary>
 		/// Data for the node
 		/// </summary>
@@ -102,11 +107,12 @@ namespace EpicGames.Horde.Storage
 		/// <summary>
 		/// Adds a new node to this tree
 		/// </summary>
+		/// <param name="typeId">Unique identifier for the type of this node</param>
 		/// <param name="data">Data for the node</param>
 		/// <param name="refs">References to other nodes</param>
 		/// <param name="cancellationToken">Cancellation token for the operation</param>
 		/// <returns>Reference to the node that was added</returns>
-		Task<ITreeBlobRef> WriteNodeAsync(ReadOnlySequence<byte> data, IReadOnlyList<ITreeBlobRef> refs, CancellationToken cancellationToken = default);
+		Task<ITreeBlobRef> WriteNodeAsync(Guid typeId, ReadOnlySequence<byte> data, IReadOnlyList<ITreeBlobRef> refs, CancellationToken cancellationToken = default);
 
 		/// <summary>
 		/// Flushes the current state using a new ref name.
@@ -125,11 +131,13 @@ namespace EpicGames.Horde.Storage
 	{
 		class TreeBlobImpl : ITreeBlob
 		{
+			public Guid TypeId { get; }
 			public ReadOnlySequence<byte> Data { get; }
 			public IReadOnlyList<ITreeBlobRef> Refs { get; }
 
-			public TreeBlobImpl(ReadOnlySequence<byte> data, IReadOnlyList<ITreeBlobRef> refs)
+			public TreeBlobImpl(Guid typeId, ReadOnlySequence<byte> data, IReadOnlyList<ITreeBlobRef> refs)
 			{
+				TypeId = typeId;
 				Data = data;
 				Refs = refs;
 			}
@@ -138,9 +146,9 @@ namespace EpicGames.Horde.Storage
 		/// <summary>
 		/// Create a blob from the given parameters
 		/// </summary>
-		public static ITreeBlob Create(ReadOnlySequence<byte> Data, IReadOnlyList<ITreeBlobRef> Refs)
+		public static ITreeBlob Create(Guid typeId, ReadOnlySequence<byte> data, IReadOnlyList<ITreeBlobRef> refs)
 		{
-			return new TreeBlobImpl(Data, Refs);
+			return new TreeBlobImpl(typeId, data, refs);
 		}
 	}
 
@@ -189,14 +197,15 @@ namespace EpicGames.Horde.Storage
 		/// Writes a ref using the given root data
 		/// </summary>
 		/// <param name="writer">Writer to modify</param>
+		/// <param name="typeId">Type of the root object</param>
 		/// <param name="name">Name of the ref</param>
 		/// <param name="data">Data for the node</param>
 		/// <param name="refs">References to other nodes</param>
 		/// <param name="cancellationToken">Cancellation token for the operation</param>
 		/// <returns>Reference to the node that was added</returns>
-		public static async Task WriteRefAsync(this ITreeWriter writer, RefName name, ReadOnlySequence<byte> data, IReadOnlyList<ITreeBlobRef> refs, CancellationToken cancellationToken = default)
+		public static async Task WriteRefAsync(this ITreeWriter writer, RefName name, Guid typeId, ReadOnlySequence<byte> data, IReadOnlyList<ITreeBlobRef> refs, CancellationToken cancellationToken = default)
 		{
-			ITreeBlobRef root = await writer.WriteNodeAsync(data, refs, cancellationToken);
+			ITreeBlobRef root = await writer.WriteNodeAsync(typeId, data, refs, cancellationToken);
 			await writer.WriteRefAsync(name, root, cancellationToken);
 		}
 	}
diff --git a/Engine/Source/Programs/Shared/EpicGames.Horde/Storage/Nodes/DirectoryNode.cs b/Engine/Source/Programs/Shared/EpicGames.Horde/Storage/Nodes/DirectoryNode.cs
--- a/Engine/Source/Programs/Shared/EpicGames.Horde/Storage/Nodes/DirectoryNode.cs
+++ b/Engine/Source/Programs/Shared/EpicGames.Horde/Storage/Nodes/DirectoryNode.cs
@@ -224,7 +224,7 @@ namespace EpicGames.Horde.Storage.Nodes
 	/// <summary>
 	/// A directory node
 	/// </summary>
-	[TreeSerializer(typeof(DirectoryNodeSerializer))]
+	[TreeSerializer("9275090B-8731-4872-BD66-BD27EA3EEDF4", typeof(DirectoryNodeSerializer))]
 	public class DirectoryNode : TreeNode
 	{
 		internal const byte TypeId = (byte)'d';
@@ -517,7 +517,7 @@ namespace EpicGames.Horde.Storage.Nodes
 		#endregion
 
 		/// <inheritdoc/>
-		public override async Task<ITreeBlob> SerializeAsync(ITreeWriter writer, CancellationToken cancellationToken)
+		public override async Task<NewTreeBlob> SerializeAsync(ITreeWriter writer, CancellationToken cancellationToken)
 		{
 			List<ITreeBlobRef> refs = new List<ITreeBlobRef>();
 
diff --git a/Engine/Source/Programs/Shared/EpicGames.Horde/Storage/Nodes/FileNode.cs b/Engine/Source/Programs/Shared/EpicGames.Horde/Storage/Nodes/FileNode.cs
--- a/Engine/Source/Programs/Shared/EpicGames.Horde/Storage/Nodes/FileNode.cs
+++ b/Engine/Source/Programs/Shared/EpicGames.Horde/Storage/Nodes/FileNode.cs
@@ -74,7 +74,7 @@ namespace EpicGames.Horde.Storage.Nodes
 	/// Chunks are pushed into a tree hierarchy as data is appended to the root, with nodes of the tree also split along content-aware boundaries with <see cref="IoHash.NumBytes"/> granularity.
 	/// Once a chunk has been written to storage, it is treated as immutable.
 	/// </summary>
-	[TreeSerializer(typeof(FileNodeSerializer))]
+	[TreeSerializer("478B5170-E0A0-42D0-BB32-305BBE55F38C", typeof(FileNodeSerializer))]
 	public abstract class FileNode : TreeNode
 	{
 		/// <summary>
@@ -293,7 +293,7 @@ namespace EpicGames.Horde.Storage.Nodes
 		public override IReadOnlyList<TreeNodeRef> GetReferences() => Array.Empty<TreeNodeRef>();
 
 		/// <inheritdoc/>
-		public override Task<ITreeBlob> SerializeAsync(ITreeWriter writer, CancellationToken cancellationToken) => Task.FromResult<ITreeBlob>(new NewTreeBlob(_writtenSequence, Array.Empty<ITreeBlobRef>()));
+		public override Task<NewTreeBlob> SerializeAsync(ITreeWriter writer, CancellationToken cancellationToken) => Task.FromResult(new NewTreeBlob(_writtenSequence, Array.Empty<ITreeBlobRef>()));
 
 		/// <inheritdoc/>
 		public override ValueTask<ReadOnlyMemory<byte>> AppendDataAsync(ReadOnlyMemory<byte> newData, ChunkingOptions options, ITreeWriter writer, CancellationToken cancellationToken)
@@ -478,7 +478,7 @@ namespace EpicGames.Horde.Storage.Nodes
 		public override IReadOnlyList<TreeNodeRef> GetReferences() => _children;
 
 		/// <inheritdoc/>
-		public override async Task<ITreeBlob> SerializeAsync(ITreeWriter writer, CancellationToken cancellationToken)
+		public override async Task<NewTreeBlob> SerializeAsync(ITreeWriter writer, CancellationToken cancellationToken)
 		{
 			byte[] data = new byte[1 + sizeof(uint) + VarInt.MeasureUnsigned((ulong)_length) + 1];
 			data[0] = TypeId;
diff --git a/Engine/Source/Programs/Shared/EpicGames.Horde/Storage/TreeNode.cs b/Engine/Source/Programs/Shared/EpicGames.Horde/Storage/TreeNode.cs
--- a/Engine/Source/Programs/Shared/EpicGames.Horde/Storage/TreeNode.cs
+++ b/Engine/Source/Programs/Shared/EpicGames.Horde/Storage/TreeNode.cs
@@ -2,6 +2,7 @@
 
 using System;
 using System.Buffers;
+using System.Collections.Concurrent;
 using System.Collections.Generic;
 using System.Diagnostics;
 using System.Reflection;
@@ -46,14 +47,45 @@ namespace EpicGames.Horde.Storage
 		static class SerializerInstance<T> where T : TreeNode
 		{
 			static readonly TreeSerializerAttribute _attribute = typeof(T).GetCustomAttribute<TreeSerializerAttribute>()!;
-			public static TreeNodeSerializer<T> Serializer { get; } = (TreeNodeSerializer<T>)Activator.CreateInstance(_attribute.Type)!;
+			public static TreeNodeSerializer<T> Serializer { get; } = (TreeNodeSerializer<T>)Activator.CreateInstance(_attribute.SerializerType)!;
+		}
+
+		/// <summary>
+		/// Static cache of type ids for known types
+		/// </summary>
+		static readonly ConcurrentDictionary<Type, Guid> s_typeToTypeId = new ConcurrentDictionary<Type, Guid>();
+
+		/// <summary>
+		/// Gets the type id for a type, using the <see cref="TreeSerializerAttribute"/> declared on it.
+		/// </summary>
+		/// <returns></returns>
+		public Guid GetTypeId() => GetTypeId(GetType());
+
+		/// <summary>
+		/// Gets the type id for a type, using the <see cref="TreeSerializerAttribute"/> declared on it.
+		/// </summary>
+		/// <param name="type"></param>
+		/// <returns></returns>
+		public static Guid GetTypeId(Type type)
+		{
+			Guid typeId;
+			if (!s_typeToTypeId.TryGetValue(type, out typeId))
+			{
+				TreeSerializerAttribute? attribute = type.GetCustomAttribute<TreeSerializerAttribute>();
+				if (attribute == null)
+				{
+					throw new InvalidOperationException($"Missing {nameof(TreeSerializerAttribute)} from type {type.Name}");
+				}
+				typeId = s_typeToTypeId.GetOrAdd(type, new Guid(attribute.Type));
+			}
+			return typeId;
 		}
 
 		/// <summary>
 		/// Serialize a node to a block of memory
 		/// </summary>
 		/// <returns>New data to be stored into a blob</returns>
-		public abstract Task<ITreeBlob> SerializeAsync(ITreeWriter writer, CancellationToken cancellationToken);
+		public abstract Task<NewTreeBlob> SerializeAsync(ITreeWriter writer, CancellationToken cancellationToken);
 
 		/// <summary>
 		/// Deserialize a node from data
@@ -70,7 +102,7 @@ namespace EpicGames.Horde.Storage
 	/// <summary>
 	/// Data to be stored in a tree blob
 	/// </summary>
-	public struct NewTreeBlob : ITreeBlob
+	public struct NewTreeBlob
 	{
 		/// <summary>
 		/// The opaque data payload
@@ -126,14 +158,23 @@ namespace EpicGames.Horde.Storage
 	public sealed class TreeSerializerAttribute : Attribute
 	{
 		/// <summary>
-		/// The factory type. Should be derived from <see cref="TreeNodeSerializer{T}"/>
+		/// The type guid
 		/// </summary>
-		public Type Type { get; }
+		public string Type { get; }
+
+		/// <summary>
+		/// The serializer type. Should be derived from <see cref="TreeNodeSerializer{T}"/>
+		/// </summary>
+		public Type SerializerType { get; }
 
 		/// <summary>
 		/// Constructor
 		/// </summary>
-		public TreeSerializerAttribute(Type type) => Type = type;
+		public TreeSerializerAttribute(string type, Type serializerType)
+		{
+			Type = type;
+			SerializerType = serializerType;
+		}
 	}
 
 	/// <summary>
@@ -149,8 +190,9 @@ namespace EpicGames.Horde.Storage
 		/// <param name="cancellationToken">Cancellation token for the operation</param>
 		public static async Task<ITreeBlobRef> WriteNodeAsync(this ITreeWriter writer, TreeNode node, CancellationToken cancellationToken = default)
 		{
-			ITreeBlob blob = await node.SerializeAsync(writer, cancellationToken);
-			return await writer.WriteNodeAsync(blob.Data, blob.Refs, cancellationToken);
+			NewTreeBlob blob = await node.SerializeAsync(writer, cancellationToken);
+			Guid typeId = TreeNode.GetTypeId(node.GetType());
+			return await writer.WriteNodeAsync(typeId, blob.Data, blob.Refs, cancellationToken);
 		}
 
 		/// <summary>
diff --git a/Engine/Source/Programs/Shared/EpicGames.Horde/Storage/TreeNodeRef.cs b/Engine/Source/Programs/Shared/EpicGames.Horde/Storage/TreeNodeRef.cs
--- a/Engine/Source/Programs/Shared/EpicGames.Horde/Storage/TreeNodeRef.cs
+++ b/Engine/Source/Programs/Shared/EpicGames.Horde/Storage/TreeNodeRef.cs
@@ -181,8 +181,9 @@ namespace EpicGames.Horde.Storage
 			ITreeBlobRef? target = _target;
 			if (target == null)
 			{
-				ITreeBlob blob = await Node!.SerializeAsync(writer, cancellationToken);
-				target = await writer.WriteNodeAsync(blob.Data, blob.Refs, cancellationToken);
+				NewTreeBlob blob = await Node!.SerializeAsync(writer, cancellationToken);
+				Guid typeId = Node!.GetTypeId();
+				target = await writer.WriteNodeAsync(typeId, blob.Data, blob.Refs, cancellationToken);
 				MarkAsClean(target);
 			}
 			return target;
@@ -240,7 +241,7 @@ namespace EpicGames.Horde.Storage
 			{
 				throw new InvalidOperationException($"No serializer is defined for {typeof(T).Name}");
 			}
-			return (TreeNodeSerializer<T>)Activator.CreateInstance(attribute.Type)!;
+			return (TreeNodeSerializer<T>)Activator.CreateInstance(attribute.SerializerType)!;
 		}
 
 		/// <summary>
