Fixed issue where net enabled properties would prevent properties from finding their getter/setter methods

diff --git a/Engine/Source/Programs/Shared/EpicGames.UHT/Types/UhtClass.cs b/Engine/Source/Programs/Shared/EpicGames.UHT/Types/UhtClass.cs
--- a/Engine/Source/Programs/Shared/EpicGames.UHT/Types/UhtClass.cs
+++ b/Engine/Source/Programs/Shared/EpicGames.UHT/Types/UhtClass.cs
@@ -855,7 +855,6 @@ namespace EpicGames.UHT.Types
 							if (property.PropertyFlags.HasAnyFlags(EPropertyFlags.Net))
 							{
 								ClassExportFlags |= UhtClassExportFlags.SelfHasReplicatedProperties;
-								break;
 							}
 
 							if (!property.PropertyExportFlags.HasAnyFlags(UhtPropertyExportFlags.GetterSpecifiedNone))
