Horde: Add signature bytes and header length back into bundles.

diff --git a/Engine/Source/Programs/Shared/EpicGames.Horde/Storage/Bundles/Bundle.cs b/Engine/Source/Programs/Shared/EpicGames.Horde/Storage/Bundles/Bundle.cs
--- a/Engine/Source/Programs/Shared/EpicGames.Horde/Storage/Bundles/Bundle.cs
+++ b/Engine/Source/Programs/Shared/EpicGames.Horde/Storage/Bundles/Bundle.cs
@@ -4,10 +4,12 @@ using EpicGames.Core;
 using EpicGames.Serialization;
 using System;
 using System.Buffers;
+using System.Buffers.Binary;
 using System.Collections.Generic;
 using System.Diagnostics.CodeAnalysis;
 using System.IO;
 using System.Linq;
+using System.Text;
 using System.Threading;
 using System.Threading.Tasks;
 
@@ -25,14 +27,16 @@ namespace EpicGames.Horde.Storage.Bundles
 		Initial = 0,
 
 		/// <summary>
-		/// Added compression format enum
+		/// Last item in the enum. Used for <see cref="Latest"/>
 		/// </summary>
-		CompressionFormat = 1,
+		LatestPlusOne,
 
+#pragma warning disable CA1069 // Enums values should not be duplicated
 		/// <summary>
 		/// The current version number
 		/// </summary>
-		Current = CompressionFormat,
+		Latest = (int)LatestPlusOne - 1,
+#pragma warning restore CA1069 // Enums values should not be duplicated
 	}
 
 	/// <summary>
@@ -128,6 +132,11 @@ namespace EpicGames.Horde.Storage.Bundles
 	/// </summary>
 	public class BundleHeader
 	{
+		/// <summary>
+		/// Signature bytes
+		/// </summary>
+		public static ReadOnlyMemory<byte> Signature { get; } = Encoding.UTF8.GetBytes("UEBN");
+
 		/// <summary>
 		/// Compression format for the bundle
 		/// </summary>
@@ -136,17 +145,17 @@ namespace EpicGames.Horde.Storage.Bundles
 		/// <summary>
 		/// References to nodes in other bundles
 		/// </summary>
-		public IReadOnlyList<BundleImport> Imports { get; } = new List<BundleImport>();
+		public IReadOnlyList<BundleImport> Imports { get; }
 
 		/// <summary>
 		/// Nodes exported from this bundle
 		/// </summary>
-		public IReadOnlyList<BundleExport> Exports { get; } = new List<BundleExport>();
+		public IReadOnlyList<BundleExport> Exports { get; }
 
 		/// <summary>
 		/// List of data packets within this bundle
 		/// </summary>
-		public IReadOnlyList<BundlePacket> Packets { get; } = new List<BundlePacket>();
+		public IReadOnlyList<BundlePacket> Packets { get; }
 
 		/// <summary>
 		/// Constructs a new bundle header
@@ -170,24 +179,65 @@ namespace EpicGames.Horde.Storage.Bundles
 		/// <returns>New header object</returns>
 		public BundleHeader(IMemoryReader reader)
 		{
+			ReadOnlyMemory<byte> prelude = reader.GetMemory(8);
+			CheckPrelude(prelude);
+			reader.Advance(8);
+
 			BundleVersion version = (BundleVersion)reader.ReadUnsignedVarInt();
-			if (version > BundleVersion.Current)
+			if (version > BundleVersion.Latest)
 			{
-				throw new InvalidDataException($"Unknown bundle version {(int)version}. Max supported is {(int)BundleVersion.Current}.");
+				throw new InvalidDataException($"Unknown bundle version {(int)version}. Max supported is {(int)BundleVersion.Latest}.");
 			}
 
-			if (version >= BundleVersion.CompressionFormat)
+			CompressionFormat = (BundleCompressionFormat)reader.ReadUnsignedVarInt();
+			Imports = reader.ReadVariableLengthArray(() => new BundleImport(reader));
+			Exports = reader.ReadVariableLengthArray(() => new BundleExport(reader));
+			Packets = reader.ReadVariableLengthArray(() => new BundlePacket(reader));
+		}
+
+		/// <summary>
+		/// Reads a bundle header from a stream
+		/// </summary>
+		/// <param name="stream">Stream to read from</param>
+		/// <param name="cancellationToken">Cancellation token for the stream</param>
+		/// <returns>New header</returns>
+		public static async Task<BundleHeader> ReadAsync(Stream stream, CancellationToken cancellationToken)
+		{
+			byte[] prelude = new byte[8];
+
+			int length = await stream.ReadAsync(prelude, cancellationToken);
+			if (length != 8)
 			{
-				CompressionFormat = (BundleCompressionFormat)reader.ReadUnsignedVarInt();
+				throw new InvalidDataException("Unexpected end of stream while reading prelude data");
 			}
-			else
+			CheckPrelude(prelude);
+
+			int headerLength = BinaryPrimitives.ReadInt32BigEndian(prelude.AsSpan(4));
+			using (IMemoryOwner<byte> owner = MemoryPool<byte>.Shared.Rent(headerLength))
 			{
-				CompressionFormat = BundleCompressionFormat.LZ4;
+				Memory<byte> header = owner.Memory.Slice(headerLength);
+				prelude.CopyTo(header);
+
+				int readLength = await stream.ReadAsync(header.Slice(8), cancellationToken);
+				if (readLength != headerLength)
+				{
+					throw new InvalidDataException("Unexpected end of file in stream");
+				}
+
+				return new BundleHeader(new MemoryReader(header));
 			}
+		}
 
-			Imports = reader.ReadVariableLengthArray(() => new BundleImport(reader));
-			Exports = reader.ReadVariableLengthArray(() => new BundleExport(reader));
-			Packets = reader.ReadVariableLengthArray(() => new BundlePacket(reader));
+		/// <summary>
+		/// Validates that the prelude bytes for a bundle header are correct
+		/// </summary>
+		/// <param name="prelude">The prelude bytes</param>
+		static void CheckPrelude(ReadOnlyMemory<byte> prelude)
+		{
+			if (!prelude.Slice(0, 4).Span.SequenceEqual(Signature.Span))
+			{
+				throw new InvalidDataException("Invalid signature bytes for bundle. Corrupt data?");
+			}
 		}
 
 		/// <summary>
@@ -196,12 +246,23 @@ namespace EpicGames.Horde.Storage.Bundles
 		/// <param name="writer">Writer to serialize to</param>
 		public void Write(IMemoryWriter writer)
 		{
-			writer.WriteUnsignedVarInt((ulong)BundleVersion.Current);
+			int initialLength = writer.Length;
+
+			// Write the header prelude; a fixed-length block containing a signature and header size.
+			Memory<byte> prelude = writer.GetMemory(8);
+			Signature.CopyTo(prelude);
+			writer.Advance(8);
+
+			// Write the header data
+			writer.WriteUnsignedVarInt((ulong)BundleVersion.Latest);
 
 			writer.WriteUnsignedVarInt((ulong)CompressionFormat);
 			writer.WriteVariableLengthArray(Imports, x => x.Write(writer));
 			writer.WriteVariableLengthArray(Exports, x => x.Write(writer));
 			writer.WriteVariableLengthArray(Packets, x => x.Write(writer));
+
+			// Go back and write the length of the header
+			BinaryPrimitives.WriteInt32BigEndian(prelude.Slice(4).Span, writer.Length - initialLength);
 		}
 	}
 
