Horde: Move Bundle implementation into EpicGames.Horde.Storage namespace.

diff --git a/Engine/Source/Programs/Shared/EpicGames.Horde/Storage/Bundle.cs b/Engine/Source/Programs/Shared/EpicGames.Horde/Storage/Bundle.cs
new file mode 100644
--- /dev/null
+++ b/Engine/Source/Programs/Shared/EpicGames.Horde/Storage/Bundle.cs
@@ -0,0 +1,454 @@
+// Copyright Epic Games, Inc. All Rights Reserved.
+
+using EpicGames.Core;
+using System;
+using System.Buffers;
+using System.Buffers.Binary;
+using System.Collections.Generic;
+using System.Diagnostics.CodeAnalysis;
+using System.IO;
+using System.Linq;
+using System.Text;
+using System.Threading;
+using System.Threading.Tasks;
+
+namespace EpicGames.Horde.Storage
+{
+	/// <summary>
+	/// Bundle version number
+	/// </summary>
+	[System.Diagnostics.CodeAnalysis.SuppressMessage("Design", "CA1027:Mark enums with FlagsAttribute")]
+	public enum BundleVersion
+	{
+		/// <summary>
+		/// Initial version number
+		/// </summary>
+		Initial = 0,
+
+		/// <summary>
+		/// Last item in the enum. Used for <see cref="Latest"/>
+		/// </summary>
+		LatestPlusOne,
+
+#pragma warning disable CA1069 // Enums values should not be duplicated
+		/// <summary>
+		/// The current version number
+		/// </summary>
+		Latest = (int)LatestPlusOne - 1,
+#pragma warning restore CA1069 // Enums values should not be duplicated
+	}
+
+	/// <summary>
+	/// Header for the contents of a bundle. May contain an inlined payload object containing the object data itself.
+	/// </summary>
+	public class Bundle
+	{
+		/// <summary>
+		/// Header for the bundle
+		/// </summary>
+		public BundleHeader Header { get; }
+
+		/// <summary>
+		/// Packet data as described in the header
+		/// </summary>
+		public ReadOnlySequence<byte> Payload { get; }
+
+		/// <summary>
+		/// Constructor
+		/// </summary>
+		public Bundle(BundleHeader header, ReadOnlySequence<byte> payload)
+		{
+			Header = header;
+			Payload = payload;
+		}
+
+		/// <summary>
+		/// Reads a bundle from a block of memory
+		/// </summary>
+		public Bundle(IMemoryReader reader)
+		{
+			Header = new BundleHeader(reader);
+			int length = (int)reader.ReadUnsignedVarInt();
+			Payload = new ReadOnlySequence<byte>(reader.ReadFixedLengthBytes(length));
+		}
+
+		/// <summary>
+		/// Create a bundle from a blob
+		/// </summary>
+		/// <param name="blob">Blob to parse from</param>
+		/// <returns>Bundle parsed from the given blob</returns>
+		[return: NotNullIfNotNull("blob")]
+		public static Bundle? FromBlob(IBlob? blob)
+		{
+			if (blob == null)
+			{
+				return null;
+			}
+
+			ReadOnlyMemory<byte> data = blob.Data;
+			MemoryReader reader = new MemoryReader(data);
+			Bundle bundle = new Bundle(reader);
+			reader.CheckEmpty();
+			return bundle;
+		}
+
+		/// <summary>
+		/// Serializes the bundle to a sequence of bytes
+		/// </summary>
+		/// <returns>Sequence for the bundle</returns>
+		public ReadOnlySequence<byte> AsSequence()
+		{
+			ByteArrayBuilder builder = new ByteArrayBuilder();
+			Header.Write(builder);
+			builder.WriteUnsignedVarInt((ulong)Payload.Length);
+
+			ReadOnlySequenceBuilder<byte> sequence = new ReadOnlySequenceBuilder<byte>();
+			sequence.Append(builder.AsSequence());
+			sequence.Append(Payload);
+
+			return sequence.Construct();
+		}
+	}
+
+	/// <summary>
+	/// Indicates the compression format in the bundle
+	/// </summary>
+	public enum BundleCompressionFormat
+	{
+		/// <summary>
+		/// Packets are uncompressed
+		/// </summary>
+		None = 0,
+
+		/// <summary>
+		/// LZ4 compression
+		/// </summary>
+		LZ4 = 1,
+	}
+
+	/// <summary>
+	/// Header for the contents of a bundle.
+	/// </summary>
+	public class BundleHeader
+	{
+		/// <summary>
+		/// Signature bytes
+		/// </summary>
+		public static ReadOnlyMemory<byte> Signature { get; } = Encoding.UTF8.GetBytes("UEBN");
+
+		/// <summary>
+		/// Compression format for the bundle
+		/// </summary>
+		public BundleCompressionFormat CompressionFormat { get; }
+
+		/// <summary>
+		/// References to nodes in other bundles
+		/// </summary>
+		public IReadOnlyList<BundleImport> Imports { get; }
+
+		/// <summary>
+		/// Nodes exported from this bundle
+		/// </summary>
+		public IReadOnlyList<BundleExport> Exports { get; }
+
+		/// <summary>
+		/// List of data packets within this bundle
+		/// </summary>
+		public IReadOnlyList<BundlePacket> Packets { get; }
+
+		/// <summary>
+		/// Constructs a new bundle header
+		/// </summary>
+		/// <param name="compressionFormat">Compression format for bundle packets</param>
+		/// <param name="imports">Imports from other bundles</param>
+		/// <param name="exports">Exports for nodes</param>
+		/// <param name="packets">Compression packets within the bundle</param>
+		public BundleHeader(BundleCompressionFormat compressionFormat, IReadOnlyList<BundleImport> imports, IReadOnlyList<BundleExport> exports, IReadOnlyList<BundlePacket> packets)
+		{
+			CompressionFormat = compressionFormat;
+			Imports = imports;
+			Exports = exports;
+			Packets = packets;
+		}
+
+		/// <summary>
+		/// Reads a bundle header from memory
+		/// </summary>
+		/// <param name="reader">Reader to deserialize from</param>
+		/// <returns>New header object</returns>
+		public BundleHeader(IMemoryReader reader)
+		{
+			ReadOnlyMemory<byte> prelude = reader.GetMemory(8);
+			CheckPrelude(prelude);
+			reader.Advance(8);
+
+			BundleVersion version = (BundleVersion)reader.ReadUnsignedVarInt();
+			if (version > BundleVersion.Latest)
+			{
+				throw new InvalidDataException($"Unknown bundle version {(int)version}. Max supported is {(int)BundleVersion.Latest}.");
+			}
+
+			CompressionFormat = (BundleCompressionFormat)reader.ReadUnsignedVarInt();
+			Imports = reader.ReadVariableLengthArray(() => new BundleImport(reader));
+			Exports = reader.ReadVariableLengthArray(() => new BundleExport(reader));
+			Packets = reader.ReadVariableLengthArray(() => new BundlePacket(reader));
+		}
+
+		/// <summary>
+		/// Reads a bundle header from a stream
+		/// </summary>
+		/// <param name="stream">Stream to read from</param>
+		/// <param name="cancellationToken">Cancellation token for the stream</param>
+		/// <returns>New header</returns>
+		public static async Task<BundleHeader> ReadAsync(Stream stream, CancellationToken cancellationToken)
+		{
+			byte[] prelude = new byte[8];
+
+			int length = await stream.ReadAsync(prelude, cancellationToken);
+			if (length != 8)
+			{
+				throw new InvalidDataException("Unexpected end of stream while reading prelude data");
+			}
+			CheckPrelude(prelude);
+
+			int headerLength = BinaryPrimitives.ReadInt32BigEndian(prelude.AsSpan(4));
+			using (IMemoryOwner<byte> owner = MemoryPool<byte>.Shared.Rent(headerLength))
+			{
+				Memory<byte> header = owner.Memory.Slice(headerLength);
+				prelude.CopyTo(header);
+
+				int readLength = await stream.ReadAsync(header.Slice(8), cancellationToken);
+				if (readLength != headerLength)
+				{
+					throw new InvalidDataException("Unexpected end of file in stream");
+				}
+
+				return new BundleHeader(new MemoryReader(header));
+			}
+		}
+
+		/// <summary>
+		/// Validates that the prelude bytes for a bundle header are correct
+		/// </summary>
+		/// <param name="prelude">The prelude bytes</param>
+		static void CheckPrelude(ReadOnlyMemory<byte> prelude)
+		{
+			if (!prelude.Slice(0, 4).Span.SequenceEqual(Signature.Span))
+			{
+				throw new InvalidDataException("Invalid signature bytes for bundle. Corrupt data?");
+			}
+		}
+
+		/// <summary>
+		/// Serializes a header to memory
+		/// </summary>
+		/// <param name="writer">Writer to serialize to</param>
+		public void Write(IMemoryWriter writer)
+		{
+			int initialLength = writer.Length;
+
+			// Write the header prelude; a fixed-length block containing a signature and header size.
+			Memory<byte> prelude = writer.GetMemory(8);
+			Signature.CopyTo(prelude);
+			writer.Advance(8);
+
+			// Write the header data
+			writer.WriteUnsignedVarInt((ulong)BundleVersion.Latest);
+
+			writer.WriteUnsignedVarInt((ulong)CompressionFormat);
+			writer.WriteVariableLengthArray(Imports, x => x.Write(writer));
+			writer.WriteVariableLengthArray(Exports, x => x.Write(writer));
+			writer.WriteVariableLengthArray(Packets, x => x.Write(writer));
+
+			// Go back and write the length of the header
+			BinaryPrimitives.WriteInt32BigEndian(prelude.Slice(4).Span, writer.Length - initialLength);
+		}
+	}
+
+	/// <summary>
+	/// Reference to another tree pack object
+	/// </summary>
+	public class BundleImport
+	{
+		/// <summary>
+		/// Blob containing the bundle data.
+		/// </summary>
+		public BlobId BlobId { get; }
+
+		/// <summary>
+		/// Number of exports from this blob.
+		/// </summary>
+		public int ExportCount { get; }
+
+		/// <summary>
+		/// Indexes of referenced nodes exported from this bundle
+		/// </summary>
+		public IReadOnlyList<(int, IoHash)> Exports { get; }
+
+		/// <summary>
+		/// Constructor
+		/// </summary>
+		public BundleImport(BlobId blobId, int exportCount, IReadOnlyList<(int, IoHash)> exports)
+		{
+			BlobId = blobId;
+			ExportCount = exportCount;
+			Exports = exports;
+		}
+
+		/// <summary>
+		/// Deserialize a bundle import
+		/// </summary>
+		/// <param name="reader">Reader to deserialize from</param>
+		internal BundleImport(IMemoryReader reader)
+		{
+			BlobId = reader.ReadBlobId();
+
+			ExportCount = (int)reader.ReadUnsignedVarInt();
+
+			int count = (int)reader.ReadUnsignedVarInt();
+
+			(int, IoHash)[] exports = new (int, IoHash)[count];
+			for (int idx = 0; idx < count; idx++)
+			{
+				int index = (int)reader.ReadUnsignedVarInt();
+				IoHash hash = reader.ReadIoHash();
+				exports[idx] = (index, hash);
+			}
+			Exports = exports;
+		}
+
+		/// <summary>
+		/// Serializes a bundle import
+		/// </summary>
+		/// <param name="writer">Writer to serialize to</param>
+		public void Write(IMemoryWriter writer)
+		{
+			writer.WriteBlobId(BlobId);
+			writer.WriteUnsignedVarInt(ExportCount);
+
+			writer.WriteUnsignedVarInt(Exports.Count);
+			for(int idx = 0; idx < Exports.Count; idx++)
+			{
+				(int index, IoHash hash) = Exports[idx];
+				writer.WriteUnsignedVarInt(index);
+				writer.WriteIoHash(hash);
+			}
+		}
+	}
+
+	/// <summary>
+	/// Descriptor for a compression packet
+	/// </summary>
+	public class BundlePacket
+	{
+		/// <summary>
+		/// Encoded length of the packet
+		/// </summary>
+		public int EncodedLength { get; set; }
+
+		/// <summary>
+		/// Decoded length of the packet
+		/// </summary>
+		public int DecodedLength { get; set; }
+
+		/// <summary>
+		/// Constructor
+		/// </summary>
+		/// <param name="encodedLength">Size of the encoded data</param>
+		/// <param name="decodedLength">Size of the decoded data</param>
+		public BundlePacket(int encodedLength, int decodedLength)
+		{
+			EncodedLength = encodedLength;
+			DecodedLength = decodedLength;
+		}
+
+		/// <summary>
+		/// Deserialize a packet header from a memory reader
+		/// </summary>
+		/// <param name="reader">Reader to deserialize from</param>
+		public BundlePacket(IMemoryReader reader)
+		{
+			EncodedLength = (int)reader.ReadUnsignedVarInt();
+			DecodedLength = (int)reader.ReadUnsignedVarInt();
+		}
+
+		/// <summary>
+		/// Serializes this packet header 
+		/// </summary>
+		/// <param name="writer"></param>
+		public void Write(IMemoryWriter writer)
+		{
+			writer.WriteUnsignedVarInt(EncodedLength);
+			writer.WriteUnsignedVarInt(DecodedLength);
+		}
+	}
+
+	/// <summary>
+	/// Entry for a node exported from an object
+	/// </summary>
+	public class BundleExport
+	{
+		/// <summary>
+		/// Hash of the node
+		/// </summary>
+		public IoHash Hash { get; }
+
+		/// <summary>
+		/// Uncompressed length of this node
+		/// </summary>
+		public int Length { get; }
+		
+		/// <summary>
+		/// Nodes referenced by this export. Indices in this array correspond to a lookup table consisting
+        /// of the imported nodes in the order they are declared in the header, followed by nodes listed in the
+        /// export table itself.
+        /// </summary>
+		public IReadOnlyList<int> References { get; }
+
+		/// <summary>
+		/// Constructor
+		/// </summary>
+		public BundleExport(IoHash hash, int length, IReadOnlyList<int> references)
+		{
+			Hash = hash;
+			Length = length;
+			References = references;
+		}
+
+		/// <summary>
+		/// Deserialize an export
+		/// </summary>
+		/// <param name="reader">Reader to deserialize from</param>
+		public BundleExport(IMemoryReader reader)
+		{
+			Hash = reader.ReadIoHash();
+			Length = (int)reader.ReadUnsignedVarInt();
+
+			int numReferences = (int)reader.ReadUnsignedVarInt();
+			int[] references = new int[numReferences];
+
+			for (int idx = 0; idx < numReferences; idx++)
+			{
+				references[idx] = (int)reader.ReadUnsignedVarInt();
+			}
+
+			References = references;
+		}
+
+		/// <summary>
+		/// Serializes an export
+		/// </summary>
+		/// <param name="writer">Writer to serialize to</param>
+		public void Write(IMemoryWriter writer)
+		{
+			writer.WriteIoHash(Hash);
+			writer.WriteUnsignedVarInt(Length);
+
+			writer.WriteUnsignedVarInt(References.Count);
+			for(int idx = 0; idx < References.Count; idx++)
+			{
+				writer.WriteUnsignedVarInt(References[idx]);
+			}
+		}
+	}
+}
diff --git a/Engine/Source/Programs/Shared/EpicGames.Horde/Storage/Bundles/Bundle.cs b/Engine/Source/Programs/Shared/EpicGames.Horde/Storage/Bundles/Bundle.cs
deleted file mode 100644
