Horde: Remove hashes from bundles, to simplify use as a generic container format.

diff --git a/Engine/Source/Programs/Horde/Horde.Agent/Commands/Bundles/DumpCommand.cs b/Engine/Source/Programs/Horde/Horde.Agent/Commands/Bundles/DumpCommand.cs
--- a/Engine/Source/Programs/Horde/Horde.Agent/Commands/Bundles/DumpCommand.cs
+++ b/Engine/Source/Programs/Horde/Horde.Agent/Commands/Bundles/DumpCommand.cs
@@ -57,9 +57,9 @@ namespace Horde.Agent.Commands.Bundles
 			foreach (BundleImport import in header.Imports)
 			{
 				logger.LogInformation("  From blob {BlobId} ({NumExports}/{TotalExports} nodes)", import.BlobId, import.Exports.Count, import.ExportCount);
-				foreach ((int exportIdx, IoHash exportHash) in import.Exports)
+				foreach (int exportIdx in import.Exports)
 				{
-					logger.LogInformation("    [{Index}] IMP {BlobId}:{ExportIdx} = {ExportHash}", refIdx, import.BlobId, exportIdx, exportHash);
+					logger.LogInformation("    [{Index}] IMP {BlobId}@{ExportIdx}", refIdx, import.BlobId, exportIdx);
 					refIdx++;
 				}
 			}
@@ -72,7 +72,7 @@ namespace Horde.Agent.Commands.Bundles
 			foreach (BundleExport export in header.Exports)
 			{
 				string refs = (export.References.Count == 0) ? "[ ]" : $"[ {String.Join(", ", export.References.Select(x => x.ToString()))} ]";
-				logger.LogInformation("  [{Index}] EXP {ExportHash} (length: {NumBytes:n0}, packet: {PacketIdx}, refs: {Refs})", refIdx, export.Hash, export.Length, packetIdx, refs);
+				logger.LogInformation("  [{Index}] EXP (length: {NumBytes:n0}, packet: {PacketIdx}, refs: {Refs})", refIdx, export.Length, packetIdx, refs);
 				refIdx++;
 
 				packetOffset += export.Length;
diff --git a/Engine/Source/Programs/Shared/EpicGames.Horde/Storage/Bundle.cs b/Engine/Source/Programs/Shared/EpicGames.Horde/Storage/Bundle.cs
--- a/Engine/Source/Programs/Shared/EpicGames.Horde/Storage/Bundle.cs
+++ b/Engine/Source/Programs/Shared/EpicGames.Horde/Storage/Bundle.cs
@@ -283,12 +283,12 @@ namespace EpicGames.Horde.Storage
 		/// <summary>
 		/// Indexes of referenced nodes exported from this bundle
 		/// </summary>
-		public IReadOnlyList<(int, IoHash)> Exports { get; }
+		public IReadOnlyList<int> Exports { get; }
 
 		/// <summary>
 		/// Constructor
 		/// </summary>
-		public BundleImport(BlobId blobId, int exportCount, IReadOnlyList<(int, IoHash)> exports)
+		public BundleImport(BlobId blobId, int exportCount, IReadOnlyList<int> exports)
 		{
 			BlobId = blobId;
 			ExportCount = exportCount;
@@ -307,12 +307,10 @@ namespace EpicGames.Horde.Storage
 
 			int count = (int)reader.ReadUnsignedVarInt();
 
-			(int, IoHash)[] exports = new (int, IoHash)[count];
+			int[] exports = new int[count];
 			for (int idx = 0; idx < count; idx++)
 			{
-				int index = (int)reader.ReadUnsignedVarInt();
-				IoHash hash = reader.ReadIoHash();
-				exports[idx] = (index, hash);
+				exports[idx] = (int)reader.ReadUnsignedVarInt();
 			}
 			Exports = exports;
 		}
@@ -329,9 +327,7 @@ namespace EpicGames.Horde.Storage
 			writer.WriteUnsignedVarInt(Exports.Count);
 			for(int idx = 0; idx < Exports.Count; idx++)
 			{
-				(int index, IoHash hash) = Exports[idx];
-				writer.WriteUnsignedVarInt(index);
-				writer.WriteIoHash(hash);
+				writer.WriteUnsignedVarInt(Exports[idx]);
 			}
 		}
 	}
@@ -388,11 +384,6 @@ namespace EpicGames.Horde.Storage
 	/// </summary>
 	public class BundleExport
 	{
-		/// <summary>
-		/// Hash of the node
-		/// </summary>
-		public IoHash Hash { get; }
-
 		/// <summary>
 		/// Uncompressed length of this node
 		/// </summary>
@@ -408,9 +399,8 @@ namespace EpicGames.Horde.Storage
 		/// <summary>
 		/// Constructor
 		/// </summary>
-		public BundleExport(IoHash hash, int length, IReadOnlyList<int> references)
+		public BundleExport(int length, IReadOnlyList<int> references)
 		{
-			Hash = hash;
 			Length = length;
 			References = references;
 		}
@@ -421,7 +411,6 @@ namespace EpicGames.Horde.Storage
 		/// <param name="reader">Reader to deserialize from</param>
 		public BundleExport(IMemoryReader reader)
 		{
-			Hash = reader.ReadIoHash();
 			Length = (int)reader.ReadUnsignedVarInt();
 
 			int numReferences = (int)reader.ReadUnsignedVarInt();
@@ -441,7 +430,6 @@ namespace EpicGames.Horde.Storage
 		/// <param name="writer">Writer to serialize to</param>
 		public void Write(IMemoryWriter writer)
 		{
-			writer.WriteIoHash(Hash);
 			writer.WriteUnsignedVarInt(Length);
 
 			writer.WriteUnsignedVarInt(References.Count);
diff --git a/Engine/Source/Programs/Shared/EpicGames.Horde/Storage/Bundles/BundleStore.cs b/Engine/Source/Programs/Shared/EpicGames.Horde/Storage/Bundles/BundleStore.cs
--- a/Engine/Source/Programs/Shared/EpicGames.Horde/Storage/Bundles/BundleStore.cs
+++ b/Engine/Source/Programs/Shared/EpicGames.Horde/Storage/Bundles/BundleStore.cs
@@ -62,19 +62,14 @@ namespace EpicGames.Horde.Storage.Bundles
 		/// <summary>
 		/// Information about a node within a bundle.
 		/// </summary>
-		[DebuggerDisplay("{Hash}")]
 		sealed class NodeInfo : ITreeBlobRef
 		{
 			public readonly BundleStore Owner;
-			public readonly IoHash Hash;
 			public NodeState State { get; private set; }
 
-			IoHash ITreeBlobRef.Hash => Hash;
-
-			public NodeInfo(BundleStore owner, IoHash hash, NodeState state)
+			public NodeInfo(BundleStore owner, NodeState state)
 			{
 				Owner = owner;
-				Hash = hash;
 				State = state;
 			}
 
@@ -219,8 +214,6 @@ namespace EpicGames.Horde.Storage.Bundles
 			readonly BundleWriter? _parent;
 			readonly Utf8String _prefix;
 
-			public readonly ConcurrentDictionary<IoHash, NodeInfo> HashToNode = new ConcurrentDictionary<IoHash, NodeInfo>(); // TODO: this needs to include some additional state from external sources.
-
 			// List of child writers that need to be flushed
 			readonly List<BundleWriter> _children = new List<BundleWriter>();
 
@@ -263,23 +256,12 @@ namespace EpicGames.Horde.Storage.Bundles
 			{
 				IReadOnlyList<NodeInfo> typedRefs = refs.Select(x => (NodeInfo)x).ToList();
 
-				IoHash hash = ComputeHash(data, typedRefs);
+				InMemoryNodeState state = new InMemoryNodeState(data, typedRefs);
+				NodeInfo newNode = new NodeInfo(_owner, state);
 
-				NodeInfo? node;
-				if (!HashToNode.TryGetValue(hash, out node))
-				{
-					InMemoryNodeState state = new InMemoryNodeState(data, typedRefs);
-					NodeInfo newNode = new NodeInfo(_owner, hash, state);
-
-					// Try to add the node again. If we succeed, check if we need to flush the current bundle being built.
-					node = HashToNode.GetOrAdd(hash, newNode);
-					if (node == newNode)
-					{
-						WriteNode(node);
-					}
-				}
+				WriteNode(newNode);
 
-				return Task.FromResult<ITreeBlobRef>(node);
+				return Task.FromResult<ITreeBlobRef>(newNode);
 			}
 
 			void WriteNode(NodeInfo node)
@@ -393,7 +375,7 @@ namespace EpicGames.Horde.Storage.Bundles
 					{
 						ITreeBlob blob = await root.GetTargetAsync(cancellationToken);
 						InMemoryNodeState state = new InMemoryNodeState(blob.Data, blob.Refs.ConvertAll(x => (NodeInfo)x));
-						WriteNode(new NodeInfo(_owner, root.Hash, state));
+						WriteNode(new NodeInfo(_owner, state));
 					}
 					last = await FlushInternalAsync(cancellationToken);
 				}
@@ -540,16 +522,16 @@ namespace EpicGames.Horde.Storage.Bundles
 				Dictionary<NodeInfo, int> nodeToIndex = new Dictionary<NodeInfo, int>();
 				foreach ((BundleInfo bundleInfo, List<NodeInfo> importedNodes) in bundleToImportedNodes)
 				{
-					(int, IoHash)[] exportInfos = new (int, IoHash)[importedNodes.Count];
+					int[] exportIdxs = new int[importedNodes.Count];
 					for (int idx = 0; idx < importedNodes.Count; idx++)
 					{
 						NodeInfo importedNode = importedNodes[idx];
 						nodeToIndex.Add(importedNode, nodeToIndex.Count);
 
 						ImportedNodeState importedNodeState = (ImportedNodeState)importedNode.State;
-						exportInfos[idx] = (importedNodeState.ExportIdx, importedNode.Hash);
+						exportIdxs[idx] = importedNodeState.ExportIdx;
 					}
-					imports.Add(new BundleImport(bundleInfo.BlobId, bundleInfo.Exports.Length, exportInfos));
+					imports.Add(new BundleImport(bundleInfo.BlobId, bundleInfo.Exports.Length, exportIdxs));
 				}
 
 				// Preallocate data in the encoded buffer to reduce fragmentation if we have to resize
@@ -580,7 +562,7 @@ namespace EpicGames.Horde.Storage.Bundles
 
 					// Create the export for this node
 					int[] references = nodeState.References.Select(x => nodeToIndex[x]).ToArray();
-					BundleExport export = new BundleExport(node.Hash, (int)nodeData.Length, references);
+					BundleExport export = new BundleExport((int)nodeData.Length, references);
 					exports.Add(export);
 					nodeToIndex[node] = nodeToIndex.Count;
 
@@ -763,7 +745,7 @@ namespace EpicGames.Horde.Storage.Bundles
 					ReadOnlyMemory<byte> nodeData = decodedPacket.AsMemory(nodeOffset, export.Length);
 					List<NodeInfo> nodeRefs = export.References.ConvertAll(x => nodes[x]);
 
-					NodeInfo node = new NodeInfo(this, export.Hash, new InMemoryNodeState(new ReadOnlySequence<byte>(nodeData), nodeRefs));
+					NodeInfo node = new NodeInfo(this, new InMemoryNodeState(new ReadOnlySequence<byte>(nodeData), nodeRefs));
 					nodes.Add(node);
 
 					nodeOffset += export.Length;
@@ -935,7 +917,7 @@ namespace EpicGames.Horde.Storage.Bundles
 						NodeInfo? node = bundleInfo.Exports[exportIdx];
 						if (node == null)
 						{
-							node = bundleInfo.Exports[exportIdx] = new NodeInfo(this, export.Hash, state);
+							node = bundleInfo.Exports[exportIdx] = new NodeInfo(this, state);
 						}
 						else
 						{
@@ -960,12 +942,12 @@ namespace EpicGames.Horde.Storage.Bundles
 			foreach (BundleImport import in header.Imports)
 			{
 				BundleInfo importBundle = context.FindOrAddBundle(import.BlobId, import.ExportCount);
-				foreach ((int exportIdx, IoHash exportHash) in import.Exports)
+				foreach (int exportIdx in import.Exports)
 				{
 					NodeInfo? node = importBundle.Exports[exportIdx];
 					if (node == null)
 					{
-						node = new NodeInfo(this, exportHash, new ImportedNodeState(importBundle, exportIdx));
+						node = new NodeInfo(this, new ImportedNodeState(importBundle, exportIdx));
 						importBundle.Exports[exportIdx] = node;
 					}
 					indexedNodes.Add(node);
@@ -1012,21 +994,6 @@ namespace EpicGames.Horde.Storage.Bundles
 
 		#region Writing bundles
 
-		static IoHash ComputeHash(ReadOnlySequence<byte> data, IReadOnlyList<NodeInfo> references)
-		{
-			byte[] buffer = new byte[IoHash.NumBytes * (references.Count + 1)];
-
-			Span<byte> span = buffer;
-			for (int idx = 0; idx < references.Count; idx++)
-			{
-				references[idx].Hash.CopyTo(span);
-				span = span[IoHash.NumBytes..];
-			}
-			IoHash.Compute(data).CopyTo(span);
-
-			return IoHash.Compute(buffer);
-		}
-
 		static void Decode(BundleCompressionFormat format, ReadOnlySpan<byte> inputSpan, Span<byte> outputSpan)
 		{
 			switch (format)
@@ -1042,16 +1009,6 @@ namespace EpicGames.Horde.Storage.Bundles
 			}
 		}
 
-		static void CreateFreeSpace(ref byte[] buffer, int usedSize, int requiredSize)
-		{
-			if (requiredSize > buffer.Length)
-			{
-				byte[] newBuffer = new byte[requiredSize];
-				buffer.AsSpan(0, usedSize).CopyTo(newBuffer);
-				buffer = newBuffer;
-			}
-		}
-
 		#endregion
 	}
 }
diff --git a/Engine/Source/Programs/Shared/EpicGames.Horde/Storage/ITreeStore.cs b/Engine/Source/Programs/Shared/EpicGames.Horde/Storage/ITreeStore.cs
--- a/Engine/Source/Programs/Shared/EpicGames.Horde/Storage/ITreeStore.cs
+++ b/Engine/Source/Programs/Shared/EpicGames.Horde/Storage/ITreeStore.cs
@@ -30,11 +30,6 @@ namespace EpicGames.Horde.Storage
 	/// </summary>
 	public interface ITreeBlobRef
 	{
-		/// <summary>
-		/// Hash of the target
-		/// </summary>
-		IoHash Hash { get; }
-
 		/// <summary>
 		/// Gets the target of a reference
 		/// </summary>
