Horde: Rename ITreeStore implementation to TreeStore, and move into EpicGames.Horde.Storage.

diff --git a/Engine/Source/Programs/Horde/Horde.Agent/Commands/Bundles/CreateCommand.cs b/Engine/Source/Programs/Horde/Horde.Agent/Commands/Bundles/CreateCommand.cs
--- a/Engine/Source/Programs/Horde/Horde.Agent/Commands/Bundles/CreateCommand.cs
+++ b/Engine/Source/Programs/Horde/Horde.Agent/Commands/Bundles/CreateCommand.cs
@@ -8,7 +8,6 @@ using System.Threading;
 using System.Threading.Tasks;
 using EpicGames.Core;
 using EpicGames.Horde.Storage;
-using EpicGames.Horde.Storage.Bundles;
 using EpicGames.Horde.Storage.Nodes;
 using Microsoft.Extensions.Caching.Memory;
 using Microsoft.Extensions.Logging;
@@ -131,7 +130,7 @@ namespace Horde.Agent.Commands.Bundles
 		protected ITreeStore CreateTreeStore(ILogger logger)
 		{
 			IBlobStore blobStore = CreateBlobStore(logger);
-			return new BundleStore(blobStore, new BundleOptions());
+			return new TreeStore(blobStore, new TreeOptions());
 		}
 	}
 
diff --git a/Engine/Source/Programs/Horde/Horde.Agent/Commands/Bundles/DumpCommand.cs b/Engine/Source/Programs/Horde/Horde.Agent/Commands/Bundles/DumpCommand.cs
--- a/Engine/Source/Programs/Horde/Horde.Agent/Commands/Bundles/DumpCommand.cs
+++ b/Engine/Source/Programs/Horde/Horde.Agent/Commands/Bundles/DumpCommand.cs
@@ -3,11 +3,9 @@
 using System;
 using System.Collections.Generic;
 using System.Linq;
-using System.Threading;
 using System.Threading.Tasks;
 using EpicGames.Core;
 using EpicGames.Horde.Storage;
-using EpicGames.Horde.Storage.Bundles;
 using Microsoft.Extensions.Logging;
 
 namespace Horde.Agent.Commands.Bundles
diff --git a/Engine/Source/Programs/Horde/Horde.Build.Tests/TestSetup.cs b/Engine/Source/Programs/Horde/Horde.Build.Tests/TestSetup.cs
--- a/Engine/Source/Programs/Horde/Horde.Build.Tests/TestSetup.cs
+++ b/Engine/Source/Programs/Horde/Horde.Build.Tests/TestSetup.cs
@@ -57,7 +57,6 @@ using Horde.Build.Logs.Storage;
 using Horde.Build.Tasks;
 using Horde.Build.Auditing;
 using Horde.Build.Storage.Backends;
-using EpicGames.Horde.Storage.Bundles;
 using Horde.Build.Compute;
 using Horde.Build.Devices;
 
@@ -242,7 +241,7 @@ namespace Horde.Build.Tests
 			services.AddSingleton<ITreeStore<CommitService>>(sp =>
 			{
 				IBlobStore blobStore = new BasicBlobStore(sp.GetRequiredService<MongoService>(), new TransientStorageBackend(), sp.GetRequiredService<IMemoryCache>(), sp.GetRequiredService<ILogger<BasicBlobStore>>());
-				return new BundleStore(blobStore, new BundleOptions()).ForType<CommitService>();
+				return new TreeStore(blobStore, new TreeOptions()).ForType<CommitService>();
 			});
 
 			services.AddSingleton<IStorageClient, BasicStorageClient>();
diff --git a/Engine/Source/Programs/Horde/Horde.Build/Perforce/ReplicationService.cs b/Engine/Source/Programs/Horde/Horde.Build/Perforce/ReplicationService.cs
--- a/Engine/Source/Programs/Horde/Horde.Build/Perforce/ReplicationService.cs
+++ b/Engine/Source/Programs/Horde/Horde.Build/Perforce/ReplicationService.cs
@@ -11,7 +11,6 @@ using System.Threading;
 using System.Threading.Tasks;
 using EpicGames.Core;
 using EpicGames.Horde.Storage;
-using EpicGames.Horde.Storage.Bundles;
 using EpicGames.Horde.Storage.Git;
 using EpicGames.Horde.Storage.Nodes;
 using EpicGames.Perforce;
@@ -51,7 +50,7 @@ namespace Horde.Build.Perforce
 		/// <summary>
 		/// Options for how objects are packed together
 		/// </summary>
-		public BundleOptions Bundle { get; set; } = new BundleOptions();
+		public TreeOptions Bundle { get; set; } = new TreeOptions();
 
 		/// <summary>
 		/// Options for how objects are sliced
diff --git a/Engine/Source/Programs/Horde/Horde.Build/ServerSettings.cs b/Engine/Source/Programs/Horde/Horde.Build/ServerSettings.cs
--- a/Engine/Source/Programs/Horde/Horde.Build/ServerSettings.cs
+++ b/Engine/Source/Programs/Horde/Horde.Build/ServerSettings.cs
@@ -3,7 +3,6 @@
 using System;
 using System.Linq;
 using EpicGames.Horde.Storage;
-using EpicGames.Horde.Storage.Bundles;
 using EpicGames.Horde.Storage.Nodes;
 using Horde.Build.Agents.Fleet;
 using Horde.Build.Storage.Backends;
@@ -101,7 +100,7 @@ namespace Horde.Build
 		/// <summary>
 		/// Options for creating bundles
 		/// </summary>
-		BundleOptions Bundle { get; }
+		TreeOptions Bundle { get; }
 
 		/// <summary>
 		/// Options for chunking content
@@ -115,7 +114,7 @@ namespace Horde.Build
 	public class TreeStoreOptions : BlobStoreOptions, ITreeStoreOptions
 	{
 		/// <inheritdoc/>
-		public BundleOptions Bundle { get; set; } = new BundleOptions();
+		public TreeOptions Bundle { get; set; } = new TreeOptions();
 
 		/// <inheritdoc/>
 		public ChunkingOptions Chunking { get; set; } = new ChunkingOptions();
diff --git a/Engine/Source/Programs/Horde/Horde.Build/Startup.cs b/Engine/Source/Programs/Horde/Horde.Build/Startup.cs
--- a/Engine/Source/Programs/Horde/Horde.Build/Startup.cs
+++ b/Engine/Source/Programs/Horde/Horde.Build/Startup.cs
@@ -92,7 +92,6 @@ using Horde.Build.Auditing;
 using Horde.Build.Agents.Fleet.Providers;
 using Horde.Build.Server.Notices;
 using Horde.Build.Notifications.Sinks;
-using EpicGames.Horde.Storage.Bundles;
 using StatusCode = Grpc.Core.StatusCode;
 using Microsoft.Extensions.Caching.Memory;
 
@@ -285,7 +284,7 @@ namespace Horde.Build
 		static ITreeStore CreateTreeStore(IServiceProvider sp, TreeStoreOptions options)
 		{
 			IBlobStore store = CreateBlobStore(sp, options);
-			return new BundleStore(store, options.Bundle);
+			return new TreeStore(store, options.Bundle);
 		}
 
 		public Startup(IConfiguration configuration)
diff --git a/Engine/Source/Programs/Shared/EpicGames.Horde.Tests/BundleStoreTests.cs b/Engine/Source/Programs/Shared/EpicGames.Horde.Tests/BundleStoreTests.cs
--- a/Engine/Source/Programs/Shared/EpicGames.Horde.Tests/BundleStoreTests.cs
+++ b/Engine/Source/Programs/Shared/EpicGames.Horde.Tests/BundleStoreTests.cs
@@ -2,7 +2,6 @@
 
 using EpicGames.Horde.Storage;
 using EpicGames.Horde.Storage.Backends;
-using EpicGames.Horde.Storage.Bundles;
 using EpicGames.Horde.Storage.Nodes;
 using Microsoft.Extensions.Caching.Memory;
 using Microsoft.VisualStudio.TestTools.UnitTesting;
@@ -26,7 +25,7 @@ namespace EpicGames.Horde.Tests
 		public async Task TestTreeAsync()
 		{
 			InMemoryBlobStore blobStore = new InMemoryBlobStore();
-			await TestTreeAsync(blobStore, new BundleOptions { MaxBlobSize = 1024 * 1024 });
+			await TestTreeAsync(blobStore, new TreeOptions { MaxBlobSize = 1024 * 1024 });
 
 			Assert.AreEqual(1, blobStore.Blobs.Count);
 			Assert.AreEqual(1, blobStore.Refs.Count);
@@ -36,16 +35,16 @@ namespace EpicGames.Horde.Tests
 		public async Task TestTreeSeparateBlobsAsync()
 		{
 			InMemoryBlobStore blobStore = new InMemoryBlobStore();
-			await TestTreeAsync(blobStore, new BundleOptions { MaxBlobSize = 1 });
+			await TestTreeAsync(blobStore, new TreeOptions { MaxBlobSize = 1 });
 
 			Assert.AreEqual(5, blobStore.Blobs.Count);
 			Assert.AreEqual(1, blobStore.Refs.Count);
 		}
 
-		static async Task TestTreeAsync(InMemoryBlobStore blobStore, BundleOptions options)
+		static async Task TestTreeAsync(InMemoryBlobStore blobStore, TreeOptions options)
 		{
 			// Generate a tree
-			using (BundleStore store = new BundleStore(blobStore, options))
+			using (TreeStore store = new TreeStore(blobStore, options))
 			{
 				ITreeWriter writer = store.CreateTreeWriter("test");
 
@@ -61,7 +60,7 @@ namespace EpicGames.Horde.Tests
 			}
 
 			// Check we can read it back in
-			using (BundleStore collection = new BundleStore(blobStore, options))
+			using (TreeStore collection = new TreeStore(blobStore, options))
 			{
 				ITreeBlob? root = await collection.TryReadTreeAsync(new RefName("test"));
 				Assert.IsNotNull(root);
@@ -108,7 +107,7 @@ namespace EpicGames.Horde.Tests
 			InMemoryBlobStore blobStore = new InMemoryBlobStore();
 
 			// Generate a tree
-			using (BundleStore store = new BundleStore(blobStore, new BundleOptions()))
+			using (TreeStore store = new TreeStore(blobStore, new TreeOptions()))
 			{
 				DirectoryNode root = new DirectoryNode(DirectoryFlags.None);
 				DirectoryNode hello = root.AddDirectory("hello");
@@ -119,7 +118,7 @@ namespace EpicGames.Horde.Tests
 			}
 
 			// Check we can read it back in
-			using (BundleStore store = new BundleStore(blobStore, new BundleOptions()))
+			using (TreeStore store = new TreeStore(blobStore, new TreeOptions()))
 			{
 				DirectoryNode root = await store.ReadTreeAsync<DirectoryNode>(new RefName("test"));
 				await CheckDirectoryTreeAsync(root);
@@ -145,7 +144,7 @@ namespace EpicGames.Horde.Tests
 			InMemoryBlobStore blobStore = new InMemoryBlobStore();
 
 			// Generate a tree
-			using (BundleStore store = new BundleStore(blobStore, new BundleOptions()))
+			using (TreeStore store = new TreeStore(blobStore, new TreeOptions()))
 			{
 				ITreeWriter writer = store.CreateTreeWriter();
 
@@ -161,7 +160,7 @@ namespace EpicGames.Horde.Tests
 			}
 
 			// Check we can read it back in
-			using (BundleStore store = new BundleStore(blobStore, new BundleOptions()))
+			using (TreeStore store = new TreeStore(blobStore, new TreeOptions()))
 			{
 				DirectoryNode root = await store.ReadTreeAsync<DirectoryNode>(new RefName("test"));
 				await CheckFileTreeAsync(root);
@@ -187,7 +186,7 @@ namespace EpicGames.Horde.Tests
 		[TestMethod]
 		public async Task ChunkingTests()
 		{
-			using BundleStore store = new BundleStore(new InMemoryBlobStore(), new BundleOptions());
+			using TreeStore store = new TreeStore(new InMemoryBlobStore(), new TreeOptions());
 			ITreeWriter writer = store.CreateTreeWriter();
 
 			byte[] chunk = RandomNumberGenerator.GetBytes(256);
@@ -227,7 +226,7 @@ namespace EpicGames.Horde.Tests
 		{
 			InMemoryBlobStore blobStore = new InMemoryBlobStore();
 
-			using (BundleStore store = new BundleStore(blobStore, new BundleOptions()))
+			using (TreeStore store = new TreeStore(blobStore, new TreeOptions()))
 			{
 				ITreeWriter writer = store.CreateTreeWriter();
 
@@ -276,7 +275,7 @@ namespace EpicGames.Horde.Tests
 			}
 
 			// Generate a tree
-			using (BundleStore store = new BundleStore(blobStore, new BundleOptions { MaxBlobSize = 1024 }))
+			using (TreeStore store = new TreeStore(blobStore, new TreeOptions { MaxBlobSize = 1024 }))
 			{
 				ITreeWriter writer = store.CreateTreeWriter();
 
@@ -296,7 +295,7 @@ namespace EpicGames.Horde.Tests
 			}
 
 			// Check we can read it back in
-			using (BundleStore store = new BundleStore(blobStore, new BundleOptions()))
+			using (TreeStore store = new TreeStore(blobStore, new TreeOptions()))
 			{
 				DirectoryNode root = await store.ReadTreeAsync<DirectoryNode>(new RefName("test"));
 				await CheckLargeFileTreeAsync(root, data);
diff --git a/Engine/Source/Programs/Shared/EpicGames.Horde.Tests/BundleTests.cs b/Engine/Source/Programs/Shared/EpicGames.Horde.Tests/BundleTests.cs
--- a/Engine/Source/Programs/Shared/EpicGames.Horde.Tests/BundleTests.cs
+++ b/Engine/Source/Programs/Shared/EpicGames.Horde.Tests/BundleTests.cs
@@ -3,15 +3,9 @@
 using EpicGames.Core;
 using EpicGames.Horde.Storage;
 using EpicGames.Horde.Storage.Backends;
-using EpicGames.Horde.Storage.Bundles;
-using EpicGames.Horde.Storage.Impl;
 using EpicGames.Horde.Storage.Nodes;
-using EpicGames.Serialization;
-using Microsoft.Extensions.Caching.Memory;
 using Microsoft.VisualStudio.TestTools.UnitTesting;
 using System;
-using System.Buffers;
-using System.Collections.Concurrent;
 using System.Collections.Generic;
 using System.Linq;
 using System.Text;
@@ -29,7 +23,7 @@ namespace EpicGames.Horde.Tests
 		public BundleTests()
 		{
 			_blobStore = new InMemoryBlobStore();
-			_treeStore = new BundleStore(_blobStore, new BundleOptions());
+			_treeStore = new TreeStore(_blobStore, new TreeOptions());
 		}
 
 		public void Dispose()
@@ -172,7 +166,7 @@ namespace EpicGames.Horde.Tests
 		public async Task BasicTestDirectory()
 		{
 			InMemoryBlobStore blobStore = new InMemoryBlobStore();
-			using ITreeStore treeStore = new BundleStore(blobStore, new BundleOptions());
+			using ITreeStore treeStore = new TreeStore(blobStore, new TreeOptions());
 
 			DirectoryNode root = new DirectoryNode(DirectoryFlags.None);
 			DirectoryNode node = root.AddDirectory("hello");
@@ -211,11 +205,11 @@ namespace EpicGames.Horde.Tests
 		[TestMethod]
 		public async Task DedupTests()
 		{
-			BundleOptions options = new BundleOptions();
+			TreeOptions options = new TreeOptions();
 			options.MaxBlobSize = 1;
 
 			InMemoryBlobStore blobStore = new InMemoryBlobStore();
-			using ITreeStore treeStore = new BundleStore(blobStore, options);
+			using ITreeStore treeStore = new TreeStore(blobStore, options);
 
 			DirectoryNode root = new DirectoryNode(DirectoryFlags.None);
 			root.AddDirectory("node1");
@@ -235,14 +229,14 @@ namespace EpicGames.Horde.Tests
 		[TestMethod]
 		public async Task ReloadTests()
 		{
-			BundleOptions options = new BundleOptions();
+			TreeOptions options = new TreeOptions();
 			options.MaxBlobSize = 1;
 
 			InMemoryBlobStore blobStore = new InMemoryBlobStore();
 
 			RefName refName = new RefName("ref");
 
-			using (ITreeStore oldBundle = new BundleStore(blobStore, options))
+			using (ITreeStore oldBundle = new TreeStore(blobStore, options))
 			{
 				DirectoryNode root = new DirectoryNode(DirectoryFlags.None);
 
@@ -257,7 +251,7 @@ namespace EpicGames.Horde.Tests
 				Assert.AreEqual(5, blobStore.Blobs.Count);
 			}
 
-			using (ITreeStore newBundle = new BundleStore(blobStore, options))
+			using (ITreeStore newBundle = new TreeStore(blobStore, options))
 			{
 				DirectoryNode root = await newBundle.ReadTreeAsync<DirectoryNode>(refName);
 
@@ -443,10 +437,10 @@ namespace EpicGames.Horde.Tests
 		[TestMethod]
 		public async Task SpillTestAsync()
 		{
-			BundleOptions options = new BundleOptions();
+			TreeOptions options = new TreeOptions();
 			options.MaxBlobSize = 1;
 
-			using BundleStore treeStore = new BundleStore(_blobStore, options);
+			using TreeStore treeStore = new TreeStore(_blobStore, options);
 
 			RefName refName = new RefName("ref");
 			ITreeWriter writer = treeStore.CreateTreeWriter(refName.Text);
diff --git a/Engine/Source/Programs/Shared/EpicGames.Horde.Tests/ChunkingTests.cs b/Engine/Source/Programs/Shared/EpicGames.Horde.Tests/ChunkingTests.cs
--- a/Engine/Source/Programs/Shared/EpicGames.Horde.Tests/ChunkingTests.cs
+++ b/Engine/Source/Programs/Shared/EpicGames.Horde.Tests/ChunkingTests.cs
@@ -10,7 +10,6 @@ using System.IO;
 using System.Threading;
 using EpicGames.Horde.Storage;
 using EpicGames.Horde.Storage.Backends;
-using EpicGames.Horde.Storage.Bundles;
 
 namespace EpicGames.Horde.Tests
 {
@@ -49,7 +48,7 @@ namespace EpicGames.Horde.Tests
 		[TestMethod]
 		public async Task BasicChunkingTests()
 		{
-			using BundleStore store = new BundleStore(new InMemoryBlobStore(), new BundleOptions());
+			using TreeStore store = new TreeStore(new InMemoryBlobStore(), new TreeOptions());
 			ITreeWriter writer = store.CreateTreeWriter("test");
 
 			ChunkingOptions options = new ChunkingOptions();
@@ -105,7 +104,7 @@ namespace EpicGames.Horde.Tests
 
 		static async Task TestChunkingAsync(ChunkingOptions options)
 		{
-			using BundleStore store = new BundleStore(new InMemoryBlobStore(), new BundleOptions());
+			using TreeStore store = new TreeStore(new InMemoryBlobStore(), new TreeOptions());
 			ITreeWriter writer = store.CreateTreeWriter();
 
 			byte[] data = new byte[4096];
diff --git a/Engine/Source/Programs/Shared/EpicGames.Horde/Storage/Bundles/BundleStore.cs b/Engine/Source/Programs/Shared/EpicGames.Horde/Storage/Bundles/BundleStore.cs
deleted file mode 100644
diff --git a/Engine/Source/Programs/Shared/EpicGames.Horde/Storage/TreeStore.cs b/Engine/Source/Programs/Shared/EpicGames.Horde/Storage/TreeStore.cs
new file mode 100644
--- /dev/null
+++ b/Engine/Source/Programs/Shared/EpicGames.Horde/Storage/TreeStore.cs
@@ -0,0 +1,1057 @@
+// Copyright Epic Games, Inc. All Rights Reserved.
+
+using System;
+using System.Buffers;
+using System.Collections.Concurrent;
+using System.Collections.Generic;
+using System.Diagnostics;
+using System.IO;
+using System.Linq;
+using System.Threading;
+using System.Threading.Tasks;
+using System.Xml.Linq;
+using Blake3;
+using EpicGames.Core;
+using K4os.Compression.LZ4;
+using Microsoft.Extensions.Caching.Memory;
+
+namespace EpicGames.Horde.Storage
+{
+	/// <summary>
+	/// Options for configuring a bundle serializer
+	/// </summary>
+	public class TreeOptions
+	{
+		/// <summary>
+		/// Maximum payload size fo a blob
+		/// </summary>
+		public int MaxBlobSize { get; set; } = 10 * 1024 * 1024;
+
+		/// <summary>
+		/// Compression format to use
+		/// </summary>
+		public BundleCompressionFormat CompressionFormat { get; set; } = BundleCompressionFormat.LZ4;
+
+		/// <summary>
+		/// Minimum size of a block to be compressed
+		/// </summary>
+		public int MinCompressionPacketSize { get; set; } = 16 * 1024;
+
+		/// <summary>
+		/// Number of nodes to retain in the working set after performing a partial flush
+		/// </summary>
+		public int TrimIgnoreCount { get; set; } = 150;
+
+		/// <summary>
+		/// Number of reads from storage to allow concurrently. This has an impact on memory usage as well as network throughput.
+		/// </summary>
+		public int MaxConcurrentReads { get; set; } = 5;
+
+		/// <summary>
+		/// Size of the bundle/decode cache
+		/// </summary>
+		public long CacheSize { get; set; }
+	}
+
+	/// <summary>
+	/// Implementation of <see cref="ITreeStore"/> which packs nodes together into <see cref="Bundle"/> objects, then stores them
+	/// in a <see cref="IBlobStore"/>.
+	/// </summary>
+	public class TreeStore : ITreeStore, IDisposable
+	{
+		/// <summary>
+		/// Information about a node within a bundle.
+		/// </summary>
+		[DebuggerDisplay("{Hash}")]
+		sealed class NodeInfo : ITreeBlobRef
+		{
+			public readonly TreeStore Owner;
+			public readonly IoHash Hash;
+			public NodeState State { get; private set; }
+
+			IoHash ITreeBlobRef.Hash => Hash;
+
+			public NodeInfo(TreeStore owner, IoHash hash, NodeState state)
+			{
+				Owner = owner;
+				Hash = hash;
+				State = state;
+			}
+
+			/// <summary>
+			/// Updates the current state. The only allowed transition is to an exported node state.
+			/// </summary>
+			/// <param name="state">The new state</param>
+			public void SetState(ExportedNodeState state)
+			{
+				State = state;
+			}
+
+			/// <inheritdoc/>
+			public ValueTask<ITreeBlob> GetTargetAsync(CancellationToken cancellationToken = default) => Owner.GetNodeAsync(this, cancellationToken);
+		}
+
+		/// <summary>
+		/// Metadata about a compression packet within a bundle
+		/// </summary>
+		class PacketInfo
+		{
+			public readonly int Offset;
+			public readonly int DecodedLength;
+			public readonly int EncodedLength;
+
+			public PacketInfo(int offset, BundlePacket packet)
+			{
+				Offset = offset;
+				DecodedLength = packet.DecodedLength;
+				EncodedLength = packet.EncodedLength;
+			}
+		}
+
+		/// <summary>
+		/// Base class for the state of a node
+		/// </summary>
+		abstract class NodeState
+		{
+		}
+
+		/// <summary>
+		/// State for data imported from a bundle, but whose location within it is not yet known because the bundle's header has not been read yet.
+		/// Can be transitioned to ExportedNodeState.
+		/// </summary>
+		class ImportedNodeState : NodeState
+		{
+			public readonly BundleInfo BundleInfo;
+			public readonly int ExportIdx;
+
+			public ImportedNodeState(BundleInfo bundleInfo, int exportIdx)
+			{
+				BundleInfo = bundleInfo;
+				ExportIdx = exportIdx;
+			}
+		}
+
+		/// <summary>
+		/// State for data persisted to a bundle. Decoded data is cached outside this structure in the store's MemoryCache.
+		/// Cannot be transitioned to any other state.
+		/// </summary>
+		class ExportedNodeState : ImportedNodeState
+		{
+			public readonly PacketInfo PacketInfo;
+			public readonly int Offset;
+			public readonly int Length;
+			public readonly IReadOnlyList<NodeInfo> References;
+
+			public ExportedNodeState(BundleInfo bundleInfo, int exportIdx, PacketInfo packetInfo, int offset, int length, IReadOnlyList<NodeInfo> references)
+				: base(bundleInfo, exportIdx)
+			{
+				PacketInfo = packetInfo;
+				Offset = offset;
+				Length = length;
+				References = references;
+			}
+		}
+
+		/// <summary>
+		/// Data for a node which exists in memory, but which does NOT exist in storage yet. Nodes which 
+		/// are read in from storage are represented by ExportedNodeState and a MemoryCache. Can transition to an ExportedNodeState only.
+		/// </summary>
+		class InMemoryNodeState : NodeState, ITreeBlob
+		{
+			public readonly ReadOnlySequence<byte> Data;
+			public readonly IReadOnlyList<NodeInfo> References;
+
+			public BundleWriter? _writer;
+			public bool _writing; // This node is currently being written asynchronously
+
+			ReadOnlySequence<byte> ITreeBlob.Data => Data;
+			IReadOnlyList<ITreeBlobRef> ITreeBlob.Refs => References;
+
+			public InMemoryNodeState(ReadOnlySequence<byte> data, IReadOnlyList<NodeInfo> references)
+			{
+				Data = data;
+				References = references;
+			}
+		}
+
+		/// <summary>
+		/// Stores a lookup from blob id to bundle info, and from there to the nodes it contains.
+		/// </summary>
+		class BundleContext
+		{
+			readonly ConcurrentDictionary<BlobId, BundleInfo> _blobIdToBundle = new ConcurrentDictionary<BlobId, BundleInfo>();
+
+			public BundleInfo FindOrAddBundle(BlobId blobId, int exportCount)
+			{
+				BundleInfo bundleInfo = _blobIdToBundle.GetOrAdd(blobId, new BundleInfo(this, blobId, exportCount));
+				Debug.Assert(bundleInfo.Exports.Length == exportCount);
+				return bundleInfo;
+			}
+		}
+
+		/// <summary>
+		/// Information about an imported bundle
+		/// </summary>
+		[DebuggerDisplay("{BlobId}")]
+		class BundleInfo
+		{
+			public readonly BundleContext Context;
+			public readonly BlobId BlobId;
+			public readonly NodeInfo?[] Exports;
+
+			public bool Mounted;
+			public Task MountTask = Task.CompletedTask;
+
+			public BundleInfo(BundleContext context, BlobId blobId, int exportCount)
+			{
+				Context = context;
+				BlobId = blobId;
+				Exports = new NodeInfo?[exportCount];
+			}
+		}
+
+		class BundleWriter : ITreeWriter
+		{
+			readonly object _lockObject = new object();
+
+			readonly TreeStore _owner;
+			readonly BundleContext _context;
+			readonly BundleWriter? _parent;
+			readonly Utf8String _prefix;
+
+			public readonly ConcurrentDictionary<IoHash, NodeInfo> HashToNode = new ConcurrentDictionary<IoHash, NodeInfo>(); // TODO: this needs to include some additional state from external sources.
+
+			// List of child writers that need to be flushed
+			readonly List<BundleWriter> _children = new List<BundleWriter>();
+
+			// Task which includes all active writes, and returns the last blob
+			Task<NodeInfo> _flushTask = Task.FromResult<NodeInfo>(null!);
+
+			// Nodes which have been queued to be written, but which are not yet part of any active write
+			readonly List<NodeInfo> _queueNodes = new List<NodeInfo>();
+
+			// Number of nodes in _queueNodes that are ready to be written (ie. all their dependencies have been written)
+			int _readyCount;
+
+			// Sum of lengths for nodes in _queueNodes up to _readyCount
+			long _readyLength;
+
+			/// <summary>
+			/// Constructor
+			/// </summary>
+			public BundleWriter(TreeStore owner, BundleContext context, BundleWriter? parent, Utf8String prefix)
+			{
+				_owner = owner;
+				_context = context;
+				_parent = parent;
+				_prefix = prefix;
+			}
+
+			/// <inheritdoc/>
+			public ITreeWriter CreateChildWriter()
+			{
+				BundleWriter writer = new BundleWriter(_owner, _context, this, _prefix);
+				lock (_lockObject)
+				{
+					_children.Add(writer);
+				}
+				return writer;
+			}
+
+			/// <inheritdoc/>
+			public Task<ITreeBlobRef> WriteNodeAsync(ReadOnlySequence<byte> data, IReadOnlyList<ITreeBlobRef> refs, CancellationToken cancellationToken = default)
+			{
+				IReadOnlyList<NodeInfo> typedRefs = refs.Select(x => (NodeInfo)x).ToList();
+
+				IoHash hash = ComputeHash(data, typedRefs);
+
+				NodeInfo? node;
+				if (!HashToNode.TryGetValue(hash, out node))
+				{
+					InMemoryNodeState state = new InMemoryNodeState(data, typedRefs);
+					NodeInfo newNode = new NodeInfo(_owner, hash, state);
+
+					// Try to add the node again. If we succeed, check if we need to flush the current bundle being built.
+					node = HashToNode.GetOrAdd(hash, newNode);
+					if (node == newNode)
+					{
+						WriteNode(node);
+					}
+				}
+
+				return Task.FromResult<ITreeBlobRef>(node);
+			}
+
+			void WriteNode(NodeInfo node)
+			{
+				lock (_lockObject)
+				{
+					// Add these nodes to the queue for writing
+					AddToQueue(node);
+
+					// Update the list of nodes which are ready to be written. We need to make sure that any child writers have flushed before
+					// we can reference the blobs containing their nodes.
+					UpdateReady();
+				}
+			}
+
+			void AddToQueue(NodeInfo root)
+			{
+				if (root.State is InMemoryNodeState state && state._writer == null)
+				{
+					lock (state)
+					{
+						// Check again to avoid race
+						if (state._writer == null) 
+						{
+							// Write all the dependencies first
+							foreach (NodeInfo reference in state.References)
+							{
+								AddToQueue(reference);
+							}
+
+							state._writer = this;
+							_queueNodes.Add(root);
+						}
+					}
+				}
+			}
+
+			void UpdateReady()
+			{
+				// Update the number of nodes which can be written
+				while (_readyCount < _queueNodes.Count)
+				{
+					NodeInfo nextNode = _queueNodes[_readyCount];
+					InMemoryNodeState nextState = (InMemoryNodeState)nextNode.State;
+
+					if (_readyCount > 0 && _readyLength + nextState.Data.Length > _owner._options.MaxBlobSize)
+					{
+						WriteReady();
+					}
+					else
+					{
+						foreach (NodeInfo other in nextState.References)
+						{
+							if (other.State is InMemoryNodeState otherState)
+							{
+								// Need to wait for nodes in another writer to be flushed first
+								if (otherState._writer != this)
+								{
+									return;
+								}
+
+								// Need to wait for previous bundles from the current writer to complete before we can reference them
+								if (otherState._writing)
+								{
+									return;
+								}
+							}
+						}
+
+						_readyCount++;
+						_readyLength += nextState.Data.Length;
+					}
+				}
+			}
+
+			void WriteReady()
+			{
+				if (_readyCount > 0)
+				{
+					NodeInfo[] writeNodes = _queueNodes.Slice(0, _readyCount).ToArray();
+					_queueNodes.RemoveRange(0, _readyCount);
+
+					// Mark the nodes as writing so nothing else will be flushed until they're ready
+					foreach (NodeInfo writeNode in writeNodes)
+					{
+						InMemoryNodeState writeState = (InMemoryNodeState)writeNode.State;
+						writeState._writing = true;
+					}
+
+					_readyCount = 0;
+					_readyLength = 0;
+
+					Task prevFlushTask = _flushTask ?? Task.CompletedTask;
+					_flushTask = Task.Run(() => WriteNodesAsync(writeNodes, prevFlushTask, CancellationToken.None));
+				}
+			}
+
+			/// <inheritdoc/>
+			public async Task WriteRefAsync(RefName refName, ITreeBlobRef root, CancellationToken cancellationToken)
+			{
+				if (_parent != null)
+				{
+					throw new InvalidOperationException("Flushing a child writer is not permitted.");
+				}
+
+				// Make sure the last written node is the desired root. If not, we'll need to write it to a new blob so it can be last.
+				NodeInfo? last = await _flushTask;
+				if (last != root)
+				{
+					if (_queueNodes.Count == 0 || _queueNodes[^1] != root)
+					{
+						ITreeBlob blob = await root.GetTargetAsync(cancellationToken);
+						InMemoryNodeState state = new InMemoryNodeState(blob.Data, blob.Refs.ConvertAll(x => (NodeInfo)x));
+						WriteNode(new NodeInfo(_owner, root.Hash, state));
+					}
+					last = await FlushInternalAsync(cancellationToken);
+				}
+
+				// Write a reference to the blob containing this node
+				ImportedNodeState importedState = (ImportedNodeState)last.State;
+				await _owner._blobStore.WriteRefTargetAsync(refName, importedState.BundleInfo.BlobId, cancellationToken);
+			}
+
+			async Task<NodeInfo> FlushInternalAsync(CancellationToken cancellationToken)
+			{
+				// Get all the child writers and flush them
+				BundleWriter[] children;
+				lock (_lockObject)
+				{
+					children = _children.ToArray();
+				}
+
+				Task[] tasks = new Task[children.Length];
+				for (int idx = 0; idx < children.Length; idx++)
+				{
+					tasks[idx] = children[idx].FlushInternalAsync(cancellationToken);
+				}
+				await Task.WhenAll(tasks);
+
+				// Wait for any current writes to finish. _flushTask may be updated through calls to UpdateReady() during its execution, so loop until it's complete.
+				while (!_flushTask.IsCompleted)
+				{
+					await Task.WhenAll(_flushTask, Task.Delay(0, cancellationToken));
+				}
+
+				// Write the last batch
+				if (_queueNodes.Count > 0)
+				{
+					lock (_lockObject)
+					{
+						WriteReady();
+					}
+					await Task.WhenAll(_flushTask, Task.Delay(0, cancellationToken));
+				}
+
+				// Return the identifier of the blob containing the root node
+				return await _flushTask;
+			}
+
+			async Task<NodeInfo> WriteNodesAsync(NodeInfo[] writeNodes, Task prevWriteTask, CancellationToken cancellationToken)
+			{
+				// Create the bundle
+				Bundle bundle = CreateBundle(writeNodes);
+				BundleHeader header = bundle.Header;
+
+				// Write it to storage
+				BlobId blobId = await _owner._blobStore.WriteBlobAsync(bundle.AsSequence(), bundle.Header.Imports.Select(x => x.BlobId).ToList(), _prefix, cancellationToken);
+				string cacheKey = GetBundleCacheKey(blobId);
+				_owner.AddBundleToCache(cacheKey, bundle);
+
+				// Create a BundleInfo for it
+				BundleInfo bundleInfo = _context.FindOrAddBundle(blobId, writeNodes.Length);
+				for (int idx = 0; idx < writeNodes.Length; idx++)
+				{
+					bundleInfo.Exports[idx] = writeNodes[idx];
+				}
+
+				// Update the node states to reference the written bundle
+				int exportIdx = 0;
+				int packetOffset = 0;
+				foreach (BundlePacket packet in bundle.Header.Packets)
+				{
+					PacketInfo packetInfo = new PacketInfo(packetOffset, packet);
+
+					int nodeOffset = 0;
+					for (; exportIdx < header.Exports.Count && nodeOffset + bundle.Header.Exports[exportIdx].Length <= packet.DecodedLength; exportIdx++)
+					{
+						InMemoryNodeState inMemoryState = (InMemoryNodeState)writeNodes[exportIdx].State;
+
+						BundleExport export = header.Exports[exportIdx];
+						writeNodes[exportIdx].SetState(new ExportedNodeState(bundleInfo, exportIdx, packetInfo, nodeOffset, export.Length, inMemoryState.References));
+
+						nodeOffset += header.Exports[exportIdx].Length;
+					}
+
+					packetOffset += packet.EncodedLength;
+				}
+
+				// Now that we've written some nodes to storage, update any parent writers that may be dependent on us
+				lock (_lockObject)
+				{
+					UpdateReady();
+				}
+
+				// Also update the parent write queue
+				if (_parent != null)
+				{
+					lock (_parent._lockObject)
+					{
+						_parent.UpdateReady();
+					}
+				}
+
+				// Wait for other writes to finish
+				await prevWriteTask;
+				return writeNodes[^1];
+			}
+
+			/// <summary>
+			/// Creates a Bundle containing a set of nodes. 
+			/// </summary>
+			Bundle CreateBundle(NodeInfo[] nodes)
+			{
+				TreeOptions options = _owner._options;
+
+				// Create a set from the nodes to be written. We use this to determine references that are imported.
+				HashSet<NodeInfo> nodeSet = new HashSet<NodeInfo>(nodes);
+
+				// Find all the imported nodes by bundle
+				Dictionary<BundleInfo, List<NodeInfo>> bundleToImportedNodes = new Dictionary<BundleInfo, List<NodeInfo>>();
+				foreach (NodeInfo node in nodes)
+				{
+					InMemoryNodeState state = (InMemoryNodeState)node.State;
+					foreach (NodeInfo reference in state.References)
+					{
+						if (nodeSet.Add(reference))
+						{
+							// Get the persisted node info
+							ImportedNodeState importedState = (ImportedNodeState)reference.State;
+							BundleInfo bundleInfo = importedState.BundleInfo;
+
+							// Get the list of nodes within it
+							List<NodeInfo>? importedNodes;
+							if (!bundleToImportedNodes.TryGetValue(bundleInfo, out importedNodes))
+							{
+								importedNodes = new List<NodeInfo>();
+								bundleToImportedNodes.Add(bundleInfo, importedNodes);
+							}
+							importedNodes.Add(reference);
+						}
+					}
+				}
+
+				// Create the import list
+				List<BundleImport> imports = new List<BundleImport>();
+
+				// Add all the imports and assign them identifiers
+				Dictionary<NodeInfo, int> nodeToIndex = new Dictionary<NodeInfo, int>();
+				foreach ((BundleInfo bundleInfo, List<NodeInfo> importedNodes) in bundleToImportedNodes)
+				{
+					(int, IoHash)[] exportInfos = new (int, IoHash)[importedNodes.Count];
+					for (int idx = 0; idx < importedNodes.Count; idx++)
+					{
+						NodeInfo importedNode = importedNodes[idx];
+						nodeToIndex.Add(importedNode, nodeToIndex.Count);
+
+						ImportedNodeState importedNodeState = (ImportedNodeState)importedNode.State;
+						exportInfos[idx] = (importedNodeState.ExportIdx, importedNode.Hash);
+					}
+					imports.Add(new BundleImport(bundleInfo.BlobId, bundleInfo.Exports.Length, exportInfos));
+				}
+
+				// Preallocate data in the encoded buffer to reduce fragmentation if we have to resize
+				ByteArrayBuilder builder = new ByteArrayBuilder(options.MinCompressionPacketSize * 2);
+
+				// Create the export list
+				List<BundleExport> exports = new List<BundleExport>();
+				List<BundlePacket> packets = new List<BundlePacket>();
+
+				// Size of data currently stored in the block buffer
+				int blockSize = 0;
+
+				// Segments of data in the current block
+				List<ReadOnlyMemory<byte>> blockSegments = new List<ReadOnlyMemory<byte>>();
+
+				// Compress all the nodes into the output buffer buffer
+				foreach (NodeInfo node in nodes)
+				{
+					InMemoryNodeState nodeState = (InMemoryNodeState)node.State;
+					ReadOnlySequence<byte> nodeData = nodeState.Data;
+
+					// If we can't fit this data into the current block, flush the contents of it first
+					if (blockSize > 0 && blockSize + nodeData.Length > options.MinCompressionPacketSize)
+					{
+						FlushPacket(_owner._options.CompressionFormat, blockSegments, blockSize, builder, packets);
+						blockSize = 0;
+					}
+
+					// Create the export for this node
+					int[] references = nodeState.References.Select(x => nodeToIndex[x]).ToArray();
+					BundleExport export = new BundleExport(node.Hash, (int)nodeData.Length, references);
+					exports.Add(export);
+					nodeToIndex[node] = nodeToIndex.Count;
+
+					// Write out the new block
+					if (nodeData.Length < options.MinCompressionPacketSize || !nodeData.IsSingleSegment)
+					{
+						blockSize += AddSegments(nodeData, blockSegments);
+					}
+					else
+					{
+						FlushPacket(options.CompressionFormat, nodeData.First, builder, packets);
+					}
+				}
+				FlushPacket(options.CompressionFormat, blockSegments, blockSize, builder, packets);
+
+				// Flush the data
+				BundleHeader header = new BundleHeader(options.CompressionFormat, imports, exports, packets);
+				return new Bundle(header, builder.AsSequence());
+			}
+
+			static int AddSegments(ReadOnlySequence<byte> sequence, List<ReadOnlyMemory<byte>> segments)
+			{
+				int size = 0;
+				foreach (ReadOnlyMemory<byte> segment in sequence)
+				{
+					segments.Add(segment);
+					size += segment.Length;
+				}
+				return size;
+			}
+
+			static void FlushPacket(BundleCompressionFormat format, List<ReadOnlyMemory<byte>> blockSegments, int blockSize, ByteArrayBuilder builder, List<BundlePacket> packets)
+			{
+				if (blockSize > 0)
+				{
+					if (blockSegments.Count == 1)
+					{
+						FlushPacket(format, blockSegments[0], builder, packets);
+					}
+					else
+					{
+						using IMemoryOwner<byte> buffer = MemoryPool<byte>.Shared.Rent(blockSize);
+
+						Memory<byte> output = buffer.Memory;
+						foreach (ReadOnlyMemory<byte> blockSegment in blockSegments)
+						{
+							blockSegment.CopyTo(output);
+							output = output.Slice(blockSegment.Length);
+						}
+
+						FlushPacket(format, buffer.Memory.Slice(0, blockSize), builder, packets);
+					}
+				}
+				blockSegments.Clear();
+			}
+
+			static void FlushPacket(BundleCompressionFormat format, ReadOnlyMemory<byte> inputData, ByteArrayBuilder builder, List<BundlePacket> packets)
+			{
+				if (inputData.Length > 0)
+				{
+					int encodedLength = Encode(format, inputData.Span, builder);
+					Debug.Assert(encodedLength >= 0);
+					packets.Add(new BundlePacket(encodedLength, inputData.Length));
+				}
+			}
+
+			static int Encode(BundleCompressionFormat format, ReadOnlySpan<byte> inputSpan, ByteArrayBuilder builder)
+			{
+				switch (format)
+				{
+					case BundleCompressionFormat.None:
+						{
+							builder.Append(inputSpan);
+							builder.Advance(inputSpan.Length);
+							return inputSpan.Length;
+						}
+					case BundleCompressionFormat.LZ4:
+						{
+							int maxSize = LZ4Codec.MaximumOutputSize(inputSpan.Length);
+
+							Span<byte> outputSpan = builder.GetMemory(maxSize).Span;
+							int encodedLength = LZ4Codec.Encode(inputSpan, outputSpan);
+							builder.Advance(encodedLength);
+
+							return encodedLength;
+						}
+					default:
+						throw new InvalidDataException($"Invalid compression format '{(int)format}'");
+				}
+			}
+		}
+
+		readonly IBlobStore _blobStore;
+		readonly IMemoryCache? _cache;
+
+		readonly object _lockObject = new object();
+
+		readonly SemaphoreSlim _readSema;
+
+		// Active read tasks at any moment. If a BundleObject is not available in the cache, we start a read and add an entry to this dictionary
+		// so that other threads can also await it.
+		readonly Dictionary<string, Task<Bundle>> _readTasks = new Dictionary<string, Task<Bundle>>();
+
+		readonly TreeOptions _options;
+
+		/// <summary>
+		/// Constructor
+		/// </summary>
+		public TreeStore(IBlobStore blobStore, TreeOptions options)
+		{
+			_blobStore = blobStore;
+			_options = options;
+			if (_options.CacheSize > 0)
+			{
+				_cache = new MemoryCache(new MemoryCacheOptions { SizeLimit = _options.CacheSize });
+			}
+			_readSema = new SemaphoreSlim(options.MaxConcurrentReads);
+		}
+
+		/// <inheritdoc/>
+		public ITreeWriter CreateTreeWriter(Utf8String prefix = default) => new BundleWriter(this, new BundleContext(), null, prefix);
+
+		/// <inheritdoc/>
+		public Task DeleteTreeAsync(RefName name, CancellationToken cancellationToken) => _blobStore.DeleteRefAsync(name, cancellationToken);
+
+		/// <inheritdoc/>
+		public void Dispose()
+		{
+			Dispose(true);
+			GC.SuppressFinalize(this);
+		}
+
+		/// <summary>
+		/// Dispose of this objects resources
+		/// </summary>
+		/// <param name="disposing"></param>
+		protected virtual void Dispose(bool disposing)
+		{
+			if (disposing)
+			{
+				_readSema.Dispose();
+				_cache?.Dispose();
+			}
+		}
+
+		/// <inheritdoc/>
+		public Task<bool> HasTreeAsync(RefName name, CancellationToken cancellationToken) => _blobStore.HasRefAsync(name, cancellationToken: cancellationToken);
+
+		/// <inheritdoc/>
+		public async Task<ITreeBlob?> TryReadTreeAsync(RefName name, TimeSpan maxAge = default, CancellationToken cancellationToken = default)
+		{
+			Bundle? bundle = Bundle.FromBlob(await _blobStore.TryReadRefAsync(name, maxAge, cancellationToken));
+			if (bundle == null)
+			{
+				return null;
+			}
+
+			// Create a new context for this bundle and its references
+			BundleContext context = new BundleContext();
+
+			// Create all the imports for the root bundle
+			BundleHeader header = bundle.Header;
+			List<NodeInfo> nodes = CreateImports(context, header);
+
+			// Decompress the bundle data and mount all the nodes in memory
+			int exportIdx = 0;
+			int packetOffset = 0;
+			foreach (BundlePacket packet in header.Packets)
+			{
+				ReadOnlyMemory<byte> encodedPacket = bundle.Payload.Slice(packetOffset, packet.EncodedLength).AsSingleSegment();
+
+				byte[] decodedPacket = new byte[packet.DecodedLength];
+				Decode(header.CompressionFormat, encodedPacket.Span, decodedPacket);
+
+				int nodeOffset = 0;
+				for (; exportIdx < header.Exports.Count && nodeOffset + header.Exports[exportIdx].Length <= packet.DecodedLength; exportIdx++)
+				{
+					BundleExport export = header.Exports[exportIdx];
+
+					ReadOnlyMemory<byte> nodeData = decodedPacket.AsMemory(nodeOffset, export.Length);
+					List<NodeInfo> nodeRefs = export.References.ConvertAll(x => nodes[x]);
+
+					NodeInfo node = new NodeInfo(this, export.Hash, new InMemoryNodeState(new ReadOnlySequence<byte>(nodeData), nodeRefs));
+					nodes.Add(node);
+
+					nodeOffset += export.Length;
+				}
+
+				packetOffset += packet.EncodedLength;
+			}
+
+			// Return the last node as the root
+			return (InMemoryNodeState)nodes[^1].State;
+		}
+
+		/// <summary>
+		/// Gets the data for a given node
+		/// </summary>
+		/// <param name="node">The node to get the data for</param>
+		/// <param name="cancellationToken">Cancellation token for the operation</param>
+		/// <returns>The node data</returns>
+		async ValueTask<ITreeBlob> GetNodeAsync(NodeInfo node, CancellationToken cancellationToken)
+		{
+			if (node.State is InMemoryNodeState inMemoryState)
+			{
+				return inMemoryState;
+			}
+
+			ExportedNodeState exportedState = await GetExportedStateAsync(node, cancellationToken);
+
+			Bundle bundle = await ReadBundleAsync(exportedState.BundleInfo.BlobId, cancellationToken);
+			ReadOnlyMemory<byte> packetData = DecodePacket(exportedState.BundleInfo.BlobId, exportedState.PacketInfo, bundle.Payload);
+			ReadOnlyMemory<byte> nodeData = packetData.Slice(exportedState.Offset, exportedState.Length);
+
+			return new InMemoryNodeState(new ReadOnlySequence<byte>(nodeData), exportedState.References);
+		}
+
+		async ValueTask<ExportedNodeState> GetExportedStateAsync(NodeInfo node, CancellationToken cancellationToken)
+		{
+			ExportedNodeState? exportedState = node.State as ExportedNodeState;
+			if (exportedState == null)
+			{
+				ImportedNodeState importedState = (ImportedNodeState)node.State;
+				await MountBundleAsync(importedState.BundleInfo, cancellationToken);
+				exportedState = (ExportedNodeState)node.State;
+			}
+			return exportedState;
+		}
+
+		#region Reading bundles
+
+		/// <summary>
+		/// Reads a bundle from the given blob id, or retrieves it from the cache
+		/// </summary>
+		/// <param name="blobId"></param>
+		/// <param name="cancellationToken"></param>
+		/// <returns></returns>
+		async Task<Bundle> ReadBundleAsync(BlobId blobId, CancellationToken cancellationToken = default)
+		{
+			string cacheKey = GetBundleCacheKey(blobId);
+			if (_cache == null || !_cache.TryGetValue<Bundle>(cacheKey, out Bundle? bundle))
+			{
+				Task<Bundle>? readTask;
+				lock (_lockObject)
+				{
+					if (!_readTasks.TryGetValue(cacheKey, out readTask))
+					{
+						readTask = Task.Run(() => ReadBundleInternalAsync(cacheKey, blobId, cancellationToken), cancellationToken);
+						_readTasks.Add(cacheKey, readTask);
+					}
+				}
+				bundle = await readTask;
+			}
+			return bundle;
+		}
+
+		async Task<Bundle> ReadBundleInternalAsync(string cacheKey, BlobId blobId, CancellationToken cancellationToken)
+		{
+			// Perform another (sequenced) check whether an object has been added to the cache, to counteract the race between a read task being added and a task completing.
+			lock (_lockObject)
+			{
+				if (_cache != null && _cache.TryGetValue<Bundle>(cacheKey, out Bundle? cachedObject))
+				{
+					return cachedObject;
+				}
+			}
+
+			// Wait for the read semaphore to avoid triggering too many operations at once.
+			await _readSema.WaitAsync(cancellationToken);
+
+			// Read the data from storage
+			Bundle bundle;
+			try
+			{
+				bundle = Bundle.FromBlob(await _blobStore.ReadBlobAsync(blobId, cancellationToken));
+			}
+			finally
+			{
+				_readSema.Release();
+			}
+
+			// Add the object to the cache
+			AddBundleToCache(cacheKey, bundle);
+
+			// Remove this object from the list of read tasks
+			lock (_lockObject)
+			{
+				_readTasks.Remove(cacheKey);
+			}
+			return bundle;
+		}
+
+		static string GetBundleCacheKey(BlobId blobId) => $"bundle:{blobId}";
+
+		void AddBundleToCache(string cacheKey, Bundle bundle)
+		{
+			if (_cache != null)
+			{
+				using (ICacheEntry entry = _cache.CreateEntry(cacheKey))
+				{
+					entry.SetSize(bundle.Payload.Length);
+					entry.SetValue(bundle);
+				}
+			}
+		}
+
+		async Task MountBundleAsync(BundleInfo bundleInfo, CancellationToken cancellationToken)
+		{
+			if (!bundleInfo.Mounted)
+			{
+				Task mountTask;
+				lock (bundleInfo)
+				{
+					Task prevMountTask = bundleInfo.MountTask.ContinueWith(x => { }, cancellationToken, TaskContinuationOptions.None, TaskScheduler.Default);
+					mountTask = Task.Run(() => MountBundleInternalAsync(prevMountTask, bundleInfo, cancellationToken), cancellationToken);
+					bundleInfo.MountTask = mountTask;
+				}
+				await mountTask;
+			}
+		}
+
+		async Task MountBundleInternalAsync(Task prevMountTask, BundleInfo bundleInfo, CancellationToken cancellationToken)
+		{
+			// Wait for the previous mount task to complete. This may succeed or be cancelled.
+			await prevMountTask;
+
+			// If it didn't succeed, try again
+			if (!bundleInfo.Mounted)
+			{
+				// Read the bundle data
+				Bundle bundle = await ReadBundleAsync(bundleInfo.BlobId, cancellationToken);
+				BundleHeader header = bundle.Header;
+
+				// Create all the imported nodes
+				List<NodeInfo> nodes = CreateImports(bundleInfo.Context, header);
+
+				// Create the exported nodes, or update the state of any imported nodes to exported
+				int exportIdx = 0;
+				int packetOffset = 0;
+				foreach (BundlePacket packet in header.Packets)
+				{
+					PacketInfo packetInfo = new PacketInfo(packetOffset, packet);
+
+					int nodeOffset = 0;
+					for (; exportIdx < header.Exports.Count && nodeOffset + header.Exports[exportIdx].Length <= packet.DecodedLength; exportIdx++)
+					{
+						BundleExport export = header.Exports[exportIdx];
+						List<NodeInfo> nodeRefs = export.References.ConvertAll(x => nodes[x]);
+
+						ExportedNodeState state = new ExportedNodeState(bundleInfo, exportIdx, packetInfo, nodeOffset, export.Length, nodeRefs);
+
+						NodeInfo? node = bundleInfo.Exports[exportIdx];
+						if (node == null)
+						{
+							node = bundleInfo.Exports[exportIdx] = new NodeInfo(this, export.Hash, state);
+						}
+						else
+						{
+							node.SetState(state);
+						}
+						nodes.Add(node);
+
+						nodeOffset += export.Length;
+					}
+
+					packetOffset += packet.EncodedLength;
+				}
+
+				// Mark the bundle as mounted
+				bundleInfo.Mounted = true;
+			}
+		}
+
+		List<NodeInfo> CreateImports(BundleContext context, BundleHeader header)
+		{
+			List<NodeInfo> indexedNodes = new List<NodeInfo>(header.Imports.Sum(x => x.Exports.Count) + header.Exports.Count);
+			foreach (BundleImport import in header.Imports)
+			{
+				BundleInfo importBundle = context.FindOrAddBundle(import.BlobId, import.ExportCount);
+				foreach ((int exportIdx, IoHash exportHash) in import.Exports)
+				{
+					NodeInfo? node = importBundle.Exports[exportIdx];
+					if (node == null)
+					{
+						node = new NodeInfo(this, exportHash, new ImportedNodeState(importBundle, exportIdx));
+						importBundle.Exports[exportIdx] = node;
+					}
+					indexedNodes.Add(node);
+				}
+			}
+			return indexedNodes;
+		}
+		/// <summary>
+		/// Gets a decoded block from the store
+		/// </summary>
+		/// <param name="blobId">Information about the bundle</param>
+		/// <param name="packetInfo">The decoded block location and size</param>
+		/// <param name="payload">The bundle payload</param>
+		/// <returns>The decoded data</returns>
+		ReadOnlyMemory<byte> DecodePacket(BlobId blobId, PacketInfo packetInfo, ReadOnlySequence<byte> payload)
+		{
+			if (_cache == null)
+			{
+				return DecodePacketUncached(packetInfo, payload);
+			}
+			else
+			{
+				string cacheKey = $"bundle-packet:{blobId}@{packetInfo.Offset}";
+				return _cache.GetOrCreate<ReadOnlyMemory<byte>>(cacheKey, entry =>
+				{
+					ReadOnlyMemory<byte> decodedPacket = DecodePacketUncached(packetInfo, payload);
+					entry.SetSize(packetInfo.DecodedLength);
+					return decodedPacket;
+				});
+			}
+		}
+
+		static ReadOnlyMemory<byte> DecodePacketUncached(PacketInfo packetInfo, ReadOnlySequence<byte> payload)
+		{
+			ReadOnlyMemory<byte> encodedPacket = payload.Slice(packetInfo.Offset, packetInfo.EncodedLength).AsSingleSegment();
+
+			byte[] decodedPacket = new byte[packetInfo.DecodedLength];
+			LZ4Codec.Decode(encodedPacket.Span, decodedPacket);
+
+			return decodedPacket;
+		}
+
+		#endregion
+
+		#region Writing bundles
+
+		static IoHash ComputeHash(ReadOnlySequence<byte> data, IReadOnlyList<NodeInfo> references)
+		{
+			byte[] buffer = new byte[IoHash.NumBytes * (references.Count + 1)];
+
+			Span<byte> span = buffer;
+			for (int idx = 0; idx < references.Count; idx++)
+			{
+				references[idx].Hash.CopyTo(span);
+				span = span[IoHash.NumBytes..];
+			}
+			IoHash.Compute(data).CopyTo(span);
+
+			return IoHash.Compute(buffer);
+		}
+
+		static void Decode(BundleCompressionFormat format, ReadOnlySpan<byte> inputSpan, Span<byte> outputSpan)
+		{
+			switch (format)
+			{
+				case BundleCompressionFormat.None:
+					inputSpan.CopyTo(outputSpan);
+					break;
+				case BundleCompressionFormat.LZ4:
+					LZ4Codec.Decode(inputSpan, outputSpan);
+					break;
+				default:
+					throw new InvalidDataException($"Invalid compression format '{(int)format}'");
+			}
+		}
+
+		static void CreateFreeSpace(ref byte[] buffer, int usedSize, int requiredSize)
+		{
+			if (requiredSize > buffer.Length)
+			{
+				byte[] newBuffer = new byte[requiredSize];
+				buffer.AsSpan(0, usedSize).CopyTo(newBuffer);
+				buffer = newBuffer;
+			}
+		}
+
+		#endregion
+	}
+}
