Horde: Use a common code path for mounting trees from a bundle.

diff --git a/Engine/Source/Programs/Horde/Horde.Build.Tests/BlobStoreTests.cs b/Engine/Source/Programs/Horde/Horde.Build.Tests/BlobStoreTests.cs
--- a/Engine/Source/Programs/Horde/Horde.Build.Tests/BlobStoreTests.cs
+++ b/Engine/Source/Programs/Horde/Horde.Build.Tests/BlobStoreTests.cs
@@ -89,7 +89,7 @@ namespace Horde.Build.Tests
 			{
 				RefName refName = new RefName("hello");
 				await store.WriteRefTargetAsync(refName, blob3.Id);
-				BlobId refTargetId = await store.ReadRefIdAsync(refName);
+				BlobId refTargetId = await store.ReadRefTargetAsync(refName);
 				Assert.AreEqual(blob3.Id, refTargetId);
 			}
 
diff --git a/Engine/Source/Programs/Shared/EpicGames.Horde/Storage/IBlobStore.cs b/Engine/Source/Programs/Shared/EpicGames.Horde/Storage/IBlobStore.cs
--- a/Engine/Source/Programs/Shared/EpicGames.Horde/Storage/IBlobStore.cs
+++ b/Engine/Source/Programs/Shared/EpicGames.Horde/Storage/IBlobStore.cs
@@ -233,7 +233,7 @@ namespace EpicGames.Horde.Storage
 		/// <param name="maxAge">Maximum age of any cached ref</param>
 		/// <param name="cancellationToken">Cancellation token for the operation</param>
 		/// <returns>The blob instance</returns>
-		public static Task<BlobId> TryReadRefIdAsync(this IBlobStore store, RefName name, TimeSpan maxAge, CancellationToken cancellationToken = default)
+		public static Task<BlobId> TryReadRefTargetAsync(this IBlobStore store, RefName name, TimeSpan maxAge, CancellationToken cancellationToken = default)
 		{
 			return store.TryReadRefTargetAsync(name, DateTime.UtcNow - maxAge, cancellationToken);
 		}
@@ -277,7 +277,7 @@ namespace EpicGames.Horde.Storage
 		/// <param name="cacheTime">Minimum coherency of any cached result</param>
 		/// <param name="cancellationToken">Cancellation token for the operation</param>
 		/// <returns>The blob instance</returns>
-		public static async Task<BlobId> ReadRefIdAsync(this IBlobStore store, RefName name, DateTime cacheTime = default, CancellationToken cancellationToken = default)
+		public static async Task<BlobId> ReadRefTargetAsync(this IBlobStore store, RefName name, DateTime cacheTime = default, CancellationToken cancellationToken = default)
 		{
 			BlobId blobId = await store.TryReadRefTargetAsync(name, cacheTime, cancellationToken);
 			if (!blobId.IsValid())
@@ -295,9 +295,9 @@ namespace EpicGames.Horde.Storage
 		/// <param name="maxAge">Maximum age for any cached result</param>
 		/// <param name="cancellationToken">Cancellation token for the operation</param>
 		/// <returns>The blob instance</returns>
-		public static Task<BlobId> ReadRefIdAsync(this IBlobStore store, RefName name, TimeSpan maxAge, CancellationToken cancellationToken = default)
+		public static Task<BlobId> ReadRefTargetAsync(this IBlobStore store, RefName name, TimeSpan maxAge, CancellationToken cancellationToken = default)
 		{
-			return ReadRefIdAsync(store, name, DateTime.UtcNow - maxAge, cancellationToken);
+			return ReadRefTargetAsync(store, name, DateTime.UtcNow - maxAge, cancellationToken);
 		}
 	}
 
diff --git a/Engine/Source/Programs/Shared/EpicGames.Horde/Storage/TreeStore.cs b/Engine/Source/Programs/Shared/EpicGames.Horde/Storage/TreeStore.cs
--- a/Engine/Source/Programs/Shared/EpicGames.Horde/Storage/TreeStore.cs
+++ b/Engine/Source/Programs/Shared/EpicGames.Horde/Storage/TreeStore.cs
@@ -732,8 +732,9 @@ namespace EpicGames.Horde.Storage
 		/// <inheritdoc/>
 		public async Task<ITreeBlob?> TryReadTreeAsync(RefName name, TimeSpan maxAge = default, CancellationToken cancellationToken = default)
 		{
-			Bundle? bundle = Bundle.FromBlob(await _blobStore.TryReadRefAsync(name, maxAge, cancellationToken));
-			if (bundle == null)
+			// Read the blob referenced by this ref
+			IBlob? blob = await _blobStore.TryReadRefAsync(name, maxAge, cancellationToken);
+			if (blob == null)
 			{
 				return null;
 			}
@@ -741,39 +742,13 @@ namespace EpicGames.Horde.Storage
 			// Create a new context for this bundle and its references
 			BundleContext context = new BundleContext();
 
-			// Create all the imports for the root bundle
-			BundleHeader header = bundle.Header;
-			List<NodeInfo> nodes = CreateImports(context, header);
-
-			// Decompress the bundle data and mount all the nodes in memory
-			int exportIdx = 0;
-			int packetOffset = 0;
-			foreach (BundlePacket packet in header.Packets)
-			{
-				ReadOnlyMemory<byte> encodedPacket = bundle.Payload.Slice(packetOffset, packet.EncodedLength).AsSingleSegment();
+			// Add the new bundle to it
+			Bundle bundle = Bundle.FromBlob(blob);
+			BundleInfo bundleInfo = context.FindOrAddBundle(blob.Id, bundle.Header.Exports.Count);
+			MountBundle(bundleInfo, bundle);
 
-				byte[] decodedPacket = new byte[packet.DecodedLength];
-				Decode(header.CompressionFormat, encodedPacket.Span, decodedPacket);
-
-				int nodeOffset = 0;
-				for (; exportIdx < header.Exports.Count && nodeOffset + header.Exports[exportIdx].Length <= packet.DecodedLength; exportIdx++)
-				{
-					BundleExport export = header.Exports[exportIdx];
-
-					ReadOnlyMemory<byte> nodeData = decodedPacket.AsMemory(nodeOffset, export.Length);
-					List<NodeInfo> nodeRefs = export.References.ConvertAll(x => nodes[x]);
-
-					NodeInfo node = new NodeInfo(this, export.Hash, new InMemoryNodeState(new ReadOnlySequence<byte>(nodeData), nodeRefs));
-					nodes.Add(node);
-
-					nodeOffset += export.Length;
-				}
-
-				packetOffset += packet.EncodedLength;
-			}
-
-			// Return the last node as the root
-			return (InMemoryNodeState)nodes[^1].State;
+			// Return the last node from the bundle as the root
+			return await GetNodeAsync(bundleInfo.Exports[^1]!, cancellationToken);
 		}
 
 		/// <summary>
@@ -910,48 +885,52 @@ namespace EpicGames.Horde.Storage
 			// If it didn't succeed, try again
 			if (!bundleInfo.Mounted)
 			{
-				// Read the bundle data
 				Bundle bundle = await ReadBundleAsync(bundleInfo.BlobId, cancellationToken);
-				BundleHeader header = bundle.Header;
+				MountBundle(bundleInfo, bundle);
+			}
+		}
 
-				// Create all the imported nodes
-				List<NodeInfo> nodes = CreateImports(bundleInfo.Context, header);
+		void MountBundle(BundleInfo bundleInfo, Bundle bundle)
+		{
+			BundleHeader header = bundle.Header;
 
-				// Create the exported nodes, or update the state of any imported nodes to exported
-				int exportIdx = 0;
-				int packetOffset = 0;
-				foreach (BundlePacket packet in header.Packets)
-				{
-					PacketInfo packetInfo = new PacketInfo(packetOffset, packet);
+			// Create all the imported nodes
+			List<NodeInfo> nodes = CreateImports(bundleInfo.Context, header);
 
-					int nodeOffset = 0;
-					for (; exportIdx < header.Exports.Count && nodeOffset + header.Exports[exportIdx].Length <= packet.DecodedLength; exportIdx++)
-					{
-						BundleExport export = header.Exports[exportIdx];
-						List<NodeInfo> nodeRefs = export.References.ConvertAll(x => nodes[x]);
+			// Create the exported nodes, or update the state of any imported nodes to exported
+			int exportIdx = 0;
+			int packetOffset = 0;
+			foreach (BundlePacket packet in header.Packets)
+			{
+				PacketInfo packetInfo = new PacketInfo(packetOffset, packet);
 
-						ExportedNodeState state = new ExportedNodeState(bundleInfo, exportIdx, packetInfo, nodeOffset, export.Length, nodeRefs);
+				int nodeOffset = 0;
+				for (; exportIdx < header.Exports.Count && nodeOffset + header.Exports[exportIdx].Length <= packet.DecodedLength; exportIdx++)
+				{
+					BundleExport export = header.Exports[exportIdx];
+					List<NodeInfo> nodeRefs = export.References.ConvertAll(x => nodes[x]);
 
-						NodeInfo? node = bundleInfo.Exports[exportIdx];
-						if (node == null)
-						{
-							node = bundleInfo.Exports[exportIdx] = new NodeInfo(this, export.Hash, state);
-						}
-						else
-						{
-							node.SetState(state);
-						}
-						nodes.Add(node);
+					ExportedNodeState state = new ExportedNodeState(bundleInfo, exportIdx, packetInfo, nodeOffset, export.Length, nodeRefs);
 
-						nodeOffset += export.Length;
+					NodeInfo? node = bundleInfo.Exports[exportIdx];
+					if (node == null)
+					{
+						node = bundleInfo.Exports[exportIdx] = new NodeInfo(this, export.Hash, state);
+					}
+					else
+					{
+						node.SetState(state);
 					}
+					nodes.Add(node);
 
-					packetOffset += packet.EncodedLength;
+					nodeOffset += export.Length;
 				}
 
-				// Mark the bundle as mounted
-				bundleInfo.Mounted = true;
+				packetOffset += packet.EncodedLength;
 			}
+
+			// Mark the bundle as mounted
+			bundleInfo.Mounted = true;
 		}
 
 		List<NodeInfo> CreateImports(BundleContext context, BundleHeader header)
