Horde: Rename identifiers used for storage.

* BlobId now represents a unique id assigned to blobs.
* HostId now identifies a node in the storage hierarchy.
* BlobLocator now represents a (HostId, BlobId) pair, and can be used to retrieve a blob.

diff --git a/Engine/Source/Programs/Horde/Horde.Agent/Commands/Bundles/CreateCommand.cs b/Engine/Source/Programs/Horde/Horde.Agent/Commands/Bundles/CreateCommand.cs
--- a/Engine/Source/Programs/Horde/Horde.Agent/Commands/Bundles/CreateCommand.cs
+++ b/Engine/Source/Programs/Horde/Horde.Agent/Commands/Bundles/CreateCommand.cs
@@ -30,7 +30,7 @@ namespace Horde.Agent.Commands.Bundles
 			}
 
 			FileReference GetRefFile(RefName name) => FileReference.Combine(_rootDir, name.ToString() + ".ref");
-			FileReference GetBlobFile(BlobId id) => FileReference.Combine(_rootDir, id.Inner.ToString() + ".blob");
+			FileReference GetBlobFile(BlobLocator id) => FileReference.Combine(_rootDir, id.Inner.ToString() + ".blob");
 
 			static async Task WriteAsync(FileReference file, ReadOnlySequence<byte> sequence, CancellationToken cancellationToken)
 			{
@@ -45,7 +45,7 @@ namespace Horde.Agent.Commands.Bundles
 
 			#region Blobs
 
-			public async Task<IBlob> ReadBlobAsync(BlobId id, CancellationToken cancellationToken = default)
+			public async Task<IBlob> ReadBlobAsync(BlobLocator id, CancellationToken cancellationToken = default)
 			{
 				FileReference file = GetBlobFile(id);
 				_logger.LogInformation("Reading {File}", file);
@@ -53,9 +53,9 @@ namespace Horde.Agent.Commands.Bundles
 				return Blob.Deserialize(id, bytes);
 			}
 
-			public async Task<BlobId> WriteBlobAsync(ReadOnlySequence<byte> data, IReadOnlyList<BlobId> references, Utf8String prefix = default, CancellationToken cancellationToken = default)
+			public async Task<BlobLocator> WriteBlobAsync(ReadOnlySequence<byte> data, IReadOnlyList<BlobLocator> references, Utf8String prefix = default, CancellationToken cancellationToken = default)
 			{
-				BlobId id = BlobId.Create(ServerId.Empty, prefix);
+				BlobLocator id = BlobLocator.Create(HostId.Empty, prefix);
 				FileReference file = GetBlobFile(id);
 				DirectoryReference.CreateDirectory(file.Directory);
 				_logger.LogInformation("Writing {File}", file);
@@ -75,7 +75,7 @@ namespace Horde.Agent.Commands.Bundles
 
 			public async Task<IBlob?> TryReadRefAsync(RefName name, DateTime cacheTime = default, CancellationToken cancellationToken = default)
 			{
-				BlobId id = await TryReadRefTargetAsync(name, cacheTime, cancellationToken);
+				BlobLocator id = await TryReadRefTargetAsync(name, cacheTime, cancellationToken);
 				if (!id.IsValid())
 				{
 					return null;
@@ -83,27 +83,27 @@ namespace Horde.Agent.Commands.Bundles
 				return await ReadBlobAsync(id, cancellationToken);
 			}
 
-			public async Task<BlobId> TryReadRefTargetAsync(RefName name, DateTime cacheTime = default, CancellationToken cancellationToken = default)
+			public async Task<BlobLocator> TryReadRefTargetAsync(RefName name, DateTime cacheTime = default, CancellationToken cancellationToken = default)
 			{
 				FileReference file = GetRefFile(name);
 				if(!FileReference.Exists(file))
 				{
-					return BlobId.Empty;
+					return BlobLocator.Empty;
 				}
 
 				_logger.LogInformation("Reading {File}", file);
 				string text = await FileReference.ReadAllTextAsync(file, cancellationToken);
-				return new BlobId(text.Trim());
+				return new BlobLocator(text.Trim());
 			}
 
-			public async Task<BlobId> WriteRefAsync(RefName name, ReadOnlySequence<byte> data, IReadOnlyList<BlobId> references, CancellationToken cancellationToken = default)
+			public async Task<BlobLocator> WriteRefAsync(RefName name, ReadOnlySequence<byte> data, IReadOnlyList<BlobLocator> references, CancellationToken cancellationToken = default)
 			{
-				BlobId blobId = await WriteBlobAsync(data, references, name.Text, cancellationToken);
-				await WriteRefTargetAsync(name, blobId, cancellationToken);
-				return blobId;
+				BlobLocator locator = await WriteBlobAsync(data, references, name.Text, cancellationToken);
+				await WriteRefTargetAsync(name, locator, cancellationToken);
+				return locator;
 			}
 
-			public async Task WriteRefTargetAsync(RefName name, BlobId id, CancellationToken cancellationToken = default)
+			public async Task WriteRefTargetAsync(RefName name, BlobLocator id, CancellationToken cancellationToken = default)
 			{
 				FileReference file = GetRefFile(name);
 				DirectoryReference.CreateDirectory(file.Directory);
diff --git a/Engine/Source/Programs/Horde/Horde.Agent/Commands/Bundles/DumpCommand.cs b/Engine/Source/Programs/Horde/Horde.Agent/Commands/Bundles/DumpCommand.cs
--- a/Engine/Source/Programs/Horde/Horde.Agent/Commands/Bundles/DumpCommand.cs
+++ b/Engine/Source/Programs/Horde/Horde.Agent/Commands/Bundles/DumpCommand.cs
@@ -17,7 +17,7 @@ namespace Horde.Agent.Commands.Bundles
 		public RefName RefName { get; set; } = DefaultRefName;
 
 		[CommandLine("-Blob=")]
-		public BlobId? BlobId { get; set; }
+		public BlobLocator? BlobId { get; set; }
 
 		public override async Task<int> ExecuteAsync(ILogger logger)
 		{
@@ -35,10 +35,10 @@ namespace Horde.Agent.Commands.Bundles
 				logger.LogInformation("Summary for ref {RefId}", RefName);
 			}
 
-			IReadOnlyList<BlobId> references = blob.References;
+			IReadOnlyList<BlobLocator> references = blob.References;
 			logger.LogInformation("");
 			logger.LogInformation("BlobRefs: {NumRefs}", references.Count);
-			foreach (BlobId reference in references)
+			foreach (BlobLocator reference in references)
 			{
 				logger.LogInformation("  {BlobId}", reference);
 			}
@@ -54,10 +54,10 @@ namespace Horde.Agent.Commands.Bundles
 			int refIdx = 0;
 			foreach (BundleImport import in header.Imports)
 			{
-				logger.LogInformation("  From blob {BlobId} ({NumExports}/{TotalExports} nodes)", import.BlobId, import.Exports.Count, import.ExportCount);
+				logger.LogInformation("  From blob {BlobId} ({NumExports}/{TotalExports} nodes)", import.Locator, import.Exports.Count, import.ExportCount);
 				foreach ((int exportIdx, IoHash exportHash) in import.Exports)
 				{
-					logger.LogInformation("    [{Index}] IMP {BlobId}:{ExportIdx} = {ExportHash}", refIdx, import.BlobId, exportIdx, exportHash);
+					logger.LogInformation("    [{Index}] IMP {BlobId}:{ExportIdx} = {ExportHash}", refIdx, import.Locator, exportIdx, exportHash);
 					refIdx++;
 				}
 			}
diff --git a/Engine/Source/Programs/Horde/Horde.Build.Tests/BlobStoreTests.cs b/Engine/Source/Programs/Horde/Horde.Build.Tests/BlobStoreTests.cs
--- a/Engine/Source/Programs/Horde/Horde.Build.Tests/BlobStoreTests.cs
+++ b/Engine/Source/Programs/Horde/Horde.Build.Tests/BlobStoreTests.cs
@@ -56,9 +56,9 @@ namespace Horde.Build.Tests
 
 			byte[] input = CreateTestData(256, 0);
 
-			BlobId blobId = await store.WriteBlobAsync(new ReadOnlySequence<byte>(input), Array.Empty<BlobId>());
+			BlobLocator locator = await store.WriteBlobAsync(new ReadOnlySequence<byte>(input), Array.Empty<BlobLocator>());
 
-			IBlob blob = await store.ReadBlobAsync(blobId);
+			IBlob blob = await store.ReadBlobAsync(locator);
 			Assert.IsTrue(blob.Data.Span.SequenceEqual(input));
 		}
 
@@ -68,41 +68,41 @@ namespace Horde.Build.Tests
 			IBlobStore store = CreateBlobStore();
 
 			byte[] input1 = CreateTestData(256, 1);
-			BlobId blobId1 = await store.WriteBlobAsync(new ReadOnlySequence<byte>(input1), Array.Empty<BlobId>());
-			IBlob blob1 = await store.ReadBlobAsync(blobId1);
+			BlobLocator locator1 = await store.WriteBlobAsync(new ReadOnlySequence<byte>(input1), Array.Empty<BlobLocator>());
+			IBlob blob1 = await store.ReadBlobAsync(locator1);
 			Assert.IsTrue(blob1.Data.Span.SequenceEqual(input1));
-			Assert.IsTrue(blob1.References.SequenceEqual(Array.Empty<BlobId>()));
+			Assert.IsTrue(blob1.References.SequenceEqual(Array.Empty<BlobLocator>()));
 
 			byte[] input2 = CreateTestData(256, 2);
-			BlobId blobId2 = await store.WriteBlobAsync(new ReadOnlySequence<byte>(input2), new BlobId[] { blob1.Id });
-			IBlob blob2 = await store.ReadBlobAsync(blobId2);
+			BlobLocator locator2 = await store.WriteBlobAsync(new ReadOnlySequence<byte>(input2), new BlobLocator[] { blob1.Id });
+			IBlob blob2 = await store.ReadBlobAsync(locator2);
 			Assert.IsTrue(blob2.Data.Span.SequenceEqual(input2));
-			Assert.IsTrue(blob2.References.SequenceEqual(new BlobId[] { blob1.Id }));
+			Assert.IsTrue(blob2.References.SequenceEqual(new BlobLocator[] { blob1.Id }));
 
 			byte[] input3 = CreateTestData(256, 3);
-			BlobId blobId3 = await store.WriteBlobAsync(new ReadOnlySequence<byte>(input3), new BlobId[] { blob1.Id, blob2.Id, blob1.Id });
-			IBlob blob3 = await store.ReadBlobAsync(blobId3);
+			BlobLocator locator3 = await store.WriteBlobAsync(new ReadOnlySequence<byte>(input3), new BlobLocator[] { blob1.Id, blob2.Id, blob1.Id });
+			IBlob blob3 = await store.ReadBlobAsync(locator3);
 			Assert.IsTrue(blob3.Data.Span.SequenceEqual(input3));
-			Assert.IsTrue(blob3.References.SequenceEqual(new BlobId[] { blob1.Id, blob2.Id, blob1.Id }));
+			Assert.IsTrue(blob3.References.SequenceEqual(new BlobLocator[] { blob1.Id, blob2.Id, blob1.Id }));
 
 			for(int idx = 0; idx < 2; idx++)
 			{
 				RefName refName = new RefName("hello");
 				await store.WriteRefTargetAsync(refName, blob3.Id);
-				BlobId refTargetId = await store.ReadRefTargetAsync(refName);
+				BlobLocator refTargetId = await store.ReadRefTargetAsync(refName);
 				Assert.AreEqual(blob3.Id, refTargetId);
 			}
 
 			RefName refName2 = new RefName("hello2");
 
-			BlobId refTargetId2 = await store.WriteRefAsync(refName2, new ReadOnlySequence<byte>(input3), new BlobId[] { blob1.Id, blob2.Id });
+			BlobLocator refTargetId2 = await store.WriteRefAsync(refName2, new ReadOnlySequence<byte>(input3), new BlobLocator[] { blob1.Id, blob2.Id });
 			IBlob refTarget2 = await store.ReadBlobAsync(refTargetId2);
 			Assert.IsTrue(refTarget2.Data.Span.SequenceEqual(input3));
-			Assert.IsTrue(refTarget2.References.SequenceEqual(new BlobId[] { blob1.Id, blob2.Id }));
+			Assert.IsTrue(refTarget2.References.SequenceEqual(new BlobLocator[] { blob1.Id, blob2.Id }));
 
 			refTarget2 = await store.ReadRefAsync(refName2);
 			Assert.IsTrue(refTarget2.Data.Span.SequenceEqual(input3));
-			Assert.IsTrue(refTarget2.References.SequenceEqual(new BlobId[] { blob1.Id, blob2.Id }));
+			Assert.IsTrue(refTarget2.References.SequenceEqual(new BlobLocator[] { blob1.Id, blob2.Id }));
 		}
 	}
 }
diff --git a/Engine/Source/Programs/Horde/Horde.Build/Startup.cs b/Engine/Source/Programs/Horde/Horde.Build/Startup.cs
--- a/Engine/Source/Programs/Horde/Horde.Build/Startup.cs
+++ b/Engine/Source/Programs/Horde/Horde.Build/Startup.cs
@@ -807,16 +807,16 @@ namespace Horde.Build
 				}
 		*/
 
-		public sealed class BlobIdBsonSerializer : SerializerBase<BlobId>
+		public sealed class BlobIdBsonSerializer : SerializerBase<BlobLocator>
 		{
 			/// <inheritdoc/>
-			public override BlobId Deserialize(BsonDeserializationContext context, BsonDeserializationArgs args)
+			public override BlobLocator Deserialize(BsonDeserializationContext context, BsonDeserializationArgs args)
 			{
-				return new BlobId(context.Reader.ReadString());
+				return new BlobLocator(context.Reader.ReadString());
 			}
 
 			/// <inheritdoc/>
-			public override void Serialize(BsonSerializationContext context, BsonSerializationArgs args, BlobId value)
+			public override void Serialize(BsonSerializationContext context, BsonSerializationArgs args, BlobLocator value)
 			{
 				context.Writer.WriteString(value.ToString());
 			}
diff --git a/Engine/Source/Programs/Horde/Horde.Build/Storage/BasicBlobStore.cs b/Engine/Source/Programs/Horde/Horde.Build/Storage/BasicBlobStore.cs
--- a/Engine/Source/Programs/Horde/Horde.Build/Storage/BasicBlobStore.cs
+++ b/Engine/Source/Programs/Horde/Horde.Build/Storage/BasicBlobStore.cs
@@ -24,15 +24,15 @@ namespace Horde.Build.Storage
 		{
 			[BsonId]
 			public RefName Name { get; set; }
-			public BlobId Value { get; set; }
+			public BlobLocator Value { get; set; }
 
 			public RefDocument()
 			{
 				Name = RefName.Empty;
-				Value = BlobId.Empty;
+				Value = BlobLocator.Empty;
 			}
 
-			public RefDocument(RefName name, BlobId value)
+			public RefDocument(RefName name, BlobLocator value)
 			{
 				Name = name;
 				Value = value;
@@ -42,10 +42,10 @@ namespace Horde.Build.Storage
 		class CachedRefValue
 		{
 			public RefName Name { get; }
-			public BlobId Value { get; }
+			public BlobLocator Value { get; }
 			public DateTime Time { get; }
 
-			public CachedRefValue(RefName name, BlobId value)
+			public CachedRefValue(RefName name, BlobLocator value)
 			{
 				Name = name;
 				Value = value;
@@ -53,7 +53,7 @@ namespace Horde.Build.Storage
 			}
 		}
 
-		readonly ServerId _serverId = ServerId.Empty;
+		readonly HostId _serverId = HostId.Empty;
 
 		readonly IMongoCollection<RefDocument> _refs;
 		readonly IStorageBackend _backend;
@@ -78,10 +78,10 @@ namespace Horde.Build.Storage
 
 		#region Blobs
 
-		static string GetBlobPath(BlobId id) => $"{id.GetContentId()}.blob";
+		static string GetBlobPath(BlobLocator id) => $"{id.BlobId}.blob";
 
 		/// <inheritdoc/>
-		public async Task<IBlob> ReadBlobAsync(BlobId id, CancellationToken cancellationToken = default)
+		public async Task<IBlob> ReadBlobAsync(BlobLocator id, CancellationToken cancellationToken = default)
 		{
 			string path = GetBlobPath(id);
 
@@ -95,9 +95,9 @@ namespace Horde.Build.Storage
 		}
 
 		/// <inheritdoc/>
-		public async Task<BlobId> WriteBlobAsync(ReadOnlySequence<byte> data, IReadOnlyList<BlobId> references, Utf8String prefix = default, CancellationToken cancellationToken = default)
+		public async Task<BlobLocator> WriteBlobAsync(ReadOnlySequence<byte> data, IReadOnlyList<BlobLocator> references, Utf8String prefix = default, CancellationToken cancellationToken = default)
 		{
-			BlobId id = BlobId.Create(_serverId, prefix);
+			BlobLocator id = BlobLocator.Create(_serverId, prefix);
 			string path = GetBlobPath(id);
 
 			ReadOnlySequence<byte> sequence = Blob.Serialize(data, references);
@@ -112,7 +112,7 @@ namespace Horde.Build.Storage
 
 		#region Refs
 
-		CachedRefValue AddRefToCache(RefName name, BlobId value)
+		CachedRefValue AddRefToCache(RefName name, BlobLocator value)
 		{
 			CachedRefValue cacheValue = new CachedRefValue(name, value);
 			using (ICacheEntry newEntry = _cache.CreateEntry(name))
@@ -128,54 +128,54 @@ namespace Horde.Build.Storage
 		public async Task DeleteRefAsync(RefName name, CancellationToken cancellationToken = default)
 		{
 			await _refs.DeleteOneAsync(x => x.Name == name, cancellationToken);
-			AddRefToCache(name, BlobId.Empty);
+			AddRefToCache(name, BlobLocator.Empty);
 		}
 
 		/// <inheritdoc/>
 		public async Task<IBlob?> TryReadRefAsync(RefName name, DateTime cacheTime = default, CancellationToken cancellationToken = default)
 		{
-			BlobId blobId = await TryReadRefTargetAsync(name, cacheTime, cancellationToken);
-			if (blobId.IsValid())
+			BlobLocator locator = await TryReadRefTargetAsync(name, cacheTime, cancellationToken);
+			if (locator.IsValid())
 			{
 				try
 				{
-					return await ReadBlobAsync(blobId, cancellationToken);
+					return await ReadBlobAsync(locator, cancellationToken);
 				}
 				catch (Exception ex)
 				{
-					_logger.LogWarning(ex, "Unable to read blob {BlobId} referenced by ref {RefName}", blobId, name);
+					_logger.LogWarning(ex, "Unable to read blob {BlobId} referenced by ref {RefName}", locator, name);
 				}
 			}
 			return null;
 		}
 
 		/// <inheritdoc/>
-		public async Task<BlobId> TryReadRefTargetAsync(RefName name, DateTime cacheTime = default, CancellationToken cancellationToken = default)
+		public async Task<BlobLocator> TryReadRefTargetAsync(RefName name, DateTime cacheTime = default, CancellationToken cancellationToken = default)
 		{
 			CachedRefValue entry;
 			if (!_cache.TryGetValue(name, out entry) || entry.Time < cacheTime)
 			{
 				RefDocument? refDocument = await _refs.Find(x => x.Name == name).FirstOrDefaultAsync(cancellationToken);
-				BlobId value = (refDocument == null) ? BlobId.Empty : refDocument.Value;
+				BlobLocator value = (refDocument == null) ? BlobLocator.Empty : refDocument.Value;
 				entry = AddRefToCache(name, value);
 			}
 			return entry.Value;
 		}
 
 		/// <inheritdoc/>
-		public async Task<BlobId> WriteRefAsync(RefName name, ReadOnlySequence<byte> data, IReadOnlyList<BlobId> references, CancellationToken cancellationToken = default)
+		public async Task<BlobLocator> WriteRefAsync(RefName name, ReadOnlySequence<byte> data, IReadOnlyList<BlobLocator> references, CancellationToken cancellationToken = default)
 		{
-			BlobId blobId = await WriteBlobAsync(data, references, name.Text, cancellationToken);
+			BlobLocator blobId = await WriteBlobAsync(data, references, name.Text, cancellationToken);
 			await WriteRefTargetAsync(name, blobId, cancellationToken);
 			return blobId;
 		}
 
 		/// <inheritdoc/>
-		public async Task WriteRefTargetAsync(RefName name, BlobId blobId, CancellationToken cancellationToken = default)
+		public async Task WriteRefTargetAsync(RefName name, BlobLocator locator, CancellationToken cancellationToken = default)
 		{
-			RefDocument refDocument = new RefDocument(name, blobId);
+			RefDocument refDocument = new RefDocument(name, locator);
 			await _refs.ReplaceOneAsync(x => x.Name == name, refDocument, new ReplaceOptions { IsUpsert = true }, cancellationToken);
-			AddRefToCache(name, blobId);
+			AddRefToCache(name, locator);
 		}
 
 		#endregion
diff --git a/Engine/Source/Programs/Shared/EpicGames.Horde.Tests/BundleTests.cs b/Engine/Source/Programs/Shared/EpicGames.Horde.Tests/BundleTests.cs
--- a/Engine/Source/Programs/Shared/EpicGames.Horde.Tests/BundleTests.cs
+++ b/Engine/Source/Programs/Shared/EpicGames.Horde.Tests/BundleTests.cs
@@ -101,8 +101,8 @@ namespace EpicGames.Horde.Tests
 			BundleHeader oldHeader;
 			{
 				List<BundleImport> imports = new List<BundleImport>();
-				imports.Add(new BundleImport(new BlobId("import1"), 10, new (int, IoHash)[] { (5, IoHash.Compute(Encoding.UTF8.GetBytes("blob1"))) }));
-				imports.Add(new BundleImport(new BlobId("import2"), 20, new (int, IoHash)[] { (6, IoHash.Compute(Encoding.UTF8.GetBytes("blob2"))) }));
+				imports.Add(new BundleImport(new BlobLocator("import1"), 10, new (int, IoHash)[] { (5, IoHash.Compute(Encoding.UTF8.GetBytes("blob1"))) }));
+				imports.Add(new BundleImport(new BlobLocator("import2"), 20, new (int, IoHash)[] { (6, IoHash.Compute(Encoding.UTF8.GetBytes("blob2"))) }));
 
 				List<BundleExport> exports = new List<BundleExport>();
 				exports.Add(new BundleExport(IoHash.Compute(Encoding.UTF8.GetBytes("export1")), 2, new int[] { 1, 2 }));
@@ -127,7 +127,7 @@ namespace EpicGames.Horde.Tests
 				BundleImport oldImport = oldHeader.Imports[idx];
 				BundleImport newImport = newHeader.Imports[idx];
 
-				Assert.AreEqual(oldImport.BlobId, newImport.BlobId);
+				Assert.AreEqual(oldImport.Locator, newImport.Locator);
 				Assert.AreEqual(oldImport.ExportCount, newImport.ExportCount);
 				Assert.AreEqual(oldImport.Exports.Count, newImport.Exports.Count);
 
diff --git a/Engine/Source/Programs/Shared/EpicGames.Horde/Storage/Backends/InMemoryBlobStore.cs b/Engine/Source/Programs/Shared/EpicGames.Horde/Storage/Backends/InMemoryBlobStore.cs
--- a/Engine/Source/Programs/Shared/EpicGames.Horde/Storage/Backends/InMemoryBlobStore.cs
+++ b/Engine/Source/Programs/Shared/EpicGames.Horde/Storage/Backends/InMemoryBlobStore.cs
@@ -18,31 +18,31 @@ namespace EpicGames.Horde.Storage.Backends
 		/// <summary>
 		/// Map of blob id to blob data
 		/// </summary>
-		readonly ConcurrentDictionary<BlobId, IBlob> _blobs = new ConcurrentDictionary<BlobId, IBlob>();
+		readonly ConcurrentDictionary<BlobLocator, IBlob> _blobs = new ConcurrentDictionary<BlobLocator, IBlob>();
 
 		/// <summary>
 		/// Map of ref name to ref data
 		/// </summary>
-		readonly ConcurrentDictionary<RefName, BlobId> _refs = new ConcurrentDictionary<RefName, BlobId>();
+		readonly ConcurrentDictionary<RefName, BlobLocator> _refs = new ConcurrentDictionary<RefName, BlobLocator>();
 
 		/// <inheritdoc cref="_blobs"/>
-		public IReadOnlyDictionary<BlobId, IBlob> Blobs => _blobs;
+		public IReadOnlyDictionary<BlobLocator, IBlob> Blobs => _blobs;
 
 		/// <inheritdoc cref="_refs"/>
-		public IReadOnlyDictionary<RefName, BlobId> Refs => _refs;
+		public IReadOnlyDictionary<RefName, BlobLocator> Refs => _refs;
 
 		#region Blobs
 
 		/// <inheritdoc/>
-		public Task<IBlob> ReadBlobAsync(BlobId blobId, CancellationToken cancellationToken = default)
+		public Task<IBlob> ReadBlobAsync(BlobLocator locator, CancellationToken cancellationToken = default)
 		{
-			return Task.FromResult(_blobs[blobId]);
+			return Task.FromResult(_blobs[locator]);
 		}
 
 		/// <inheritdoc/>
-		public Task<BlobId> WriteBlobAsync(ReadOnlySequence<byte> data, IReadOnlyList<BlobId> references, Utf8String prefix = default, CancellationToken cancellationToken = default)
+		public Task<BlobLocator> WriteBlobAsync(ReadOnlySequence<byte> data, IReadOnlyList<BlobLocator> references, Utf8String prefix = default, CancellationToken cancellationToken = default)
 		{
-			BlobId blobId = BlobId.Create(ServerId.Empty, prefix);
+			BlobLocator blobId = BlobLocator.Create(HostId.Empty, prefix);
 
 			IBlob blob = Blob.FromMemory(blobId, data.ToArray(), references);
 			_blobs[blobId] = blob;
@@ -64,36 +64,36 @@ namespace EpicGames.Horde.Storage.Backends
 		public Task<IBlob?> TryReadRefAsync(RefName name, DateTime cacheTime = default, CancellationToken cancellationToken = default)
 		{
 			IBlob? blob = null;
-			if (_refs.TryGetValue(name, out BlobId blobId))
+			if (_refs.TryGetValue(name, out BlobLocator locator))
 			{
-				_blobs.TryGetValue(blobId, out blob);
+				_blobs.TryGetValue(locator, out blob);
 			}
 			return Task.FromResult(blob);
 		}
 
 		/// <inheritdoc/>
-		public Task<BlobId> TryReadRefTargetAsync(RefName name, DateTime cacheTime = default, CancellationToken cancellationToken = default)
+		public Task<BlobLocator> TryReadRefTargetAsync(RefName name, DateTime cacheTime = default, CancellationToken cancellationToken = default)
 		{
-			if (_refs.TryGetValue(name, out BlobId blobId))
+			if (_refs.TryGetValue(name, out BlobLocator blobId))
 			{
 				return Task.FromResult(blobId);
 			}
 			else
 			{
-				return Task.FromResult(BlobId.Empty);
+				return Task.FromResult(BlobLocator.Empty);
 			}
 		}
 
 		/// <inheritdoc/>
-		public async Task<BlobId> WriteRefAsync(RefName name, ReadOnlySequence<byte> data, IReadOnlyList<BlobId> references, CancellationToken cancellationToken)
+		public async Task<BlobLocator> WriteRefAsync(RefName name, ReadOnlySequence<byte> data, IReadOnlyList<BlobLocator> references, CancellationToken cancellationToken)
 		{
-			BlobId blobId = await WriteBlobAsync(data, references, name.Text, cancellationToken);
-			_refs[name] = blobId;
-			return blobId;
+			BlobLocator locator = await WriteBlobAsync(data, references, name.Text, cancellationToken);
+			_refs[name] = locator;
+			return locator;
 		}
 
 		/// <inheritdoc/>
-		public Task WriteRefTargetAsync(RefName name, BlobId blobId, CancellationToken cancellationToken = default)
+		public Task WriteRefTargetAsync(RefName name, BlobLocator blobId, CancellationToken cancellationToken = default)
 		{
 			_refs[name] = blobId;
 			return Task.CompletedTask;
diff --git a/Engine/Source/Programs/Shared/EpicGames.Horde/Storage/BlobId.cs b/Engine/Source/Programs/Shared/EpicGames.Horde/Storage/BlobId.cs
--- a/Engine/Source/Programs/Shared/EpicGames.Horde/Storage/BlobId.cs
+++ b/Engine/Source/Programs/Shared/EpicGames.Horde/Storage/BlobId.cs
@@ -3,16 +3,18 @@
 using EpicGames.Core;
 using EpicGames.Serialization;
 using System;
+using System.Buffers.Binary;
 using System.ComponentModel;
 using System.Globalization;
+using System.Text;
 using System.Text.Json;
 using System.Text.Json.Serialization;
+using System.Threading;
 
 namespace EpicGames.Horde.Storage
 {
 	/// <summary>
-	/// Unique identifier for a blob, as a utf-8 string. Clients should not assume any internal structure to this identifier; it only
-	/// has meaning to the <see cref="IBlobStore"/> implementation.
+	/// Globally unique identifier for a blob, as a utf-8 string. Clients should not assume any internal structure to this identifier.
 	/// </summary>
 	[JsonConverter(typeof(BlobIdJsonConverter))]
 	[TypeConverter(typeof(BlobIdTypeConverter))]
@@ -61,73 +63,37 @@ namespace EpicGames.Horde.Storage
 			_ = sanitize;
 		}
 
+		static readonly Utf8String s_process;
+		static int s_counter;
+
 		/// <summary>
-		/// Create a unique content id, optionally including a ref name
+		/// Static constructor
 		/// </summary>
-		/// <param name="serverId">The server id</param>
-		/// <param name="prefix">Prefix for blob names. Follows the same restrictions as for content ids.</param>
-		/// <returns>New content id</returns>
-		public static BlobId Create(ServerId serverId, Utf8String prefix = default)
+		[System.Diagnostics.CodeAnalysis.SuppressMessage("Usage", "CA2207:Initialize value type static fields inline")]
+		static BlobId()
 		{
-			if (prefix.Length == 0)
-			{
-				DateTime now = DateTime.UtcNow.Date;
-				prefix = $"_by_date_/{now.Year}-{now.Month:D2}/{now.Day:D2}";
-			}
-			else
-			{
-				ContentId.ValidateArgument(nameof(prefix), prefix);
-			}
+			Random rnd = new Random(HashCode.Combine(DateTime.UtcNow.Ticks, Environment.ProcessId, Environment.TickCount64, Environment.MachineName.GetHashCode(StringComparison.Ordinal)));
 
-			byte[] buffer;
-			Span<byte> span;
+			Span<byte> process = stackalloc byte[4];
+			rnd.NextBytes(process);
+			s_process = StringUtils.FormatUtf8HexString(process);
 
-			if (serverId.IsValid())
-			{
-				buffer = new byte[serverId.Inner.Length + 1 + prefix.Length + 1 + 24];
-				span = buffer;
-
-				serverId.Inner.Span.CopyTo(span);
-				span = span.Slice(serverId.Inner.Length);
-
-				span[0] = (byte)':';
-				span = span.Slice(1);
-			}
-			else
-			{
-				buffer = new byte[prefix.Length + 1 + 24];
-				span = buffer;
-			}
-
-			prefix.Span.CopyTo(span);
-			span = span.Slice(prefix.Length);
-
-			span[0] = (byte)'/';
-			span = span.Slice(1);
-
-			ContentId.GenerateUniqueId(span);
-			return new BlobId(new Utf8String(buffer), Sanitize.None);
+			s_counter = rnd.Next();
 		}
 
 		/// <summary>
-		/// Validates a given string as a content id
+		/// Creates a new identifier in the given buffer.
 		/// </summary>
-		/// <param name="name">Name of the argument</param>
-		/// <param name="text">String to validate</param>
-		public static void ValidateArgument(string name, Utf8String text)
+		/// <param name="output">Buffer to receive the output. The first 24 bytes will be written to.</param>
+		public static void GenerateUniqueId(Span<byte> output)
 		{
-			if (text.Length == 0)
-			{
-				throw new ArgumentException("Blob identifiers cannot be empty", name);
-			}
+			uint timestamp = (uint)((DateTime.UtcNow - DateTime.UnixEpoch).Ticks / TimeSpan.TicksPerSecond);
+			StringUtils.FormatUtf8HexString(timestamp, output);
 
-			int colonIdx = text.LastIndexOf((byte)':');
-			if (colonIdx != -1)
-			{
-				ServerId.ValidateArgument(name, text.Substring(0, colonIdx));
-			}
+			s_process.Span.CopyTo(output.Slice(8));
 
-			ContentId.ValidateArgument(name, text.Substring(colonIdx + 1));
+			uint counter = (uint)Interlocked.Increment(ref s_counter);
+			StringUtils.FormatUtf8HexString(counter, output.Slice(16));
 		}
 
 		/// <summary>
@@ -137,40 +103,61 @@ namespace EpicGames.Horde.Storage
 		public bool IsValid() => Inner.Length > 0;
 
 		/// <summary>
-		/// Gets the server id for this blob
+		/// Validates a given string as a blob id
 		/// </summary>
-		/// <returns>Server id</returns>
-		public ServerId GetServerId()
+		/// <param name="name">Name of the argument</param>
+		/// <param name="text">String to validate</param>
+		public static void ValidateArgument(string name, Utf8String text)
 		{
-			int colonIdx = Inner.LastIndexOf((byte)':');
-			if (colonIdx == -1)
+			if (text.Length == 0)
+			{
+				throw new ArgumentException("Blob identifiers cannot be empty", name);
+			}
+			if (text[^1] == '/')
 			{
-				return ServerId.Empty;
+				throw new ArgumentException("Blob identifiers cannot start or end with a slash", name);
 			}
-			else
+
+			int lastSlashIdx = -1;
+
+			for (int idx = 0; idx < text.Length; idx++)
 			{
-				return new ServerId(Inner.Substring(0, colonIdx), ServerId.Sanitize.None);
+				if (text[idx] == '/')
+				{
+					if (lastSlashIdx == idx - 1)
+					{
+						throw new ArgumentException("Leading and consecutive slashes are not permitted in blob identifiers", name);
+					}
+					else
+					{
+						lastSlashIdx = idx;
+					}
+				}
+				else
+				{
+					if (!IsValidChar(text[idx]))
+					{
+						throw new ArgumentException($"'{(char)text[idx]} is not a valid blob identifier character", name);
+					}
+				}
 			}
 		}
 
 		/// <summary>
-		/// Gets the content id for this blob
+		/// Test if a given character is valid in a store id
 		/// </summary>
-		/// <returns>Content id</returns>
-		public ContentId GetContentId()
-		{
-			int colonIdx = Inner.LastIndexOf((byte)':');
-			return new ContentId(Inner.Substring(colonIdx + 1), ContentId.Sanitize.None);
-		}
+		/// <param name="c">Character to test</param>
+		/// <returns>True if the character is valid</returns>
+		static bool IsValidChar(byte c) => (c >= 'a' && c <= 'z') || (c >= '0' && c <= '9') || (c == '_' || c == '-');
 
 		/// <inheritdoc/>
-		public override bool Equals(object? obj) => obj is BlobId blobId && Equals(blobId);
+		public override bool Equals(object? obj) => obj is BlobId BlobId && Equals(BlobId);
 
 		/// <inheritdoc/>
 		public override int GetHashCode() => Inner.GetHashCode();
 
 		/// <inheritdoc/>
-		public bool Equals(BlobId blobId) => Inner == blobId.Inner;
+		public bool Equals(BlobId BlobId) => Inner == BlobId.Inner;
 
 		/// <inheritdoc/>
 		public override string ToString() => Inner.ToString();
@@ -183,7 +170,7 @@ namespace EpicGames.Horde.Storage
 	}
 
 	/// <summary>
-	/// Type converter for BlobId to and from JSON
+	/// Type converter for <see cref="BlobId"/> to and from JSON
 	/// </summary>
 	sealed class BlobIdJsonConverter : JsonConverter<BlobId>
 	{
@@ -195,7 +182,7 @@ namespace EpicGames.Horde.Storage
 	}
 
 	/// <summary>
-	/// Type converter from strings to BlobId objects
+	/// Type converter from strings to <see cref="BlobId"/> objects
 	/// </summary>
 	sealed class BlobIdTypeConverter : TypeConverter
 	{
diff --git a/Engine/Source/Programs/Shared/EpicGames.Horde/Storage/BlobLocator.cs b/Engine/Source/Programs/Shared/EpicGames.Horde/Storage/BlobLocator.cs
new file mode 100644
--- /dev/null
+++ b/Engine/Source/Programs/Shared/EpicGames.Horde/Storage/BlobLocator.cs
@@ -0,0 +1,254 @@
+// Copyright Epic Games, Inc. All Rights Reserved.
+
+using EpicGames.Core;
+using EpicGames.Serialization;
+using System;
+using System.ComponentModel;
+using System.Globalization;
+using System.Text.Json;
+using System.Text.Json.Serialization;
+
+namespace EpicGames.Horde.Storage
+{
+	/// <summary>
+	/// Unique identifier for a blob, as a utf-8 string. Clients should not assume any internal structure to this identifier; it only
+	/// has meaning to the <see cref="IBlobStore"/> implementation.
+	/// </summary>
+	[JsonConverter(typeof(BlobLocatorJsonConverter))]
+	[TypeConverter(typeof(BlobLocatorTypeConverter))]
+	[CbConverter(typeof(BlobLocatorCbConverter))]
+	public struct BlobLocator : IEquatable<BlobLocator>
+	{
+		/// <summary>
+		/// Dummy enum to allow invoking the constructor which takes a sanitized full path
+		/// </summary>
+		public enum Sanitize
+		{
+			/// <summary>
+			/// Dummy value
+			/// </summary>
+			None
+		}
+
+		/// <summary>
+		/// Empty blob locator
+		/// </summary>
+		public static BlobLocator Empty { get; } = default;
+
+		/// <summary>
+		/// Identifier for the blob
+		/// </summary>
+		public Utf8String Inner { get; }
+
+		/// <summary>
+		/// Gets the server id for this blob
+		/// </summary>
+		/// <returns>Server id</returns>
+		public HostId HostId
+		{
+			get
+			{
+				int colonIdx = Inner.LastIndexOf((byte)':');
+				return (colonIdx == -1) ? HostId.Empty : new HostId(Inner.Substring(0, colonIdx), HostId.Sanitize.None);
+			}
+		}
+
+		/// <summary>
+		/// Gets the content id for this blob
+		/// </summary>
+		/// <returns>Content id</returns>
+		public BlobId BlobId
+		{
+			get
+			{
+				int colonIdx = Inner.LastIndexOf((byte)':');
+				return new BlobId(Inner.Substring(colonIdx + 1), BlobId.Sanitize.None);
+			}
+		}
+
+		/// <summary>
+		/// Constructor
+		/// </summary>
+		/// <param name="inner"></param>
+		public BlobLocator(Utf8String inner)
+		{
+			Inner = inner;
+			ValidateArgument(nameof(inner), inner);
+		}
+
+		/// <summary>
+		/// Constructor
+		/// </summary>
+		/// <param name="inner"></param>
+		/// <param name="sanitize"></param>
+		public BlobLocator(Utf8String inner, Sanitize sanitize)
+		{
+			Inner = inner;
+			_ = sanitize;
+		}
+
+		/// <summary>
+		/// Create a unique content id, optionally including a ref name
+		/// </summary>
+		/// <param name="serverId">The server id</param>
+		/// <param name="prefix">Prefix for blob names. Follows the same restrictions as for content ids.</param>
+		/// <returns>New content id</returns>
+		public static BlobLocator Create(HostId serverId, Utf8String prefix = default)
+		{
+			if (prefix.Length == 0)
+			{
+				DateTime now = DateTime.UtcNow.Date;
+				prefix = $"_by_date_/{now.Year}-{now.Month:D2}/{now.Day:D2}";
+			}
+			else
+			{
+				BlobId.ValidateArgument(nameof(prefix), prefix);
+			}
+
+			byte[] buffer;
+			Span<byte> span;
+
+			if (serverId.IsValid())
+			{
+				buffer = new byte[serverId.Inner.Length + 1 + prefix.Length + 1 + 24];
+				span = buffer;
+
+				serverId.Inner.Span.CopyTo(span);
+				span = span.Slice(serverId.Inner.Length);
+
+				span[0] = (byte)':';
+				span = span.Slice(1);
+			}
+			else
+			{
+				buffer = new byte[prefix.Length + 1 + 24];
+				span = buffer;
+			}
+
+			prefix.Span.CopyTo(span);
+			span = span.Slice(prefix.Length);
+
+			span[0] = (byte)'/';
+			span = span.Slice(1);
+
+			BlobId.GenerateUniqueId(span);
+			return new BlobLocator(new Utf8String(buffer), Sanitize.None);
+		}
+
+		/// <summary>
+		/// Validates a given string as a content id
+		/// </summary>
+		/// <param name="name">Name of the argument</param>
+		/// <param name="text">String to validate</param>
+		public static void ValidateArgument(string name, Utf8String text)
+		{
+			if (text.Length == 0)
+			{
+				throw new ArgumentException("Blob identifiers cannot be empty", name);
+			}
+
+			int colonIdx = text.LastIndexOf((byte)':');
+			if (colonIdx != -1)
+			{
+				HostId.ValidateArgument(name, text.Substring(0, colonIdx));
+			}
+
+			BlobId.ValidateArgument(name, text.Substring(colonIdx + 1));
+		}
+
+		/// <summary>
+		/// Checks whether this blob id is valid
+		/// </summary>
+		/// <returns>True if the identifier is valid</returns>
+		public bool IsValid() => Inner.Length > 0;
+
+		/// <inheritdoc/>
+		public override bool Equals(object? obj) => obj is BlobLocator blobId && Equals(blobId);
+
+		/// <inheritdoc/>
+		public override int GetHashCode() => Inner.GetHashCode();
+
+		/// <inheritdoc/>
+		public bool Equals(BlobLocator locator) => Inner == locator.Inner;
+
+		/// <inheritdoc/>
+		public override string ToString() => Inner.ToString();
+
+		/// <inheritdoc/>
+		public static bool operator ==(BlobLocator lhs, BlobLocator rhs) => lhs.Equals(rhs);
+
+		/// <inheritdoc/>
+		public static bool operator !=(BlobLocator lhs, BlobLocator rhs) => !lhs.Equals(rhs);
+	}
+
+	/// <summary>
+	/// Type converter for BlobId to and from JSON
+	/// </summary>
+	sealed class BlobLocatorJsonConverter : JsonConverter<BlobLocator>
+	{
+		/// <inheritdoc/>
+		public override BlobLocator Read(ref Utf8JsonReader reader, Type typeToConvert, JsonSerializerOptions options) => new BlobLocator(new Utf8String(reader.GetUtf8String().ToArray()));
+
+		/// <inheritdoc/>
+		public override void Write(Utf8JsonWriter writer, BlobLocator value, JsonSerializerOptions options) => writer.WriteStringValue(value.Inner.Span);
+	}
+
+	/// <summary>
+	/// Type converter from strings to BlobId objects
+	/// </summary>
+	sealed class BlobLocatorTypeConverter : TypeConverter
+	{
+		/// <inheritdoc/>
+		public override bool CanConvertFrom(ITypeDescriptorContext? context, Type sourceType)
+		{
+			return sourceType == typeof(string);
+		}
+
+		/// <inheritdoc/>
+		public override object ConvertFrom(ITypeDescriptorContext? context, CultureInfo? culture, object? value)
+		{
+			return new BlobLocator((string)value!);
+		}
+	}
+
+	/// <summary>
+	/// Type converter to compact binary
+	/// </summary>
+	sealed class BlobLocatorCbConverter : CbConverterBase<BlobLocator>
+	{
+		/// <inheritdoc/>
+		public override BlobLocator Read(CbField field) => new BlobLocator(field.AsUtf8String());
+
+		/// <inheritdoc/>
+		public override void Write(CbWriter writer, BlobLocator value) => writer.WriteUtf8StringValue(value.Inner);
+
+		/// <inheritdoc/>
+		public override void WriteNamed(CbWriter writer, Utf8String name, BlobLocator value) => writer.WriteUtf8String(name, value.Inner);
+	}
+
+	/// <summary>
+	/// Extension methods for blob locators
+	/// </summary>
+	public static class BlobLocatorExtensions
+	{
+		/// <summary>
+		/// Deserialize a blob locator
+		/// </summary>
+		/// <param name="reader">Reader to deserialize from</param>
+		/// <returns>The blob id that was read</returns>
+		public static BlobLocator ReadBlobLocator(this IMemoryReader reader)
+		{
+			return new BlobLocator(reader.ReadUtf8String());
+		}
+
+		/// <summary>
+		/// Serialize a blob locator
+		/// </summary>
+		/// <param name="writer">Writer to serialize to</param>
+		/// <param name="value">Value to serialize</param>
+		public static void WriteBlobLocator(this IMemoryWriter writer, BlobLocator value)
+		{
+			writer.WriteUtf8String(value.Inner);
+		}
+	}
+}
diff --git a/Engine/Source/Programs/Shared/EpicGames.Horde/Storage/Bundle.cs b/Engine/Source/Programs/Shared/EpicGames.Horde/Storage/Bundle.cs
--- a/Engine/Source/Programs/Shared/EpicGames.Horde/Storage/Bundle.cs
+++ b/Engine/Source/Programs/Shared/EpicGames.Horde/Storage/Bundle.cs
@@ -273,7 +273,7 @@ namespace EpicGames.Horde.Storage
 		/// <summary>
 		/// Blob containing the bundle data.
 		/// </summary>
-		public BlobId BlobId { get; }
+		public BlobLocator Locator { get; }
 
 		/// <summary>
 		/// Number of exports from this blob.
@@ -288,9 +288,9 @@ namespace EpicGames.Horde.Storage
 		/// <summary>
 		/// Constructor
 		/// </summary>
-		public BundleImport(BlobId blobId, int exportCount, IReadOnlyList<(int, IoHash)> exports)
+		public BundleImport(BlobLocator locator, int exportCount, IReadOnlyList<(int, IoHash)> exports)
 		{
-			BlobId = blobId;
+			Locator = locator;
 			ExportCount = exportCount;
 			Exports = exports;
 		}
@@ -301,7 +301,7 @@ namespace EpicGames.Horde.Storage
 		/// <param name="reader">Reader to deserialize from</param>
 		internal BundleImport(IMemoryReader reader)
 		{
-			BlobId = reader.ReadBlobId();
+			Locator = reader.ReadBlobLocator();
 
 			ExportCount = (int)reader.ReadUnsignedVarInt();
 
@@ -323,7 +323,7 @@ namespace EpicGames.Horde.Storage
 		/// <param name="writer">Writer to serialize to</param>
 		public void Write(IMemoryWriter writer)
 		{
-			writer.WriteBlobId(BlobId);
+			writer.WriteBlobLocator(Locator);
 			writer.WriteUnsignedVarInt(ExportCount);
 
 			writer.WriteUnsignedVarInt(Exports.Count);
diff --git a/Engine/Source/Programs/Shared/EpicGames.Horde/Storage/ContentId.cs b/Engine/Source/Programs/Shared/EpicGames.Horde/Storage/ContentId.cs
deleted file mode 100644
diff --git a/Engine/Source/Programs/Shared/EpicGames.Horde/Storage/HostId.cs b/Engine/Source/Programs/Shared/EpicGames.Horde/Storage/HostId.cs
new file mode 100644
--- /dev/null
+++ b/Engine/Source/Programs/Shared/EpicGames.Horde/Storage/HostId.cs
@@ -0,0 +1,234 @@
+// Copyright Epic Games, Inc. All Rights Reserved.
+
+using EpicGames.Core;
+using EpicGames.Serialization;
+using System;
+using System.ComponentModel;
+using System.Globalization;
+using System.Text;
+using System.Text.Json;
+using System.Text.Json.Serialization;
+
+namespace EpicGames.Horde.Storage
+{
+	/// <summary>
+	/// Identifies a host in the storage hierarchy, via a string of identifiers separated with colons.
+	/// </summary>
+	[JsonConverter(typeof(HostIdJsonConverter))]
+	[TypeConverter(typeof(HostIdTypeConverter))]
+	[CbConverter(typeof(HostIdCbConverter))]
+	public struct HostId : IEquatable<HostId>
+	{
+		/// <summary>
+		/// Dummy enum to allow invoking the constructor which takes a sanitized full path
+		/// </summary>
+		public enum Sanitize
+		{
+			/// <summary>
+			/// Dummy value
+			/// </summary>
+			None
+		}
+
+		/// <summary>
+		/// Empty host id
+		/// </summary>
+		public static HostId Empty { get; } = default;
+
+		/// <summary>
+		/// Identifier for the host
+		/// </summary>
+		public Utf8String Inner { get; }
+
+		/// <summary>
+		/// Constructor
+		/// </summary>
+		/// <param name="inner"></param>
+		public HostId(Utf8String inner)
+		{
+			Inner = inner;
+			ValidateArgument(nameof(inner), inner);
+		}
+
+		/// <summary>
+		/// Constructor
+		/// </summary>
+		/// <param name="inner"></param>
+		/// <param name="sanitize"></param>
+		public HostId(Utf8String inner, Sanitize sanitize)
+		{
+			Inner = inner;
+			_ = sanitize;
+		}
+
+		/// <summary>
+		/// Appends another identifier to this host id
+		/// </summary>
+		/// <param name="identifier"></param>
+		/// <returns></returns>
+		public HostId Append(Utf8String identifier)
+		{
+			ValidateArgument(nameof(identifier), identifier);
+
+			byte[] buffer = new byte[Inner.Length + 1 + identifier.Length];
+			Inner.Span.CopyTo(buffer);
+			buffer[Inner.Length] = (byte)':';
+			identifier.Span.CopyTo(buffer.AsSpan(Inner.Length + 1));
+
+			return new HostId(new Utf8String(buffer), Sanitize.None);
+		}
+
+		/// <summary>
+		/// Validates a given string as a host id
+		/// </summary>
+		/// <param name="name">Name of the argument</param>
+		/// <param name="text">String to validate</param>
+		public static void ValidateArgument(string name, Utf8String text)
+		{
+			if (text.Length == 0)
+			{
+				throw new ArgumentException("Host identifiers cannot be empty", name);
+			}
+			if (text[^1] == ':')
+			{
+				throw new ArgumentException("Host identifiers cannot end with a colon", name);
+			}
+
+			int lastColonIdx = -1;
+
+			for (int idx = 0; idx < text.Length; idx++)
+			{
+				if (text[idx] == ':')
+				{
+					if (lastColonIdx == idx - 1)
+					{
+						throw new ArgumentException("Leading and consecutive colons are not permitted in host identifiers", name);
+					}
+					else
+					{
+						lastColonIdx = idx;
+					}
+				}
+				else
+				{
+					if (!IsValidChar(text[idx]))
+					{
+						throw new ArgumentException($"'{(char)text[idx]} is not a valid host identifier character", name);
+					}
+				}
+			}
+		}
+
+		/// <summary>
+		/// Test if a given character is valid in a store id
+		/// </summary>
+		/// <param name="c">Character to test</param>
+		/// <returns>True if the character is valid</returns>
+		static bool IsValidChar(byte c) => (c >= 'a' && c <= 'z') || (c >= '0' && c <= '9') || (c == '_' || c == '-');
+
+		/// <summary>
+		/// Checks whether this blob id is valid
+		/// </summary>
+		/// <returns>True if the identifier is valid</returns>
+		public bool IsValid() => Inner.Length > 0;
+
+		/// <summary>
+		/// Gets the upstream host id
+		/// </summary>
+		/// <returns>The upstream host id</returns>
+		public HostId GetUpstream()
+		{
+			int idx = Inner.LastIndexOf(':');
+			return (idx == -1) ? Empty : new HostId(Inner.Substring(0, idx), Sanitize.None);
+		}
+
+		/// <inheritdoc/>
+		public override bool Equals(object? obj) => obj is HostId HostId && Equals(HostId);
+
+		/// <inheritdoc/>
+		public override int GetHashCode() => Inner.GetHashCode();
+
+		/// <inheritdoc/>
+		public bool Equals(HostId HostId) => Inner == HostId.Inner;
+
+		/// <inheritdoc/>
+		public override string ToString() => Inner.ToString();
+
+		/// <inheritdoc/>
+		public static bool operator ==(HostId lhs, HostId rhs) => lhs.Equals(rhs);
+
+		/// <inheritdoc/>
+		public static bool operator !=(HostId lhs, HostId rhs) => !lhs.Equals(rhs);
+	}
+
+	/// <summary>
+	/// Type converter for <see cref="HostId"/> to and from JSON
+	/// </summary>
+	sealed class HostIdJsonConverter : JsonConverter<HostId>
+	{
+		/// <inheritdoc/>
+		public override HostId Read(ref Utf8JsonReader reader, Type typeToConvert, JsonSerializerOptions options) => new HostId(new Utf8String(reader.GetUtf8String().ToArray()));
+
+		/// <inheritdoc/>
+		public override void Write(Utf8JsonWriter writer, HostId value, JsonSerializerOptions options) => writer.WriteStringValue(value.Inner.Span);
+	}
+
+	/// <summary>
+	/// Type converter from strings to <see cref="HostId"/> objects
+	/// </summary>
+	sealed class HostIdTypeConverter : TypeConverter
+	{
+		/// <inheritdoc/>
+		public override bool CanConvertFrom(ITypeDescriptorContext? context, Type sourceType)
+		{
+			return sourceType == typeof(string);
+		}
+
+		/// <inheritdoc/>
+		public override object ConvertFrom(ITypeDescriptorContext? context, CultureInfo? culture, object? value)
+		{
+			return new HostId((string)value!);
+		}
+	}
+
+	/// <summary>
+	/// Type converter to compact binary
+	/// </summary>
+	sealed class HostIdCbConverter : CbConverterBase<HostId>
+	{
+		/// <inheritdoc/>
+		public override HostId Read(CbField field) => new HostId(field.AsUtf8String());
+
+		/// <inheritdoc/>
+		public override void Write(CbWriter writer, HostId value) => writer.WriteUtf8StringValue(value.Inner);
+
+		/// <inheritdoc/>
+		public override void WriteNamed(CbWriter writer, Utf8String name, HostId value) => writer.WriteUtf8String(name, value.Inner);
+	}
+
+	/// <summary>
+	/// Extension methods for <see cref="HostId"/>s.
+	/// </summary>
+	public static class HostIdExtensions
+	{
+		/// <summary>
+		/// Deserialize a <see cref="HostId"/>
+		/// </summary>
+		/// <param name="reader">Reader to deserialize from</param>
+		/// <returns>The blob id that was read</returns>
+		public static HostId ReadHostId(this IMemoryReader reader)
+		{
+			return new HostId(reader.ReadUtf8String());
+		}
+
+		/// <summary>
+		/// Serialize a <see cref="HostId"/>
+		/// </summary>
+		/// <param name="writer">Writer to serialize to</param>
+		/// <param name="value">Value to serialize</param>
+		public static void WriteHostId(this IMemoryWriter writer, HostId value)
+		{
+			writer.WriteUtf8String(value.Inner);
+		}
+	}
+}
diff --git a/Engine/Source/Programs/Shared/EpicGames.Horde/Storage/IBlobStore.cs b/Engine/Source/Programs/Shared/EpicGames.Horde/Storage/IBlobStore.cs
--- a/Engine/Source/Programs/Shared/EpicGames.Horde/Storage/IBlobStore.cs
+++ b/Engine/Source/Programs/Shared/EpicGames.Horde/Storage/IBlobStore.cs
@@ -39,7 +39,7 @@ namespace EpicGames.Horde.Storage
 		/// <summary>
 		/// Identifier for the blob
 		/// </summary>
-		BlobId Id { get; }
+		BlobLocator Id { get; }
 
 		/// <summary>
 		/// Gets the data for this blob
@@ -51,11 +51,11 @@ namespace EpicGames.Horde.Storage
 		/// Find the outward references for a node
 		/// </summary>
 		/// <returns></returns>
-		IReadOnlyList<BlobId> References { get; }
+		IReadOnlyList<BlobLocator> References { get; }
 	}
 
 	/// <summary>
-	/// Base interface for a low-level storage backend. Blobs added to this store are not content addressed, but referenced by <see cref="BlobId"/>.
+	/// Base interface for a low-level storage backend. Blobs added to this store are not content addressed, but referenced by <see cref="BlobLocator"/>.
 	/// </summary>
 	public interface IBlobStore
 	{
@@ -67,7 +67,7 @@ namespace EpicGames.Horde.Storage
 		/// <param name="id">The blob identifier</param>
 		/// <param name="cancellationToken">Cancellation token for the operation</param>
 		/// <returns></returns>
-		Task<IBlob> ReadBlobAsync(BlobId id, CancellationToken cancellationToken = default);
+		Task<IBlob> ReadBlobAsync(BlobLocator id, CancellationToken cancellationToken = default);
 
 		/// <summary>
 		/// Writes a new object to the store
@@ -77,7 +77,7 @@ namespace EpicGames.Horde.Storage
 		/// <param name="cancellationToken">Cancellation token for the operation</param>
 		/// <param name="prefix">Prefix for blob names. While the returned BlobId is guaranteed to be unique, this name can be used as a prefix to aid debugging.</param>
 		/// <returns>Unique identifier for the blob</returns>
-		Task<BlobId> WriteBlobAsync(ReadOnlySequence<byte> data, IReadOnlyList<BlobId> references, Utf8String prefix = default, CancellationToken cancellationToken = default);
+		Task<BlobLocator> WriteBlobAsync(ReadOnlySequence<byte> data, IReadOnlyList<BlobLocator> references, Utf8String prefix = default, CancellationToken cancellationToken = default);
 
 		#endregion
 
@@ -99,7 +99,7 @@ namespace EpicGames.Horde.Storage
 		/// <param name="cacheTime">Minimum coherency for any cached value to be returned</param>
 		/// <param name="cancellationToken">Cancellation token for the operation</param>
 		/// <returns></returns>
-		Task<BlobId> TryReadRefTargetAsync(RefName name, DateTime cacheTime = default, CancellationToken cancellationToken = default);
+		Task<BlobLocator> TryReadRefTargetAsync(RefName name, DateTime cacheTime = default, CancellationToken cancellationToken = default);
 
 		/// <summary>
 		/// Writes a new ref to the store which points to a new blob
@@ -109,7 +109,7 @@ namespace EpicGames.Horde.Storage
 		/// <param name="references">Object references</param>
 		/// <param name="cancellationToken">Cancellation token for the operation</param>
 		/// <returns>Unique identifier for the blob</returns>
-		Task<BlobId> WriteRefAsync(RefName name, ReadOnlySequence<byte> data, IReadOnlyList<BlobId> references, CancellationToken cancellationToken = default);
+		Task<BlobLocator> WriteRefAsync(RefName name, ReadOnlySequence<byte> data, IReadOnlyList<BlobLocator> references, CancellationToken cancellationToken = default);
 
 		/// <summary>
 		/// Writes a new ref to the store
@@ -118,7 +118,7 @@ namespace EpicGames.Horde.Storage
 		/// <param name="blobId">The target for the ref</param>
 		/// <param name="cancellationToken">Cancellation token for the operation</param>
 		/// <returns>Unique identifier for the blob</returns>
-		Task WriteRefTargetAsync(RefName name, BlobId blobId, CancellationToken cancellationToken = default);
+		Task WriteRefTargetAsync(RefName name, BlobLocator blobId, CancellationToken cancellationToken = default);
 
 		/// <summary>
 		/// Reads data for a ref from the store
@@ -151,10 +151,10 @@ namespace EpicGames.Horde.Storage
 			#region Blobs
 
 			/// <inheritdoc/>
-			public Task<IBlob> ReadBlobAsync(BlobId id, CancellationToken cancellationToken = default) => _inner.ReadBlobAsync(id, cancellationToken);
+			public Task<IBlob> ReadBlobAsync(BlobLocator id, CancellationToken cancellationToken = default) => _inner.ReadBlobAsync(id, cancellationToken);
 
 			/// <inheritdoc/>
-			public Task<BlobId> WriteBlobAsync(ReadOnlySequence<byte> data, IReadOnlyList<BlobId> references, Utf8String prefix = default, CancellationToken cancellationToken = default) => _inner.WriteBlobAsync(data, references, prefix, cancellationToken);
+			public Task<BlobLocator> WriteBlobAsync(ReadOnlySequence<byte> data, IReadOnlyList<BlobLocator> references, Utf8String prefix = default, CancellationToken cancellationToken = default) => _inner.WriteBlobAsync(data, references, prefix, cancellationToken);
 
 			#endregion
 
@@ -167,13 +167,13 @@ namespace EpicGames.Horde.Storage
 			public Task<IBlob?> TryReadRefAsync(RefName name, DateTime cacheTime = default, CancellationToken cancellationToken = default) => _inner.TryReadRefAsync(name, cacheTime, cancellationToken);
 
 			/// <inheritdoc/>
-			public Task<BlobId> TryReadRefTargetAsync(RefName name, DateTime cacheTime = default, CancellationToken cancellationToken = default) => _inner.TryReadRefTargetAsync(name, cacheTime, cancellationToken);
+			public Task<BlobLocator> TryReadRefTargetAsync(RefName name, DateTime cacheTime = default, CancellationToken cancellationToken = default) => _inner.TryReadRefTargetAsync(name, cacheTime, cancellationToken);
 
 			/// <inheritdoc/>
-			public Task<BlobId> WriteRefAsync(RefName name, ReadOnlySequence<byte> data, IReadOnlyList<BlobId> references, CancellationToken cancellationToken = default) => _inner.WriteRefAsync(name, data, references, cancellationToken);
+			public Task<BlobLocator> WriteRefAsync(RefName name, ReadOnlySequence<byte> data, IReadOnlyList<BlobLocator> references, CancellationToken cancellationToken = default) => _inner.WriteRefAsync(name, data, references, cancellationToken);
 
 			/// <inheritdoc/>
-			public Task WriteRefTargetAsync(RefName name, BlobId blobId, CancellationToken cancellationToken = default) => _inner.WriteRefTargetAsync(name, blobId, cancellationToken);
+			public Task WriteRefTargetAsync(RefName name, BlobLocator blobId, CancellationToken cancellationToken = default) => _inner.WriteRefTargetAsync(name, blobId, cancellationToken);
 
 			#endregion
 		}
@@ -195,8 +195,8 @@ namespace EpicGames.Horde.Storage
 		/// <returns>True if the ref exists, false if it did not exist</returns>
 		public static async Task<bool> HasRefAsync(this IBlobStore store, RefName name, DateTime cacheTime = default, CancellationToken cancellationToken = default)
 		{
-			BlobId blobId = await store.TryReadRefTargetAsync(name, cacheTime, cancellationToken);
-			return blobId.IsValid();
+			BlobLocator locator = await store.TryReadRefTargetAsync(name, cacheTime, cancellationToken);
+			return locator.IsValid();
 		}
 
 		/// <summary>
@@ -233,7 +233,7 @@ namespace EpicGames.Horde.Storage
 		/// <param name="maxAge">Maximum age of any cached ref</param>
 		/// <param name="cancellationToken">Cancellation token for the operation</param>
 		/// <returns>The blob instance</returns>
-		public static Task<BlobId> TryReadRefTargetAsync(this IBlobStore store, RefName name, TimeSpan maxAge, CancellationToken cancellationToken = default)
+		public static Task<BlobLocator> TryReadRefTargetAsync(this IBlobStore store, RefName name, TimeSpan maxAge, CancellationToken cancellationToken = default)
 		{
 			return store.TryReadRefTargetAsync(name, DateTime.UtcNow - maxAge, cancellationToken);
 		}
@@ -277,14 +277,14 @@ namespace EpicGames.Horde.Storage
 		/// <param name="cacheTime">Minimum coherency of any cached result</param>
 		/// <param name="cancellationToken">Cancellation token for the operation</param>
 		/// <returns>The blob instance</returns>
-		public static async Task<BlobId> ReadRefTargetAsync(this IBlobStore store, RefName name, DateTime cacheTime = default, CancellationToken cancellationToken = default)
+		public static async Task<BlobLocator> ReadRefTargetAsync(this IBlobStore store, RefName name, DateTime cacheTime = default, CancellationToken cancellationToken = default)
 		{
-			BlobId blobId = await store.TryReadRefTargetAsync(name, cacheTime, cancellationToken);
-			if (!blobId.IsValid())
+			BlobLocator locator = await store.TryReadRefTargetAsync(name, cacheTime, cancellationToken);
+			if (!locator.IsValid())
 			{
 				throw new RefNameNotFoundException(name);
 			}
-			return blobId;
+			return locator;
 		}
 
 		/// <summary>
@@ -295,7 +295,7 @@ namespace EpicGames.Horde.Storage
 		/// <param name="maxAge">Maximum age for any cached result</param>
 		/// <param name="cancellationToken">Cancellation token for the operation</param>
 		/// <returns>The blob instance</returns>
-		public static Task<BlobId> ReadRefTargetAsync(this IBlobStore store, RefName name, TimeSpan maxAge, CancellationToken cancellationToken = default)
+		public static Task<BlobLocator> ReadRefTargetAsync(this IBlobStore store, RefName name, TimeSpan maxAge, CancellationToken cancellationToken = default)
 		{
 			return ReadRefTargetAsync(store, name, DateTime.UtcNow - maxAge, cancellationToken);
 		}
@@ -309,15 +309,15 @@ namespace EpicGames.Horde.Storage
 		class InMemoryBlob : IBlob
 		{
 			/// <inheritdoc/>
-			public BlobId Id { get; }
+			public BlobLocator Id { get; }
 
 			/// <inheritdoc/>
 			public ReadOnlyMemory<byte> Data { get; }
 
 			/// <inheritdoc/>
-			public IReadOnlyList<BlobId> References { get; }
+			public IReadOnlyList<BlobLocator> References { get; }
 
-			public InMemoryBlob(BlobId id, ReadOnlyMemory<byte> data, IReadOnlyList<BlobId> references)
+			public InMemoryBlob(BlobLocator id, ReadOnlyMemory<byte> data, IReadOnlyList<BlobLocator> references)
 			{
 				Id = id;
 				Data = data;
@@ -332,7 +332,7 @@ namespace EpicGames.Horde.Storage
 		/// <param name="data">Payload for the blob</param>
 		/// <param name="references">References to other blobs</param>
 		/// <returns>Blob instance</returns>
-		public static IBlob FromMemory(BlobId id, ReadOnlyMemory<byte> data, IReadOnlyList<BlobId> references) => new InMemoryBlob(id, data, references);
+		public static IBlob FromMemory(BlobLocator id, ReadOnlyMemory<byte> data, IReadOnlyList<BlobLocator> references) => new InMemoryBlob(id, data, references);
 
 		/// <summary>
 		/// Serialize a blob into a flat memory buffer
@@ -350,7 +350,7 @@ namespace EpicGames.Horde.Storage
 		/// <param name="data">Data for the blob</param>
 		/// <param name="references">List of references to other blobs</param>
 		/// <returns>Serialized data</returns>
-		public static ReadOnlySequence<byte> Serialize(ReadOnlyMemory<byte> data, IReadOnlyList<BlobId> references)
+		public static ReadOnlySequence<byte> Serialize(ReadOnlyMemory<byte> data, IReadOnlyList<BlobLocator> references)
 		{
 			return Serialize(new ReadOnlySequence<byte>(data), references);
 		}
@@ -361,10 +361,10 @@ namespace EpicGames.Horde.Storage
 		/// <param name="data">Data for the blob</param>
 		/// <param name="references">List of references to other blobs</param>
 		/// <returns>Serialized data</returns>
-		public static ReadOnlySequence<byte> Serialize(ReadOnlySequence<byte> data, IReadOnlyList<BlobId> references)
+		public static ReadOnlySequence<byte> Serialize(ReadOnlySequence<byte> data, IReadOnlyList<BlobLocator> references)
 		{
 			ByteArrayBuilder writer = new ByteArrayBuilder();
-			writer.WriteVariableLengthArray(references, x => writer.WriteBlobId(x));
+			writer.WriteVariableLengthArray(references, x => writer.WriteBlobLocator(x));
 			writer.WriteUnsignedVarInt((ulong)data.Length);
 
 			ReadOnlySequenceBuilder<byte> builder = new ReadOnlySequenceBuilder<byte>();
@@ -380,7 +380,7 @@ namespace EpicGames.Horde.Storage
 		/// <param name="id">Id for the blob</param>
 		/// <param name="memory">Memory to deserialize from</param>
 		/// <returns>Deserialized blob data. May reference the supplied memory.</returns>
-		public static IBlob Deserialize(BlobId id, ReadOnlyMemory<byte> memory) => Deserialize(id, new ReadOnlySequence<byte>(memory));
+		public static IBlob Deserialize(BlobLocator id, ReadOnlyMemory<byte> memory) => Deserialize(id, new ReadOnlySequence<byte>(memory));
 
 		/// <summary>
 		/// Deserialize a blob from a block of memory
@@ -388,10 +388,10 @@ namespace EpicGames.Horde.Storage
 		/// <param name="id">Id for the blob</param>
 		/// <param name="sequence">Sequence to deserialize from</param>
 		/// <returns>Deserialized blob data. May reference the supplied memory.</returns>
-		public static IBlob Deserialize(BlobId id, ReadOnlySequence<byte> sequence)
+		public static IBlob Deserialize(BlobLocator id, ReadOnlySequence<byte> sequence)
 		{
 			MemoryReader reader = new MemoryReader(sequence.First);
-			IReadOnlyList<BlobId> references = reader.ReadVariableLengthArray(() => reader.ReadBlobId());
+			IReadOnlyList<BlobLocator> references = reader.ReadVariableLengthArray(() => reader.ReadBlobLocator());
 			long length = (long)reader.ReadUnsignedVarInt();
 
 			ReadOnlyMemory<byte> data = sequence.Slice(sequence.First.Length - reader.Memory.Length).AsSingleSegment();
diff --git a/Engine/Source/Programs/Shared/EpicGames.Horde/Storage/ITreeStore.cs b/Engine/Source/Programs/Shared/EpicGames.Horde/Storage/ITreeStore.cs
--- a/Engine/Source/Programs/Shared/EpicGames.Horde/Storage/ITreeStore.cs
+++ b/Engine/Source/Programs/Shared/EpicGames.Horde/Storage/ITreeStore.cs
@@ -51,7 +51,7 @@ namespace EpicGames.Horde.Storage
 		/// <summary>
 		/// Creates a new context for reading trees.
 		/// </summary>
-		/// <param name="prefix">Prefix for blobs. See <see cref="IBlobStore.WriteBlobAsync(ReadOnlySequence{Byte}, IReadOnlyList{BlobId}, Utf8String, CancellationToken)"/></param>
+		/// <param name="prefix">Prefix for blobs. See <see cref="IBlobStore.WriteBlobAsync(ReadOnlySequence{Byte}, IReadOnlyList{BlobLocator}, Utf8String, CancellationToken)"/></param>
 		/// <returns>New context instance</returns>
 		ITreeWriter CreateTreeWriter(Utf8String prefix = default);
 
diff --git a/Engine/Source/Programs/Shared/EpicGames.Horde/Storage/RefName.cs b/Engine/Source/Programs/Shared/EpicGames.Horde/Storage/RefName.cs
--- a/Engine/Source/Programs/Shared/EpicGames.Horde/Storage/RefName.cs
+++ b/Engine/Source/Programs/Shared/EpicGames.Horde/Storage/RefName.cs
@@ -36,7 +36,7 @@ namespace EpicGames.Horde.Storage
 		public RefName(Utf8String text)
 		{
 			Text = text;
-			ContentId.ValidateArgument(nameof(text), text);
+			BlobId.ValidateArgument(nameof(text), text);
 		}
 
 		/// <inheritdoc/>
diff --git a/Engine/Source/Programs/Shared/EpicGames.Horde/Storage/ServerId.cs b/Engine/Source/Programs/Shared/EpicGames.Horde/Storage/ServerId.cs
deleted file mode 100644
diff --git a/Engine/Source/Programs/Shared/EpicGames.Horde/Storage/TreeStore.cs b/Engine/Source/Programs/Shared/EpicGames.Horde/Storage/TreeStore.cs
--- a/Engine/Source/Programs/Shared/EpicGames.Horde/Storage/TreeStore.cs
+++ b/Engine/Source/Programs/Shared/EpicGames.Horde/Storage/TreeStore.cs
@@ -179,9 +179,9 @@ namespace EpicGames.Horde.Storage
 		/// </summary>
 		class BundleContext
 		{
-			readonly ConcurrentDictionary<BlobId, BundleInfo> _blobIdToBundle = new ConcurrentDictionary<BlobId, BundleInfo>();
+			readonly ConcurrentDictionary<BlobLocator, BundleInfo> _blobIdToBundle = new ConcurrentDictionary<BlobLocator, BundleInfo>();
 
-			public BundleInfo FindOrAddBundle(BlobId blobId, int exportCount)
+			public BundleInfo FindOrAddBundle(BlobLocator blobId, int exportCount)
 			{
 				BundleInfo bundleInfo = _blobIdToBundle.GetOrAdd(blobId, new BundleInfo(this, blobId, exportCount));
 				Debug.Assert(bundleInfo.Exports.Length == exportCount);
@@ -196,16 +196,16 @@ namespace EpicGames.Horde.Storage
 		class BundleInfo
 		{
 			public readonly BundleContext Context;
-			public readonly BlobId BlobId;
+			public readonly BlobLocator Locator;
 			public readonly NodeInfo?[] Exports;
 
 			public bool Mounted;
 			public Task MountTask = Task.CompletedTask;
 
-			public BundleInfo(BundleContext context, BlobId blobId, int exportCount)
+			public BundleInfo(BundleContext context, BlobLocator locator, int exportCount)
 			{
 				Context = context;
-				BlobId = blobId;
+				Locator = locator;
 				Exports = new NodeInfo?[exportCount];
 			}
 		}
@@ -400,7 +400,7 @@ namespace EpicGames.Horde.Storage
 
 				// Write a reference to the blob containing this node
 				ImportedNodeState importedState = (ImportedNodeState)last.State;
-				await _owner._blobStore.WriteRefTargetAsync(refName, importedState.BundleInfo.BlobId, cancellationToken);
+				await _owner._blobStore.WriteRefTargetAsync(refName, importedState.BundleInfo.Locator, cancellationToken);
 			}
 
 			async Task<NodeInfo> FlushInternalAsync(CancellationToken cancellationToken)
@@ -446,12 +446,12 @@ namespace EpicGames.Horde.Storage
 				BundleHeader header = bundle.Header;
 
 				// Write it to storage
-				BlobId blobId = await _owner._blobStore.WriteBlobAsync(bundle.AsSequence(), bundle.Header.Imports.Select(x => x.BlobId).ToList(), _prefix, cancellationToken);
-				string cacheKey = GetBundleCacheKey(blobId);
+				BlobLocator locator = await _owner._blobStore.WriteBlobAsync(bundle.AsSequence(), bundle.Header.Imports.Select(x => x.Locator).ToList(), _prefix, cancellationToken);
+				string cacheKey = GetBundleCacheKey(locator);
 				_owner.AddBundleToCache(cacheKey, bundle);
 
 				// Create a BundleInfo for it
-				BundleInfo bundleInfo = _context.FindOrAddBundle(blobId, writeNodes.Length);
+				BundleInfo bundleInfo = _context.FindOrAddBundle(locator, writeNodes.Length);
 				for (int idx = 0; idx < writeNodes.Length; idx++)
 				{
 					bundleInfo.Exports[idx] = writeNodes[idx];
@@ -549,7 +549,7 @@ namespace EpicGames.Horde.Storage
 						ImportedNodeState importedNodeState = (ImportedNodeState)importedNode.State;
 						exportInfos[idx] = (importedNodeState.ExportIdx, importedNode.Hash);
 					}
-					imports.Add(new BundleImport(bundleInfo.BlobId, bundleInfo.Exports.Length, exportInfos));
+					imports.Add(new BundleImport(bundleInfo.Locator, bundleInfo.Exports.Length, exportInfos));
 				}
 
 				// Preallocate data in the encoded buffer to reduce fragmentation if we have to resize
@@ -766,8 +766,8 @@ namespace EpicGames.Horde.Storage
 
 			ExportedNodeState exportedState = await GetExportedStateAsync(node, cancellationToken);
 
-			Bundle bundle = await ReadBundleAsync(exportedState.BundleInfo.BlobId, cancellationToken);
-			ReadOnlyMemory<byte> packetData = DecodePacket(exportedState.BundleInfo.BlobId, exportedState.PacketInfo, bundle.Payload);
+			Bundle bundle = await ReadBundleAsync(exportedState.BundleInfo.Locator, cancellationToken);
+			ReadOnlyMemory<byte> packetData = DecodePacket(exportedState.BundleInfo.Locator, exportedState.PacketInfo, bundle.Payload);
 			ReadOnlyMemory<byte> nodeData = packetData.Slice(exportedState.Offset, exportedState.Length);
 
 			return new InMemoryNodeState(new ReadOnlySequence<byte>(nodeData), exportedState.References);
@@ -790,12 +790,12 @@ namespace EpicGames.Horde.Storage
 		/// <summary>
 		/// Reads a bundle from the given blob id, or retrieves it from the cache
 		/// </summary>
-		/// <param name="blobId"></param>
+		/// <param name="locator"></param>
 		/// <param name="cancellationToken"></param>
 		/// <returns></returns>
-		async Task<Bundle> ReadBundleAsync(BlobId blobId, CancellationToken cancellationToken = default)
+		async Task<Bundle> ReadBundleAsync(BlobLocator locator, CancellationToken cancellationToken = default)
 		{
-			string cacheKey = GetBundleCacheKey(blobId);
+			string cacheKey = GetBundleCacheKey(locator);
 			if (_cache == null || !_cache.TryGetValue<Bundle>(cacheKey, out Bundle? bundle))
 			{
 				Task<Bundle>? readTask;
@@ -803,7 +803,7 @@ namespace EpicGames.Horde.Storage
 				{
 					if (!_readTasks.TryGetValue(cacheKey, out readTask))
 					{
-						readTask = Task.Run(() => ReadBundleInternalAsync(cacheKey, blobId, cancellationToken), cancellationToken);
+						readTask = Task.Run(() => ReadBundleInternalAsync(cacheKey, locator, cancellationToken), cancellationToken);
 						_readTasks.Add(cacheKey, readTask);
 					}
 				}
@@ -812,7 +812,7 @@ namespace EpicGames.Horde.Storage
 			return bundle;
 		}
 
-		async Task<Bundle> ReadBundleInternalAsync(string cacheKey, BlobId blobId, CancellationToken cancellationToken)
+		async Task<Bundle> ReadBundleInternalAsync(string cacheKey, BlobLocator blobId, CancellationToken cancellationToken)
 		{
 			// Perform another (sequenced) check whether an object has been added to the cache, to counteract the race between a read task being added and a task completing.
 			lock (_lockObject)
@@ -848,7 +848,7 @@ namespace EpicGames.Horde.Storage
 			return bundle;
 		}
 
-		static string GetBundleCacheKey(BlobId blobId) => $"bundle:{blobId}";
+		static string GetBundleCacheKey(BlobLocator blobId) => $"bundle:{blobId}";
 
 		void AddBundleToCache(string cacheKey, Bundle bundle)
 		{
@@ -885,7 +885,7 @@ namespace EpicGames.Horde.Storage
 			// If it didn't succeed, try again
 			if (!bundleInfo.Mounted)
 			{
-				Bundle bundle = await ReadBundleAsync(bundleInfo.BlobId, cancellationToken);
+				Bundle bundle = await ReadBundleAsync(bundleInfo.Locator, cancellationToken);
 				MountBundle(bundleInfo, bundle);
 			}
 		}
@@ -938,7 +938,7 @@ namespace EpicGames.Horde.Storage
 			List<NodeInfo> indexedNodes = new List<NodeInfo>(header.Imports.Sum(x => x.Exports.Count) + header.Exports.Count);
 			foreach (BundleImport import in header.Imports)
 			{
-				BundleInfo importBundle = context.FindOrAddBundle(import.BlobId, import.ExportCount);
+				BundleInfo importBundle = context.FindOrAddBundle(import.Locator, import.ExportCount);
 				foreach ((int exportIdx, IoHash exportHash) in import.Exports)
 				{
 					NodeInfo? node = importBundle.Exports[exportIdx];
@@ -955,11 +955,11 @@ namespace EpicGames.Horde.Storage
 		/// <summary>
 		/// Gets a decoded block from the store
 		/// </summary>
-		/// <param name="blobId">Information about the bundle</param>
+		/// <param name="locator">Information about the bundle</param>
 		/// <param name="packetInfo">The decoded block location and size</param>
 		/// <param name="payload">The bundle payload</param>
 		/// <returns>The decoded data</returns>
-		ReadOnlyMemory<byte> DecodePacket(BlobId blobId, PacketInfo packetInfo, ReadOnlySequence<byte> payload)
+		ReadOnlyMemory<byte> DecodePacket(BlobLocator locator, PacketInfo packetInfo, ReadOnlySequence<byte> payload)
 		{
 			if (_cache == null)
 			{
@@ -967,7 +967,7 @@ namespace EpicGames.Horde.Storage
 			}
 			else
 			{
-				string cacheKey = $"bundle-packet:{blobId}@{packetInfo.Offset}";
+				string cacheKey = $"bundle-packet:{locator.BlobId}@{packetInfo.Offset}";
 				return _cache.GetOrCreate<ReadOnlyMemory<byte>>(cacheKey, entry =>
 				{
 					ReadOnlyMemory<byte> decodedPacket = DecodePacketUncached(packetInfo, payload);
