Horde: Add helper methods to construct BlobId and BlobLocator instances.

diff --git a/Engine/Source/Programs/Shared/EpicGames.Horde/Storage/BlobId.cs b/Engine/Source/Programs/Shared/EpicGames.Horde/Storage/BlobId.cs
--- a/Engine/Source/Programs/Shared/EpicGames.Horde/Storage/BlobId.cs
+++ b/Engine/Source/Programs/Shared/EpicGames.Horde/Storage/BlobId.cs
@@ -81,6 +81,36 @@ namespace EpicGames.Horde.Storage
 			s_counter = rnd.Next();
 		}
 
+		/// <summary>
+		/// Create a unique blob id, optionally including a prefix
+		/// </summary>
+		/// <param name="prefix">Prefix for blob ids.</param>
+		/// <returns>New content id</returns>
+		public static BlobId CreateNew(Utf8String prefix = default)
+		{
+			if (prefix.Length == 0)
+			{
+				DateTime now = DateTime.UtcNow.Date;
+				prefix = $"_by_date_/{now.Year}-{now.Month:D2}/{now.Day:D2}";
+			}
+			else
+			{
+				ValidateArgument(nameof(prefix), prefix);
+			}
+
+			byte[] buffer = new byte[prefix.Length + 1 + 24];
+			Span<byte> span = buffer;
+
+			prefix.Span.CopyTo(span);
+			span = span.Slice(prefix.Length);
+
+			span[0] = (byte)'/';
+			span = span.Slice(1);
+
+			GenerateUniqueId(span);
+			return new BlobId(new Utf8String(buffer), Sanitize.None);
+		}
+
 		/// <summary>
 		/// Creates a new identifier in the given buffer.
 		/// </summary>
diff --git a/Engine/Source/Programs/Shared/EpicGames.Horde/Storage/BlobLocator.cs b/Engine/Source/Programs/Shared/EpicGames.Horde/Storage/BlobLocator.cs
--- a/Engine/Source/Programs/Shared/EpicGames.Horde/Storage/BlobLocator.cs
+++ b/Engine/Source/Programs/Shared/EpicGames.Horde/Storage/BlobLocator.cs
@@ -2,6 +2,7 @@
 
 using EpicGames.Core;
 using EpicGames.Serialization;
+using Microsoft.CodeAnalysis.CSharp.Syntax;
 using System;
 using System.ComponentModel;
 using System.Globalization;
@@ -76,6 +77,29 @@ namespace EpicGames.Horde.Storage
 			ValidateArgument(nameof(inner), inner);
 		}
 
+		/// <summary>
+		/// Constructor
+		/// </summary>
+		/// <param name="hostId"></param>
+		/// <param name="blobId"></param>
+		public BlobLocator(HostId hostId, BlobId blobId)
+		{
+			if (hostId.IsValid())
+			{
+				byte[] buffer = new byte[hostId.Inner.Length + 1 + blobId.Inner.Length];
+
+				hostId.Inner.Span.CopyTo(buffer);
+				buffer[hostId.Inner.Length] = (byte)':';
+				blobId.Inner.Span.CopyTo(buffer.AsSpan(hostId.Inner.Length + 1));
+
+				Inner = new Utf8String(buffer);
+			}
+			else
+			{
+				Inner = blobId.Inner;
+			}
+		}
+
 		/// <summary>
 		/// Constructor
 		/// </summary>
