Horde: Return bundles directly from IBlobStore

diff --git a/Engine/Source/Programs/Horde/Horde.Agent/Commands/Bundles/CreateCommand.cs b/Engine/Source/Programs/Horde/Horde.Agent/Commands/Bundles/CreateCommand.cs
--- a/Engine/Source/Programs/Horde/Horde.Agent/Commands/Bundles/CreateCommand.cs
+++ b/Engine/Source/Programs/Horde/Horde.Agent/Commands/Bundles/CreateCommand.cs
@@ -45,22 +45,21 @@ namespace Horde.Agent.Commands.Bundles
 
 			#region Blobs
 
-			public async Task<IBlob> ReadBlobAsync(BlobLocator id, CancellationToken cancellationToken = default)
+			public async Task<Bundle> ReadBundleAsync(BlobLocator id, CancellationToken cancellationToken = default)
 			{
 				FileReference file = GetBlobFile(id);
 				_logger.LogInformation("Reading {File}", file);
 				byte[] bytes = await FileReference.ReadAllBytesAsync(file, cancellationToken);
-				return Blob.Deserialize(id, bytes);
+				return new Bundle(new MemoryReader(bytes));
 			}
 
-			public async Task<BlobLocator> WriteBlobAsync(ReadOnlySequence<byte> data, IReadOnlyList<BlobLocator> references, Utf8String prefix = default, CancellationToken cancellationToken = default)
+			public async Task<BlobLocator> WriteBundleAsync(Bundle bundle, Utf8String prefix = default, CancellationToken cancellationToken = default)
 			{
 				BlobLocator id = BlobLocator.Create(HostId.Empty, prefix);
 				FileReference file = GetBlobFile(id);
 				DirectoryReference.CreateDirectory(file.Directory);
 				_logger.LogInformation("Writing {File}", file);
-				IBlob blob = Blob.FromMemory(id, data.AsSingleSegment(), references);
-				await WriteAsync(file, Blob.Serialize(blob), cancellationToken);
+				await WriteAsync(file, bundle.AsSequence(), cancellationToken);
 				return id;
 			}
 
@@ -73,14 +72,16 @@ namespace Horde.Agent.Commands.Bundles
 				throw new NotImplementedException();
 			}
 
-			public async Task<IBlob?> TryReadRefAsync(RefName name, DateTime cacheTime = default, CancellationToken cancellationToken = default)
+			public async Task<BundleRef?> TryReadRefAsync(RefName name, DateTime cacheTime = default, CancellationToken cancellationToken = default)
 			{
-				BlobLocator id = await TryReadRefTargetAsync(name, cacheTime, cancellationToken);
-				if (!id.IsValid())
+				BlobLocator locator = await TryReadRefTargetAsync(name, cacheTime, cancellationToken);
+				if (!locator.IsValid())
 				{
 					return null;
 				}
-				return await ReadBlobAsync(id, cancellationToken);
+
+				Bundle bundle = await ReadBundleAsync(locator, cancellationToken);
+				return new BundleRef(locator, bundle);
 			}
 
 			public async Task<BlobLocator> TryReadRefTargetAsync(RefName name, DateTime cacheTime = default, CancellationToken cancellationToken = default)
@@ -96,9 +97,9 @@ namespace Horde.Agent.Commands.Bundles
 				return new BlobLocator(text.Trim());
 			}
 
-			public async Task<BlobLocator> WriteRefAsync(RefName name, ReadOnlySequence<byte> data, IReadOnlyList<BlobLocator> references, CancellationToken cancellationToken = default)
+			public async Task<BlobLocator> WriteRefAsync(RefName name, Bundle bundle, CancellationToken cancellationToken = default)
 			{
-				BlobLocator locator = await WriteBlobAsync(data, references, name.Text, cancellationToken);
+				BlobLocator locator = await WriteBundleAsync(bundle, name.Text, cancellationToken);
 				await WriteRefTargetAsync(name, locator, cancellationToken);
 				return locator;
 			}
diff --git a/Engine/Source/Programs/Horde/Horde.Agent/Commands/Bundles/DumpCommand.cs b/Engine/Source/Programs/Horde/Horde.Agent/Commands/Bundles/DumpCommand.cs
--- a/Engine/Source/Programs/Horde/Horde.Agent/Commands/Bundles/DumpCommand.cs
+++ b/Engine/Source/Programs/Horde/Horde.Agent/Commands/Bundles/DumpCommand.cs
@@ -23,30 +23,21 @@ namespace Horde.Agent.Commands.Bundles
 		{
 			IBlobStore blobStore = CreateBlobStore(logger);
 
-			IBlob blob;
+			Bundle blob;
 			if (BlobId != null)
 			{
-				blob = await blobStore.ReadBlobAsync(BlobId.Value);
+				blob = await blobStore.ReadBundleAsync(BlobId.Value);
 				logger.LogInformation("Summary for blob {BlobId}", BlobId.Value);
 			}
 			else
 			{
-				blob = await blobStore.ReadRefAsync(RefName);
+				blob = (await blobStore.ReadRefAsync(RefName)).Bundle;
 				logger.LogInformation("Summary for ref {RefId}", RefName);
 			}
 
-			IReadOnlyList<BlobLocator> references = blob.References;
-			logger.LogInformation("");
-			logger.LogInformation("BlobRefs: {NumRefs}", references.Count);
-			foreach (BlobLocator reference in references)
-			{
-				logger.LogInformation("  {BlobId}", reference);
-			}
-
-			ReadOnlyMemory<byte> data = blob.Data;
-			MemoryReader reader = new MemoryReader(data);
-			BundleHeader header = new BundleHeader(reader);
-			int packetStart = data.Length - reader.Memory.Length;
+			ReadOnlyMemory<byte> data = blob.Payload.AsSingleSegment();
+			BundleHeader header = blob.Header;
+			int packetStart = 0;
 
 			logger.LogInformation("");
 			logger.LogInformation("Imports: {NumImports}", header.Imports.Count);
diff --git a/Engine/Source/Programs/Horde/Horde.Build.Tests/BlobStoreTests.cs b/Engine/Source/Programs/Horde/Horde.Build.Tests/BlobStoreTests.cs
--- a/Engine/Source/Programs/Horde/Horde.Build.Tests/BlobStoreTests.cs
+++ b/Engine/Source/Programs/Horde/Horde.Build.Tests/BlobStoreTests.cs
@@ -49,17 +49,62 @@ namespace Horde.Build.Tests
 			return data;
 		}
 
+		class Blob
+		{
+			public ReadOnlyMemory<byte> Data { get; set; } = ReadOnlyMemory<byte>.Empty;
+			public List<BlobLocator> References { get; set; } = new List<BlobLocator>();
+		}
+
+		static Bundle CreateTestBundle(ReadOnlySequence<byte> data, IReadOnlyList<BlobLocator> refs)
+		{
+			List<BundleImport> imports = new List<BundleImport>();
+
+			Dictionary<BlobLocator, int> locatorToIndex = new Dictionary<BlobLocator, int>();
+			foreach (BlobLocator locator in refs)
+			{
+				int index;
+				if (!locatorToIndex.TryGetValue(locator, out index))
+				{
+					BundleImport import = new BundleImport(locator, 1, new List<(int, IoHash)> { (0, IoHash.Zero) });
+					imports.Add(import);
+					locatorToIndex.Add(locator, imports.Count - 1);
+				}
+			}
+
+			List<BundleExport> exports = new List<BundleExport>();
+			exports.Add(new BundleExport(IoHash.Compute(data), (int)data.Length, refs.Select(x => locatorToIndex[x]).ToArray()));
+
+			List<BundlePacket> packets = new List<BundlePacket>();
+			packets.Add(new BundlePacket((int)data.Length, (int)data.Length));
+
+			BundleHeader header = new BundleHeader(BundleCompressionFormat.None, imports, exports, packets);
+			return new Bundle(header, data);
+		}
+
+		static async Task<Blob> ReadBlobAsync(IBlobStore store, BlobLocator locator)
+		{
+			Bundle bundle = await store.ReadBundleAsync(locator);
+			return ExtractBlobFromBundle(bundle);
+		}
+
+		static Blob ExtractBlobFromBundle(Bundle bundle)
+		{
+			Blob blob = new Blob();
+			blob.Data = bundle.Payload.AsSingleSegment();
+			blob.References = bundle.Header.Exports[0].References.Select(x => bundle.Header.Imports[x].Locator).ToList();
+			return blob;
+		}
+
 		[TestMethod]
 		public async Task LeafTest()
 		{
 			IBlobStore store = CreateBlobStore();
 
 			byte[] input = CreateTestData(256, 0);
+			BlobLocator locator = await store.WriteBundleAsync(CreateTestBundle(new ReadOnlySequence<byte>(input), Array.Empty<BlobLocator>()));
 
-			BlobLocator locator = await store.WriteBlobAsync(new ReadOnlySequence<byte>(input), Array.Empty<BlobLocator>());
-
-			IBlob blob = await store.ReadBlobAsync(locator);
-			Assert.IsTrue(blob.Data.Span.SequenceEqual(input));
+			Bundle outputBundle = await store.ReadBundleAsync(locator);
+			Assert.IsTrue(outputBundle.Payload.AsSingleSegment().Span.SequenceEqual(input));
 		}
 
 		[TestMethod]
@@ -68,41 +113,37 @@ namespace Horde.Build.Tests
 			IBlobStore store = CreateBlobStore();
 
 			byte[] input1 = CreateTestData(256, 1);
-			BlobLocator locator1 = await store.WriteBlobAsync(new ReadOnlySequence<byte>(input1), Array.Empty<BlobLocator>());
-			IBlob blob1 = await store.ReadBlobAsync(locator1);
+			BlobLocator locator1 = await store.WriteBundleAsync(CreateTestBundle(new ReadOnlySequence<byte>(input1), Array.Empty<BlobLocator>()));
+			Blob blob1 = await ReadBlobAsync(store, locator1);
 			Assert.IsTrue(blob1.Data.Span.SequenceEqual(input1));
 			Assert.IsTrue(blob1.References.SequenceEqual(Array.Empty<BlobLocator>()));
 
 			byte[] input2 = CreateTestData(256, 2);
-			BlobLocator locator2 = await store.WriteBlobAsync(new ReadOnlySequence<byte>(input2), new BlobLocator[] { blob1.Id });
-			IBlob blob2 = await store.ReadBlobAsync(locator2);
+			BlobLocator locator2 = await store.WriteBundleAsync(CreateTestBundle(new ReadOnlySequence<byte>(input2), new BlobLocator[] { locator1 }));
+			Blob blob2 = await ReadBlobAsync(store, locator2);
 			Assert.IsTrue(blob2.Data.Span.SequenceEqual(input2));
-			Assert.IsTrue(blob2.References.SequenceEqual(new BlobLocator[] { blob1.Id }));
+			Assert.IsTrue(blob2.References.SequenceEqual(new BlobLocator[] { locator1 }));
 
 			byte[] input3 = CreateTestData(256, 3);
-			BlobLocator locator3 = await store.WriteBlobAsync(new ReadOnlySequence<byte>(input3), new BlobLocator[] { blob1.Id, blob2.Id, blob1.Id });
-			IBlob blob3 = await store.ReadBlobAsync(locator3);
+			BlobLocator locator3 = await store.WriteBundleAsync(CreateTestBundle(new ReadOnlySequence<byte>(input3), new BlobLocator[] { locator1, locator2, locator1 }));
+			Blob blob3 = await ReadBlobAsync(store, locator3);
 			Assert.IsTrue(blob3.Data.Span.SequenceEqual(input3));
-			Assert.IsTrue(blob3.References.SequenceEqual(new BlobLocator[] { blob1.Id, blob2.Id, blob1.Id }));
+			Assert.IsTrue(blob3.References.SequenceEqual(new BlobLocator[] { locator1, locator2, locator1 }));
 
 			for(int idx = 0; idx < 2; idx++)
 			{
 				RefName refName = new RefName("hello");
-				await store.WriteRefTargetAsync(refName, blob3.Id);
+				await store.WriteRefTargetAsync(refName, locator3);
 				BlobLocator refTargetId = await store.ReadRefTargetAsync(refName);
-				Assert.AreEqual(blob3.Id, refTargetId);
+				Assert.AreEqual(locator3, refTargetId);
 			}
 
 			RefName refName2 = new RefName("hello2");
 
-			BlobLocator refTargetId2 = await store.WriteRefAsync(refName2, new ReadOnlySequence<byte>(input3), new BlobLocator[] { blob1.Id, blob2.Id });
-			IBlob refTarget2 = await store.ReadBlobAsync(refTargetId2);
-			Assert.IsTrue(refTarget2.Data.Span.SequenceEqual(input3));
-			Assert.IsTrue(refTarget2.References.SequenceEqual(new BlobLocator[] { blob1.Id, blob2.Id }));
-
-			refTarget2 = await store.ReadRefAsync(refName2);
+			BlobLocator refTargetId2 = await store.WriteRefAsync(refName2, CreateTestBundle(new ReadOnlySequence<byte>(input3), new BlobLocator[] { locator1, locator2 }));
+			Blob refTarget2 = await ReadBlobAsync(store, refTargetId2);
 			Assert.IsTrue(refTarget2.Data.Span.SequenceEqual(input3));
-			Assert.IsTrue(refTarget2.References.SequenceEqual(new BlobLocator[] { blob1.Id, blob2.Id }));
+			Assert.IsTrue(refTarget2.References.SequenceEqual(new BlobLocator[] { locator1, locator2 }));
 		}
 	}
 }
diff --git a/Engine/Source/Programs/Horde/Horde.Build/Storage/BasicBlobStore.cs b/Engine/Source/Programs/Horde/Horde.Build/Storage/BasicBlobStore.cs
--- a/Engine/Source/Programs/Horde/Horde.Build/Storage/BasicBlobStore.cs
+++ b/Engine/Source/Programs/Horde/Horde.Build/Storage/BasicBlobStore.cs
@@ -78,12 +78,12 @@ namespace Horde.Build.Storage
 
 		#region Blobs
 
-		static string GetBlobPath(BlobLocator id) => $"{id.BlobId}.blob";
+		static string GetBlobPath(BlobId id) => $"{id}.blob";
 
 		/// <inheritdoc/>
-		public async Task<IBlob> ReadBlobAsync(BlobLocator id, CancellationToken cancellationToken = default)
+		public async Task<Bundle> ReadBundleAsync(BlobLocator locator, CancellationToken cancellationToken = default)
 		{
-			string path = GetBlobPath(id);
+			string path = GetBlobPath(locator.BlobId);
 
 			ReadOnlyMemory<byte>? data = await _backend.ReadBytesAsync(path, cancellationToken);
 			if (data == null)
@@ -91,21 +91,22 @@ namespace Horde.Build.Storage
 				throw new KeyNotFoundException($"Unable to read data from {path}");
 			}
 
-			return Blob.Deserialize(id, data.Value);
+			return new Bundle(new MemoryReader(data.Value));
 		}
 
 		/// <inheritdoc/>
-		public async Task<BlobLocator> WriteBlobAsync(ReadOnlySequence<byte> data, IReadOnlyList<BlobLocator> references, Utf8String prefix = default, CancellationToken cancellationToken = default)
+		public async Task<BlobLocator> WriteBundleAsync(Bundle bundle, Utf8String prefix = default, CancellationToken cancellationToken = default)
 		{
-			BlobLocator id = BlobLocator.Create(_serverId, prefix);
+			BlobId id = BlobId.CreateNew(prefix);
 			string path = GetBlobPath(id);
 
-			ReadOnlySequence<byte> sequence = Blob.Serialize(data, references);
+			ReadOnlySequence<byte> sequence = bundle.AsSequence();
 			using (ReadOnlySequenceStream stream = new ReadOnlySequenceStream(sequence))
 			{
 				await _backend.WriteAsync(path, stream, cancellationToken);
 			}
-			return id;
+
+			return new BlobLocator(_serverId, id);
 		}
 
 		#endregion
@@ -132,18 +133,20 @@ namespace Horde.Build.Storage
 		}
 
 		/// <inheritdoc/>
-		public async Task<IBlob?> TryReadRefAsync(RefName name, DateTime cacheTime = default, CancellationToken cancellationToken = default)
+		public async Task<BundleRef?> TryReadRefAsync(RefName name, DateTime cacheTime = default, CancellationToken cancellationToken = default)
 		{
 			BlobLocator locator = await TryReadRefTargetAsync(name, cacheTime, cancellationToken);
 			if (locator.IsValid())
 			{
 				try
 				{
-					return await ReadBlobAsync(locator, cancellationToken);
+					Bundle? bundle = await ReadBundleAsync(locator, cancellationToken);
+					return new BundleRef(locator, bundle);
 				}
 				catch (Exception ex)
 				{
 					_logger.LogWarning(ex, "Unable to read blob {BlobId} referenced by ref {RefName}", locator, name);
+					return null;
 				}
 			}
 			return null;
@@ -163,11 +166,11 @@ namespace Horde.Build.Storage
 		}
 
 		/// <inheritdoc/>
-		public async Task<BlobLocator> WriteRefAsync(RefName name, ReadOnlySequence<byte> data, IReadOnlyList<BlobLocator> references, CancellationToken cancellationToken = default)
+		public async Task<BlobLocator> WriteRefAsync(RefName name, Bundle bundle, CancellationToken cancellationToken = default)
 		{
-			BlobLocator blobId = await WriteBlobAsync(data, references, name.Text, cancellationToken);
-			await WriteRefTargetAsync(name, blobId, cancellationToken);
-			return blobId;
+			BlobLocator locator = await WriteBundleAsync(bundle, name.Text, cancellationToken);
+			await WriteRefTargetAsync(name, locator, cancellationToken);
+			return locator;
 		}
 
 		/// <inheritdoc/>
diff --git a/Engine/Source/Programs/Shared/EpicGames.Horde.Tests/BundleStoreTests.cs b/Engine/Source/Programs/Shared/EpicGames.Horde.Tests/BundleStoreTests.cs
--- a/Engine/Source/Programs/Shared/EpicGames.Horde.Tests/BundleStoreTests.cs
+++ b/Engine/Source/Programs/Shared/EpicGames.Horde.Tests/BundleStoreTests.cs
@@ -3,7 +3,6 @@
 using EpicGames.Horde.Storage;
 using EpicGames.Horde.Storage.Backends;
 using EpicGames.Horde.Storage.Nodes;
-using Microsoft.Extensions.Caching.Memory;
 using Microsoft.VisualStudio.TestTools.UnitTesting;
 using System;
 using System.Buffers;
diff --git a/Engine/Source/Programs/Shared/EpicGames.Horde.Tests/BundleTests.cs b/Engine/Source/Programs/Shared/EpicGames.Horde.Tests/BundleTests.cs
--- a/Engine/Source/Programs/Shared/EpicGames.Horde.Tests/BundleTests.cs
+++ b/Engine/Source/Programs/Shared/EpicGames.Horde.Tests/BundleTests.cs
@@ -180,8 +180,8 @@ namespace EpicGames.Horde.Tests
 			Assert.AreEqual(1, blobStore.Blobs.Count);
 
 			// Check the ref
-			IBlob blob = await blobStore.ReadRefAsync(refName);
-			Bundle bundle = new Bundle(new MemoryReader(blob.Data));
+			BundleRef bundleRef = await blobStore.ReadRefAsync(refName);
+			Bundle bundle = bundleRef.Bundle;
 			Assert.AreEqual(0, bundle.Header.Imports.Count);
 			Assert.AreEqual(3, bundle.Header.Exports.Count);
 
diff --git a/Engine/Source/Programs/Shared/EpicGames.Horde/Storage/Backends/InMemoryBlobStore.cs b/Engine/Source/Programs/Shared/EpicGames.Horde/Storage/Backends/InMemoryBlobStore.cs
--- a/Engine/Source/Programs/Shared/EpicGames.Horde/Storage/Backends/InMemoryBlobStore.cs
+++ b/Engine/Source/Programs/Shared/EpicGames.Horde/Storage/Backends/InMemoryBlobStore.cs
@@ -18,7 +18,7 @@ namespace EpicGames.Horde.Storage.Backends
 		/// <summary>
 		/// Map of blob id to blob data
 		/// </summary>
-		readonly ConcurrentDictionary<BlobLocator, IBlob> _blobs = new ConcurrentDictionary<BlobLocator, IBlob>();
+		readonly ConcurrentDictionary<BlobLocator, Bundle> _blobs = new ConcurrentDictionary<BlobLocator, Bundle>();
 
 		/// <summary>
 		/// Map of ref name to ref data
@@ -26,7 +26,7 @@ namespace EpicGames.Horde.Storage.Backends
 		readonly ConcurrentDictionary<RefName, BlobLocator> _refs = new ConcurrentDictionary<RefName, BlobLocator>();
 
 		/// <inheritdoc cref="_blobs"/>
-		public IReadOnlyDictionary<BlobLocator, IBlob> Blobs => _blobs;
+		public IReadOnlyDictionary<BlobLocator, Bundle> Blobs => _blobs;
 
 		/// <inheritdoc cref="_refs"/>
 		public IReadOnlyDictionary<RefName, BlobLocator> Refs => _refs;
@@ -34,20 +34,17 @@ namespace EpicGames.Horde.Storage.Backends
 		#region Blobs
 
 		/// <inheritdoc/>
-		public Task<IBlob> ReadBlobAsync(BlobLocator locator, CancellationToken cancellationToken = default)
+		public Task<Bundle> ReadBundleAsync(BlobLocator locator, CancellationToken cancellationToken = default)
 		{
 			return Task.FromResult(_blobs[locator]);
 		}
 
 		/// <inheritdoc/>
-		public Task<BlobLocator> WriteBlobAsync(ReadOnlySequence<byte> data, IReadOnlyList<BlobLocator> references, Utf8String prefix = default, CancellationToken cancellationToken = default)
+		public Task<BlobLocator> WriteBundleAsync(Bundle bundle, Utf8String prefix = default, CancellationToken cancellationToken = default)
 		{
-			BlobLocator blobId = BlobLocator.Create(HostId.Empty, prefix);
-
-			IBlob blob = Blob.FromMemory(blobId, data.ToArray(), references);
-			_blobs[blobId] = blob;
-
-			return Task.FromResult(blobId);
+			BlobLocator locator = BlobLocator.Create(HostId.Empty, prefix);
+			_blobs[locator] = bundle;
+			return Task.FromResult(locator);
 		}
 
 		#endregion
@@ -58,25 +55,22 @@ namespace EpicGames.Horde.Storage.Backends
 		public Task DeleteRefAsync(RefName name, CancellationToken cancellationToken) => Task.FromResult(_refs.TryRemove(name, out _));
 
 		/// <inheritdoc/>
-		public Task<bool> HasRefAsync(RefName name, CancellationToken cancellationToken) => Task.FromResult(_refs.ContainsKey(name));
-
-		/// <inheritdoc/>
-		public Task<IBlob?> TryReadRefAsync(RefName name, DateTime cacheTime = default, CancellationToken cancellationToken = default)
+		public Task<BundleRef?> TryReadRefAsync(RefName name, DateTime cacheTime = default, CancellationToken cancellationToken = default)
 		{
-			IBlob? blob = null;
-			if (_refs.TryGetValue(name, out BlobLocator locator))
+			BundleRef? result = null;
+			if (_refs.TryGetValue(name, out BlobLocator locator) && _blobs.TryGetValue(locator, out Bundle? bundle))
 			{
-				_blobs.TryGetValue(locator, out blob);
+				result = new BundleRef(locator, bundle);
 			}
-			return Task.FromResult(blob);
+			return Task.FromResult(result);
 		}
 
 		/// <inheritdoc/>
 		public Task<BlobLocator> TryReadRefTargetAsync(RefName name, DateTime cacheTime = default, CancellationToken cancellationToken = default)
 		{
-			if (_refs.TryGetValue(name, out BlobLocator blobId))
+			if (_refs.TryGetValue(name, out BlobLocator locator))
 			{
-				return Task.FromResult(blobId);
+				return Task.FromResult(locator);
 			}
 			else
 			{
@@ -85,17 +79,17 @@ namespace EpicGames.Horde.Storage.Backends
 		}
 
 		/// <inheritdoc/>
-		public async Task<BlobLocator> WriteRefAsync(RefName name, ReadOnlySequence<byte> data, IReadOnlyList<BlobLocator> references, CancellationToken cancellationToken)
+		public async Task<BlobLocator> WriteRefAsync(RefName name, Bundle bundle, CancellationToken cancellationToken)
 		{
-			BlobLocator locator = await WriteBlobAsync(data, references, name.Text, cancellationToken);
+			BlobLocator locator = await WriteBundleAsync(bundle, name.Text, cancellationToken);
 			_refs[name] = locator;
 			return locator;
 		}
 
 		/// <inheritdoc/>
-		public Task WriteRefTargetAsync(RefName name, BlobLocator blobId, CancellationToken cancellationToken = default)
+		public Task WriteRefTargetAsync(RefName name, BlobLocator locator, CancellationToken cancellationToken = default)
 		{
-			_refs[name] = blobId;
+			_refs[name] = locator;
 			return Task.CompletedTask;
 		}
 
diff --git a/Engine/Source/Programs/Shared/EpicGames.Horde/Storage/Bundle.cs b/Engine/Source/Programs/Shared/EpicGames.Horde/Storage/Bundle.cs
--- a/Engine/Source/Programs/Shared/EpicGames.Horde/Storage/Bundle.cs
+++ b/Engine/Source/Programs/Shared/EpicGames.Horde/Storage/Bundle.cs
@@ -8,6 +8,7 @@ using System.Collections.Generic;
 using System.Diagnostics.CodeAnalysis;
 using System.IO;
 using System.Linq;
+using System.Reflection.Metadata;
 using System.Text;
 using System.Threading;
 using System.Threading.Tasks;
@@ -72,26 +73,6 @@ namespace EpicGames.Horde.Storage
 			Payload = new ReadOnlySequence<byte>(reader.ReadFixedLengthBytes(length));
 		}
 
-		/// <summary>
-		/// Create a bundle from a blob
-		/// </summary>
-		/// <param name="blob">Blob to parse from</param>
-		/// <returns>Bundle parsed from the given blob</returns>
-		[return: NotNullIfNotNull("blob")]
-		public static Bundle? FromBlob(IBlob? blob)
-		{
-			if (blob == null)
-			{
-				return null;
-			}
-
-			ReadOnlyMemory<byte> data = blob.Data;
-			MemoryReader reader = new MemoryReader(data);
-			Bundle bundle = new Bundle(reader);
-			reader.CheckEmpty();
-			return bundle;
-		}
-
 		/// <summary>
 		/// Serializes the bundle to a sequence of bytes
 		/// </summary>
diff --git a/Engine/Source/Programs/Shared/EpicGames.Horde/Storage/IBlobStore.cs b/Engine/Source/Programs/Shared/EpicGames.Horde/Storage/IBlobStore.cs
--- a/Engine/Source/Programs/Shared/EpicGames.Horde/Storage/IBlobStore.cs
+++ b/Engine/Source/Programs/Shared/EpicGames.Horde/Storage/IBlobStore.cs
@@ -32,26 +32,28 @@ namespace EpicGames.Horde.Storage
 	}
 
 	/// <summary>
-	/// Interface for accessing an object's data
+	/// Value for a ref
 	/// </summary>
-	public interface IBlob
+	public class BundleRef
 	{
 		/// <summary>
-		/// Identifier for the blob
+		/// Location of the blob
 		/// </summary>
-		BlobLocator Id { get; }
+		public BlobLocator Locator { get; }
 
 		/// <summary>
-		/// Gets the data for this blob
+		/// Bundle stored for this ref
 		/// </summary>
-		/// <returns></returns>
-		ReadOnlyMemory<byte> Data { get; }
+		public Bundle Bundle { get; }
 
 		/// <summary>
-		/// Find the outward references for a node
+		/// Constructor
 		/// </summary>
-		/// <returns></returns>
-		IReadOnlyList<BlobLocator> References { get; }
+		public BundleRef(BlobLocator locator, Bundle bundle)
+		{
+			Locator = locator;
+			Bundle = bundle;
+		}
 	}
 
 	/// <summary>
@@ -59,25 +61,24 @@ namespace EpicGames.Horde.Storage
 	/// </summary>
 	public interface IBlobStore
 	{
-		#region Blobs
+		#region Bundles
 
 		/// <summary>
-		/// Reads data for a blob from the store
+		/// Reads data for a bundle from the store
 		/// </summary>
-		/// <param name="id">The blob identifier</param>
+		/// <param name="locator">The blob identifier</param>
 		/// <param name="cancellationToken">Cancellation token for the operation</param>
 		/// <returns></returns>
-		Task<IBlob> ReadBlobAsync(BlobLocator id, CancellationToken cancellationToken = default);
+		Task<Bundle> ReadBundleAsync(BlobLocator locator, CancellationToken cancellationToken = default);
 
 		/// <summary>
-		/// Writes a new object to the store
+		/// Writes a new bundle to the store
 		/// </summary>
-		/// <param name="data">Payload for the object</param>
-		/// <param name="references">Object references</param>
+		/// <param name="bundle">Bundle data</param>
 		/// <param name="cancellationToken">Cancellation token for the operation</param>
 		/// <param name="prefix">Prefix for blob names. While the returned BlobId is guaranteed to be unique, this name can be used as a prefix to aid debugging.</param>
 		/// <returns>Unique identifier for the blob</returns>
-		Task<BlobLocator> WriteBlobAsync(ReadOnlySequence<byte> data, IReadOnlyList<BlobLocator> references, Utf8String prefix = default, CancellationToken cancellationToken = default);
+		Task<BlobLocator> WriteBundleAsync(Bundle bundle, Utf8String prefix = default, CancellationToken cancellationToken = default);
 
 		#endregion
 
@@ -90,7 +91,7 @@ namespace EpicGames.Horde.Storage
 		/// <param name="cacheTime">Minimum coherency for any cached value to be returned</param>
 		/// <param name="cancellationToken">Cancellation token for the operation</param>
 		/// <returns></returns>
-		Task<IBlob?> TryReadRefAsync(RefName name, DateTime cacheTime = default, CancellationToken cancellationToken = default);
+		Task<BundleRef?> TryReadRefAsync(RefName name, DateTime cacheTime = default, CancellationToken cancellationToken = default);
 
 		/// <summary>
 		/// Reads data for a ref from the store
@@ -105,20 +106,19 @@ namespace EpicGames.Horde.Storage
 		/// Writes a new ref to the store which points to a new blob
 		/// </summary>
 		/// <param name="name">Ref to write</param>
-		/// <param name="data">Payload for the object</param>
-		/// <param name="references">Object references</param>
+		/// <param name="bundle">Bundle to write</param>
 		/// <param name="cancellationToken">Cancellation token for the operation</param>
 		/// <returns>Unique identifier for the blob</returns>
-		Task<BlobLocator> WriteRefAsync(RefName name, ReadOnlySequence<byte> data, IReadOnlyList<BlobLocator> references, CancellationToken cancellationToken = default);
+		Task<BlobLocator> WriteRefAsync(RefName name, Bundle bundle, CancellationToken cancellationToken = default);
 
 		/// <summary>
 		/// Writes a new ref to the store
 		/// </summary>
 		/// <param name="name">Ref to write</param>
-		/// <param name="blobId">The target for the ref</param>
+		/// <param name="locator">The target for the ref</param>
 		/// <param name="cancellationToken">Cancellation token for the operation</param>
 		/// <returns>Unique identifier for the blob</returns>
-		Task WriteRefTargetAsync(RefName name, BlobLocator blobId, CancellationToken cancellationToken = default);
+		Task WriteRefTargetAsync(RefName name, BlobLocator locator, CancellationToken cancellationToken = default);
 
 		/// <summary>
 		/// Reads data for a ref from the store
@@ -151,10 +151,10 @@ namespace EpicGames.Horde.Storage
 			#region Blobs
 
 			/// <inheritdoc/>
-			public Task<IBlob> ReadBlobAsync(BlobLocator id, CancellationToken cancellationToken = default) => _inner.ReadBlobAsync(id, cancellationToken);
+			public Task<Bundle> ReadBundleAsync(BlobLocator id, CancellationToken cancellationToken = default) => _inner.ReadBundleAsync(id, cancellationToken);
 
 			/// <inheritdoc/>
-			public Task<BlobLocator> WriteBlobAsync(ReadOnlySequence<byte> data, IReadOnlyList<BlobLocator> references, Utf8String prefix = default, CancellationToken cancellationToken = default) => _inner.WriteBlobAsync(data, references, prefix, cancellationToken);
+			public Task<BlobLocator> WriteBundleAsync(Bundle bundle, Utf8String prefix = default, CancellationToken cancellationToken = default) => _inner.WriteBundleAsync(bundle, prefix, cancellationToken);
 
 			#endregion
 
@@ -164,13 +164,13 @@ namespace EpicGames.Horde.Storage
 			public Task DeleteRefAsync(RefName name, CancellationToken cancellationToken = default) => _inner.DeleteRefAsync(name, cancellationToken);
 
 			/// <inheritdoc/>
-			public Task<IBlob?> TryReadRefAsync(RefName name, DateTime cacheTime = default, CancellationToken cancellationToken = default) => _inner.TryReadRefAsync(name, cacheTime, cancellationToken);
+			public Task<BundleRef?> TryReadRefAsync(RefName name, DateTime cacheTime = default, CancellationToken cancellationToken = default) => _inner.TryReadRefAsync(name, cacheTime, cancellationToken);
 
 			/// <inheritdoc/>
 			public Task<BlobLocator> TryReadRefTargetAsync(RefName name, DateTime cacheTime = default, CancellationToken cancellationToken = default) => _inner.TryReadRefTargetAsync(name, cacheTime, cancellationToken);
 
 			/// <inheritdoc/>
-			public Task<BlobLocator> WriteRefAsync(RefName name, ReadOnlySequence<byte> data, IReadOnlyList<BlobLocator> references, CancellationToken cancellationToken = default) => _inner.WriteRefAsync(name, data, references, cancellationToken);
+			public Task<BlobLocator> WriteRefAsync(RefName name, Bundle bundle, CancellationToken cancellationToken = default) => _inner.WriteRefAsync(name, bundle, cancellationToken);
 
 			/// <inheritdoc/>
 			public Task WriteRefTargetAsync(RefName name, BlobLocator blobId, CancellationToken cancellationToken = default) => _inner.WriteRefTargetAsync(name, blobId, cancellationToken);
@@ -220,7 +220,7 @@ namespace EpicGames.Horde.Storage
 		/// <param name="maxAge">Maximum age of any cached ref</param>
 		/// <param name="cancellationToken">Cancellation token for the operation</param>
 		/// <returns>The blob instance</returns>
-		public static Task<IBlob?> TryReadRefAsync(this IBlobStore store, RefName name, TimeSpan maxAge, CancellationToken cancellationToken = default)
+		public static Task<BundleRef?> TryReadRefAsync(this IBlobStore store, RefName name, TimeSpan maxAge, CancellationToken cancellationToken = default)
 		{
 			return store.TryReadRefAsync(name, DateTime.UtcNow - maxAge, cancellationToken);
 		}
@@ -246,14 +246,14 @@ namespace EpicGames.Horde.Storage
 		/// <param name="cacheTime">Minimum coherency of any cached result</param>
 		/// <param name="cancellationToken">Cancellation token for the operation</param>
 		/// <returns>The blob instance</returns>
-		public static async Task<IBlob> ReadRefAsync(this IBlobStore store, RefName name, DateTime cacheTime = default, CancellationToken cancellationToken = default)
+		public static async Task<BundleRef> ReadRefAsync(this IBlobStore store, RefName name, DateTime cacheTime = default, CancellationToken cancellationToken = default)
 		{
-			IBlob? blob = await store.TryReadRefAsync(name, cacheTime, cancellationToken);
-			if (blob == null)
+			BundleRef? bundleRef = await store.TryReadRefAsync(name, cacheTime, cancellationToken);
+			if (bundleRef == null)
 			{
 				throw new RefNameNotFoundException(name);
 			}
-			return blob;
+			return bundleRef;
 		}
 
 		/// <summary>
@@ -264,7 +264,7 @@ namespace EpicGames.Horde.Storage
 		/// <param name="maxAge">Maximum age for any cached result</param>
 		/// <param name="cancellationToken">Cancellation token for the operation</param>
 		/// <returns>The blob instance</returns>
-		public static Task<IBlob> ReadRefAsync(this IBlobStore store, RefName name, TimeSpan maxAge, CancellationToken cancellationToken = default)
+		public static Task<BundleRef> ReadRefAsync(this IBlobStore store, RefName name, TimeSpan maxAge, CancellationToken cancellationToken = default)
 		{
 			return ReadRefAsync(store, name, DateTime.UtcNow - maxAge, cancellationToken);
 		}
@@ -300,102 +300,4 @@ namespace EpicGames.Horde.Storage
 			return ReadRefTargetAsync(store, name, DateTime.UtcNow - maxAge, cancellationToken);
 		}
 	}
-
-	/// <summary>
-	/// Utility methods for blobs
-	/// </summary>
-	public static class Blob
-	{
-		class InMemoryBlob : IBlob
-		{
-			/// <inheritdoc/>
-			public BlobLocator Id { get; }
-
-			/// <inheritdoc/>
-			public ReadOnlyMemory<byte> Data { get; }
-
-			/// <inheritdoc/>
-			public IReadOnlyList<BlobLocator> References { get; }
-
-			public InMemoryBlob(BlobLocator id, ReadOnlyMemory<byte> data, IReadOnlyList<BlobLocator> references)
-			{
-				Id = id;
-				Data = data;
-				References = references;
-			}
-		}
-
-		/// <summary>
-		/// Create a blob from memory
-		/// </summary>
-		/// <param name="id">Id for the blob</param>
-		/// <param name="data">Payload for the blob</param>
-		/// <param name="references">References to other blobs</param>
-		/// <returns>Blob instance</returns>
-		public static IBlob FromMemory(BlobLocator id, ReadOnlyMemory<byte> data, IReadOnlyList<BlobLocator> references) => new InMemoryBlob(id, data, references);
-
-		/// <summary>
-		/// Serialize a blob into a flat memory buffer
-		/// </summary>
-		/// <param name="blob">The blob to serialize</param>
-		/// <returns>Serialized data</returns>
-		public static ReadOnlySequence<byte> Serialize(IBlob blob)
-		{
-			return Serialize(blob.Data, blob.References);
-		}
-
-		/// <summary>
-		/// Serialize a blob into a flat memory buffer
-		/// </summary>
-		/// <param name="data">Data for the blob</param>
-		/// <param name="references">List of references to other blobs</param>
-		/// <returns>Serialized data</returns>
-		public static ReadOnlySequence<byte> Serialize(ReadOnlyMemory<byte> data, IReadOnlyList<BlobLocator> references)
-		{
-			return Serialize(new ReadOnlySequence<byte>(data), references);
-		}
-
-		/// <summary>
-		/// Serialize a blob into a flat memory buffer
-		/// </summary>
-		/// <param name="data">Data for the blob</param>
-		/// <param name="references">List of references to other blobs</param>
-		/// <returns>Serialized data</returns>
-		public static ReadOnlySequence<byte> Serialize(ReadOnlySequence<byte> data, IReadOnlyList<BlobLocator> references)
-		{
-			ByteArrayBuilder writer = new ByteArrayBuilder();
-			writer.WriteVariableLengthArray(references, x => writer.WriteBlobLocator(x));
-			writer.WriteUnsignedVarInt((ulong)data.Length);
-
-			ReadOnlySequenceBuilder<byte> builder = new ReadOnlySequenceBuilder<byte>();
-			builder.Append(writer.ToByteArray());
-			builder.Append(data);
-
-			return builder.Construct();
-		}
-
-		/// <summary>
-		/// Deserialize a blob from a block of memory
-		/// </summary>
-		/// <param name="id">Id for the blob</param>
-		/// <param name="memory">Memory to deserialize from</param>
-		/// <returns>Deserialized blob data. May reference the supplied memory.</returns>
-		public static IBlob Deserialize(BlobLocator id, ReadOnlyMemory<byte> memory) => Deserialize(id, new ReadOnlySequence<byte>(memory));
-
-		/// <summary>
-		/// Deserialize a blob from a block of memory
-		/// </summary>
-		/// <param name="id">Id for the blob</param>
-		/// <param name="sequence">Sequence to deserialize from</param>
-		/// <returns>Deserialized blob data. May reference the supplied memory.</returns>
-		public static IBlob Deserialize(BlobLocator id, ReadOnlySequence<byte> sequence)
-		{
-			MemoryReader reader = new MemoryReader(sequence.First);
-			IReadOnlyList<BlobLocator> references = reader.ReadVariableLengthArray(() => reader.ReadBlobLocator());
-			long length = (long)reader.ReadUnsignedVarInt();
-
-			ReadOnlyMemory<byte> data = sequence.Slice(sequence.First.Length - reader.Memory.Length).AsSingleSegment();
-			return new InMemoryBlob(id, data, references);
-		}
-	}
 }
diff --git a/Engine/Source/Programs/Shared/EpicGames.Horde/Storage/ITreeStore.cs b/Engine/Source/Programs/Shared/EpicGames.Horde/Storage/ITreeStore.cs
--- a/Engine/Source/Programs/Shared/EpicGames.Horde/Storage/ITreeStore.cs
+++ b/Engine/Source/Programs/Shared/EpicGames.Horde/Storage/ITreeStore.cs
@@ -51,7 +51,7 @@ namespace EpicGames.Horde.Storage
 		/// <summary>
 		/// Creates a new context for reading trees.
 		/// </summary>
-		/// <param name="prefix">Prefix for blobs. See <see cref="IBlobStore.WriteBlobAsync(ReadOnlySequence{Byte}, IReadOnlyList{BlobLocator}, Utf8String, CancellationToken)"/></param>
+		/// <param name="prefix">Prefix for blobs. See <see cref="IBlobStore.WriteBundleAsync(Bundle, Utf8String, CancellationToken)"/></param>
 		/// <returns>New context instance</returns>
 		ITreeWriter CreateTreeWriter(Utf8String prefix = default);
 
diff --git a/Engine/Source/Programs/Shared/EpicGames.Horde/Storage/TreeStore.cs b/Engine/Source/Programs/Shared/EpicGames.Horde/Storage/TreeStore.cs
--- a/Engine/Source/Programs/Shared/EpicGames.Horde/Storage/TreeStore.cs
+++ b/Engine/Source/Programs/Shared/EpicGames.Horde/Storage/TreeStore.cs
@@ -446,7 +446,7 @@ namespace EpicGames.Horde.Storage
 				BundleHeader header = bundle.Header;
 
 				// Write it to storage
-				BlobLocator locator = await _owner._blobStore.WriteBlobAsync(bundle.AsSequence(), bundle.Header.Imports.Select(x => x.Locator).ToList(), _prefix, cancellationToken);
+				BlobLocator locator = await _owner._blobStore.WriteBundleAsync(bundle, _prefix, cancellationToken);
 				string cacheKey = GetBundleCacheKey(locator);
 				_owner.AddBundleToCache(cacheKey, bundle);
 
@@ -733,8 +733,8 @@ namespace EpicGames.Horde.Storage
 		public async Task<ITreeBlob?> TryReadTreeAsync(RefName name, TimeSpan maxAge = default, CancellationToken cancellationToken = default)
 		{
 			// Read the blob referenced by this ref
-			IBlob? blob = await _blobStore.TryReadRefAsync(name, maxAge, cancellationToken);
-			if (blob == null)
+			BundleRef? bundleRef = await _blobStore.TryReadRefAsync(name, maxAge, cancellationToken);
+			if (bundleRef == null)
 			{
 				return null;
 			}
@@ -743,9 +743,8 @@ namespace EpicGames.Horde.Storage
 			BundleContext context = new BundleContext();
 
 			// Add the new bundle to it
-			Bundle bundle = Bundle.FromBlob(blob);
-			BundleInfo bundleInfo = context.FindOrAddBundle(blob.Id, bundle.Header.Exports.Count);
-			MountBundle(bundleInfo, bundle);
+			BundleInfo bundleInfo = context.FindOrAddBundle(bundleRef.Locator, bundleRef.Bundle.Header.Exports.Count);
+			MountBundle(bundleInfo, bundleRef.Bundle);
 
 			// Return the last node from the bundle as the root
 			return await GetNodeAsync(bundleInfo.Exports[^1]!, cancellationToken);
@@ -812,7 +811,7 @@ namespace EpicGames.Horde.Storage
 			return bundle;
 		}
 
-		async Task<Bundle> ReadBundleInternalAsync(string cacheKey, BlobLocator blobId, CancellationToken cancellationToken)
+		async Task<Bundle> ReadBundleInternalAsync(string cacheKey, BlobLocator locator, CancellationToken cancellationToken)
 		{
 			// Perform another (sequenced) check whether an object has been added to the cache, to counteract the race between a read task being added and a task completing.
 			lock (_lockObject)
@@ -830,7 +829,7 @@ namespace EpicGames.Horde.Storage
 			Bundle bundle;
 			try
 			{
-				bundle = Bundle.FromBlob(await _blobStore.ReadBlobAsync(blobId, cancellationToken));
+				bundle = await _blobStore.ReadBundleAsync(locator, cancellationToken);
 			}
 			finally
 			{
