Horde: Derive BundleRef from Bundle, to simplify access to ref objects.

diff --git a/Engine/Source/Programs/Horde/Horde.Agent/Commands/Bundles/CreateCommand.cs b/Engine/Source/Programs/Horde/Horde.Agent/Commands/Bundles/CreateCommand.cs
--- a/Engine/Source/Programs/Horde/Horde.Agent/Commands/Bundles/CreateCommand.cs
+++ b/Engine/Source/Programs/Horde/Horde.Agent/Commands/Bundles/CreateCommand.cs
@@ -2,14 +2,12 @@
 
 using System;
 using System.Buffers;
-using System.Collections.Generic;
 using System.IO;
 using System.Threading;
 using System.Threading.Tasks;
 using EpicGames.Core;
 using EpicGames.Horde.Storage;
 using EpicGames.Horde.Storage.Nodes;
-using Microsoft.Extensions.Caching.Memory;
 using Microsoft.Extensions.Logging;
 
 namespace Horde.Agent.Commands.Bundles
diff --git a/Engine/Source/Programs/Horde/Horde.Agent/Commands/Bundles/DumpCommand.cs b/Engine/Source/Programs/Horde/Horde.Agent/Commands/Bundles/DumpCommand.cs
--- a/Engine/Source/Programs/Horde/Horde.Agent/Commands/Bundles/DumpCommand.cs
+++ b/Engine/Source/Programs/Horde/Horde.Agent/Commands/Bundles/DumpCommand.cs
@@ -31,7 +31,7 @@ namespace Horde.Agent.Commands.Bundles
 			}
 			else
 			{
-				blob = (await blobStore.ReadRefAsync(RefName)).Bundle;
+				blob = await blobStore.ReadRefAsync(RefName);
 				logger.LogInformation("Summary for ref {RefId}", RefName);
 			}
 
diff --git a/Engine/Source/Programs/Shared/EpicGames.Horde.Tests/BundleTests.cs b/Engine/Source/Programs/Shared/EpicGames.Horde.Tests/BundleTests.cs
--- a/Engine/Source/Programs/Shared/EpicGames.Horde.Tests/BundleTests.cs
+++ b/Engine/Source/Programs/Shared/EpicGames.Horde.Tests/BundleTests.cs
@@ -180,8 +180,7 @@ namespace EpicGames.Horde.Tests
 			Assert.AreEqual(1, blobStore.Blobs.Count);
 
 			// Check the ref
-			BundleRef bundleRef = await blobStore.ReadRefAsync(refName);
-			Bundle bundle = bundleRef.Bundle;
+			BundleRef bundle = await blobStore.ReadRefAsync(refName);
 			Assert.AreEqual(0, bundle.Header.Imports.Count);
 			Assert.AreEqual(3, bundle.Header.Exports.Count);
 
diff --git a/Engine/Source/Programs/Shared/EpicGames.Horde/Storage/IBlobStore.cs b/Engine/Source/Programs/Shared/EpicGames.Horde/Storage/IBlobStore.cs
--- a/Engine/Source/Programs/Shared/EpicGames.Horde/Storage/IBlobStore.cs
+++ b/Engine/Source/Programs/Shared/EpicGames.Horde/Storage/IBlobStore.cs
@@ -34,25 +34,20 @@ namespace EpicGames.Horde.Storage
 	/// <summary>
 	/// Value for a ref
 	/// </summary>
-	public class BundleRef
+	public class BundleRef : Bundle
 	{
 		/// <summary>
 		/// Location of the blob
 		/// </summary>
 		public BlobLocator Locator { get; }
 
-		/// <summary>
-		/// Bundle stored for this ref
-		/// </summary>
-		public Bundle Bundle { get; }
-
 		/// <summary>
 		/// Constructor
 		/// </summary>
 		public BundleRef(BlobLocator locator, Bundle bundle)
+			: base(bundle.Header, bundle.Payload)
 		{
 			Locator = locator;
-			Bundle = bundle;
 		}
 	}
 
diff --git a/Engine/Source/Programs/Shared/EpicGames.Horde/Storage/TreeStore.cs b/Engine/Source/Programs/Shared/EpicGames.Horde/Storage/TreeStore.cs
--- a/Engine/Source/Programs/Shared/EpicGames.Horde/Storage/TreeStore.cs
+++ b/Engine/Source/Programs/Shared/EpicGames.Horde/Storage/TreeStore.cs
@@ -743,8 +743,8 @@ namespace EpicGames.Horde.Storage
 			BundleContext context = new BundleContext();
 
 			// Add the new bundle to it
-			BundleInfo bundleInfo = context.FindOrAddBundle(bundleRef.Locator, bundleRef.Bundle.Header.Exports.Count);
-			MountBundle(bundleInfo, bundleRef.Bundle);
+			BundleInfo bundleInfo = context.FindOrAddBundle(bundleRef.Locator, bundleRef.Header.Exports.Count);
+			MountBundle(bundleInfo, bundleRef);
 
 			// Return the last node from the bundle as the root
 			return await GetNodeAsync(bundleInfo.Exports[^1]!, cancellationToken);
