Horde: Remove restriction on maximum concurrent reads for tree stores. This can be implemented at the blob store level instead.

diff --git a/Engine/Source/Programs/Shared/EpicGames.Horde/Storage/TreeStore.cs b/Engine/Source/Programs/Shared/EpicGames.Horde/Storage/TreeStore.cs
--- a/Engine/Source/Programs/Shared/EpicGames.Horde/Storage/TreeStore.cs
+++ b/Engine/Source/Programs/Shared/EpicGames.Horde/Storage/TreeStore.cs
@@ -9,8 +9,6 @@ using System.IO;
 using System.Linq;
 using System.Threading;
 using System.Threading.Tasks;
-using System.Xml.Linq;
-using Blake3;
 using EpicGames.Core;
 using K4os.Compression.LZ4;
 using Microsoft.Extensions.Caching.Memory;
@@ -42,11 +40,6 @@ namespace EpicGames.Horde.Storage
 		/// </summary>
 		public int TrimIgnoreCount { get; set; } = 150;
 
-		/// <summary>
-		/// Number of reads from storage to allow concurrently. This has an impact on memory usage as well as network throughput.
-		/// </summary>
-		public int MaxConcurrentReads { get; set; } = 5;
-
 		/// <summary>
 		/// Size of the bundle/decode cache
 		/// </summary>
@@ -678,8 +671,6 @@ namespace EpicGames.Horde.Storage
 
 		readonly object _lockObject = new object();
 
-		readonly SemaphoreSlim _readSema;
-
 		// Active read tasks at any moment. If a BundleObject is not available in the cache, we start a read and add an entry to this dictionary
 		// so that other threads can also await it.
 		readonly Dictionary<string, Task<Bundle>> _readTasks = new Dictionary<string, Task<Bundle>>();
@@ -697,7 +688,6 @@ namespace EpicGames.Horde.Storage
 			{
 				_cache = new MemoryCache(new MemoryCacheOptions { SizeLimit = _options.CacheSize });
 			}
-			_readSema = new SemaphoreSlim(options.MaxConcurrentReads);
 		}
 
 		/// <inheritdoc/>
@@ -721,7 +711,6 @@ namespace EpicGames.Horde.Storage
 		{
 			if (disposing)
 			{
-				_readSema.Dispose();
 				_cache?.Dispose();
 			}
 		}
@@ -822,19 +811,8 @@ namespace EpicGames.Horde.Storage
 				}
 			}
 
-			// Wait for the read semaphore to avoid triggering too many operations at once.
-			await _readSema.WaitAsync(cancellationToken);
-
 			// Read the data from storage
-			Bundle bundle;
-			try
-			{
-				bundle = await _blobStore.ReadBundleAsync(locator, cancellationToken);
-			}
-			finally
-			{
-				_readSema.Release();
-			}
+			Bundle bundle = await _blobStore.ReadBundleAsync(locator, cancellationToken);
 
 			// Add the object to the cache
 			AddBundleToCache(cacheKey, bundle);
