Horde: Allow specifying a blob prefix when writing a ref directly to IBlobStore.

diff --git a/Engine/Source/Programs/Horde/Horde.Agent/Commands/Bundles/CreateCommand.cs b/Engine/Source/Programs/Horde/Horde.Agent/Commands/Bundles/CreateCommand.cs
--- a/Engine/Source/Programs/Horde/Horde.Agent/Commands/Bundles/CreateCommand.cs
+++ b/Engine/Source/Programs/Horde/Horde.Agent/Commands/Bundles/CreateCommand.cs
@@ -95,7 +95,7 @@ namespace Horde.Agent.Commands.Bundles
 				return new BlobLocator(text.Trim());
 			}
 
-			public async Task<BlobLocator> WriteRefAsync(RefName name, Bundle bundle, CancellationToken cancellationToken = default)
+			public async Task<BlobLocator> WriteRefAsync(RefName name, Bundle bundle, Utf8String prefix = default, CancellationToken cancellationToken = default)
 			{
 				BlobLocator locator = await WriteBundleAsync(bundle, name.Text, cancellationToken);
 				await WriteRefTargetAsync(name, locator, cancellationToken);
diff --git a/Engine/Source/Programs/Horde/Horde.Build/Storage/BasicBlobStore.cs b/Engine/Source/Programs/Horde/Horde.Build/Storage/BasicBlobStore.cs
--- a/Engine/Source/Programs/Horde/Horde.Build/Storage/BasicBlobStore.cs
+++ b/Engine/Source/Programs/Horde/Horde.Build/Storage/BasicBlobStore.cs
@@ -166,7 +166,7 @@ namespace Horde.Build.Storage
 		}
 
 		/// <inheritdoc/>
-		public async Task<BlobLocator> WriteRefAsync(RefName name, Bundle bundle, CancellationToken cancellationToken = default)
+		public async Task<BlobLocator> WriteRefAsync(RefName name, Bundle bundle, Utf8String prefix = default, CancellationToken cancellationToken = default)
 		{
 			BlobLocator locator = await WriteBundleAsync(bundle, name.Text, cancellationToken);
 			await WriteRefTargetAsync(name, locator, cancellationToken);
diff --git a/Engine/Source/Programs/Shared/EpicGames.Horde/Storage/Backends/InMemoryBlobStore.cs b/Engine/Source/Programs/Shared/EpicGames.Horde/Storage/Backends/InMemoryBlobStore.cs
--- a/Engine/Source/Programs/Shared/EpicGames.Horde/Storage/Backends/InMemoryBlobStore.cs
+++ b/Engine/Source/Programs/Shared/EpicGames.Horde/Storage/Backends/InMemoryBlobStore.cs
@@ -79,7 +79,7 @@ namespace EpicGames.Horde.Storage.Backends
 		}
 
 		/// <inheritdoc/>
-		public async Task<BlobLocator> WriteRefAsync(RefName name, Bundle bundle, CancellationToken cancellationToken)
+		public async Task<BlobLocator> WriteRefAsync(RefName name, Bundle bundle, Utf8String prefix, CancellationToken cancellationToken)
 		{
 			BlobLocator locator = await WriteBundleAsync(bundle, name.Text, cancellationToken);
 			_refs[name] = locator;
diff --git a/Engine/Source/Programs/Shared/EpicGames.Horde/Storage/IBlobStore.cs b/Engine/Source/Programs/Shared/EpicGames.Horde/Storage/IBlobStore.cs
--- a/Engine/Source/Programs/Shared/EpicGames.Horde/Storage/IBlobStore.cs
+++ b/Engine/Source/Programs/Shared/EpicGames.Horde/Storage/IBlobStore.cs
@@ -102,9 +102,10 @@ namespace EpicGames.Horde.Storage
 		/// </summary>
 		/// <param name="name">Ref to write</param>
 		/// <param name="bundle">Bundle to write</param>
+		/// <param name="prefix">Prefix for blob names.</param>
 		/// <param name="cancellationToken">Cancellation token for the operation</param>
 		/// <returns>Unique identifier for the blob</returns>
-		Task<BlobLocator> WriteRefAsync(RefName name, Bundle bundle, CancellationToken cancellationToken = default);
+		Task<BlobLocator> WriteRefAsync(RefName name, Bundle bundle, Utf8String prefix = default, CancellationToken cancellationToken = default);
 
 		/// <summary>
 		/// Writes a new ref to the store
@@ -165,7 +166,7 @@ namespace EpicGames.Horde.Storage
 			public Task<BlobLocator> TryReadRefTargetAsync(RefName name, DateTime cacheTime = default, CancellationToken cancellationToken = default) => _inner.TryReadRefTargetAsync(name, cacheTime, cancellationToken);
 
 			/// <inheritdoc/>
-			public Task<BlobLocator> WriteRefAsync(RefName name, Bundle bundle, CancellationToken cancellationToken = default) => _inner.WriteRefAsync(name, bundle, cancellationToken);
+			public Task<BlobLocator> WriteRefAsync(RefName name, Bundle bundle, Utf8String prefix = default, CancellationToken cancellationToken = default) => _inner.WriteRefAsync(name, bundle, prefix, cancellationToken);
 
 			/// <inheritdoc/>
 			public Task WriteRefTargetAsync(RefName name, BlobLocator blobId, CancellationToken cancellationToken = default) => _inner.WriteRefTargetAsync(name, blobId, cancellationToken);
