Horde: Rename current IStorageClient implementation to ILegacyStorageClient.

diff --git a/Engine/Source/Programs/Horde/Horde.Agent/Commands/Compute/AwsLambdaListenerCommand.cs b/Engine/Source/Programs/Horde/Horde.Agent/Commands/Compute/AwsLambdaListenerCommand.cs
--- a/Engine/Source/Programs/Horde/Horde.Agent/Commands/Compute/AwsLambdaListenerCommand.cs
+++ b/Engine/Source/Programs/Horde/Horde.Agent/Commands/Compute/AwsLambdaListenerCommand.cs
@@ -135,7 +135,7 @@ namespace Horde.Agent.Commands.Compute
 				OAuthHandler<HttpStorageClient> oAuthHandler = oAuthHandlerFactory.Create<HttpStorageClient>(clientOptions);
 				using HttpClient client = new (oAuthHandler);
 				client.BaseAddress = clientOptions.Url;
-				IStorageClient storageClient = new HttpStorageClient(client);
+				ILegacyStorageClient storageClient = new HttpStorageClient(client);
 
 				AgentSettings agentSettings = await CreateAgentSettingsAsync(cts.Token);
 				agentSettings.WorkingDir = Environment.GetEnvironmentVariable("UE_HORDE_WORKING_DIR") ?? "/tmp";
diff --git a/Engine/Source/Programs/Horde/Horde.Agent/Commands/Compute/ComputeCommand.cs b/Engine/Source/Programs/Horde/Horde.Agent/Commands/Compute/ComputeCommand.cs
--- a/Engine/Source/Programs/Horde/Horde.Agent/Commands/Compute/ComputeCommand.cs
+++ b/Engine/Source/Programs/Horde/Horde.Agent/Commands/Compute/ComputeCommand.cs
@@ -165,7 +165,7 @@ namespace Horde.Agent.Commands
 
 			using (IHost host = hostBuilder.Build())
 			{
-				IStorageClient storageClient = host.Services.GetRequiredService<IStorageClient>();
+				ILegacyStorageClient storageClient = host.Services.GetRequiredService<ILegacyStorageClient>();
 				IComputeClient computeClient = host.Services.GetRequiredService<IComputeClient>();
 
 				ByteString saltBytes = ByteString.Empty;
@@ -278,7 +278,7 @@ namespace Horde.Agent.Commands
 			return magic == ExpectedMagic && method > CompressionMethodNone;
 		}
 
-		static async Task UploadSandbox(IComputeClusterInfo cluster, IStorageClient storageClient, Dictionary<IoHash, (byte[] data, bool isCompressed)> uploadList, ILogger logger)
+		static async Task UploadSandbox(IComputeClusterInfo cluster, ILegacyStorageClient storageClient, Dictionary<IoHash, (byte[] data, bool isCompressed)> uploadList, ILogger logger)
 		{
 			int totalUploaded = 0;
 			int totalSkipped = 0;
@@ -322,7 +322,7 @@ namespace Horde.Agent.Commands
 			logger.LogInformation("Uploaded {UploadedCount} blobs ({UploadedBytes} bytes), Skipped {SkippedCount} blobs ({SkippedBytes} bytes)", totalUploaded, uploadedBytes, totalSkipped, skippedBytes);
 		}
 
-		private async Task<ComputeTaskStatus> ExecuteAction(ILogger logger, IComputeClient computeClient, IStorageClient storageClient, ClusterId clusterId, ComputeTask task, Dictionary<IoHash, (byte[] data, bool isCompressed)> uploadList)
+		private async Task<ComputeTaskStatus> ExecuteAction(ILogger logger, IComputeClient computeClient, ILegacyStorageClient storageClient, ClusterId clusterId, ComputeTask task, Dictionary<IoHash, (byte[] data, bool isCompressed)> uploadList)
 		{
 			IComputeClusterInfo cluster = await computeClient.GetClusterInfoAsync(clusterId);
 
@@ -380,7 +380,7 @@ namespace Horde.Agent.Commands
 			}
 		}
 
-		async Task HandleCompleteTask(IStorageClient storageClient, NamespaceId namespaceId, BucketId outputBucketId, RefId outputRefId, ILogger logger)
+		async Task HandleCompleteTask(ILegacyStorageClient storageClient, NamespaceId namespaceId, BucketId outputBucketId, RefId outputRefId, ILogger logger)
 		{
 			ComputeTaskResult result = await storageClient.GetRefAsync<ComputeTaskResult>(namespaceId, outputBucketId, outputRefId);
 			logger.LogInformation("exit: {ExitCode}", result.ExitCode);
@@ -394,7 +394,7 @@ namespace Horde.Agent.Commands
 			}
 		}
 
-		static async Task LogTaskOutputAsync(IStorageClient storageClient, string channel, NamespaceId namespaceId, IoHash? logHash, ILogger logger)
+		static async Task LogTaskOutputAsync(ILegacyStorageClient storageClient, string channel, NamespaceId namespaceId, IoHash? logHash, ILogger logger)
 		{
 			if (logHash != null)
 			{
@@ -409,7 +409,7 @@ namespace Horde.Agent.Commands
 			}
 		}
 
-		async Task WriteOutputAsync(IStorageClient storageClient, NamespaceId namespaceId, IoHash treeHash, DirectoryReference outputDir)
+		async Task WriteOutputAsync(ILegacyStorageClient storageClient, NamespaceId namespaceId, IoHash treeHash, DirectoryReference outputDir)
 		{
 			DirectoryTree tree = await storageClient.ReadBlobAsync<DirectoryTree>(namespaceId, treeHash);
 
@@ -428,7 +428,7 @@ namespace Horde.Agent.Commands
 			await Task.WhenAll(tasks);
 		}
 
-		static async Task WriteObjectToFileAsync(IStorageClient storageClient, NamespaceId namespaceId, IoHash hash, FileReference outputFile)
+		static async Task WriteObjectToFileAsync(ILegacyStorageClient storageClient, NamespaceId namespaceId, IoHash hash, FileReference outputFile)
 		{
 			DirectoryReference.CreateDirectory(outputFile.Directory);
 
diff --git a/Engine/Source/Programs/Horde/Horde.Agent/Execution/ComputeTaskExecutor.cs b/Engine/Source/Programs/Horde/Horde.Agent/Execution/ComputeTaskExecutor.cs
--- a/Engine/Source/Programs/Horde/Horde.Agent/Execution/ComputeTaskExecutor.cs
+++ b/Engine/Source/Programs/Horde/Horde.Agent/Execution/ComputeTaskExecutor.cs
@@ -30,7 +30,7 @@ namespace Horde.Agent.Execution
 			public Dictionary<string, OutputTree> _subDirs = new Dictionary<string, OutputTree>();
 		}
 
-		private readonly IStorageClient _storageClient;
+		private readonly ILegacyStorageClient _storageClient;
 		private readonly ILogger _logger;
 
 		/// <summary>
@@ -38,7 +38,7 @@ namespace Horde.Agent.Execution
 		/// </summary>
 		/// <param name="storageClient"></param>
 		/// <param name="logger"></param>
-		public ComputeTaskExecutor(IStorageClient storageClient, ILogger logger)
+		public ComputeTaskExecutor(ILegacyStorageClient storageClient, ILogger logger)
 		{
 			_storageClient = storageClient;
 			_logger = logger;
diff --git a/Engine/Source/Programs/Horde/Horde.Agent/Services/WorkerService.cs b/Engine/Source/Programs/Horde/Horde.Agent/Services/WorkerService.cs
--- a/Engine/Source/Programs/Horde/Horde.Agent/Services/WorkerService.cs
+++ b/Engine/Source/Programs/Horde/Horde.Agent/Services/WorkerService.cs
@@ -162,7 +162,7 @@ namespace Horde.Agent.Services
 		/// <summary>
 		/// Client interface to the storage system
 		/// </summary>
-		readonly IStorageClient _storageClient;
+		readonly ILegacyStorageClient _storageClient;
 
 		/// <summary>
 		/// The working directory
@@ -256,7 +256,7 @@ namespace Horde.Agent.Services
 		/// <param name="storageClient">Instance of the storage client</param>
 		/// <param name="createExecutor">Optional factory method for creating an executor</param>
 		/// <param name="createHordeRpcClient">Optional factory method for creating a HordeRpcClient</param>
-		public WorkerService(ILogger<WorkerService> logger, IOptions<AgentSettings> options, GrpcService grpcService, IStorageClient storageClient,
+		public WorkerService(ILogger<WorkerService> logger, IOptions<AgentSettings> options, GrpcService grpcService, ILegacyStorageClient storageClient,
 			Func<IRpcConnection, ExecuteJobTask, BeginBatchResponse, IExecutor>? createExecutor = null,
 			Func<GrpcChannel, HordeRpc.HordeRpcClient>? createHordeRpcClient = null)
 		{
@@ -1010,7 +1010,7 @@ namespace Horde.Agent.Services
 					}
 				}
 			}
-			catch (BlobNotFoundException ex)
+			catch (LegacyBlobNotFoundException ex)
 			{
 				logger.LogError(ex, "Blob not found: {Hash}", ex.Hash);
 				result = new ComputeTaskResultMessage(ComputeTaskOutcome.BlobNotFound, ex.Hash.ToString());
diff --git a/Engine/Source/Programs/Horde/Horde.Build.Tests/Compute/ComputeServiceTest.cs b/Engine/Source/Programs/Horde/Horde.Build.Tests/Compute/ComputeServiceTest.cs
--- a/Engine/Source/Programs/Horde/Horde.Build.Tests/Compute/ComputeServiceTest.cs
+++ b/Engine/Source/Programs/Horde/Horde.Build.Tests/Compute/ComputeServiceTest.cs
@@ -25,7 +25,7 @@ namespace Horde.Build.Tests.Compute
 		private static int s_channelId;
 		private static int s_refId;
 
-		private IStorageClient StorageClient => ServiceProvider.GetRequiredService<IStorageClient>();
+		private ILegacyStorageClient StorageClient => ServiceProvider.GetRequiredService<ILegacyStorageClient>();
 
 		public ComputeServiceTest()
 		{
diff --git a/Engine/Source/Programs/Horde/Horde.Build.Tests/StorageTests.cs b/Engine/Source/Programs/Horde/Horde.Build.Tests/StorageTests.cs
--- a/Engine/Source/Programs/Horde/Horde.Build.Tests/StorageTests.cs
+++ b/Engine/Source/Programs/Horde/Horde.Build.Tests/StorageTests.cs
@@ -20,12 +20,12 @@ namespace Horde.Build.Tests
 		{
 			NamespaceId namespaceId = new NamespaceId("test-ns");
 
-			IStorageClient storageClient = ServiceProvider.GetRequiredService<IStorageClient>();
+			ILegacyStorageClient storageClient = ServiceProvider.GetRequiredService<ILegacyStorageClient>();
 
 			byte[] testData = Encoding.UTF8.GetBytes("Hello world");
 			IoHash hash = IoHash.Compute(testData);
 
-			await Assert.ThrowsExceptionAsync<BlobNotFoundException>(() => storageClient.ReadBlobAsync(namespaceId, hash));
+			await Assert.ThrowsExceptionAsync<LegacyBlobNotFoundException>(() => storageClient.ReadBlobAsync(namespaceId, hash));
 
 			IoHash returnedHash = await storageClient.WriteBlobFromMemoryAsync(namespaceId, testData);
 			Assert.AreEqual(hash, returnedHash);
@@ -35,7 +35,7 @@ namespace Horde.Build.Tests
 			Assert.IsTrue(testData.AsSpan().SequenceEqual(storedData.Span));
 
 			NamespaceId otherNamespaceId = new NamespaceId("other-ns");
-			await Assert.ThrowsExceptionAsync<BlobNotFoundException>(() => storageClient.ReadBlobAsync(otherNamespaceId, hash));
+			await Assert.ThrowsExceptionAsync<LegacyBlobNotFoundException>(() => storageClient.ReadBlobAsync(otherNamespaceId, hash));
 		}
 
 		[TestMethod]
@@ -43,7 +43,7 @@ namespace Horde.Build.Tests
 		{
 			NamespaceId namespaceId = new NamespaceId("test-ns");
 
-			IStorageClient storageClient = ServiceProvider.GetRequiredService<IStorageClient>();
+			ILegacyStorageClient storageClient = ServiceProvider.GetRequiredService<ILegacyStorageClient>();
 
 			CbObject objectD = CbObject.Build(writer =>
 			{
@@ -51,7 +51,7 @@ namespace Horde.Build.Tests
 			});
 			IoHash hashD = objectD.GetHash();
 
-			await Assert.ThrowsExceptionAsync<BlobNotFoundException>(() => storageClient.ReadBlobToMemoryAsync(namespaceId, hashD));
+			await Assert.ThrowsExceptionAsync<LegacyBlobNotFoundException>(() => storageClient.ReadBlobToMemoryAsync(namespaceId, hashD));
 
 			await storageClient.WriteBlobFromMemoryAsync(namespaceId, hashD, objectD.GetView());
 
@@ -62,7 +62,7 @@ namespace Horde.Build.Tests
 		[TestMethod]
 		public async Task RefCollectionTest()
 		{
-			IStorageClient storageClient = ServiceProvider.GetRequiredService<IStorageClient>();
+			ILegacyStorageClient storageClient = ServiceProvider.GetRequiredService<ILegacyStorageClient>();
 
 			byte[] blobA = Encoding.UTF8.GetBytes("This is blob A");
 			IoHash hashA = IoHash.Compute(blobA);
@@ -98,7 +98,7 @@ namespace Horde.Build.Tests
 			Assert.IsTrue(missingHashes.Contains(hashA));
 			Assert.IsTrue(missingHashes.Contains(hashD));
 
-			await Assert.ThrowsExceptionAsync<RefNotFoundException>(() => storageClient.GetRefAsync(namespaceId, bucketId, refId));
+			await Assert.ThrowsExceptionAsync<LegacyRefNotFoundException>(() => storageClient.GetRefAsync(namespaceId, bucketId, refId));
 
 			// Add object D, and check that changes the missing hashes to just the blobs
 			await storageClient.WriteBlobFromMemoryAsync(namespaceId, hashD, objectD.GetView());
@@ -109,7 +109,7 @@ namespace Horde.Build.Tests
 			Assert.IsTrue(missingHashes.Contains(hashB));
 			Assert.IsTrue(missingHashes.Contains(hashC));
 
-			await Assert.ThrowsExceptionAsync<RefNotFoundException>(() => storageClient.GetRefAsync(namespaceId, bucketId, refId));
+			await Assert.ThrowsExceptionAsync<LegacyRefNotFoundException>(() => storageClient.GetRefAsync(namespaceId, bucketId, refId));
 
 			// Add blobs A, B and C and check that the object can be finalized
 			await storageClient.WriteBlobFromMemoryAsync(namespaceId, hashA, blobA);
@@ -119,7 +119,7 @@ namespace Horde.Build.Tests
 			missingHashes = await storageClient.TryFinalizeRefAsync(namespaceId, bucketId, refId, hashE);
 			Assert.AreEqual(0, missingHashes.Count);
 
-			IRef refValue = await storageClient.GetRefAsync(namespaceId, bucketId, refId);
+			ILegacyRef refValue = await storageClient.GetRefAsync(namespaceId, bucketId, refId);
 			Assert.IsNotNull(refValue);
 
 			// Add a new ref to existing objects and check it finalizes correctly 
diff --git a/Engine/Source/Programs/Horde/Horde.Build.Tests/TestSetup.cs b/Engine/Source/Programs/Horde/Horde.Build.Tests/TestSetup.cs
--- a/Engine/Source/Programs/Horde/Horde.Build.Tests/TestSetup.cs
+++ b/Engine/Source/Programs/Horde/Horde.Build.Tests/TestSetup.cs
@@ -249,7 +249,7 @@ namespace Horde.Build.Tests
 				return new TreeStore(blobStore, new TreeOptions()).ForType<CommitService>();
 			});
 
-			services.AddSingleton<IStorageClient, BasicStorageClient>();
+			services.AddSingleton<ILegacyStorageClient, BasicStorageClient>();
 
 			services.AddSingleton<ISingletonDocument<AgentSoftwareChannels>>(new SingletonDocumentStub<AgentSoftwareChannels>());
 			services.AddSingleton<ISingletonDocument<DevicePlatformMapV1>>(new SingletonDocumentStub<DevicePlatformMapV1>());
diff --git a/Engine/Source/Programs/Horde/Horde.Build/Commands/Bundles/PerforceCommand.cs b/Engine/Source/Programs/Horde/Horde.Build/Commands/Bundles/PerforceCommand.cs
--- a/Engine/Source/Programs/Horde/Horde.Build/Commands/Bundles/PerforceCommand.cs
+++ b/Engine/Source/Programs/Horde/Horde.Build/Commands/Bundles/PerforceCommand.cs
@@ -65,7 +65,7 @@ namespace Horde.Build.Commands.Bundles
 			ICommitService commitService = serviceProvider.GetRequiredService<ICommitService>();
 			ReplicationService replicationService = serviceProvider.GetRequiredService<ReplicationService>();
 			IStreamCollection streamCollection = serviceProvider.GetRequiredService<IStreamCollection>();
-			IStorageClient storageClient = serviceProvider.GetRequiredService<IStorageClient>();
+			ILegacyStorageClient storageClient = serviceProvider.GetRequiredService<ILegacyStorageClient>();
 
 			IStream? stream = await streamCollection.GetAsync(new StreamId(StreamId));
 			if (stream == null)
diff --git a/Engine/Source/Programs/Horde/Horde.Build/Compute/ComputeService.cs b/Engine/Source/Programs/Horde/Horde.Build/Compute/ComputeService.cs
--- a/Engine/Source/Programs/Horde/Horde.Build/Compute/ComputeService.cs
+++ b/Engine/Source/Programs/Horde/Horde.Build/Compute/ComputeService.cs
@@ -121,7 +121,7 @@ namespace Horde.Build.Compute
 		/// </summary>
 		public static NamespaceId DefaultNamespaceId { get; } = new NamespaceId("default");
 
-		readonly IStorageClient _storageClient;
+		readonly ILegacyStorageClient _storageClient;
 		readonly ITaskScheduler<QueueKey, ComputeTaskInfo> _taskScheduler;
 		readonly RedisMessageQueue<ComputeTaskStatus> _messageQueue;
 		readonly ITicker _expireTasksTicker;
@@ -137,7 +137,7 @@ namespace Horde.Build.Compute
 		/// <summary>
 		/// Constructor
 		/// </summary>
-		public ComputeService(GlobalsService globalsService, RedisService redisService, IStorageClient storageClient, IClock clock, ILogger<ComputeService> logger)
+		public ComputeService(GlobalsService globalsService, RedisService redisService, ILegacyStorageClient storageClient, IClock clock, ILogger<ComputeService> logger)
 		{
 			_storageClient = storageClient;
 			_taskScheduler = new RedisTaskScheduler<QueueKey, ComputeTaskInfo>(redisService.ConnectionPool, "compute/tasks/", logger);
@@ -460,7 +460,7 @@ namespace Horde.Build.Compute
 				{
 					requirements = await _storageClient.ReadBlobAsync<Requirements>(namespaceId, queueKey.RequirementsHash);
 				}
-				catch (BlobNotFoundException)
+				catch (LegacyBlobNotFoundException)
 				{
 				}
 				catch (Exception ex)
diff --git a/Engine/Source/Programs/Horde/Horde.Build/Storage/Services/BasicStorageClient.cs b/Engine/Source/Programs/Horde/Horde.Build/Storage/Services/BasicStorageClient.cs
--- a/Engine/Source/Programs/Horde/Horde.Build/Storage/Services/BasicStorageClient.cs
+++ b/Engine/Source/Programs/Horde/Horde.Build/Storage/Services/BasicStorageClient.cs
@@ -15,15 +15,15 @@ using MongoDB.Driver;
 
 namespace Horde.Build.Storage.Services
 {
-	using BlobNotFoundException = EpicGames.Horde.Storage.BlobNotFoundException;
-	using IRef = EpicGames.Horde.Storage.IRef;
+	using BlobNotFoundException = EpicGames.Horde.Storage.LegacyBlobNotFoundException;
+	using ILegacyRef = EpicGames.Horde.Storage.ILegacyRef;
 
 	/// <summary>
-	/// Simple implementation of <see cref="IStorageClient"/> which uses the current <see cref="IStorageBackend"/> implementation without garbage collection.
+	/// Simple implementation of <see cref="ILegacyStorageClient"/> which uses the current <see cref="IStorageBackend"/> implementation without garbage collection.
 	/// </summary>
-	public class BasicStorageClient : IStorageClient
+	public class BasicStorageClient : ILegacyStorageClient
 	{
-		class RefImpl : IRef
+		class RefImpl : ILegacyRef
 		{
 			public string Id { get; set; } = String.Empty;
 
@@ -142,14 +142,14 @@ namespace Horde.Build.Storage.Services
 		}
 
 		/// <inheritdoc/>
-		public async Task<IRef> GetRefAsync(NamespaceId namespaceId, BucketId bucketId, RefId refId, CancellationToken cancellationToken)
+		public async Task<ILegacyRef> GetRefAsync(NamespaceId namespaceId, BucketId bucketId, RefId refId, CancellationToken cancellationToken)
 		{
 			string id = GetFullRefId(namespaceId, bucketId, refId);
 
 			RefImpl? refImpl = await _refs.FindOneAndUpdateAsync<RefImpl>(x => x.Id == id, Builders<RefImpl>.Update.Set(x => x.LastAccessTime, DateTime.UtcNow), new FindOneAndUpdateOptions<RefImpl> { ReturnDocument = ReturnDocument.After }, cancellationToken);
 			if (refImpl == null || !refImpl.Finalized)
 			{
-				throw new RefNotFoundException(namespaceId, bucketId, refId);
+				throw new LegacyRefNotFoundException(namespaceId, bucketId, refId);
 			}
 
 			return refImpl;
diff --git a/Engine/Source/Programs/Horde/Horde.Build/Tools/ToolCollection.cs b/Engine/Source/Programs/Horde/Horde.Build/Tools/ToolCollection.cs
--- a/Engine/Source/Programs/Horde/Horde.Build/Tools/ToolCollection.cs
+++ b/Engine/Source/Programs/Horde/Horde.Build/Tools/ToolCollection.cs
@@ -55,7 +55,7 @@ namespace Horde.Build.Tools
 
 		private readonly RedisService _redisService;
 		private readonly VersionedCollection<ToolId, Tool> _tools;
-		private readonly IStorageClient _storage;
+		private readonly ILegacyStorageClient _storage;
 		private readonly IClock _clock;
 		private readonly NamespaceId _namespaceId;
 		private readonly ILogger _logger;
@@ -68,7 +68,7 @@ namespace Horde.Build.Tools
 		/// <summary>
 		/// Constructor
 		/// </summary>
-		public ToolCollection(MongoService mongoService, RedisService redisService, IStorageClient storage, IClock clock, IOptions<ServerSettings> options, ILogger<ToolCollection> logger)
+		public ToolCollection(MongoService mongoService, RedisService redisService, ILegacyStorageClient storage, IClock clock, IOptions<ServerSettings> options, ILogger<ToolCollection> logger)
 		{
 			_redisService = redisService;
 			_tools = new VersionedCollection<ToolId, Tool>(mongoService, "Tools", redisService, s_baseKey, s_types);
diff --git a/Engine/Source/Programs/Horde/Horde.Build/Utilities/ObjectSet.cs b/Engine/Source/Programs/Horde/Horde.Build/Utilities/ObjectSet.cs
--- a/Engine/Source/Programs/Horde/Horde.Build/Utilities/ObjectSet.cs
+++ b/Engine/Source/Programs/Horde/Horde.Build/Utilities/ObjectSet.cs
@@ -135,7 +135,7 @@ namespace Horde.Build.Utilities
 	/// </summary>
 	class ObjectSet
 	{
-		readonly IStorageClient _storageClient;
+		readonly ILegacyStorageClient _storageClient;
 		readonly NamespaceId _namespaceId;
 		public int MaxPackSize { get; }
 
@@ -158,7 +158,7 @@ namespace Horde.Build.Utilities
 		/// <param name="namespaceId"></param>
 		/// <param name="maxPackSize"></param>
 		/// <param name="time">The initial update time; used to determine the age of blobs</param>
-		public ObjectSet(IStorageClient storageClient, NamespaceId namespaceId, int maxPackSize, DateTime time)
+		public ObjectSet(ILegacyStorageClient storageClient, NamespaceId namespaceId, int maxPackSize, DateTime time)
 		{
 			_storageClient = storageClient;
 			_namespaceId = namespaceId;
diff --git a/Engine/Source/Programs/Shared/EpicGames.Horde.Tests/StorageClientTest.cs b/Engine/Source/Programs/Shared/EpicGames.Horde.Tests/StorageClientTest.cs
--- a/Engine/Source/Programs/Shared/EpicGames.Horde.Tests/StorageClientTest.cs
+++ b/Engine/Source/Programs/Shared/EpicGames.Horde.Tests/StorageClientTest.cs
@@ -32,12 +32,12 @@ namespace EpicGames.Horde.Tests
 			_compressedHash = IoHash.Compute(_compressedData.ToArray());
 		}
 
-		protected abstract IStorageClient GetStorageClient();
+		protected abstract ILegacyStorageClient GetStorageClient();
 
 		[TestMethod]
 		public async Task WriteBlob()
 		{
-			IStorageClient client = GetStorageClient();
+			ILegacyStorageClient client = GetStorageClient();
 			IoHash writtenHash = await client.WriteBlobFromMemoryAsync(_ns, _uncompressedData);
 			byte[] actual = await client.ReadBlobToMemoryAsync(_ns, writtenHash);
 			CollectionAssert.AreEqual(_uncompressedData.ToArray(), actual);
@@ -46,7 +46,7 @@ namespace EpicGames.Horde.Tests
 		[TestMethod]
 		public async Task WriteBlobWithPreCalculatedHash()
 		{
-			IStorageClient client = GetStorageClient();
+			ILegacyStorageClient client = GetStorageClient();
 			await client.WriteBlobFromMemoryAsync(_ns, _uncompressedHash, _uncompressedData);
 			byte[] actual = await client.ReadBlobToMemoryAsync(_ns, _uncompressedHash);
 			CollectionAssert.AreEqual(_uncompressedData.ToArray(), actual);
@@ -55,7 +55,7 @@ namespace EpicGames.Horde.Tests
 		[TestMethod]
 		public async Task WriteAndReadCompressedBlob()
 		{
-			IStorageClient client = GetStorageClient();
+			ILegacyStorageClient client = GetStorageClient();
 			using MemoryStream compressedStream = new(_compressedData.ToArray());
 			IoHash returnedHash = await client.WriteCompressedBlobAsync(_ns, compressedStream);
 			Assert.AreEqual(_uncompressedHash, returnedHash);
@@ -66,7 +66,7 @@ namespace EpicGames.Horde.Tests
 		[TestMethod]
 		public async Task WriteAndReadCompressedBlobWithPreCalculatedHash()
 		{
-			IStorageClient client = GetStorageClient();
+			ILegacyStorageClient client = GetStorageClient();
 			using MemoryStream compressedStream = new(_compressedData.ToArray());
 			await client.WriteCompressedBlobAsync(_ns, _uncompressedHash, compressedStream);
 			byte[] returnedData = await client.ReadCompressedBlobToMemoryAsync(_ns, _uncompressedHash);
@@ -76,7 +76,7 @@ namespace EpicGames.Horde.Tests
 		[TestMethod]
 		public async Task WriteWithBadPreCalculatedHash()
 		{
-			IStorageClient client = GetStorageClient();
+			ILegacyStorageClient client = GetStorageClient();
 			IoHash badHash = IoHash.Compute(new byte[] {0xBA, 0xDB, 0xAD});
 			using MemoryStream compressedStream = new(_compressedData.ToArray());
 			await Assert.ThrowsExceptionAsync<ArgumentException>(() => client.WriteCompressedBlobAsync(_ns, badHash, compressedStream));
@@ -85,7 +85,7 @@ namespace EpicGames.Horde.Tests
 		[TestMethod]
 		public async Task WriteWithInvalidCompression()
 		{
-			IStorageClient client = GetStorageClient();
+			ILegacyStorageClient client = GetStorageClient();
 			using MemoryStream compressedStream = new(new byte[] {0xBA, 0xDB, 0xAD});
 			await Assert.ThrowsExceptionAsync<ArgumentException>(() => client.WriteCompressedBlobAsync(_ns, compressedStream));
 		}
@@ -94,7 +94,7 @@ namespace EpicGames.Horde.Tests
 	[TestClass]
 	public class MemoryStorageClientTest : BaseStorageClientTest
 	{
-		protected override IStorageClient GetStorageClient()
+		protected override ILegacyStorageClient GetStorageClient()
 		{
 			return new MemoryStorageClient();
 		}
diff --git a/Engine/Source/Programs/Shared/EpicGames.Horde/Storage/ILegacyStorageClient.cs b/Engine/Source/Programs/Shared/EpicGames.Horde/Storage/ILegacyStorageClient.cs
new file mode 100644
--- /dev/null
+++ b/Engine/Source/Programs/Shared/EpicGames.Horde/Storage/ILegacyStorageClient.cs
@@ -0,0 +1,558 @@
+// Copyright Epic Games, Inc. All Rights Reserved.
+
+using EpicGames.Core;
+using EpicGames.Serialization;
+using System;
+using System.Collections.Generic;
+using System.IO;
+using System.Threading;
+using System.Threading.Tasks;
+
+namespace EpicGames.Horde.Storage
+{
+	/// <summary>
+	/// Base class for exceptions related to store
+	/// </summary>
+	public class StorageException : Exception
+	{
+		/// <summary>
+		/// Constructor
+		/// </summary>
+		public StorageException(string message, Exception? innerException)
+			: base(message, innerException)
+		{
+		}
+	}
+
+	/// <summary>
+	/// Base class for blob exceptions
+	/// </summary>
+	public class LegacyBlobException : StorageException
+	{
+		/// <summary>
+		/// Namespace containing the blob
+		/// </summary>
+		public NamespaceId NamespaceId { get; }
+
+		/// <summary>
+		/// Hash of the blob
+		/// </summary>
+		public IoHash Hash { get; }
+
+		/// <summary>
+		/// Constructor
+		/// </summary>
+		public LegacyBlobException(NamespaceId namespaceId, IoHash hash, string message, Exception? innerException = null)
+			: base(message, innerException)
+		{
+			NamespaceId = namespaceId;
+			Hash = hash;
+		}
+	}
+
+	/// <summary>
+	/// Exception thrown for missing blobs
+	/// </summary>
+	public sealed class LegacyBlobNotFoundException : LegacyBlobException
+	{
+		/// <summary>
+		/// Constructor
+		/// </summary>
+		public LegacyBlobNotFoundException(NamespaceId namespaceId, IoHash hash, Exception? innerException = null)
+			: base(namespaceId, hash, $"Unable to find blob {hash} in {namespaceId}", innerException)
+		{
+		}
+	}
+
+	/// <summary>
+	/// Base class for ref exceptions
+	/// </summary>
+	public class LegacyRefException : StorageException
+	{
+		/// <summary>
+		/// Namespace containing the ref
+		/// </summary>
+		public NamespaceId NamespaceId { get; }
+
+		/// <summary>
+		/// Bucket containing the ref
+		/// </summary>
+		public BucketId BucketId { get; }
+
+		/// <summary>
+		/// Identifier for the ref
+		/// </summary>
+		public RefId RefId { get; }
+
+		/// <summary>
+		/// Constructor
+		/// </summary>
+		public LegacyRefException(NamespaceId namespaceId, BucketId bucketId, RefId refId, string message, Exception? innerException = null)
+			: base(message, innerException)
+		{
+			NamespaceId = namespaceId;
+			BucketId = bucketId;
+			RefId = refId;
+		}
+	}
+
+	/// <summary>
+	/// Indicates that a named reference wasn't found
+	/// </summary>
+	public sealed class LegacyRefNotFoundException : LegacyRefException
+	{
+		/// <summary>
+		/// Constructor
+		/// </summary>
+		public LegacyRefNotFoundException(NamespaceId namespaceId, BucketId bucketId, RefId refId, Exception? innerException = null)
+			: base(namespaceId, bucketId, refId, $"Ref {namespaceId}/{bucketId}/{refId} not found", innerException)
+		{
+		}
+	}
+
+	/// <summary>
+	/// Indicates that a ref cannot be finalized due to a missing blob
+	/// </summary>
+	public sealed class LegacyRefMissingBlobException : LegacyRefException
+	{
+		/// <summary>
+		/// Constructor
+		/// </summary>
+		public LegacyRefMissingBlobException(NamespaceId namespaceId, BucketId bucketId, RefId refId, List<IoHash> missingBlobs, Exception? innerException = null)
+			: base(namespaceId, bucketId, refId, $"Ref {namespaceId}/{bucketId}/{refId} cannot be finalized; missing {missingBlobs.Count} blobs ({missingBlobs[0]}...)", innerException)
+		{
+		}
+	}
+
+	/// <summary>
+	/// Interface for an object reference
+	/// </summary>
+	public interface ILegacyRef
+	{
+		/// <summary>
+		/// Namespace identifier
+		/// </summary>
+		NamespaceId NamespaceId { get; }
+
+		/// <summary>
+		/// Bucket identifier
+		/// </summary>
+		BucketId BucketId { get; }
+
+		/// <summary>
+		/// Ref identifier
+		/// </summary>
+		RefId RefId { get; }
+
+		/// <summary>
+		/// The value stored for this ref
+		/// </summary>
+		CbObject Value { get; }
+	}
+
+	/// <summary>
+	/// Base interface for a storage client that only records blobs.
+	/// </summary>
+	public interface ILegacyStorageClient
+	{
+		#region Blobs
+
+		/// <summary>
+		/// Opens a blob read stream
+		/// </summary>
+		/// <param name="namespaceId">Namespace to operate on</param>
+		/// <param name="hash">Hash of the blob</param>
+		/// <param name="cancellationToken">Cancellation token for the operation</param>
+		/// <returns>Stream for the blob, or null if it does not exist</returns>
+		Task<Stream> ReadBlobAsync(NamespaceId namespaceId, IoHash hash, CancellationToken cancellationToken = default);
+		
+		/// <summary>
+		/// Opens a blob read stream for compressed blobs
+		/// </summary>
+		/// <param name="namespaceId">Namespace to operate on</param>
+		/// <param name="uncompressedHash">Hash of the uncompressed blob</param>
+		/// <param name="cancellationToken">Cancellation token for the operation</param>
+		/// <returns>Stream for the blob, or null if it does not exist</returns>
+		Task<Stream> ReadCompressedBlobAsync(NamespaceId namespaceId, IoHash uncompressedHash, CancellationToken cancellationToken = default);
+
+		/// <summary>
+		/// Writes a blob to storage
+		/// </summary>
+		/// <param name="namespaceId">Namespace to operate on</param>
+		/// <param name="hash">Hash of the blob</param>
+		/// <param name="stream">The stream to write</param>
+		/// <param name="cancellationToken">Cancellation token for the operation</param>
+		Task WriteBlobAsync(NamespaceId namespaceId, IoHash hash, Stream stream, CancellationToken cancellationToken = default);
+
+		/// <summary>
+		/// Writes a blob to storage and calculates hash
+		/// </summary>
+		/// <param name="namespaceId">Namespace to operate on</param>
+		/// <param name="stream">The stream to write</param>
+		/// <param name="cancellationToken">Cancellation token for the operation</param>
+		/// <returns>Hash of the blob written</returns>
+		Task<IoHash> WriteBlobAsync(NamespaceId namespaceId, Stream stream, CancellationToken cancellationToken = default);
+
+		/// <summary>
+		/// Writes a compressed blob to storage
+		/// </summary>
+		/// <param name="namespaceId">Namespace to operate on</param>
+		/// <param name="uncompressedHash">Hash of the blob</param>
+		/// <param name="compressedStream">Compressed stream to write</param>
+		/// <param name="cancellationToken">Cancellation token for the operation</param>
+		Task WriteCompressedBlobAsync(NamespaceId namespaceId, IoHash uncompressedHash, Stream compressedStream, CancellationToken cancellationToken = default);
+		
+		/// <summary>
+		/// Writes a compressed blob to storage and calculates hash
+		/// </summary>
+		/// <param name="namespaceId">Namespace to operate on</param>
+		/// <param name="compressedStream">Compressed stream to write</param>
+		/// <param name="cancellationToken">Cancellation token for the operation</param>
+		/// <returns>Hash of uncompressed blob</returns>
+		Task<IoHash> WriteCompressedBlobAsync(NamespaceId namespaceId, Stream compressedStream, CancellationToken cancellationToken = default);
+		
+		/// <summary>
+		/// Checks if the given blob exists
+		/// </summary>
+		/// <param name="namespaceId">Namespace to operate on</param>
+		/// <param name="hash">Hash of the blob</param>
+		/// <param name="cancellationToken">Cancellation token for the operation</param>
+		/// <returns>True if the blob exists, false if it did not exist</returns>
+		Task<bool> HasBlobAsync(NamespaceId namespaceId, IoHash hash, CancellationToken cancellationToken = default);
+
+		/// <summary>
+		/// Checks if a list of blobs exist, returning the set of missing blobs
+		/// </summary>
+		/// <param name="namespaceId">Namespace to operate on</param>
+		/// <param name="hashes">Set of hashes to check</param>
+		/// <param name="cancellationToken">Cancellation token for the operation</param>
+		/// <returns>A set of the missing hashes</returns>
+		Task<HashSet<IoHash>> FindMissingBlobsAsync(NamespaceId namespaceId, HashSet<IoHash> hashes, CancellationToken cancellationToken = default);
+
+		#endregion
+		#region Refs
+
+		/// <summary>
+		/// Gets the given reference
+		/// </summary>
+		/// <param name="namespaceId">Namespace identifier</param>
+		/// <param name="bucketId">Bucket identifier</param>
+		/// <param name="refId">Name of the reference</param>
+		/// <param name="cancellationToken">Cancellation token for the operation</param>
+		/// <returns>The reference data if the ref exists</returns>
+		Task<ILegacyRef> GetRefAsync(NamespaceId namespaceId, BucketId bucketId, RefId refId, CancellationToken cancellationToken = default);
+
+		/// <summary>
+		/// Checks if the given reference exists
+		/// </summary>
+		/// <param name="namespaceId">Namespace identifier</param>
+		/// <param name="bucketId">Bucket identifier</param>
+		/// <param name="refId">Ref identifier</param>
+		/// <param name="cancellationToken">Cancellation token for the operation</param>
+		/// <returns>True if the ref exists, false if it did not exist</returns>
+		Task<bool> HasRefAsync(NamespaceId namespaceId, BucketId bucketId, RefId refId, CancellationToken cancellationToken = default);
+
+		/// <summary>
+		/// Determines which refs are missing
+		/// </summary>
+		/// <param name="namespaceId">Namespace identifier</param>
+		/// <param name="bucketId">Bucket identifier</param>
+		/// <param name="refIds">Names of the references</param>
+		/// <param name="cancellationToken">Cancellation token for the operation</param>
+		/// <returns>List of missing references</returns>
+		Task<List<RefId>> FindMissingRefsAsync(NamespaceId namespaceId, BucketId bucketId, List<RefId> refIds, CancellationToken cancellationToken = default);
+
+		/// <summary>
+		/// Attempts to sets the given reference, returning a list of missing objects on failure.
+		/// </summary>
+		/// <param name="namespaceId">Namespace identifier</param>
+		/// <param name="bucketId">Bucket identifier</param>
+		/// <param name="refId">Ref identifier</param>
+		/// <param name="value">New value for the reference</param>
+		/// <param name="cancellationToken">Cancellation token for the operation</param>
+		/// <returns>List of missing references</returns>
+		Task<List<IoHash>> TrySetRefAsync(NamespaceId namespaceId, BucketId bucketId, RefId refId, CbObject value, CancellationToken cancellationToken = default);
+
+		/// <summary>
+		/// Attempts to finalize a reference, turning its references into hard references
+		/// </summary>
+		/// <param name="namespaceId">Namespace identifier</param>
+		/// <param name="bucketId">Bucket identifier</param>
+		/// <param name="refId">Ref identifier</param>
+		/// <param name="hash">Hash of the referenced object</param>
+		/// <param name="cancellationToken">Cancellation token for the operation</param>
+		/// <returns></returns>
+		Task<List<IoHash>> TryFinalizeRefAsync(NamespaceId namespaceId, BucketId bucketId, RefId refId, IoHash hash, CancellationToken cancellationToken = default);
+
+		/// <summary>
+		/// Removes the given reference
+		/// </summary>
+		/// <param name="namespaceId">Namespace identifier</param>
+		/// <param name="bucketId">Bucket identifier</param>
+		/// <param name="refId">Ref identifier</param>
+		/// <param name="cancellationToken">Cancellation token for the operation</param>
+		/// <returns>True if the ref was deleted, false if it did not exist</returns>
+		Task<bool> DeleteRefAsync(NamespaceId namespaceId, BucketId bucketId, RefId refId, CancellationToken cancellationToken = default);
+
+		#endregion
+	}
+
+	/// <summary>
+	/// Extension methods for <see cref="ILegacyStorageClient"/>
+	/// </summary>
+	public static class LegacyStorageClientExtensions
+	{
+		const int DefaultMaxInMemoryBlobLength = 128 * 1024 * 1024;
+
+		/// <summary>
+		/// Gets a blob as a byte array
+		/// </summary>
+		/// <param name="storageClient">The storage interface</param>
+		/// <param name="namespaceId">Namespace containing the blob</param>
+		/// <param name="hash">Hash of the blob to read</param>
+		/// <param name="maxInMemoryBlobLength">Maximum allowed memory allocation to store the blob</param>
+		/// <param name="cancellationToken">Cancellation token for the operation</param>
+		/// <returns>Data for the blob that was read. Throws an exception if the blob was not found.</returns>
+		public static Task<byte[]> ReadBlobToMemoryAsync(this ILegacyStorageClient storageClient, NamespaceId namespaceId, IoHash hash, int maxInMemoryBlobLength = DefaultMaxInMemoryBlobLength, CancellationToken cancellationToken = default)
+		{
+			return ReadBlobToMemoryAsync(storageClient, false, namespaceId, hash, maxInMemoryBlobLength, cancellationToken);
+		}
+		
+		/// <summary>
+		/// Gets a compressed blob as a byte array
+		/// </summary>
+		/// <param name="storageClient">The storage interface</param>
+		/// <param name="namespaceId">Namespace containing the blob</param>
+		/// <param name="hash">Hash of the blob to read</param>
+		/// <param name="maxInMemoryBlobLength">Maximum allowed memory allocation to store the blob</param>
+		/// <param name="cancellationToken">Cancellation token for the operation</param>
+		/// <returns>Data for the blob that was read. Throws an exception if the blob was not found.</returns>
+		public static Task<byte[]> ReadCompressedBlobToMemoryAsync(this ILegacyStorageClient storageClient, NamespaceId namespaceId, IoHash hash, int maxInMemoryBlobLength = DefaultMaxInMemoryBlobLength, CancellationToken cancellationToken = default)
+		{
+			return ReadBlobToMemoryAsync(storageClient, true, namespaceId, hash, maxInMemoryBlobLength, cancellationToken);
+		}
+		
+		private static async Task<byte[]> ReadBlobToMemoryAsync(ILegacyStorageClient storageClient, bool isBlobCompressed, NamespaceId namespaceId, IoHash hash, int maxInMemoryBlobLength = DefaultMaxInMemoryBlobLength, CancellationToken cancellationToken = default)
+		{
+			using Stream stream = isBlobCompressed
+				? await storageClient.ReadCompressedBlobAsync(namespaceId, hash, cancellationToken)
+				: await storageClient.ReadBlobAsync(namespaceId, hash, cancellationToken);
+
+			long length = stream.Length;
+			if (length > maxInMemoryBlobLength)
+			{
+				throw new LegacyBlobException(namespaceId, hash, $"Blob {hash} is too large ({length} > {maxInMemoryBlobLength})");
+			}
+
+			byte[] buffer = new byte[length];
+			for (int offset = 0; offset < length;)
+			{
+				int count = await stream.ReadAsync(buffer, offset, (int)length - offset, cancellationToken);
+				if (count == 0)
+				{
+					throw new LegacyBlobException(namespaceId, hash, $"Unexpected end of stream reading blob {hash}");
+				}
+				offset += count;
+			}
+
+			return buffer;
+		}
+
+		/// <summary>
+		/// Writes a blob from memory to storage
+		/// </summary>
+		/// <param name="storageClient">The storage interface</param>
+		/// <param name="namespaceId">Namespace containing the blob</param>
+		/// <param name="data">Data to write</param>
+		/// <param name="cancellationToken">Cancellation token for the operation</param>
+		/// <returns></returns>
+		public static async Task<IoHash> WriteBlobFromMemoryAsync(this ILegacyStorageClient storageClient, NamespaceId namespaceId, ReadOnlyMemory<byte> data, CancellationToken cancellationToken = default)
+		{
+			IoHash hash = IoHash.Compute(data.Span);
+			await WriteBlobFromMemoryAsync(storageClient, namespaceId, hash, data, cancellationToken);
+			return hash;
+		}
+
+		/// <summary>
+		/// Writes a blob from memory to storage
+		/// </summary>
+		/// <param name="storageClient">The storage interface</param>
+		/// <param name="namespaceId">Namespace containing the blob</param>
+		/// <param name="hash">Hash of the data</param>
+		/// <param name="data">The data to be written</param>
+		/// <param name="cancellationToken">Cancellation token for the operation</param>
+		/// <returns></returns>
+		public static async Task WriteBlobFromMemoryAsync(this ILegacyStorageClient storageClient, NamespaceId namespaceId, IoHash hash, ReadOnlyMemory<byte> data, CancellationToken cancellationToken = default)
+		{
+			using ReadOnlyMemoryStream stream = new ReadOnlyMemoryStream(data);
+			await storageClient.WriteBlobAsync(namespaceId, hash, stream, cancellationToken);
+		}
+
+		/// <summary>
+		/// Reads a blob and deserializes it as the given compact-binary encoded type
+		/// </summary>
+		/// <typeparam name="T">Type of object to deserialize</typeparam>
+		/// <param name="storageClient">The storage interface</param>
+		/// <param name="namespaceId">Namespace containing the blob</param>
+		/// <param name="hash">Hash of the blob to read</param>
+		/// <param name="maxInMemoryBlobLength">Maximum allowed memory allocation to store the blob</param>
+		/// <param name="cancellationToken">Cancellation token for the operation</param>
+		/// <returns>The decoded object</returns>
+		public static async Task<T> ReadBlobAsync<T>(this ILegacyStorageClient storageClient, NamespaceId namespaceId, IoHash hash, int maxInMemoryBlobLength = DefaultMaxInMemoryBlobLength, CancellationToken cancellationToken = default)
+		{
+			ReadOnlyMemory<byte> data = await ReadBlobToMemoryAsync(storageClient, namespaceId, hash, maxInMemoryBlobLength, cancellationToken);
+			return CbSerializer.Deserialize<T>(data);
+		}
+
+		/// <summary>
+		/// Serializes an object to compact binary format and writes it to storage
+		/// </summary>
+		/// <param name="storageClient">The storage interface</param>
+		/// <param name="namespaceId">Namespace containing the blob</param>
+		/// <param name="blob">The object to be written</param>
+		/// <param name="cancellationToken">Cancellation token for the operation</param>
+		/// <returns></returns>
+		public static async Task<IoHash> WriteBlobAsync<T>(this ILegacyStorageClient storageClient, NamespaceId namespaceId, T blob, CancellationToken cancellationToken = default)
+		{
+			CbWriter writer = new CbWriter();
+			CbSerializer.Serialize<T>(writer, blob);
+
+			IoHash hash = writer.ComputeHash();
+			await storageClient.WriteBlobAsync(namespaceId, hash, writer.AsStream(), cancellationToken);
+			return hash;
+		}
+
+		/// <summary>
+		/// Gets the given reference
+		/// </summary>
+		/// <param name="storageClient">The storage interface</param>
+		/// <param name="qualifiedRefId">Name of the reference</param>
+		/// <param name="cancellationToken">Cancellation token for the operation</param>
+		/// <returns>The reference data if the ref exists</returns>
+		public static Task<ILegacyRef> GetRefAsync(this ILegacyStorageClient storageClient, QualifiedRefId qualifiedRefId, CancellationToken cancellationToken = default)
+		{
+			return storageClient.GetRefAsync(qualifiedRefId.NamespaceId, qualifiedRefId.BucketId, qualifiedRefId.RefId, cancellationToken);
+		}
+
+		/// <summary>
+		/// Checks if the given reference exists
+		/// </summary>
+		/// <param name="storageClient">The storage interface</param>
+		/// <param name="qualifiedRefId">Name of the reference</param>
+		/// <param name="cancellationToken">Cancellation token for the operation</param>
+		/// <returns>True if the ref exists, false if it did not exist</returns>
+		public static Task<bool> HasRefAsync(this ILegacyStorageClient storageClient, QualifiedRefId qualifiedRefId, CancellationToken cancellationToken = default)
+		{
+			return storageClient.HasRefAsync(qualifiedRefId.NamespaceId, qualifiedRefId.BucketId, qualifiedRefId.RefId, cancellationToken);
+		}
+
+		/// <summary>
+		/// Attempts to sets the given reference, returning a list of missing objects on failure.
+		/// </summary>
+		/// <param name="storageClient">The storage interface</param>
+		/// <param name="qualifiedRefId">Name of the reference</param>
+		/// <param name="value">New value for the reference</param>
+		/// <param name="cancellationToken">Cancellation token for the operation</param>
+		/// <returns>List of missing references</returns>
+		public static Task<List<IoHash>> TrySetRefAsync(this ILegacyStorageClient storageClient, QualifiedRefId qualifiedRefId, CbObject value, CancellationToken cancellationToken = default)
+		{
+			return storageClient.TrySetRefAsync(qualifiedRefId.NamespaceId, qualifiedRefId.BucketId, qualifiedRefId.RefId, value, cancellationToken);
+		}
+
+		/// <summary>
+		/// Reads a reference as a specific type
+		/// </summary>
+		/// <typeparam name="T"></typeparam>
+		/// <param name="storageClient">The storage interface</param>
+		/// <param name="namespaceId">Namespace containing the ref</param>
+		/// <param name="bucketId">Bucket containing the ref</param>
+		/// <param name="refId">The ref id</param>
+		/// <param name="cancellationToken">Cancellation token for the operation</param>
+		/// <returns>Deserialized object for the given ref</returns>
+		public static async Task<T> GetRefAsync<T>(this ILegacyStorageClient storageClient, NamespaceId namespaceId, BucketId bucketId, RefId refId, CancellationToken cancellationToken = default)
+		{
+			ILegacyRef item = await storageClient.GetRefAsync(namespaceId, bucketId, refId, cancellationToken);
+			return CbSerializer.Deserialize<T>(item.Value);
+		}
+
+		/// <summary>
+		/// Sets a ref to a particular value
+		/// </summary>
+		/// <param name="storageClient">The storage interface</param>
+		/// <param name="namespaceId">Namespace containing the ref</param>
+		/// <param name="bucketId">Bucket containing the ref</param>
+		/// <param name="refId">The ref id</param>
+		/// <param name="value">The new object for the ref</param>
+		/// <param name="cancellationToken">Cancellation token for the operation</param>
+		public static Task SetRefAsync<T>(this ILegacyStorageClient storageClient, NamespaceId namespaceId, BucketId bucketId, RefId refId, T value, CancellationToken cancellationToken = default)
+		{
+			CbObject objectValue = CbSerializer.Serialize<T>(value);
+			return SetRefAsync(storageClient, namespaceId, bucketId, refId, objectValue, cancellationToken);
+		}
+
+		/// <summary>
+		/// Sets a ref to a particular value
+		/// </summary>
+		/// <param name="storageClient">The storage interface</param>
+		/// <param name="namespaceId">Namespace containing the ref</param>
+		/// <param name="bucketId">Bucket containing the ref</param>
+		/// <param name="refId">The ref id</param>
+		/// <param name="value">The new object for the ref</param>
+		/// <param name="cancellationToken">Cancellation token for the operation</param>
+		public static async Task SetRefAsync(this ILegacyStorageClient storageClient, NamespaceId namespaceId, BucketId bucketId, RefId refId, CbObject value, CancellationToken cancellationToken = default)
+		{
+			List<IoHash> missingHashes = await storageClient.TrySetRefAsync(namespaceId, bucketId, refId, value, cancellationToken);
+			if (missingHashes.Count > 0)
+			{
+				throw new LegacyRefMissingBlobException(namespaceId, bucketId, refId, missingHashes);
+			}
+		}
+
+		/// <summary>
+		/// Attempts to set a ref to a particular value
+		/// </summary>
+		/// <typeparam name="T"></typeparam>
+		/// <param name="storageClient">The storage interface</param>
+		/// <param name="namespaceId">Namespace containing the ref</param>
+		/// <param name="bucketId">Bucket containing the ref</param>
+		/// <param name="refId">The ref id</param>
+		/// <param name="value">The new object for the ref</param>
+		/// <param name="cancellationToken">Cancellation token for the operation</param>
+		/// <returns>List of missing blob hashes</returns>
+		public static Task<List<IoHash>> TrySetRefAsync<T>(this ILegacyStorageClient storageClient, NamespaceId namespaceId, BucketId bucketId, RefId refId, T value, CancellationToken cancellationToken = default)
+		{
+			CbObject objectValue = CbSerializer.Serialize<T>(value);
+			return storageClient.TrySetRefAsync(namespaceId, bucketId, refId, objectValue, cancellationToken);
+		}
+
+		/// <summary>
+		/// Finalize a ref, throwing an exception if finalization fails
+		/// </summary>
+		/// <param name="storageClient">The storage interface</param>
+		/// <param name="namespaceId">Namespace containing the ref</param>
+		/// <param name="bucketId">Bucket containing the ref</param>
+		/// <param name="refId">The ref id</param>
+		/// <param name="valueHash">Hash of the ref value</param>
+		/// <param name="cancellationToken">Cancellation token for the operation</param>
+		public static async Task FinalizeRefAsync(this ILegacyStorageClient storageClient, NamespaceId namespaceId, BucketId bucketId, RefId refId, IoHash valueHash, CancellationToken cancellationToken = default)
+		{
+			List<IoHash> missingHashes = await storageClient.TryFinalizeRefAsync(namespaceId, bucketId, refId, valueHash, cancellationToken);
+			if (missingHashes.Count > 0)
+			{
+				throw new LegacyRefMissingBlobException(namespaceId, bucketId, refId, missingHashes);
+			}
+		}
+
+		/// <summary>
+		/// Removes the given reference
+		/// </summary>
+		/// <param name="storageClient">The storage interface</param>
+		/// <param name="qualifiedRefId">Name of the reference</param>
+		/// <param name="cancellationToken">Cancellation token for the operation</param>
+		/// <returns>True if the ref was deleted, false if it did not exist</returns>
+		public static Task<bool> DeleteRefAsync(this ILegacyStorageClient storageClient, QualifiedRefId qualifiedRefId, CancellationToken cancellationToken = default)
+		{
+			return storageClient.DeleteRefAsync(qualifiedRefId.NamespaceId, qualifiedRefId.BucketId, qualifiedRefId.RefId, cancellationToken);
+		}
+	}
+}
diff --git a/Engine/Source/Programs/Shared/EpicGames.Horde/Storage/IStorageClient.cs b/Engine/Source/Programs/Shared/EpicGames.Horde/Storage/IStorageClient.cs
deleted file mode 100644
diff --git a/Engine/Source/Programs/Shared/EpicGames.Horde/Storage/Impl/FileStorageClient.cs b/Engine/Source/Programs/Shared/EpicGames.Horde/Storage/Impl/FileStorageClient.cs
--- a/Engine/Source/Programs/Shared/EpicGames.Horde/Storage/Impl/FileStorageClient.cs
+++ b/Engine/Source/Programs/Shared/EpicGames.Horde/Storage/Impl/FileStorageClient.cs
@@ -13,11 +13,11 @@ using System.Threading.Tasks;
 namespace EpicGames.Horde.Storage.Impl
 {
 	/// <summary>
-	/// Implementation of <see cref="IStorageClient"/> which writes directly to the local filesystem for testing. Not intended for production use.
+	/// Implementation of <see cref="ILegacyStorageClient"/> which writes directly to the local filesystem for testing. Not intended for production use.
 	/// </summary>
-	public class FileStorageClient : IStorageClient
+	public class FileStorageClient : ILegacyStorageClient
 	{
-		class Ref : IRef
+		class Ref : ILegacyRef
 		{
 			public NamespaceId NamespaceId { get; set; }
 			public BucketId BucketId { get; set; }
@@ -145,7 +145,7 @@ namespace EpicGames.Horde.Storage.Impl
 		}
 
 		/// <inheritdoc/>
-		public async Task<IRef> GetRefAsync(NamespaceId namespaceId, BucketId bucketId, RefId refId, CancellationToken cancellationToken = default)
+		public async Task<ILegacyRef> GetRefAsync(NamespaceId namespaceId, BucketId bucketId, RefId refId, CancellationToken cancellationToken = default)
 		{
 			FileReference file = GetRefFile(namespaceId, bucketId, refId);
 			_logger.LogInformation("Reading {File} ({Size:n0} bytes)", file, new FileInfo(file.FullName).Length);
diff --git a/Engine/Source/Programs/Shared/EpicGames.Horde/Storage/Impl/HttpStorageClient.cs b/Engine/Source/Programs/Shared/EpicGames.Horde/Storage/Impl/HttpStorageClient.cs
--- a/Engine/Source/Programs/Shared/EpicGames.Horde/Storage/Impl/HttpStorageClient.cs
+++ b/Engine/Source/Programs/Shared/EpicGames.Horde/Storage/Impl/HttpStorageClient.cs
@@ -50,11 +50,11 @@ namespace EpicGames.Horde.Storage.Impl
 	}
 
 	/// <summary>
-	/// Implementation of <see cref="IStorageClient"/> using REST requests.
+	/// Implementation of <see cref="ILegacyStorageClient"/> using REST requests.
 	/// </summary>
-	public sealed class HttpStorageClient : IStorageClient
+	public sealed class HttpStorageClient : ILegacyStorageClient
 	{
-		class RefImpl : IRef
+		class RefImpl : ILegacyRef
 		{
 			public NamespaceId NamespaceId { get; set; }
 
@@ -112,7 +112,7 @@ namespace EpicGames.Horde.Storage.Impl
 			{
 				if (response.StatusCode == System.Net.HttpStatusCode.NotFound)
 				{
-					throw new BlobNotFoundException(namespaceId, hash);
+					throw new LegacyBlobNotFoundException(namespaceId, hash);
 				}
 
 				response.EnsureSuccessStatusCode();
@@ -248,7 +248,7 @@ namespace EpicGames.Horde.Storage.Impl
 		#region Refs
 
 		/// <inheritdoc/>
-		public async Task<IRef> GetRefAsync(NamespaceId namespaceId, BucketId bucketId, RefId refId, CancellationToken cancellationToken)
+		public async Task<ILegacyRef> GetRefAsync(NamespaceId namespaceId, BucketId bucketId, RefId refId, CancellationToken cancellationToken)
 		{
 			using HttpRequestMessage request = new HttpRequestMessage(HttpMethod.Get, $"api/v1/refs/{namespaceId}/{bucketId}/{refId}");
 
@@ -260,11 +260,11 @@ namespace EpicGames.Horde.Storage.Impl
 			{
 				if (response.StatusCode == System.Net.HttpStatusCode.NotFound)
 				{
-					throw new RefNotFoundException(namespaceId, bucketId, refId);
+					throw new LegacyRefNotFoundException(namespaceId, bucketId, refId);
 				}
 				else
 				{
-					throw new RefException(namespaceId, bucketId, refId, await GetMessageFromResponse(response));
+					throw new LegacyRefException(namespaceId, bucketId, refId, await GetMessageFromResponse(response));
 				}
 			}
 
@@ -290,13 +290,13 @@ namespace EpicGames.Horde.Storage.Impl
 			using HttpResponseMessage response = await _httpClient.SendAsync(request, cancellationToken);
 			if (!response.IsSuccessStatusCode)
 			{
-				throw new RefException(namespaceId, bucketId, refId, await GetMessageFromResponse(response));
+				throw new LegacyRefException(namespaceId, bucketId, refId, await GetMessageFromResponse(response));
 			}
 
 			RefPutResponse? responseBody = await ReadJsonResponse<RefPutResponse>(response.Content);
 			if (responseBody == null)
 			{
-				throw new RefException(namespaceId, bucketId, refId, "Unable to parse response body");
+				throw new LegacyRefException(namespaceId, bucketId, refId, "Unable to parse response body");
 			}
 
 			return responseBody.Needs;
@@ -313,13 +313,13 @@ namespace EpicGames.Horde.Storage.Impl
 			using HttpResponseMessage response = await _httpClient.SendAsync(request, cancellationToken);
 			if (!response.IsSuccessStatusCode)
 			{
-				throw new RefException(namespaceId, bucketId, refId, await GetMessageFromResponse(response));
+				throw new LegacyRefException(namespaceId, bucketId, refId, await GetMessageFromResponse(response));
 			}
 
 			RefPutResponse? responseBody = await ReadJsonResponse<RefPutResponse>(response.Content);
 			if (responseBody == null)
 			{
-				throw new RefException(namespaceId, bucketId, refId, "Unable to parse response body");
+				throw new LegacyRefException(namespaceId, bucketId, refId, "Unable to parse response body");
 			}
 
 			return responseBody.Needs;
@@ -350,7 +350,7 @@ namespace EpicGames.Horde.Storage.Impl
 			}
 			else
 			{
-				throw new RefException(namespaceId, bucketId, refId, $"Unexpected response for {nameof(HasRefAsync)} call: {response.StatusCode}");
+				throw new LegacyRefException(namespaceId, bucketId, refId, $"Unexpected response for {nameof(HasRefAsync)} call: {response.StatusCode}");
 			}
 		}
 
@@ -364,13 +364,13 @@ namespace EpicGames.Horde.Storage.Impl
 			using HttpResponseMessage response = await _httpClient.SendAsync(request, cancellationToken);
 			if (!response.IsSuccessStatusCode)
 			{
-				throw new RefException(namespaceId, bucketId, new RefId(IoHash.Zero), await GetMessageFromResponse(response));
+				throw new LegacyRefException(namespaceId, bucketId, new RefId(IoHash.Zero), await GetMessageFromResponse(response));
 			}
 
 			RefExistsResponse? responseContent = await ReadJsonResponse<RefExistsResponse>(response.Content);
 			if (responseContent == null)
 			{
-				throw new RefException(namespaceId, bucketId, new RefId(IoHash.Zero), "Unable to parse response body");
+				throw new LegacyRefException(namespaceId, bucketId, new RefId(IoHash.Zero), "Unable to parse response body");
 			}
 
 			return responseContent.Needs;
diff --git a/Engine/Source/Programs/Shared/EpicGames.Horde/Storage/Impl/MemoryStorageClient.cs b/Engine/Source/Programs/Shared/EpicGames.Horde/Storage/Impl/MemoryStorageClient.cs
--- a/Engine/Source/Programs/Shared/EpicGames.Horde/Storage/Impl/MemoryStorageClient.cs
+++ b/Engine/Source/Programs/Shared/EpicGames.Horde/Storage/Impl/MemoryStorageClient.cs
@@ -13,11 +13,11 @@ using System.Threading.Tasks;
 namespace EpicGames.Horde.Storage.Impl
 {
 	/// <summary>
-	/// Implementation of <see cref="IStorageClient"/> which writes directly to memory for testing. Not intended for production use.
+	/// Implementation of <see cref="ILegacyStorageClient"/> which writes directly to memory for testing. Not intended for production use.
 	/// </summary>
-	public class MemoryStorageClient : IStorageClient
+	public class MemoryStorageClient : ILegacyStorageClient
 	{
-		class Ref : IRef
+		class Ref : ILegacyRef
 		{
 			public NamespaceId NamespaceId { get; set; }
 			public BucketId BucketId { get; set; }
@@ -42,7 +42,7 @@ namespace EpicGames.Horde.Storage.Impl
 		public Dictionary<(NamespaceId, IoHash), IoHash> UncompressedToCompressedHash { get; } = new Dictionary<(NamespaceId, IoHash), IoHash>();
 
 		/// <inheritdoc/>
-		public Dictionary<(NamespaceId, BucketId, RefId), IRef> Refs { get; } = new Dictionary<(NamespaceId, BucketId, RefId), IRef>();
+		public Dictionary<(NamespaceId, BucketId, RefId), ILegacyRef> Refs { get; } = new Dictionary<(NamespaceId, BucketId, RefId), ILegacyRef>();
 
 		private readonly FakeCompressor _compressor = new FakeCompressor();
 
@@ -52,7 +52,7 @@ namespace EpicGames.Horde.Storage.Impl
 			ReadOnlyMemory<byte> data;
 			if(!Blobs.TryGetValue((namespaceId, hash), out data))
 			{
-				throw new BlobNotFoundException(namespaceId, hash);
+				throw new LegacyBlobNotFoundException(namespaceId, hash);
 			}
 			return Task.FromResult<Stream>(new ReadOnlyMemoryStream(data));
 		}
@@ -147,12 +147,12 @@ namespace EpicGames.Horde.Storage.Impl
 		}
 
 		/// <inheritdoc/>
-		public Task<IRef> GetRefAsync(NamespaceId namespaceId, BucketId bucketId, RefId refId, CancellationToken cancellationToken = default)
+		public Task<ILegacyRef> GetRefAsync(NamespaceId namespaceId, BucketId bucketId, RefId refId, CancellationToken cancellationToken = default)
 		{
-			IRef? result;
+			ILegacyRef? result;
 			if(!Refs.TryGetValue((namespaceId, bucketId, refId), out result))
 			{
-				throw new RefNotFoundException(namespaceId, bucketId, refId);
+				throw new LegacyRefNotFoundException(namespaceId, bucketId, refId);
 			}
 			return Task.FromResult(result);
 		}
diff --git a/Engine/Source/Programs/Shared/EpicGames.Horde/Storage/Impl/TracingStorageClient.cs b/Engine/Source/Programs/Shared/EpicGames.Horde/Storage/Impl/TracingStorageClient.cs
--- a/Engine/Source/Programs/Shared/EpicGames.Horde/Storage/Impl/TracingStorageClient.cs
+++ b/Engine/Source/Programs/Shared/EpicGames.Horde/Storage/Impl/TracingStorageClient.cs
@@ -11,11 +11,11 @@ using System.Threading.Tasks;
 namespace EpicGames.Horde.Storage.Impl
 {
 	/// <summary>
-	/// Implementation of <see cref="IStorageClient"/> which writes directly to the local filesystem for testing. Not intended for production use.
+	/// Implementation of <see cref="ILegacyStorageClient"/> which writes directly to the local filesystem for testing. Not intended for production use.
 	/// </summary>
-	public class TracingStorageClient : IStorageClient
+	public class TracingStorageClient : ILegacyStorageClient
 	{
-		readonly IStorageClient _inner;
+		readonly ILegacyStorageClient _inner;
 		readonly ILogger _logger;
 
 		/// <summary>
@@ -23,7 +23,7 @@ namespace EpicGames.Horde.Storage.Impl
 		/// </summary>
 		/// <param name="inner"></param>
 		/// <param name="logger"></param>
-		public TracingStorageClient(IStorageClient inner, ILogger logger)
+		public TracingStorageClient(ILegacyStorageClient inner, ILogger logger)
 		{
 			_inner = inner;
 			_logger = logger;
@@ -99,7 +99,7 @@ namespace EpicGames.Horde.Storage.Impl
 		}
 
 		/// <inheritdoc/>
-		public Task<IRef> GetRefAsync(NamespaceId namespaceId, BucketId bucketId, RefId refId, CancellationToken cancellationToken = default)
+		public Task<ILegacyRef> GetRefAsync(NamespaceId namespaceId, BucketId bucketId, RefId refId, CancellationToken cancellationToken = default)
 		{
 			_logger.LogDebug("Getting ref {NamespaceId}/{BucketId}/{RefId}", namespaceId, bucketId, refId);
 			return _inner.GetRefAsync(namespaceId, bucketId, refId, cancellationToken);
diff --git a/Engine/Source/Programs/Shared/EpicGames.Horde/Storage/StorageExtensions.cs b/Engine/Source/Programs/Shared/EpicGames.Horde/Storage/StorageExtensions.cs
--- a/Engine/Source/Programs/Shared/EpicGames.Horde/Storage/StorageExtensions.cs
+++ b/Engine/Source/Programs/Shared/EpicGames.Horde/Storage/StorageExtensions.cs
@@ -28,8 +28,8 @@ namespace EpicGames.Horde.Storage
 		{
 			services.AddOptions<StorageOptions>();
 
-			services.AddScoped<IStorageClient, HttpStorageClient>();
-			services.AddHttpClientWithAuth<IStorageClient, HttpStorageClient>(serviceProvider => serviceProvider.GetRequiredService<IOptions<StorageOptions>>().Value);
+			services.AddScoped<ILegacyStorageClient, HttpStorageClient>();
+			services.AddHttpClientWithAuth<ILegacyStorageClient, HttpStorageClient>(serviceProvider => serviceProvider.GetRequiredService<IOptions<StorageOptions>>().Value);
 		}
 
 		/// <summary>
