Fixed issue where project plugins would not compile with an installed engine.

diff --git a/Engine/Source/Programs/Shared/EpicGames.Build/System/CompileScriptModules.cs b/Engine/Source/Programs/Shared/EpicGames.Build/System/CompileScriptModules.cs
--- a/Engine/Source/Programs/Shared/EpicGames.Build/System/CompileScriptModules.cs
+++ b/Engine/Source/Programs/Shared/EpicGames.Build/System/CompileScriptModules.cs
@@ -16,6 +16,42 @@ namespace UnrealBuildBase
 {
 	public static class CompileScriptModule
 	{
+		[Flags]
+		public enum BuildFlags
+		{
+			/// <summary>
+			/// No build flags specified
+			/// </summary>
+			None = 0,
+
+			/// <summary>
+			/// Do not allow compiles of any projects
+			/// </summary>
+			NoCompile = 1 << 0,
+
+			/// <summary>
+			/// If not compiling, generate error if target does not exist
+			/// </summary>
+			ErrorOnMissingTarget = 1 << 1,
+
+			/// <summary>
+			/// When NoCompile options are specified, this flag will test to see if the build records
+			/// are valid.  But the results of that test are currently ignored.  Only the existence of the
+			/// target path is considered.
+			/// 
+			/// When NoCompile options are not specified, then the validity of the build records will be used
+			/// to test to see if project files need building.  If ForceCompile is specified, then this flag
+			/// serves no useful purpose.  If this flag is not specified, then ForceCompile is assumed.
+			/// </summary>
+			UseBuildRecords = 1 << 2,
+
+			/// <summary>
+			/// Force compile when the NoCompile options aren't specified regardless of the 
+			/// state of the build records.
+			/// </summary>
+			ForceCompile = 1 << 3,
+		}
+
 		class Hook : CsProjBuildHook
 		{
 			private ILogger Logger;
@@ -127,7 +163,25 @@ namespace UnrealBuildBase
 				GameFolders: Unreal.IsEngineInstalled() ? GameDirectories : new List<DirectoryReference>(), 
 				ForeignPlugins: null, AdditionalSearchPaths: AdditionalDirectories.Concat(GameBuildDirectories).ToList()));
 
-			return GetTargetPaths(Build(RulesFileType, FoundProjects, BaseDirectories, null, bForceCompile, bNoCompile, bUseBuildRecords, out bBuildSuccess, OnBuildingProjects, Logger));
+			BuildFlags BuildFlags = BuildFlags.None;
+			if (bForceCompile)
+			{
+				BuildFlags |= BuildFlags.ForceCompile;
+			}
+			if (bNoCompile)
+			{
+				BuildFlags |= BuildFlags.NoCompile;
+			}
+			if (Unreal.IsEngineInstalled())
+			{
+				BuildFlags |= BuildFlags.NoCompile | BuildFlags.ErrorOnMissingTarget;
+			}
+			if (bUseBuildRecords)
+			{
+				BuildFlags |= BuildFlags.UseBuildRecords;
+			}
+
+			return GetTargetPaths(Build(RulesFileType, FoundProjects, BaseDirectories, null, BuildFlags, out bBuildSuccess, OnBuildingProjects, Logger));
 		}
 
 		/// <summary>
@@ -161,34 +215,35 @@ namespace UnrealBuildBase
 		/// <param name="FoundProjects">Projects to be compiled</param>
 		/// <param name="BaseDirectories">Base directories for all the projects</param>
 		/// <param name="DefineConstants">Collection of constants to be added to the project</param>
-		/// <param name="bForceCompile"></param>
-		/// <param name="bNoCompile"></param>
-		/// <param name="bUseBuildRecords"></param>
+		/// <param name="BuildFlags">Collection of flags to customize compilation</param>
 		/// <param name="bBuildSuccess"></param>
 		/// <param name="OnBuildingProjects">Action to invoke when projects get built</param>
 		/// <param name="Logger"></param>
 		/// <returns>Collection of all the projects.  They will have been compiled.</returns>
 		public static Dictionary<FileReference, CsProjBuildRecordEntry> Build(Rules.RulesFileType RulesFileType,
 			HashSet<FileReference> FoundProjects, List<DirectoryReference> BaseDirectories, List<string>? DefineConstants, 
-			bool bForceCompile, bool bNoCompile, bool bUseBuildRecords,
-			out bool bBuildSuccess, Action<int> OnBuildingProjects, ILogger Logger)
+			BuildFlags BuildFlags, out bool bBuildSuccess, Action<int> OnBuildingProjects, ILogger Logger)
 		{
 			CsProjBuildHook Hook = new Hook(Logger, RulesFileType);
 
-			bool bUseBuildRecordsOnlyForProjectDiscovery = bNoCompile || Unreal.IsEngineInstalled();
-
 			// Load existing build records, validating them only if (re)compiling script projects is an option
 			Dictionary<FileReference, CsProjBuildRecordEntry> BuildRecords = LoadExistingBuildRecords(Hook, BaseDirectories, Logger);
 
-			if (bUseBuildRecords)
+			// Only validate the build records if requested
+			if (BuildFlags.HasFlag(BuildFlags.UseBuildRecords))
 			{
 				foreach (FileReference Project in FoundProjects)
 				{
 					ValidateBuildRecordRecursively(BuildRecords, Project, Hook, Logger);
 				}
 			}
+			else
+			{
+				BuildFlags |= BuildFlags.ForceCompile;
+			}
 
-			if (bUseBuildRecordsOnlyForProjectDiscovery)
+			// If we are not compiling, then test to see if the targets exist
+			if (BuildFlags.HasFlag(BuildFlags.NoCompile))
 			{
 				string FilterExtension = String.Empty;
 				switch (RulesFileType)
@@ -213,7 +268,7 @@ namespace UnrealBuildBase
 					}
 					else
 					{
-						if (bNoCompile)
+						if (!BuildFlags.HasFlag(BuildFlags.ErrorOnMissingTarget))
 						{
 							// when -NoCompile is on the command line, try to run with whatever is available
 							Logger.LogWarning("Script module \"{TargetPath}\" not found for record \"{BuildRecordPath}\"", TargetPath, Record.Value.BuildRecordFile);
@@ -242,7 +297,7 @@ namespace UnrealBuildBase
 				}
 			}
 
-			if (!bForceCompile && bUseBuildRecords)
+			if (!BuildFlags.HasFlag(BuildFlags.ForceCompile))
 			{
 				// If all found records are valid, we can return their targets directly
 				Dictionary<FileReference, CsProjBuildRecordEntry> ValidBuildRecords = new(BuildRecords.Where(x => x.Value.Status == CsProjBuildRecordStatus.Valid));
@@ -254,8 +309,7 @@ namespace UnrealBuildBase
 			}
 
 			// Fall back to the slower approach: use msbuild to load csproj files & build as necessary
-			return Build(FoundProjects, bForceCompile || !bUseBuildRecords, out bBuildSuccess, Hook, BaseDirectories,
-				DefineConstants, OnBuildingProjects, Logger);
+			return Build(FoundProjects, BuildFlags.HasFlag(BuildFlags.ForceCompile), out bBuildSuccess, Hook, BaseDirectories, DefineConstants, OnBuildingProjects, Logger);
 		}
 
 		/// <summary>
diff --git a/Engine/Source/Programs/Shared/EpicGames.Build/System/EnumeratePlugins.cs b/Engine/Source/Programs/Shared/EpicGames.Build/System/EnumeratePlugins.cs
--- a/Engine/Source/Programs/Shared/EpicGames.Build/System/EnumeratePlugins.cs
+++ b/Engine/Source/Programs/Shared/EpicGames.Build/System/EnumeratePlugins.cs
@@ -151,48 +151,96 @@ namespace UnrealBuildBase
 			List<FileReference> Built = FoundPlugins.Where(P => !UbtPluginFileCache.ContainsKey(P)).ToList();
 
 			// Collect the list of plugins not build
-			IEnumerable<FileReference> NotBuilt = FoundPlugins.Where(P => !UbtPluginFileCache.ContainsKey(P));
+			List<FileReference> NotBuilt = FoundPlugins.Where(P => !UbtPluginFileCache.ContainsKey(P)).ToList();
 
 			// If we have anything left to build, then build it
-			if (NotBuilt.Any())
+			if (NotBuilt.Count > 0)
 			{
-				Dictionary<FileReference, CsProjBuildRecordEntry>? BuiltPlugins =
-					CompileScriptModule.Build(UnrealBuildBase.Rules.RulesFileType.UbtPlugin, new HashSet<FileReference>(NotBuilt),
-					GetBaseDirectories(ProjectFile), DefineConstants, false, false, true, out bBuildSuccess,
-					Count =>
-					{
-						if (Log.OutputFile != null)
-						{
-							Logger.LogInformation("Building {Count} plugins (see Log '{LogFile}' for more details)", Count, Log.OutputFile);
-						}
-						else
-						{
-							Logger.LogInformation("Building {Count} plugins", Count);
-						}
-					}, Logger
-				);
-				
-				// Add any built plugins back into the cache
-				foreach (KeyValuePair<FileReference, CsProjBuildRecordEntry> BuiltPlugin in BuiltPlugins)
+				if (Unreal.IsEngineInstalled())
+				{
+					bBuildSuccess = BuildUbtPluginsInternal(PluginSet.EngineOnly, ProjectFile, Built, NotBuilt, DefineConstants, Logger);
+					bBuildSuccess |= BuildUbtPluginsInternal(PluginSet.ProjectOnly, ProjectFile, Built, NotBuilt, DefineConstants, Logger);
+				}
+				else
 				{
-					Built.Add(BuiltPlugin.Key);
-					UbtPluginFileCache.Add(BuiltPlugin.Key, CompileScriptModule.GetTargetPath(BuiltPlugin));
+					bBuildSuccess =
+						BuildUbtPluginsInternal(PluginSet.Both, ProjectFile, Built, NotBuilt, DefineConstants, Logger);
 				}
 			}
 
-			Plugins = Built.Select(P => (P, UbtPluginFileCache[P])).OrderBy(X => X.P.FullName).ToArray();
+			Plugins = Built.Where(P => UbtPluginFileCache.ContainsKey(P)).Select(P => (P, UbtPluginFileCache[P])).OrderBy(X => X.P.FullName).ToArray();
 			return bBuildSuccess;
 		}
 
-		static List<DirectoryReference> GetBaseDirectories(FileReference? ProjectFile)
+		private enum PluginSet
 		{
-			List<DirectoryReference> BaseDirectories = new List<DirectoryReference>(2);
-			BaseDirectories.Add(Unreal.EngineDirectory);
-			if (ProjectFile != null)
+			EngineOnly,
+			ProjectOnly,
+			Both,
+		}
+
+		private static bool BuildUbtPluginsInternal(PluginSet PluginSet, FileReference? ProjectFile, 
+			List<FileReference> Built, IEnumerable<FileReference> NotBuilt, List<string>? DefineConstants, ILogger Logger)
+		{
+			// Collect the work to be done
+			IEnumerable<FileReference>? ToBuild = null;
+			List<DirectoryReference> BaseDirectories = new();
+			CompileScriptModule.BuildFlags BuildFlags = CompileScriptModule.BuildFlags.UseBuildRecords;
+			switch (PluginSet)
 			{
-				BaseDirectories.Add(ProjectFile.Directory);
+				case PluginSet.EngineOnly:
+					BaseDirectories.Add(Unreal.EngineDirectory);
+					BuildFlags |= CompileScriptModule.BuildFlags.NoCompile | CompileScriptModule.BuildFlags.ErrorOnMissingTarget;
+					ToBuild = NotBuilt.Where(P => P.IsUnderDirectory(Unreal.EngineDirectory));
+					break;
+				case PluginSet.ProjectOnly:
+					if (ProjectFile != null)
+					{
+						BaseDirectories.Add(ProjectFile.Directory);
+						ToBuild = NotBuilt.Where(P => !P.IsUnderDirectory(Unreal.EngineDirectory));
+					}
+					break;
+				case PluginSet.Both:
+					BaseDirectories.Add(Unreal.EngineDirectory);
+					if (ProjectFile != null)
+					{
+						BaseDirectories.Add(ProjectFile.Directory);
+					}
+					ToBuild = NotBuilt;
+					break;
 			}
-			return BaseDirectories;
+
+			// Just return if there is nothing
+			if (ToBuild == null || !ToBuild.Any() || BaseDirectories.Count == 0)
+			{
+				return true;
+			}
+
+			// Build the things needing to be built
+			bool bBuildSuccess = true;
+			Dictionary<FileReference, CsProjBuildRecordEntry>? BuiltPlugins =
+				CompileScriptModule.Build(UnrealBuildBase.Rules.RulesFileType.UbtPlugin, new HashSet<FileReference>(NotBuilt),
+				BaseDirectories, DefineConstants, CompileScriptModule.BuildFlags.UseBuildRecords, out bBuildSuccess,
+				Count =>
+				{
+					if (Log.OutputFile != null)
+					{
+						Logger.LogInformation("Building {Count} plugins (see Log '{LogFile}' for more details)", Count, Log.OutputFile);
+					}
+					else
+					{
+						Logger.LogInformation("Building {Count} plugins", Count);
+					}
+				}, Logger
+			);
+
+			// Add any built plugins back into the cache
+			foreach (KeyValuePair<FileReference, CsProjBuildRecordEntry> BuiltPlugin in BuiltPlugins)
+			{
+				Built.Add(BuiltPlugin.Key);
+				UbtPluginFileCache.Add(BuiltPlugin.Key, CompileScriptModule.GetTargetPath(BuiltPlugin));
+			}
+			return bBuildSuccess;
 		}
 	}
 }
