Horde: Omit list of packets from uncompressed bundles.

diff --git a/Engine/Source/Programs/Shared/EpicGames.Horde/Storage/Bundle.cs b/Engine/Source/Programs/Shared/EpicGames.Horde/Storage/Bundle.cs
--- a/Engine/Source/Programs/Shared/EpicGames.Horde/Storage/Bundle.cs
+++ b/Engine/Source/Programs/Shared/EpicGames.Horde/Storage/Bundle.cs
@@ -172,7 +172,15 @@ namespace EpicGames.Horde.Storage
 			CompressionFormat = (BundleCompressionFormat)reader.ReadUnsignedVarInt();
 			Imports = reader.ReadVariableLengthArray(() => new BundleImport(reader));
 			Exports = reader.ReadVariableLengthArray(() => new BundleExport(reader));
-			Packets = reader.ReadVariableLengthArray(() => new BundlePacket(reader));
+
+			if (CompressionFormat == BundleCompressionFormat.None)
+			{
+				Packets = Exports.ConvertAll(x => new BundlePacket(x.Length, x.Length));
+			}
+			else
+			{
+				Packets = reader.ReadVariableLengthArray(() => new BundlePacket(reader));
+			}
 		}
 
 		/// <summary>
@@ -239,7 +247,11 @@ namespace EpicGames.Horde.Storage
 			writer.WriteUnsignedVarInt((ulong)CompressionFormat);
 			writer.WriteVariableLengthArray(Imports, x => x.Write(writer));
 			writer.WriteVariableLengthArray(Exports, x => x.Write(writer));
-			writer.WriteVariableLengthArray(Packets, x => x.Write(writer));
+
+			if (CompressionFormat != BundleCompressionFormat.None)
+			{
+				writer.WriteVariableLengthArray(Packets, x => x.Write(writer));
+			}
 
 			// Go back and write the length of the header
 			BinaryPrimitives.WriteInt32BigEndian(prelude.Slice(4).Span, writer.Length - initialLength);
