Horde: Merge BlobStore/TreeStore functionality together in IStorageClient.

diff --git a/Engine/Source/Programs/Horde/Horde.Agent/Commands/Bundles/CreateCommand.cs b/Engine/Source/Programs/Horde/Horde.Agent/Commands/Bundles/CreateCommand.cs
--- a/Engine/Source/Programs/Horde/Horde.Agent/Commands/Bundles/CreateCommand.cs
+++ b/Engine/Source/Programs/Horde/Horde.Agent/Commands/Bundles/CreateCommand.cs
@@ -14,12 +14,13 @@ namespace Horde.Agent.Commands.Bundles
 {
 	abstract class BundleCommandBase : Command
 	{
-		class FileBlobStore : IBlobStore
+		class FileBlobStore : StorageClientBase
 		{
 			readonly DirectoryReference _rootDir;
 			readonly ILogger _logger;
 
 			public FileBlobStore(DirectoryReference rootDir, ILogger logger)
+				: base(null)
 			{
 				_rootDir = rootDir;
 				_logger = logger;
@@ -43,7 +44,7 @@ namespace Horde.Agent.Commands.Bundles
 
 			#region Blobs
 
-			public async Task<Bundle> ReadBundleAsync(BlobLocator id, CancellationToken cancellationToken = default)
+			public override async Task<Bundle> ReadBundleAsync(BlobLocator id, CancellationToken cancellationToken = default)
 			{
 				FileReference file = GetBlobFile(id);
 				_logger.LogInformation("Reading {File}", file);
@@ -51,7 +52,7 @@ namespace Horde.Agent.Commands.Bundles
 				return new Bundle(new MemoryReader(bytes));
 			}
 
-			public async Task<BlobLocator> WriteBundleAsync(Bundle bundle, Utf8String prefix = default, CancellationToken cancellationToken = default)
+			public override async Task<BlobLocator> WriteBundleAsync(Bundle bundle, Utf8String prefix = default, CancellationToken cancellationToken = default)
 			{
 				BlobLocator id = BlobLocator.Create(HostId.Empty, prefix);
 				FileReference file = GetBlobFile(id);
@@ -65,24 +66,12 @@ namespace Horde.Agent.Commands.Bundles
 
 			#region Refs
 
-			public Task DeleteRefAsync(RefName name, CancellationToken cancellationToken = default)
+			public override Task DeleteRefAsync(RefName name, CancellationToken cancellationToken = default)
 			{
 				throw new NotImplementedException();
 			}
 
-			public async Task<BundleRef?> TryReadRefAsync(RefName name, DateTime cacheTime = default, CancellationToken cancellationToken = default)
-			{
-				BlobLocator locator = await TryReadRefTargetAsync(name, cacheTime, cancellationToken);
-				if (!locator.IsValid())
-				{
-					return null;
-				}
-
-				Bundle bundle = await ReadBundleAsync(locator, cancellationToken);
-				return new BundleRef(locator, bundle);
-			}
-
-			public async Task<BlobLocator> TryReadRefTargetAsync(RefName name, DateTime cacheTime = default, CancellationToken cancellationToken = default)
+			public override async Task<BlobLocator> TryReadRefTargetAsync(RefName name, DateTime cacheTime = default, CancellationToken cancellationToken = default)
 			{
 				FileReference file = GetRefFile(name);
 				if(!FileReference.Exists(file))
@@ -95,14 +84,7 @@ namespace Horde.Agent.Commands.Bundles
 				return new BlobLocator(text.Trim());
 			}
 
-			public async Task<BlobLocator> WriteRefAsync(RefName name, Bundle bundle, Utf8String prefix = default, CancellationToken cancellationToken = default)
-			{
-				BlobLocator locator = await WriteBundleAsync(bundle, name.Text, cancellationToken);
-				await WriteRefTargetAsync(name, locator, cancellationToken);
-				return locator;
-			}
-
-			public async Task WriteRefTargetAsync(RefName name, BlobLocator id, CancellationToken cancellationToken = default)
+			public override async Task WriteRefTargetAsync(RefName name, BlobLocator id, CancellationToken cancellationToken = default)
 			{
 				FileReference file = GetRefFile(name);
 				DirectoryReference.CreateDirectory(file.Directory);
@@ -118,18 +100,12 @@ namespace Horde.Agent.Commands.Bundles
 		[CommandLine("-StorageDir=", Description = "Overrides the default storage server with a local directory")]
 		public DirectoryReference StorageDir { get; set; } = DirectoryReference.Combine(Program.AppDir, "bundles");
 
-		IBlobStore? _blobStore;
-
-		protected IBlobStore CreateBlobStore(ILogger logger)
-		{
-			_blobStore ??= new FileBlobStore(StorageDir, logger);
-			return _blobStore;
-		}
+		IStorageClient? _storageClient;
 
-		protected ITreeStore CreateTreeStore(ILogger logger)
+		protected IStorageClient CreateStorageClient(ILogger logger)
 		{
-			IBlobStore blobStore = CreateBlobStore(logger);
-			return new TreeStore(blobStore, new TreeOptions());
+			_storageClient ??= new FileBlobStore(StorageDir, logger);
+			return _storageClient;
 		}
 	}
 
@@ -144,7 +120,7 @@ namespace Horde.Agent.Commands.Bundles
 
 		public override async Task<int> ExecuteAsync(ILogger logger)
 		{
-			using ITreeStore store = CreateTreeStore(logger);
+			IStorageClient store = CreateStorageClient(logger);
 			ITreeWriter writer = store.CreateTreeWriter(RefName.Text);
 
 			DirectoryNode node = new DirectoryNode(DirectoryFlags.None);
diff --git a/Engine/Source/Programs/Horde/Horde.Agent/Commands/Bundles/DumpCommand.cs b/Engine/Source/Programs/Horde/Horde.Agent/Commands/Bundles/DumpCommand.cs
--- a/Engine/Source/Programs/Horde/Horde.Agent/Commands/Bundles/DumpCommand.cs
+++ b/Engine/Source/Programs/Horde/Horde.Agent/Commands/Bundles/DumpCommand.cs
@@ -21,17 +21,17 @@ namespace Horde.Agent.Commands.Bundles
 
 		public override async Task<int> ExecuteAsync(ILogger logger)
 		{
-			IBlobStore blobStore = CreateBlobStore(logger);
+			IStorageClient storage = CreateStorageClient(logger);
 
 			Bundle blob;
 			if (BlobId != null)
 			{
-				blob = await blobStore.ReadBundleAsync(BlobId.Value);
+				blob = await storage.ReadBundleAsync(BlobId.Value);
 				logger.LogInformation("Summary for blob {BlobId}", BlobId.Value);
 			}
 			else
 			{
-				blob = await blobStore.ReadRefAsync(RefName);
+				blob = await storage.ReadRefAsync(RefName);
 				logger.LogInformation("Summary for ref {RefId}", RefName);
 			}
 
diff --git a/Engine/Source/Programs/Horde/Horde.Build.Tests/BlobStoreTests.cs b/Engine/Source/Programs/Horde/Horde.Build.Tests/BlobStoreTests.cs
--- a/Engine/Source/Programs/Horde/Horde.Build.Tests/BlobStoreTests.cs
+++ b/Engine/Source/Programs/Horde/Horde.Build.Tests/BlobStoreTests.cs
@@ -37,7 +37,7 @@ namespace Horde.Build.Tests
 	[TestClass]
 	public class BlobStoreTests : TestSetup
 	{
-		IBlobStore CreateBlobStore()
+		IStorageClient CreateBlobStore()
 		{
 			return new BasicBlobStore(MongoService, new TransientStorageBackend(), ServiceProvider.GetRequiredService<IMemoryCache>(), ServiceProvider.GetRequiredService<ILogger<BasicBlobStore>>());
 		}
@@ -81,7 +81,7 @@ namespace Horde.Build.Tests
 			return new Bundle(header, data);
 		}
 
-		static async Task<Blob> ReadBlobAsync(IBlobStore store, BlobLocator locator)
+		static async Task<Blob> ReadBlobAsync(IStorageClient store, BlobLocator locator)
 		{
 			Bundle bundle = await store.ReadBundleAsync(locator);
 			return ExtractBlobFromBundle(bundle);
@@ -98,7 +98,7 @@ namespace Horde.Build.Tests
 		[TestMethod]
 		public async Task LeafTest()
 		{
-			IBlobStore store = CreateBlobStore();
+			IStorageClient store = CreateBlobStore();
 
 			byte[] input = CreateTestData(256, 0);
 			BlobLocator locator = await store.WriteBundleAsync(CreateTestBundle(new ReadOnlySequence<byte>(input), Array.Empty<BlobLocator>()));
@@ -110,7 +110,7 @@ namespace Horde.Build.Tests
 		[TestMethod]
 		public async Task ReferenceTest()
 		{
-			IBlobStore store = CreateBlobStore();
+			IStorageClient store = CreateBlobStore();
 
 			byte[] input1 = CreateTestData(256, 1);
 			BlobLocator locator1 = await store.WriteBundleAsync(CreateTestBundle(new ReadOnlySequence<byte>(input1), Array.Empty<BlobLocator>()));
diff --git a/Engine/Source/Programs/Horde/Horde.Build.Tests/TestSetup.cs b/Engine/Source/Programs/Horde/Horde.Build.Tests/TestSetup.cs
--- a/Engine/Source/Programs/Horde/Horde.Build.Tests/TestSetup.cs
+++ b/Engine/Source/Programs/Horde/Horde.Build.Tests/TestSetup.cs
@@ -242,10 +242,10 @@ namespace Horde.Build.Tests
 			services.AddSingleton<IStorageBackend<ArtifactCollection>>(sp => new TransientStorageBackend().ForType<ArtifactCollection>());
 			services.AddSingleton<IStorageBackend<BasicStorageClient>>(sp => new TransientStorageBackend().ForType<BasicStorageClient>());
 
-			services.AddSingleton<ITreeStore<CommitService>>(sp =>
+			services.AddSingleton<IStorageClient<CommitService>>(sp =>
 			{
-				IBlobStore blobStore = new BasicBlobStore(sp.GetRequiredService<MongoService>(), new TransientStorageBackend(), sp.GetRequiredService<IMemoryCache>(), sp.GetRequiredService<ILogger<BasicBlobStore>>());
-				return new TreeStore(blobStore, new TreeOptions()).ForType<CommitService>();
+				IStorageClient blobStore = new BasicBlobStore(sp.GetRequiredService<MongoService>(), new TransientStorageBackend(), sp.GetRequiredService<IMemoryCache>(), sp.GetRequiredService<ILogger<BasicBlobStore>>());
+				return blobStore.ForType<CommitService>();
 			});
 
 			services.AddSingleton<ILegacyStorageClient, BasicStorageClient>();
diff --git a/Engine/Source/Programs/Horde/Horde.Build/Commands/Bundles/CreateCommand.cs b/Engine/Source/Programs/Horde/Horde.Build/Commands/Bundles/CreateCommand.cs
--- a/Engine/Source/Programs/Horde/Horde.Build/Commands/Bundles/CreateCommand.cs
+++ b/Engine/Source/Programs/Horde/Horde.Build/Commands/Bundles/CreateCommand.cs
@@ -34,7 +34,7 @@ namespace Horde.Build.Commands.Bundles
 		{
 			using ServiceProvider serviceProvider = Startup.CreateServiceProvider(_configuration, _loggerProvider);
 
-			ITreeStore store = serviceProvider.GetRequiredService<ITreeStore<ReplicationService>>();
+			IStorageClient store = serviceProvider.GetRequiredService<IStorageClient<ReplicationService>>();
 			ITreeWriter writer = store.CreateTreeWriter(RefName.Text);
 
 			DirectoryNode node = new DirectoryNode(DirectoryFlags.None);
diff --git a/Engine/Source/Programs/Horde/Horde.Build/Commands/Bundles/ExtractCommand.cs b/Engine/Source/Programs/Horde/Horde.Build/Commands/Bundles/ExtractCommand.cs
--- a/Engine/Source/Programs/Horde/Horde.Build/Commands/Bundles/ExtractCommand.cs
+++ b/Engine/Source/Programs/Horde/Horde.Build/Commands/Bundles/ExtractCommand.cs
@@ -35,7 +35,7 @@ namespace Horde.Build.Commands.Bundles
 		{
 			using ServiceProvider serviceProvider = Startup.CreateServiceProvider(_configuration, _loggerProvider);
 
-			ITreeStore store = serviceProvider.GetRequiredService<ITreeStore<ReplicationService>>();
+			IStorageClient store = serviceProvider.GetRequiredService<IStorageClient<ReplicationService>>();
 
 			ReplicationNode node = await store.ReadTreeAsync<ReplicationNode>(RefName);
 			DirectoryNode root = await node.Contents.ExpandAsync();
diff --git a/Engine/Source/Programs/Horde/Horde.Build/Commands/Bundles/RepoCommand.cs b/Engine/Source/Programs/Horde/Horde.Build/Commands/Bundles/RepoCommand.cs
--- a/Engine/Source/Programs/Horde/Horde.Build/Commands/Bundles/RepoCommand.cs
+++ b/Engine/Source/Programs/Horde/Horde.Build/Commands/Bundles/RepoCommand.cs
@@ -37,7 +37,7 @@ namespace Horde.Build.Commands.Bundles
 		{
 			using ServiceProvider serviceProvider = Startup.CreateServiceProvider(_configuration, _loggerProvider);
 
-			ITreeStore store = serviceProvider.GetRequiredService<ITreeStore<ReplicationService>>();
+			IStorageClient store = serviceProvider.GetRequiredService<IStorageClient<ReplicationService>>();
 
 			ReplicationNode node = await store.ReadTreeAsync<ReplicationNode>(RefName);
 			DirectoryNode root = await node.Contents.ExpandAsync();
diff --git a/Engine/Source/Programs/Horde/Horde.Build/Perforce/ReplicationService.cs b/Engine/Source/Programs/Horde/Horde.Build/Perforce/ReplicationService.cs
--- a/Engine/Source/Programs/Horde/Horde.Build/Perforce/ReplicationService.cs
+++ b/Engine/Source/Programs/Horde/Horde.Build/Perforce/ReplicationService.cs
@@ -163,7 +163,7 @@ namespace Horde.Build.Perforce
 
 		readonly IStreamCollection _streamCollection;
 		readonly IPerforceService _perforceService;
-		readonly ITreeStore _treeStore;
+		readonly IStorageClient _treeStore;
 		readonly ILogger _logger;
 
 		const int MaxBackgroundTasks = 2;
@@ -183,7 +183,7 @@ namespace Horde.Build.Perforce
 		/// <summary>
 		/// Constructor
 		/// </summary>
-		public ReplicationService(RedisService redisService, /* CommitService commitService, */IStreamCollection streamCollection, IPerforceService perforceService, ITreeStore<ReplicationService> treeStore, IOptions<ReplicationServiceOptions> options, ILogger<ReplicationService> logger)
+		public ReplicationService(RedisService redisService, /* CommitService commitService, */IStreamCollection streamCollection, IPerforceService perforceService, IStorageClient<ReplicationService> treeStore, IOptions<ReplicationServiceOptions> options, ILogger<ReplicationService> logger)
 		{
 			Options = options.Value;
 
@@ -877,7 +877,7 @@ namespace Horde.Build.Perforce
 			// Delete the incremental state
 			if (deleteIncRef)
 			{
-				await _treeStore.DeleteTreeAsync(incRefName, cancellationToken);
+				await _treeStore.DeleteRefAsync(incRefName, cancellationToken);
 			}
 			return syncNode;
 		}
diff --git a/Engine/Source/Programs/Horde/Horde.Build/Startup.cs b/Engine/Source/Programs/Horde/Horde.Build/Startup.cs
--- a/Engine/Source/Programs/Horde/Horde.Build/Startup.cs
+++ b/Engine/Source/Programs/Horde/Horde.Build/Startup.cs
@@ -277,18 +277,12 @@ namespace Horde.Build
 			}
 		}
 
-		static IBlobStore CreateBlobStore(IServiceProvider sp, BlobStoreOptions options)
+		static IStorageClient CreateStorageClient(IServiceProvider sp, BlobStoreOptions options)
 		{
 			IStorageBackend backend = CreateStorageBackend(sp, options);
 			return new BasicBlobStore(sp.GetRequiredService<MongoService>(), backend, sp.GetRequiredService<IMemoryCache>(), sp.GetRequiredService<ILogger<BasicBlobStore>>());
 		}
 
-		static ITreeStore CreateTreeStore(IServiceProvider sp, TreeStoreOptions options)
-		{
-			IBlobStore store = CreateBlobStore(sp, options);
-			return new TreeStore(store, options.Bundle);
-		}
-
 		public Startup(IConfiguration configuration)
 		{
 			Configuration = configuration;
@@ -466,7 +460,7 @@ namespace Horde.Build
 			// Storage providers
 			services.AddSingleton(sp => CreateStorageBackend(sp, settings.LogStorage).ForType<PersistentLogStorage>());
 			services.AddSingleton(sp => CreateStorageBackend(sp, settings.ArtifactStorage).ForType<ArtifactCollection>());
-			services.AddSingleton(sp => CreateTreeStore(sp, settings.CommitStorage).ForType<ReplicationService>());
+			services.AddSingleton(sp => CreateStorageClient(sp, settings.CommitStorage).ForType<ReplicationService>());
 
 			services.AddHordeStorage(settings => configSection.GetSection("Storage").Bind(settings));
 			
diff --git a/Engine/Source/Programs/Horde/Horde.Build/Storage/BasicBlobStore.cs b/Engine/Source/Programs/Horde/Horde.Build/Storage/BasicBlobStore.cs
--- a/Engine/Source/Programs/Horde/Horde.Build/Storage/BasicBlobStore.cs
+++ b/Engine/Source/Programs/Horde/Horde.Build/Storage/BasicBlobStore.cs
@@ -16,9 +16,9 @@ using MongoDB.Driver;
 namespace Horde.Build.Storage
 {
 	/// <summary>
-	/// Implementation of <see cref="IBlobStore"/> using Mongo for storing refs, and an <see cref="IStorageBackend"/> implementation for bulk data.
+	/// Implementation of <see cref="IStorageClient"/> using Mongo for storing refs, and an <see cref="IStorageBackend"/> implementation for bulk data.
 	/// </summary>
-	public class BasicBlobStore : IBlobStore
+	public class BasicBlobStore : StorageClientBase
 	{
 		class RefDocument
 		{
@@ -69,6 +69,7 @@ namespace Horde.Build.Storage
 		/// <param name="cache">Cache for ref values</param>
 		/// <param name="logger">Logger for this instance</param>
 		public BasicBlobStore(MongoService mongoService, IStorageBackend backend, IMemoryCache cache, ILogger<BasicBlobStore> logger)
+			: base(cache)
 		{
 			_refs = mongoService.GetCollection<RefDocument>("Refs");
 			_backend = backend;
@@ -81,7 +82,7 @@ namespace Horde.Build.Storage
 		static string GetBlobPath(BlobId id) => $"{id}.blob";
 
 		/// <inheritdoc/>
-		public async Task<Bundle> ReadBundleAsync(BlobLocator locator, CancellationToken cancellationToken = default)
+		public override async Task<Bundle> ReadBundleAsync(BlobLocator locator, CancellationToken cancellationToken = default)
 		{
 			string path = GetBlobPath(locator.BlobId);
 
@@ -95,7 +96,7 @@ namespace Horde.Build.Storage
 		}
 
 		/// <inheritdoc/>
-		public async Task<BlobLocator> WriteBundleAsync(Bundle bundle, Utf8String prefix = default, CancellationToken cancellationToken = default)
+		public override async Task<BlobLocator> WriteBundleAsync(Bundle bundle, Utf8String prefix = default, CancellationToken cancellationToken = default)
 		{
 			BlobId id = BlobId.CreateNew(prefix);
 			string path = GetBlobPath(id);
@@ -126,34 +127,14 @@ namespace Horde.Build.Storage
 		}
 
 		/// <inheritdoc/>
-		public async Task DeleteRefAsync(RefName name, CancellationToken cancellationToken = default)
+		public override async Task DeleteRefAsync(RefName name, CancellationToken cancellationToken = default)
 		{
 			await _refs.DeleteOneAsync(x => x.Name == name, cancellationToken);
 			AddRefToCache(name, BlobLocator.Empty);
 		}
 
 		/// <inheritdoc/>
-		public async Task<BundleRef?> TryReadRefAsync(RefName name, DateTime cacheTime = default, CancellationToken cancellationToken = default)
-		{
-			BlobLocator locator = await TryReadRefTargetAsync(name, cacheTime, cancellationToken);
-			if (locator.IsValid())
-			{
-				try
-				{
-					Bundle? bundle = await ReadBundleAsync(locator, cancellationToken);
-					return new BundleRef(locator, bundle);
-				}
-				catch (Exception ex)
-				{
-					_logger.LogWarning(ex, "Unable to read blob {BlobId} referenced by ref {RefName}", locator, name);
-					return null;
-				}
-			}
-			return null;
-		}
-
-		/// <inheritdoc/>
-		public async Task<BlobLocator> TryReadRefTargetAsync(RefName name, DateTime cacheTime = default, CancellationToken cancellationToken = default)
+		public override async Task<BlobLocator> TryReadRefTargetAsync(RefName name, DateTime cacheTime = default, CancellationToken cancellationToken = default)
 		{
 			CachedRefValue entry;
 			if (!_cache.TryGetValue(name, out entry) || entry.Time < cacheTime)
@@ -166,15 +147,7 @@ namespace Horde.Build.Storage
 		}
 
 		/// <inheritdoc/>
-		public async Task<BlobLocator> WriteRefAsync(RefName name, Bundle bundle, Utf8String prefix = default, CancellationToken cancellationToken = default)
-		{
-			BlobLocator locator = await WriteBundleAsync(bundle, name.Text, cancellationToken);
-			await WriteRefTargetAsync(name, locator, cancellationToken);
-			return locator;
-		}
-
-		/// <inheritdoc/>
-		public async Task WriteRefTargetAsync(RefName name, BlobLocator locator, CancellationToken cancellationToken = default)
+		public override async Task WriteRefTargetAsync(RefName name, BlobLocator locator, CancellationToken cancellationToken = default)
 		{
 			RefDocument refDocument = new RefDocument(name, locator);
 			await _refs.ReplaceOneAsync(x => x.Name == name, refDocument, new ReplaceOptions { IsUpsert = true }, cancellationToken);
diff --git a/Engine/Source/Programs/Shared/EpicGames.Horde.Tests/BundleStoreTests.cs b/Engine/Source/Programs/Shared/EpicGames.Horde.Tests/BundleStoreTests.cs
--- a/Engine/Source/Programs/Shared/EpicGames.Horde.Tests/BundleStoreTests.cs
+++ b/Engine/Source/Programs/Shared/EpicGames.Horde.Tests/BundleStoreTests.cs
@@ -40,12 +40,11 @@ namespace EpicGames.Horde.Tests
 			Assert.AreEqual(1, blobStore.Refs.Count);
 		}
 
-		static async Task TestTreeAsync(InMemoryBlobStore blobStore, TreeOptions options)
+		static async Task TestTreeAsync(InMemoryBlobStore store, TreeOptions options)
 		{
 			// Generate a tree
-			using (TreeStore store = new TreeStore(blobStore, options))
 			{
-				ITreeWriter writer = store.CreateTreeWriter("test");
+				ITreeWriter writer = store.CreateTreeWriter(options, "test");
 
 				ITreeBlobRef node1 = await writer.WriteNodeAsync(new ReadOnlySequence<byte>(new byte[] { 1 }), Array.Empty<ITreeBlobRef>(), CancellationToken.None);
 				ITreeBlobRef node2 = await writer.WriteNodeAsync(new ReadOnlySequence<byte>(new byte[] { 2 }), new[] { node1 }, CancellationToken.None);
@@ -59,9 +58,8 @@ namespace EpicGames.Horde.Tests
 			}
 
 			// Check we can read it back in
-			using (TreeStore collection = new TreeStore(blobStore, options))
 			{
-				ITreeBlob? root = await collection.TryReadTreeAsync(new RefName("test"));
+				ITreeBlob? root = await store.TryReadTreeAsync(new RefName("test"));
 				Assert.IsNotNull(root);
 				await CheckTree(root!);
 			}
@@ -103,10 +101,9 @@ namespace EpicGames.Horde.Tests
 		[TestMethod]
 		public async Task DirectoryNodesAsync()
 		{
-			InMemoryBlobStore blobStore = new InMemoryBlobStore();
+			InMemoryBlobStore store = new InMemoryBlobStore();
 
 			// Generate a tree
-			using (TreeStore store = new TreeStore(blobStore, new TreeOptions()))
 			{
 				DirectoryNode root = new DirectoryNode(DirectoryFlags.None);
 				DirectoryNode hello = root.AddDirectory("hello");
@@ -117,7 +114,6 @@ namespace EpicGames.Horde.Tests
 			}
 
 			// Check we can read it back in
-			using (TreeStore store = new TreeStore(blobStore, new TreeOptions()))
 			{
 				DirectoryNode root = await store.ReadTreeAsync<DirectoryNode>(new RefName("test"));
 				await CheckDirectoryTreeAsync(root);
@@ -140,10 +136,9 @@ namespace EpicGames.Horde.Tests
 		[TestMethod]
 		public async Task FileNodesAsync()
 		{
-			InMemoryBlobStore blobStore = new InMemoryBlobStore();
+			InMemoryBlobStore store = new InMemoryBlobStore();
 
 			// Generate a tree
-			using (TreeStore store = new TreeStore(blobStore, new TreeOptions()))
 			{
 				ITreeWriter writer = store.CreateTreeWriter();
 
@@ -159,7 +154,6 @@ namespace EpicGames.Horde.Tests
 			}
 
 			// Check we can read it back in
-			using (TreeStore store = new TreeStore(blobStore, new TreeOptions()))
 			{
 				DirectoryNode root = await store.ReadTreeAsync<DirectoryNode>(new RefName("test"));
 				await CheckFileTreeAsync(root);
@@ -185,8 +179,8 @@ namespace EpicGames.Horde.Tests
 		[TestMethod]
 		public async Task ChunkingTests()
 		{
-			using TreeStore store = new TreeStore(new InMemoryBlobStore(), new TreeOptions());
-			ITreeWriter writer = store.CreateTreeWriter();
+			InMemoryBlobStore store = new InMemoryBlobStore();
+			ITreeWriter writer = store.CreateTreeWriter(new TreeOptions());
 
 			byte[] chunk = RandomNumberGenerator.GetBytes(256);
 
@@ -223,11 +217,10 @@ namespace EpicGames.Horde.Tests
 		[TestMethod]
 		public async Task ChildWriterAsync()
 		{
-			InMemoryBlobStore blobStore = new InMemoryBlobStore();
+			InMemoryBlobStore store = new InMemoryBlobStore();
 
-			using (TreeStore store = new TreeStore(blobStore, new TreeOptions()))
 			{
-				ITreeWriter writer = store.CreateTreeWriter();
+				ITreeWriter writer = store.CreateTreeWriter(new TreeOptions());
 
 				ITreeWriter childWriter1 = writer.CreateChildWriter();
 				ITreeBlobRef childNode1 = await childWriter1.WriteNodeAsync(new ReadOnlySequence<byte>(new byte[] { 4, 5, 6 }), Array.Empty<ITreeBlobRef>());
@@ -238,8 +231,8 @@ namespace EpicGames.Horde.Tests
 				RefName refName = new RefName("test");
 				await writer.WriteRefAsync(refName, new ReadOnlySequence<byte>(new byte[] { 1, 2, 3 }), new[] { childNode1, childNode2 });
 
-				Assert.AreEqual(3, blobStore.Blobs.Count);
-				Assert.AreEqual(1, blobStore.Refs.Count);
+				Assert.AreEqual(3, store.Blobs.Count);
+				Assert.AreEqual(1, store.Refs.Count);
 
 				ITreeBlob? rootBlob = await store.TryReadTreeAsync(refName);
 				Assert.IsNotNull(rootBlob);
@@ -259,7 +252,7 @@ namespace EpicGames.Horde.Tests
 		[TestMethod]
 		public async Task LargeFileTestAsync()
 		{
-			InMemoryBlobStore blobStore = new InMemoryBlobStore();
+			InMemoryBlobStore store = new InMemoryBlobStore();
 
 			const int length = 1024;
 			const int copies = 4096;
@@ -274,9 +267,8 @@ namespace EpicGames.Horde.Tests
 			}
 
 			// Generate a tree
-			using (TreeStore store = new TreeStore(blobStore, new TreeOptions { MaxBlobSize = 1024 }))
 			{
-				ITreeWriter writer = store.CreateTreeWriter();
+				ITreeWriter writer = store.CreateTreeWriter(new TreeOptions { MaxBlobSize = 1024 });
 
 				DirectoryNode root = new DirectoryNode(DirectoryFlags.None);
 
@@ -294,7 +286,6 @@ namespace EpicGames.Horde.Tests
 			}
 
 			// Check we can read it back in
-			using (TreeStore store = new TreeStore(blobStore, new TreeOptions()))
 			{
 				DirectoryNode root = await store.ReadTreeAsync<DirectoryNode>(new RefName("test"));
 				await CheckLargeFileTreeAsync(root, data);
diff --git a/Engine/Source/Programs/Shared/EpicGames.Horde.Tests/BundleTests.cs b/Engine/Source/Programs/Shared/EpicGames.Horde.Tests/BundleTests.cs
--- a/Engine/Source/Programs/Shared/EpicGames.Horde.Tests/BundleTests.cs
+++ b/Engine/Source/Programs/Shared/EpicGames.Horde.Tests/BundleTests.cs
@@ -15,20 +15,13 @@ using System.Threading.Tasks;
 namespace EpicGames.Horde.Tests
 {
 	[TestClass]
-	public sealed class BundleTests : IDisposable
+	public sealed class BundleTests
 	{
-		readonly InMemoryBlobStore _blobStore;
-		readonly ITreeStore _treeStore;
+		readonly InMemoryBlobStore _storage;
 
 		public BundleTests()
 		{
-			_blobStore = new InMemoryBlobStore();
-			_treeStore = new TreeStore(_blobStore, new TreeOptions());
-		}
-
-		public void Dispose()
-		{
-			_treeStore.Dispose();
+			_storage = new InMemoryBlobStore();
 		}
 
 		[TestMethod]
@@ -64,7 +57,7 @@ namespace EpicGames.Horde.Tests
 		public async Task BasicChunkingTests()
 		{
 			RefName refName = new RefName("test");
-			ITreeWriter writer = _treeStore.CreateTreeWriter(refName.Text);
+			ITreeWriter writer = _storage.CreateTreeWriter(refName.Text);
 
 			ChunkingOptions options = new ChunkingOptions();
 			options.LeafOptions = new ChunkingOptionsForNodeType(8, 8, 8);
@@ -165,27 +158,26 @@ namespace EpicGames.Horde.Tests
 		[TestMethod]
 		public async Task BasicTestDirectory()
 		{
-			InMemoryBlobStore blobStore = new InMemoryBlobStore();
-			using ITreeStore treeStore = new TreeStore(blobStore, new TreeOptions());
+			InMemoryBlobStore store = new InMemoryBlobStore();
 
 			DirectoryNode root = new DirectoryNode(DirectoryFlags.None);
 			DirectoryNode node = root.AddDirectory("hello");
 			DirectoryNode node2 = node.AddDirectory("world");
 
 			RefName refName = new RefName("testref");
-			await treeStore.WriteTreeAsync(refName, root);
+			await store.WriteTreeAsync(refName, root);
 
 			// Should be stored inline
-			Assert.AreEqual(1, blobStore.Refs.Count);
-			Assert.AreEqual(1, blobStore.Blobs.Count);
+			Assert.AreEqual(1, store.Refs.Count);
+			Assert.AreEqual(1, store.Blobs.Count);
 
 			// Check the ref
-			BundleRef bundle = await blobStore.ReadRefAsync(refName);
+			BundleRef bundle = await store.ReadRefAsync(refName);
 			Assert.AreEqual(0, bundle.Header.Imports.Count);
 			Assert.AreEqual(3, bundle.Header.Exports.Count);
 
 			// Create a new bundle and read it back in again
-			DirectoryNode newRoot = await treeStore.ReadTreeAsync<DirectoryNode>(refName);
+			DirectoryNode newRoot = await store.ReadTreeAsync<DirectoryNode>(refName);
 
 			Assert.AreEqual(0, newRoot.Files.Count);
 			Assert.AreEqual(1, newRoot.Directories.Count);
@@ -207,22 +199,21 @@ namespace EpicGames.Horde.Tests
 			TreeOptions options = new TreeOptions();
 			options.MaxBlobSize = 1;
 
-			InMemoryBlobStore blobStore = new InMemoryBlobStore();
-			using ITreeStore treeStore = new TreeStore(blobStore, options);
+			InMemoryBlobStore storage = new InMemoryBlobStore();
 
 			DirectoryNode root = new DirectoryNode(DirectoryFlags.None);
 			root.AddDirectory("node1");
 			root.AddDirectory("node2");
 			root.AddDirectory("node3");
 
-			Assert.AreEqual(0, blobStore.Refs.Count);
-			Assert.AreEqual(0, blobStore.Blobs.Count);
+			Assert.AreEqual(0, storage.Refs.Count);
+			Assert.AreEqual(0, storage.Blobs.Count);
 
 			RefName refName = new RefName("ref");
-			await treeStore.WriteTreeAsync(refName, root);
+			await storage.WriteTreeAsync(refName, root, options);
 
-			Assert.AreEqual(1, blobStore.Refs.Count);
-			Assert.AreEqual(2, blobStore.Blobs.Count);
+			Assert.AreEqual(1, storage.Refs.Count);
+			Assert.AreEqual(2, storage.Blobs.Count);
 		}
 
 		[TestMethod]
@@ -231,11 +222,8 @@ namespace EpicGames.Horde.Tests
 			TreeOptions options = new TreeOptions();
 			options.MaxBlobSize = 1;
 
-			InMemoryBlobStore blobStore = new InMemoryBlobStore();
-
 			RefName refName = new RefName("ref");
 
-			using (ITreeStore oldBundle = new TreeStore(blobStore, options))
 			{
 				DirectoryNode root = new DirectoryNode(DirectoryFlags.None);
 
@@ -244,15 +232,14 @@ namespace EpicGames.Horde.Tests
 				DirectoryNode node3 = node2.AddDirectory("node3");
 				DirectoryNode node4 = node3.AddDirectory("node4");
 
-				await oldBundle.WriteTreeAsync(refName, root);
+				await _storage.WriteTreeAsync(refName, root, options);
 
-				Assert.AreEqual(1, blobStore.Refs.Count);
-				Assert.AreEqual(5, blobStore.Blobs.Count);
+				Assert.AreEqual(1, _storage.Refs.Count);
+				Assert.AreEqual(5, _storage.Blobs.Count);
 			}
 
-			using (ITreeStore newBundle = new TreeStore(blobStore, options))
 			{
-				DirectoryNode root = await newBundle.ReadTreeAsync<DirectoryNode>(refName);
+				DirectoryNode root = await _storage.ReadTreeAsync<DirectoryNode>(refName);
 
 				DirectoryNode? newNode1 = await root.FindDirectoryAsync("node1", CancellationToken.None);
 				Assert.IsNotNull(newNode1);
@@ -340,7 +327,7 @@ namespace EpicGames.Horde.Tests
 		[TestMethod]
 		public async Task CoreAppendTest()
 		{
-			ITreeWriter writer = _treeStore.CreateTreeWriter();
+			ITreeWriter writer = _storage.CreateTreeWriter();
 
 			byte[] data = new byte[4096];
 			new Random(0).NextBytes(data);
@@ -380,7 +367,7 @@ namespace EpicGames.Horde.Tests
 
 		async Task ChunkingTests(ChunkingOptions options)
 		{
-			ITreeWriter writer = _treeStore.CreateTreeWriter();
+			ITreeWriter writer = _storage.CreateTreeWriter(new TreeOptions());
 
 			byte[] data = new byte[4096];
 			new Random(0).NextBytes(data);
@@ -439,10 +426,8 @@ namespace EpicGames.Horde.Tests
 			TreeOptions options = new TreeOptions();
 			options.MaxBlobSize = 1;
 
-			using TreeStore treeStore = new TreeStore(_blobStore, options);
-
 			RefName refName = new RefName("ref");
-			ITreeWriter writer = treeStore.CreateTreeWriter(refName.Text);
+			ITreeWriter writer = _storage.CreateTreeWriter(refName.Text);
 
 			DirectoryNode root = new DirectoryNode(DirectoryFlags.None);
 
@@ -466,22 +451,22 @@ namespace EpicGames.Horde.Tests
 					}
 
 					int oldWorkingSetSize = GetWorkingSetSize(root);
-					await treeStore.WriteTreeAsync(refName, root);
+					await _storage.WriteTreeAsync(refName, root, new TreeOptions());
 					int newWorkingSetSize = GetWorkingSetSize(root);
 					Assert.IsTrue(newWorkingSetSize <= oldWorkingSetSize);
 					Assert.IsTrue(newWorkingSetSize <= 20);
 				}
 			}
 
-			Assert.IsTrue(_blobStore.Blobs.Count > 0);
-			Assert.IsTrue(_blobStore.Refs.Count == 1);
+			Assert.IsTrue(_storage.Blobs.Count > 0);
+			Assert.IsTrue(_storage.Refs.Count == 1);
 
-			await treeStore.WriteTreeAsync(refName, root, CancellationToken.None);
+			await _storage.WriteTreeAsync(refName, root, new TreeOptions(), CancellationToken.None);
 
 			Assert.AreEqual(totalLength, root.Length);
 
-			Assert.IsTrue(_blobStore.Blobs.Count > 0);
-			Assert.IsTrue(_blobStore.Refs.Count == 1);
+			Assert.IsTrue(_storage.Blobs.Count > 0);
+			Assert.IsTrue(_storage.Refs.Count == 1);
 		}
 
 		int GetWorkingSetSize(TreeNode node)
diff --git a/Engine/Source/Programs/Shared/EpicGames.Horde.Tests/ChunkingTests.cs b/Engine/Source/Programs/Shared/EpicGames.Horde.Tests/ChunkingTests.cs
--- a/Engine/Source/Programs/Shared/EpicGames.Horde.Tests/ChunkingTests.cs
+++ b/Engine/Source/Programs/Shared/EpicGames.Horde.Tests/ChunkingTests.cs
@@ -48,7 +48,7 @@ namespace EpicGames.Horde.Tests
 		[TestMethod]
 		public async Task BasicChunkingTests()
 		{
-			using TreeStore store = new TreeStore(new InMemoryBlobStore(), new TreeOptions());
+			InMemoryBlobStore store = new InMemoryBlobStore();
 			ITreeWriter writer = store.CreateTreeWriter("test");
 
 			ChunkingOptions options = new ChunkingOptions();
@@ -104,7 +104,7 @@ namespace EpicGames.Horde.Tests
 
 		static async Task TestChunkingAsync(ChunkingOptions options)
 		{
-			using TreeStore store = new TreeStore(new InMemoryBlobStore(), new TreeOptions());
+			InMemoryBlobStore store = new InMemoryBlobStore();
 			ITreeWriter writer = store.CreateTreeWriter();
 
 			byte[] data = new byte[4096];
diff --git a/Engine/Source/Programs/Shared/EpicGames.Horde/Storage/Backends/InMemoryBlobStore.cs b/Engine/Source/Programs/Shared/EpicGames.Horde/Storage/Backends/InMemoryBlobStore.cs
--- a/Engine/Source/Programs/Shared/EpicGames.Horde/Storage/Backends/InMemoryBlobStore.cs
+++ b/Engine/Source/Programs/Shared/EpicGames.Horde/Storage/Backends/InMemoryBlobStore.cs
@@ -11,9 +11,9 @@ using EpicGames.Core;
 namespace EpicGames.Horde.Storage.Backends
 {
 	/// <summary>
-	/// Implementation of <see cref="IBlobStore"/> which stores data in memory. Not intended for production use.
+	/// Implementation of <see cref="IStorageClient"/> which stores data in memory. Not intended for production use.
 	/// </summary>
-	public class InMemoryBlobStore : IBlobStore
+	public class InMemoryBlobStore : StorageClientBase
 	{
 		/// <summary>
 		/// Map of blob id to blob data
@@ -31,16 +31,24 @@ namespace EpicGames.Horde.Storage.Backends
 		/// <inheritdoc cref="_refs"/>
 		public IReadOnlyDictionary<RefName, BlobLocator> Refs => _refs;
 
+		/// <summary>
+		/// Constructor
+		/// </summary>
+		public InMemoryBlobStore() 
+			: base(null)
+		{
+		}
+
 		#region Blobs
 
 		/// <inheritdoc/>
-		public Task<Bundle> ReadBundleAsync(BlobLocator locator, CancellationToken cancellationToken = default)
+		public override Task<Bundle> ReadBundleAsync(BlobLocator locator, CancellationToken cancellationToken = default)
 		{
 			return Task.FromResult(_blobs[locator]);
 		}
 
 		/// <inheritdoc/>
-		public Task<BlobLocator> WriteBundleAsync(Bundle bundle, Utf8String prefix = default, CancellationToken cancellationToken = default)
+		public override Task<BlobLocator> WriteBundleAsync(Bundle bundle, Utf8String prefix = default, CancellationToken cancellationToken = default)
 		{
 			BlobLocator locator = BlobLocator.Create(HostId.Empty, prefix);
 			_blobs[locator] = bundle;
@@ -52,21 +60,10 @@ namespace EpicGames.Horde.Storage.Backends
 		#region Refs
 
 		/// <inheritdoc/>
-		public Task DeleteRefAsync(RefName name, CancellationToken cancellationToken) => Task.FromResult(_refs.TryRemove(name, out _));
-
-		/// <inheritdoc/>
-		public Task<BundleRef?> TryReadRefAsync(RefName name, DateTime cacheTime = default, CancellationToken cancellationToken = default)
-		{
-			BundleRef? result = null;
-			if (_refs.TryGetValue(name, out BlobLocator locator) && _blobs.TryGetValue(locator, out Bundle? bundle))
-			{
-				result = new BundleRef(locator, bundle);
-			}
-			return Task.FromResult(result);
-		}
+		public override Task DeleteRefAsync(RefName name, CancellationToken cancellationToken) => Task.FromResult(_refs.TryRemove(name, out _));
 
 		/// <inheritdoc/>
-		public Task<BlobLocator> TryReadRefTargetAsync(RefName name, DateTime cacheTime = default, CancellationToken cancellationToken = default)
+		public override Task<BlobLocator> TryReadRefTargetAsync(RefName name, DateTime cacheTime = default, CancellationToken cancellationToken = default)
 		{
 			if (_refs.TryGetValue(name, out BlobLocator locator))
 			{
@@ -79,15 +76,7 @@ namespace EpicGames.Horde.Storage.Backends
 		}
 
 		/// <inheritdoc/>
-		public async Task<BlobLocator> WriteRefAsync(RefName name, Bundle bundle, Utf8String prefix, CancellationToken cancellationToken)
-		{
-			BlobLocator locator = await WriteBundleAsync(bundle, name.Text, cancellationToken);
-			_refs[name] = locator;
-			return locator;
-		}
-
-		/// <inheritdoc/>
-		public Task WriteRefTargetAsync(RefName name, BlobLocator locator, CancellationToken cancellationToken = default)
+		public override Task WriteRefTargetAsync(RefName name, BlobLocator locator, CancellationToken cancellationToken = default)
 		{
 			_refs[name] = locator;
 			return Task.CompletedTask;
diff --git a/Engine/Source/Programs/Shared/EpicGames.Horde/Storage/BlobLocator.cs b/Engine/Source/Programs/Shared/EpicGames.Horde/Storage/BlobLocator.cs
--- a/Engine/Source/Programs/Shared/EpicGames.Horde/Storage/BlobLocator.cs
+++ b/Engine/Source/Programs/Shared/EpicGames.Horde/Storage/BlobLocator.cs
@@ -13,7 +13,7 @@ namespace EpicGames.Horde.Storage
 {
 	/// <summary>
 	/// Unique identifier for a blob, as a utf-8 string. Clients should not assume any internal structure to this identifier; it only
-	/// has meaning to the <see cref="IBlobStore"/> implementation.
+	/// has meaning to the <see cref="IStorageClient"/> implementation.
 	/// </summary>
 	[JsonConverter(typeof(BlobLocatorJsonConverter))]
 	[TypeConverter(typeof(BlobLocatorTypeConverter))]
diff --git a/Engine/Source/Programs/Shared/EpicGames.Horde/Storage/IBlobStore.cs b/Engine/Source/Programs/Shared/EpicGames.Horde/Storage/IBlobStore.cs
deleted file mode 100644
diff --git a/Engine/Source/Programs/Shared/EpicGames.Horde/Storage/IStorageClient.cs b/Engine/Source/Programs/Shared/EpicGames.Horde/Storage/IStorageClient.cs
new file mode 100644
--- /dev/null
+++ b/Engine/Source/Programs/Shared/EpicGames.Horde/Storage/IStorageClient.cs
@@ -0,0 +1,341 @@
+// Copyright Epic Games, Inc. All Rights Reserved.
+
+using System;
+using System.Buffers;
+using System.Collections.Generic;
+using System.Threading;
+using System.Threading.Tasks;
+using EpicGames.Core;
+using Microsoft.CodeAnalysis.Rename;
+
+namespace EpicGames.Horde.Storage
+{
+	/// <summary>
+	/// Exception for a ref not existing
+	/// </summary>
+	public sealed class RefNameNotFoundException : Exception
+	{
+		/// <summary>
+		/// Name of the missing ref
+		/// </summary>
+		public RefName Name { get; }
+
+		/// <summary>
+		/// Constructor
+		/// </summary>
+		/// <param name="name"></param>
+		public RefNameNotFoundException(RefName name)
+			: base($"Ref name '{name}' not found")
+		{
+			Name = name;
+		}
+	}
+
+	/// <summary>
+	/// Value for a ref
+	/// </summary>
+	public class BundleRef : Bundle
+	{
+		/// <summary>
+		/// Location of the blob
+		/// </summary>
+		public BlobLocator Locator { get; }
+
+		/// <summary>
+		/// Constructor
+		/// </summary>
+		public BundleRef(BlobLocator locator, Bundle bundle)
+			: base(bundle.Header, bundle.Payload)
+		{
+			Locator = locator;
+		}
+	}
+
+	/// <summary>
+	/// Base interface for a low-level storage backend. Blobs added to this store are not content addressed, but referenced by <see cref="BlobLocator"/>.
+	/// </summary>
+	public interface IStorageClient
+	{
+		#region Blobs
+
+		/// <summary>
+		/// Reads data for a bundle from the store
+		/// </summary>
+		/// <param name="locator">The blob identifier</param>
+		/// <param name="cancellationToken">Cancellation token for the operation</param>
+		/// <returns></returns>
+		Task<Bundle> ReadBundleAsync(BlobLocator locator, CancellationToken cancellationToken = default);
+
+		/// <summary>
+		/// Writes a new bundle to the store
+		/// </summary>
+		/// <param name="bundle">Bundle data</param>
+		/// <param name="cancellationToken">Cancellation token for the operation</param>
+		/// <param name="prefix">Prefix for blob names. While the returned BlobId is guaranteed to be unique, this name can be used as a prefix to aid debugging.</param>
+		/// <returns>Unique identifier for the blob</returns>
+		Task<BlobLocator> WriteBundleAsync(Bundle bundle, Utf8String prefix = default, CancellationToken cancellationToken = default);
+
+		#endregion
+
+		#region Trees
+
+		/// <summary>
+		/// Creates a new context for reading trees.
+		/// </summary>
+		/// <param name="options">Options for the tree writer</param>
+		/// <param name="prefix">Prefix for blob names.</param>
+		/// <returns>New context instance</returns>
+		ITreeWriter CreateTreeWriter(TreeOptions options, Utf8String prefix = default);
+
+		/// <summary>
+		/// Attempts to read a tree with the given name
+		/// </summary>
+		/// <param name="name">Name of the ref used to store the tree</param>
+		/// <param name="maxAge"></param>
+		/// <param name="cancellationToken">Cancellation token for the operation</param>
+		/// <returns>Reference to the root of the tree. Null if the ref does not exist.</returns>
+		Task<ITreeBlob?> TryReadTreeAsync(RefName name, TimeSpan maxAge = default, CancellationToken cancellationToken = default);
+
+		#endregion
+
+		#region Refs
+
+		/// <summary>
+		/// Reads data for a ref from the store, along with the blob contents.
+		/// </summary>
+		/// <param name="name">The ref name</param>
+		/// <param name="cacheTime">Minimum coherency for any cached value to be returned</param>
+		/// <param name="cancellationToken">Cancellation token for the operation</param>
+		/// <returns></returns>
+		Task<BundleRef?> TryReadRefAsync(RefName name, DateTime cacheTime = default, CancellationToken cancellationToken = default);
+
+		/// <summary>
+		/// Reads data for a ref from the store
+		/// </summary>
+		/// <param name="name">The ref name</param>
+		/// <param name="cacheTime">Minimum coherency for any cached value to be returned</param>
+		/// <param name="cancellationToken">Cancellation token for the operation</param>
+		/// <returns></returns>
+		Task<BlobLocator> TryReadRefTargetAsync(RefName name, DateTime cacheTime = default, CancellationToken cancellationToken = default);
+
+		/// <summary>
+		/// Writes a new ref to the store which points to a new blob
+		/// </summary>
+		/// <param name="name">Ref to write</param>
+		/// <param name="bundle">Bundle to write</param>
+		/// <param name="prefix">Prefix for blob names.</param>
+		/// <param name="cancellationToken">Cancellation token for the operation</param>
+		/// <returns>Unique identifier for the blob</returns>
+		Task<BlobLocator> WriteRefAsync(RefName name, Bundle bundle, Utf8String prefix = default, CancellationToken cancellationToken = default);
+
+		/// <summary>
+		/// Writes a new ref to the store
+		/// </summary>
+		/// <param name="name">Ref to write</param>
+		/// <param name="locator">The target for the ref</param>
+		/// <param name="cancellationToken">Cancellation token for the operation</param>
+		/// <returns>Unique identifier for the blob</returns>
+		Task WriteRefTargetAsync(RefName name, BlobLocator locator, CancellationToken cancellationToken = default);
+
+		/// <summary>
+		/// Reads data for a ref from the store
+		/// </summary>
+		/// <param name="name">The ref identifier</param>
+		/// <param name="cancellationToken">Cancellation token for the operation</param>
+		Task DeleteRefAsync(RefName name, CancellationToken cancellationToken = default);
+
+		#endregion
+	}
+
+	/// <summary>
+	/// Typed implementation of <see cref="IStorageClient"/> for use with dependency injection
+	/// </summary>
+	public interface IStorageClient<T> : IStorageClient
+	{
+	}
+
+	/// <summary>
+	/// Extension methods for <see cref="IStorageClient"/>
+	/// </summary>
+	public static class StorageClientExtensions
+	{
+		class TypedStorageClient<T> : IStorageClient<T>
+		{
+			readonly IStorageClient _inner;
+
+			public TypedStorageClient(IStorageClient inner) => _inner = inner;
+
+			#region Blobs
+
+			/// <inheritdoc/>
+			public Task<Bundle> ReadBundleAsync(BlobLocator id, CancellationToken cancellationToken = default) => _inner.ReadBundleAsync(id, cancellationToken);
+
+			/// <inheritdoc/>
+			public Task<BlobLocator> WriteBundleAsync(Bundle bundle, Utf8String prefix = default, CancellationToken cancellationToken = default) => _inner.WriteBundleAsync(bundle, prefix, cancellationToken);
+
+			#endregion
+
+			#region Trees
+
+			/// <inheritdoc/>
+			public ITreeWriter CreateTreeWriter(TreeOptions options, Utf8String prefix) => _inner.CreateTreeWriter(options, prefix);
+
+			/// <inheritdoc/>
+			public Task<ITreeBlob?> TryReadTreeAsync(RefName name, TimeSpan maxAge = default, CancellationToken cancellationToken = default) => _inner.TryReadTreeAsync(name, maxAge, cancellationToken);
+
+			#endregion
+
+			#region Refs
+
+			/// <inheritdoc/>
+			public Task DeleteRefAsync(RefName name, CancellationToken cancellationToken = default) => _inner.DeleteRefAsync(name, cancellationToken);
+
+			/// <inheritdoc/>
+			public Task<BundleRef?> TryReadRefAsync(RefName name, DateTime cacheTime = default, CancellationToken cancellationToken = default) => _inner.TryReadRefAsync(name, cacheTime, cancellationToken);
+
+			/// <inheritdoc/>
+			public Task<BlobLocator> TryReadRefTargetAsync(RefName name, DateTime cacheTime = default, CancellationToken cancellationToken = default) => _inner.TryReadRefTargetAsync(name, cacheTime, cancellationToken);
+
+			/// <inheritdoc/>
+			public Task<BlobLocator> WriteRefAsync(RefName name, Bundle bundle, Utf8String prefix = default, CancellationToken cancellationToken = default) => _inner.WriteRefAsync(name, bundle, prefix, cancellationToken);
+
+			/// <inheritdoc/>
+			public Task WriteRefTargetAsync(RefName name, BlobLocator blobId, CancellationToken cancellationToken = default) => _inner.WriteRefTargetAsync(name, blobId, cancellationToken);
+
+			#endregion
+		}
+
+		/// <summary>
+		/// Wraps a <see cref="IStorageClient"/> interface with a type argument
+		/// </summary>
+		/// <param name="blobStore">Regular blob store instance</param>
+		/// <returns></returns>
+		public static IStorageClient<T> ForType<T>(this IStorageClient blobStore) => new TypedStorageClient<T>(blobStore);
+
+		/// <summary>
+		/// Checks if the given ref exists
+		/// </summary>
+		/// <param name="store">The store instance to read from</param>
+		/// <param name="name">Name of the reference to look for</param>
+		/// <param name="cacheTime">Minimum coherency for any cached value to be returned</param>
+		/// <param name="cancellationToken">Cancellation token for the operation</param>
+		/// <returns>True if the ref exists, false if it did not exist</returns>
+		public static async Task<bool> HasRefAsync(this IStorageClient store, RefName name, DateTime cacheTime = default, CancellationToken cancellationToken = default)
+		{
+			BlobLocator locator = await store.TryReadRefTargetAsync(name, cacheTime, cancellationToken);
+			return locator.IsValid();
+		}
+
+		/// <summary>
+		/// Checks if the given ref exists
+		/// </summary>
+		/// <param name="store">The store instance to read from</param>
+		/// <param name="name">Name of the reference to look for</param>
+		/// <param name="maxAge">Maximum age of any cached ref</param>
+		/// <param name="cancellationToken">Cancellation token for the operation</param>
+		/// <returns>True if the ref exists, false if it did not exist</returns>
+		public static Task<bool> HasRefAsync(this IStorageClient store, RefName name, TimeSpan maxAge, CancellationToken cancellationToken = default)
+		{
+			return HasRefAsync(store, name, DateTime.UtcNow - maxAge, cancellationToken);
+		}
+
+		/// <summary>
+		/// Creates a new tree writer with default options
+		/// </summary>
+		/// <param name="store">Store instance</param>
+		/// <param name="prefix"></param>
+		/// <returns></returns>
+		public static ITreeWriter CreateTreeWriter(this IStorageClient store, Utf8String prefix = default)
+		{
+			return store.CreateTreeWriter(new TreeOptions(), prefix);
+		}
+
+		/// <summary>
+		/// Reads a ref from the store, throwing an exception if it does not exist
+		/// </summary>
+		/// <param name="store">The store instance to read from</param>
+		/// <param name="name">Id for the ref</param>
+		/// <param name="maxAge">Maximum age of any cached ref</param>
+		/// <param name="cancellationToken">Cancellation token for the operation</param>
+		/// <returns>The blob instance</returns>
+		public static Task<BundleRef?> TryReadRefAsync(this IStorageClient store, RefName name, TimeSpan maxAge, CancellationToken cancellationToken = default)
+		{
+			return store.TryReadRefAsync(name, DateTime.UtcNow - maxAge, cancellationToken);
+		}
+
+		/// <summary>
+		/// Reads a ref from the store, throwing an exception if it does not exist
+		/// </summary>
+		/// <param name="store">The store instance to read from</param>
+		/// <param name="name">Id for the ref</param>
+		/// <param name="maxAge">Maximum age of any cached ref</param>
+		/// <param name="cancellationToken">Cancellation token for the operation</param>
+		/// <returns>The blob instance</returns>
+		public static Task<BlobLocator> TryReadRefTargetAsync(this IStorageClient store, RefName name, TimeSpan maxAge, CancellationToken cancellationToken = default)
+		{
+			return store.TryReadRefTargetAsync(name, DateTime.UtcNow - maxAge, cancellationToken);
+		}
+
+		/// <summary>
+		/// Reads a ref from the store, throwing an exception if it does not exist
+		/// </summary>
+		/// <param name="store">The store instance to read from</param>
+		/// <param name="name">Id for the ref</param>
+		/// <param name="cacheTime">Minimum coherency of any cached result</param>
+		/// <param name="cancellationToken">Cancellation token for the operation</param>
+		/// <returns>The blob instance</returns>
+		public static async Task<BundleRef> ReadRefAsync(this IStorageClient store, RefName name, DateTime cacheTime = default, CancellationToken cancellationToken = default)
+		{
+			BundleRef? bundleRef = await store.TryReadRefAsync(name, cacheTime, cancellationToken);
+			if (bundleRef == null)
+			{
+				throw new RefNameNotFoundException(name);
+			}
+			return bundleRef;
+		}
+
+		/// <summary>
+		/// Reads a ref from the store, throwing an exception if it does not exist
+		/// </summary>
+		/// <param name="store">The store instance to read from</param>
+		/// <param name="name">Id for the ref</param>
+		/// <param name="maxAge">Maximum age for any cached result</param>
+		/// <param name="cancellationToken">Cancellation token for the operation</param>
+		/// <returns>The blob instance</returns>
+		public static Task<BundleRef> ReadRefAsync(this IStorageClient store, RefName name, TimeSpan maxAge, CancellationToken cancellationToken = default)
+		{
+			return ReadRefAsync(store, name, DateTime.UtcNow - maxAge, cancellationToken);
+		}
+
+		/// <summary>
+		/// Reads a ref from the store, throwing an exception if it does not exist
+		/// </summary>
+		/// <param name="store">The store instance to read from</param>
+		/// <param name="name">Id for the ref</param>
+		/// <param name="cacheTime">Minimum coherency of any cached result</param>
+		/// <param name="cancellationToken">Cancellation token for the operation</param>
+		/// <returns>The blob instance</returns>
+		public static async Task<BlobLocator> ReadRefTargetAsync(this IStorageClient store, RefName name, DateTime cacheTime = default, CancellationToken cancellationToken = default)
+		{
+			BlobLocator locator = await store.TryReadRefTargetAsync(name, cacheTime, cancellationToken);
+			if (!locator.IsValid())
+			{
+				throw new RefNameNotFoundException(name);
+			}
+			return locator;
+		}
+
+		/// <summary>
+		/// Reads a ref from the store, throwing an exception if it does not exist
+		/// </summary>
+		/// <param name="store">The store instance to read from</param>
+		/// <param name="name">Id for the ref</param>
+		/// <param name="maxAge">Maximum age for any cached result</param>
+		/// <param name="cancellationToken">Cancellation token for the operation</param>
+		/// <returns>The blob instance</returns>
+		public static Task<BlobLocator> ReadRefTargetAsync(this IStorageClient store, RefName name, TimeSpan maxAge, CancellationToken cancellationToken = default)
+		{
+			return ReadRefTargetAsync(store, name, DateTime.UtcNow - maxAge, cancellationToken);
+		}
+	}
+}
diff --git a/Engine/Source/Programs/Shared/EpicGames.Horde/Storage/ITreeStore.cs b/Engine/Source/Programs/Shared/EpicGames.Horde/Storage/ITreeStore.cs
--- a/Engine/Source/Programs/Shared/EpicGames.Horde/Storage/ITreeStore.cs
+++ b/Engine/Source/Programs/Shared/EpicGames.Horde/Storage/ITreeStore.cs
@@ -1,6 +1,5 @@
 // Copyright Epic Games, Inc. All Rights Reserved.
 
-using System;
 using System.Buffers;
 using System.Collections.Generic;
 using System.Threading;
@@ -43,50 +42,6 @@ namespace EpicGames.Horde.Storage
 		ValueTask<ITreeBlob> GetTargetAsync(CancellationToken cancellationToken = default);
 	}
 
-	/// <summary>
-	/// Base interface for a tree store
-	/// </summary>
-	public interface ITreeStore : IDisposable
-	{
-		/// <summary>
-		/// Creates a new context for reading trees.
-		/// </summary>
-		/// <param name="prefix">Prefix for blobs. See <see cref="IBlobStore.WriteBundleAsync(Bundle, Utf8String, CancellationToken)"/></param>
-		/// <returns>New context instance</returns>
-		ITreeWriter CreateTreeWriter(Utf8String prefix = default);
-
-		/// <summary>
-		/// Deletes a tree with the given name
-		/// </summary>
-		/// <param name="name">Name of the ref used to store the tree</param>
-		/// <param name="cancellationToken">Cancellation token for the operation</param>
-		Task DeleteTreeAsync(RefName name, CancellationToken cancellationToken = default);
-
-		/// <summary>
-		/// Tests whether a tree exists with the given name
-		/// </summary>
-		/// <param name="name">Name of the ref used to store the tree</param>
-		/// <param name="cancellationToken">Cancellation token for the operation</param>
-		/// <returns>Tree reader instance</returns>
-		Task<bool> HasTreeAsync(RefName name, CancellationToken cancellationToken = default);
-
-		/// <summary>
-		/// Attempts to read a tree with the given name
-		/// </summary>
-		/// <param name="name">Name of the ref used to store the tree</param>
-		/// <param name="maxAge"></param>
-		/// <param name="cancellationToken">Cancellation token for the operation</param>
-		/// <returns>Reference to the root of the tree. Null if the ref does not exist.</returns>
-		Task<ITreeBlob?> TryReadTreeAsync(RefName name, TimeSpan maxAge = default, CancellationToken cancellationToken = default);
-	}
-
-	/// <summary>
-	/// Interface for a typed tree store
-	/// </summary>
-	public interface ITreeStore<T> : ITreeStore
-	{
-	}
-
 	/// <summary>
 	/// Allows incrementally writing new trees
 	/// </summary>
@@ -149,42 +104,6 @@ namespace EpicGames.Horde.Storage
 	/// </summary>
 	public static class TreeStoreExtensions
 	{
-		sealed class TypedTreeStore<T> : ITreeStore<T>
-		{
-			readonly ITreeStore _inner;
-
-			public TypedTreeStore(ITreeStore inner)
-			{
-				_inner = inner;
-			}
-
-			/// <inheritdoc/>
-			public ITreeWriter CreateTreeWriter(Utf8String prefix) => _inner.CreateTreeWriter(prefix);
-
-			/// <inheritdoc/>
-			public Task DeleteTreeAsync(RefName name, CancellationToken cancellationToken = default) => _inner.DeleteTreeAsync(name, cancellationToken);
-
-			/// <inheritdoc/>
-			public void Dispose() => _inner.Dispose();
-
-			/// <inheritdoc/>
-			public Task<bool> HasTreeAsync(RefName name, CancellationToken cancellationToken = default) => _inner.HasTreeAsync(name, cancellationToken);
-
-			/// <inheritdoc/>
-			public Task<ITreeBlob?> TryReadTreeAsync(RefName name, TimeSpan maxAge = default, CancellationToken cancellationToken = default) => _inner.TryReadTreeAsync(name, maxAge, cancellationToken);
-		}
-
-		/// <summary>
-		/// Wraps an <see cref="ITreeStore"/> interface with a type for dependency injection
-		/// </summary>
-		/// <typeparam name="T"></typeparam>
-		/// <param name="store">The instance to wrap</param>
-		/// <returns>Wrapped instance of the tree store</returns>
-		public static ITreeStore<T> ForType<T>(this ITreeStore store)
-		{
-			return new TypedTreeStore<T>(store);
-		}
-
 		/// <summary>
 		/// Writes a ref using the given root data
 		/// </summary>
diff --git a/Engine/Source/Programs/Shared/EpicGames.Horde/Storage/StorageClientBase.cs b/Engine/Source/Programs/Shared/EpicGames.Horde/Storage/StorageClientBase.cs
new file mode 100644
--- /dev/null
+++ b/Engine/Source/Programs/Shared/EpicGames.Horde/Storage/StorageClientBase.cs
@@ -0,0 +1,92 @@
+// Copyright Epic Games, Inc. All Rights Reserved.
+
+using System;
+using System.Collections.Generic;
+using System.Linq;
+using System.Runtime.CompilerServices;
+using System.Text;
+using System.Threading;
+using System.Threading.Tasks;
+using EpicGames.Core;
+using Microsoft.Extensions.Caching.Memory;
+
+namespace EpicGames.Horde.Storage
+{
+	/// <summary>
+	/// Base class for an implementation of <see cref="IStorageClient"/>, providing implementations for some common functionality.
+	/// </summary>
+	public abstract class StorageClientBase : IStorageClient
+	{
+		readonly TreeStore _treeStore;
+
+		/// <summary>
+		/// Constructor
+		/// </summary>
+		/// <param name="cache"></param>
+		protected StorageClientBase(IMemoryCache? cache)
+		{
+			_treeStore = new TreeStore(this, cache);
+		}
+
+		#region Blobs
+
+		/// <inheritdoc/>
+		public abstract Task<Bundle> ReadBundleAsync(BlobLocator locator, CancellationToken cancellationToken = default);
+
+		/// <inheritdoc/>
+		public abstract Task<BlobLocator> WriteBundleAsync(Bundle bundle, Utf8String prefix = default, CancellationToken cancellationToken = default);
+
+		#endregion
+
+		#region Trees
+
+		/// <inheritdoc/>
+		public ITreeWriter CreateTreeWriter(TreeOptions options, Utf8String prefix = default)
+		{
+			return _treeStore.CreateTreeWriter(options, prefix);
+		}
+
+		/// <inheritdoc/>
+		public virtual Task<ITreeBlob?> TryReadTreeAsync(RefName name, TimeSpan maxAge = default, CancellationToken cancellationToken = default)
+		{
+			return _treeStore.TryReadTreeAsync(name, maxAge, cancellationToken);
+		}
+
+		#endregion
+
+		#region Refs
+
+		/// <inheritdoc/>
+		public abstract Task DeleteRefAsync(RefName name, CancellationToken cancellationToken = default);
+
+		/// <inheritdoc/>
+		public virtual async Task<BundleRef?> TryReadRefAsync(RefName name, DateTime cacheTime = default, CancellationToken cancellationToken = default)
+		{
+			BlobLocator locator = await TryReadRefTargetAsync(name, cacheTime, cancellationToken);
+			if (locator.IsValid())
+			{
+				return new BundleRef(locator, await ReadBundleAsync(locator, cancellationToken));
+			}
+			else
+			{
+				return null;
+			}
+		}
+
+		/// <inheritdoc/>
+		public abstract Task<BlobLocator> TryReadRefTargetAsync(RefName name, DateTime cacheTime = default, CancellationToken cancellationToken = default);
+
+		/// <inheritdoc/>
+		public virtual async Task<BlobLocator> WriteRefAsync(RefName name, Bundle bundle, Utf8String prefix = default, CancellationToken cancellationToken = default)
+		{
+			BlobLocator locator = await WriteBundleAsync(bundle, prefix, cancellationToken);
+			await WriteRefTargetAsync(name, locator, cancellationToken);
+			return locator;
+		}
+
+		/// <inheritdoc/>
+		public abstract Task WriteRefTargetAsync(RefName name, BlobLocator locator, CancellationToken cancellationToken = default);
+
+		#endregion
+	}
+}
diff --git a/Engine/Source/Programs/Shared/EpicGames.Horde/Storage/TreeNode.cs b/Engine/Source/Programs/Shared/EpicGames.Horde/Storage/TreeNode.cs
--- a/Engine/Source/Programs/Shared/EpicGames.Horde/Storage/TreeNode.cs
+++ b/Engine/Source/Programs/Shared/EpicGames.Horde/Storage/TreeNode.cs
@@ -167,7 +167,7 @@ namespace EpicGames.Horde.Storage
 		}
 
 		/// <inheritdoc/>
-		public static async Task<T?> TryReadTreeAsync<T>(this ITreeStore store, RefName name, TimeSpan maxAge = default, CancellationToken cancellationToken = default) where T : TreeNode
+		public static async Task<T?> TryReadTreeAsync<T>(this IStorageClient store, RefName name, TimeSpan maxAge = default, CancellationToken cancellationToken = default) where T : TreeNode
 		{
 			ITreeBlob? root = await store.TryReadTreeAsync(name, maxAge, cancellationToken);
 			if (root == null)
@@ -178,7 +178,7 @@ namespace EpicGames.Horde.Storage
 		}
 
 		/// <inheritdoc/>
-		public static async Task<T> ReadTreeAsync<T>(this ITreeStore store, RefName name, TimeSpan maxAge = default, CancellationToken cancellationToken = default) where T : TreeNode
+		public static async Task<T> ReadTreeAsync<T>(this IStorageClient store, RefName name, TimeSpan maxAge = default, CancellationToken cancellationToken = default) where T : TreeNode
 		{
 			T? result = await store.TryReadTreeAsync<T>(name, maxAge, cancellationToken);
 			if (result == null)
@@ -189,9 +189,15 @@ namespace EpicGames.Horde.Storage
 		}
 
 		/// <inheritdoc/>
-		public static async Task WriteTreeAsync(this ITreeStore store, RefName name, TreeNode root, CancellationToken cancellationToken = default)
+		public static Task WriteTreeAsync(this IStorageClient store, RefName name, TreeNode root, CancellationToken cancellationToken = default)
 		{
-			ITreeWriter writer = store.CreateTreeWriter(name.Text);
+			return WriteTreeAsync(store, name, root, new TreeOptions(), cancellationToken);
+		}
+
+		/// <inheritdoc/>
+		public static async Task WriteTreeAsync(this IStorageClient store, RefName name, TreeNode root, TreeOptions options, CancellationToken cancellationToken = default)
+		{
+			ITreeWriter writer = store.CreateTreeWriter(options);
 			await writer.WriteRefAsync(name, root, cancellationToken);
 		}
 	}
diff --git a/Engine/Source/Programs/Shared/EpicGames.Horde/Storage/TreeStore.cs b/Engine/Source/Programs/Shared/EpicGames.Horde/Storage/TreeStore.cs
--- a/Engine/Source/Programs/Shared/EpicGames.Horde/Storage/TreeStore.cs
+++ b/Engine/Source/Programs/Shared/EpicGames.Horde/Storage/TreeStore.cs
@@ -47,10 +47,10 @@ namespace EpicGames.Horde.Storage
 	}
 
 	/// <summary>
-	/// Implementation of <see cref="ITreeStore"/> which packs nodes together into <see cref="Bundle"/> objects, then stores them
-	/// in a <see cref="IBlobStore"/>.
+	/// Implementation of tree storage which packs nodes together into <see cref="Bundle"/> objects, then stores them
+	/// in a <see cref="IStorageClient"/>.
 	/// </summary>
-	public class TreeStore : ITreeStore, IDisposable
+	public class TreeStore
 	{
 		/// <summary>
 		/// Information about a node within a bundle.
@@ -210,6 +210,7 @@ namespace EpicGames.Horde.Storage
 			readonly TreeStore _owner;
 			readonly BundleContext _context;
 			readonly BundleWriter? _parent;
+			readonly TreeOptions _options;
 			readonly Utf8String _prefix;
 
 			public readonly ConcurrentDictionary<IoHash, NodeInfo> HashToNode = new ConcurrentDictionary<IoHash, NodeInfo>(); // TODO: this needs to include some additional state from external sources.
@@ -232,18 +233,19 @@ namespace EpicGames.Horde.Storage
 			/// <summary>
 			/// Constructor
 			/// </summary>
-			public BundleWriter(TreeStore owner, BundleContext context, BundleWriter? parent, Utf8String prefix)
+			public BundleWriter(TreeStore owner, BundleContext context, BundleWriter? parent, TreeOptions options, Utf8String prefix)
 			{
 				_owner = owner;
 				_context = context;
 				_parent = parent;
+				_options = options;
 				_prefix = prefix;
 			}
 
 			/// <inheritdoc/>
 			public ITreeWriter CreateChildWriter()
 			{
-				BundleWriter writer = new BundleWriter(_owner, _context, this, _prefix);
+				BundleWriter writer = new BundleWriter(_owner, _context, this, _options, _prefix);
 				lock (_lockObject)
 				{
 					_children.Add(writer);
@@ -318,7 +320,7 @@ namespace EpicGames.Horde.Storage
 					NodeInfo nextNode = _queueNodes[_readyCount];
 					InMemoryNodeState nextState = (InMemoryNodeState)nextNode.State;
 
-					if (_readyCount > 0 && _readyLength + nextState.Data.Length > _owner._options.MaxBlobSize)
+					if (_readyCount > 0 && _readyLength + nextState.Data.Length > _options.MaxBlobSize)
 					{
 						WriteReady();
 					}
@@ -496,7 +498,7 @@ namespace EpicGames.Horde.Storage
 			/// </summary>
 			Bundle CreateBundle(NodeInfo[] nodes)
 			{
-				TreeOptions options = _owner._options;
+				TreeOptions options = _options;
 
 				// Create a set from the nodes to be written. We use this to determine references that are imported.
 				HashSet<NodeInfo> nodeSet = new HashSet<NodeInfo>(nodes);
@@ -567,7 +569,7 @@ namespace EpicGames.Horde.Storage
 					// If we can't fit this data into the current block, flush the contents of it first
 					if (blockSize > 0 && blockSize + nodeData.Length > options.MinCompressionPacketSize)
 					{
-						FlushPacket(_owner._options.CompressionFormat, blockSegments, blockSize, builder, packets);
+						FlushPacket(_options.CompressionFormat, blockSegments, blockSize, builder, packets);
 						blockSize = 0;
 					}
 
@@ -666,7 +668,7 @@ namespace EpicGames.Horde.Storage
 			}
 		}
 
-		readonly IBlobStore _blobStore;
+		readonly IStorageClient _blobStore;
 		readonly IMemoryCache? _cache;
 
 		readonly object _lockObject = new object();
@@ -675,46 +677,21 @@ namespace EpicGames.Horde.Storage
 		// so that other threads can also await it.
 		readonly Dictionary<string, Task<Bundle>> _readTasks = new Dictionary<string, Task<Bundle>>();
 
-		readonly TreeOptions _options;
-
 		/// <summary>
 		/// Constructor
 		/// </summary>
-		public TreeStore(IBlobStore blobStore, TreeOptions options)
+		public TreeStore(IStorageClient blobStore, IMemoryCache? cache)
 		{
 			_blobStore = blobStore;
-			_options = options;
-			if (_options.CacheSize > 0)
-			{
-				_cache = new MemoryCache(new MemoryCacheOptions { SizeLimit = _options.CacheSize });
-			}
+			_cache = cache;
 		}
 
 		/// <inheritdoc/>
-		public ITreeWriter CreateTreeWriter(Utf8String prefix = default) => new BundleWriter(this, new BundleContext(), null, prefix);
+		public ITreeWriter CreateTreeWriter(TreeOptions options, Utf8String prefix = default) => new BundleWriter(this, new BundleContext(), null, options, prefix);
 
 		/// <inheritdoc/>
 		public Task DeleteTreeAsync(RefName name, CancellationToken cancellationToken) => _blobStore.DeleteRefAsync(name, cancellationToken);
 
-		/// <inheritdoc/>
-		public void Dispose()
-		{
-			Dispose(true);
-			GC.SuppressFinalize(this);
-		}
-
-		/// <summary>
-		/// Dispose of this objects resources
-		/// </summary>
-		/// <param name="disposing"></param>
-		protected virtual void Dispose(bool disposing)
-		{
-			if (disposing)
-			{
-				_cache?.Dispose();
-			}
-		}
-
 		/// <inheritdoc/>
 		public Task<bool> HasTreeAsync(RefName name, CancellationToken cancellationToken) => _blobStore.HasRefAsync(name, cancellationToken: cancellationToken);
 
