Horde: Fix compression format not being respected in bundles.

diff --git a/Engine/Source/Programs/Shared/EpicGames.Horde/Storage/TreeStore.cs b/Engine/Source/Programs/Shared/EpicGames.Horde/Storage/TreeStore.cs
--- a/Engine/Source/Programs/Shared/EpicGames.Horde/Storage/TreeStore.cs
+++ b/Engine/Source/Programs/Shared/EpicGames.Horde/Storage/TreeStore.cs
@@ -732,7 +732,7 @@ namespace EpicGames.Horde.Storage
 			ExportedNodeState exportedState = await GetExportedStateAsync(node, cancellationToken);
 
 			Bundle bundle = await ReadBundleAsync(exportedState.BundleInfo.Locator, cancellationToken);
-			ReadOnlyMemory<byte> packetData = DecodePacket(exportedState.BundleInfo.Locator, exportedState.PacketInfo, bundle.Payload);
+			ReadOnlyMemory<byte> packetData = DecodePacket(exportedState.BundleInfo.Locator, exportedState.PacketInfo, bundle.Header.CompressionFormat, bundle.Payload);
 			ReadOnlyMemory<byte> nodeData = packetData.Slice(exportedState.Offset, exportedState.Length);
 
 			return new InMemoryNodeState(new ReadOnlySequence<byte>(nodeData), exportedState.References);
@@ -911,32 +911,33 @@ namespace EpicGames.Horde.Storage
 		/// </summary>
 		/// <param name="locator">Information about the bundle</param>
 		/// <param name="packetInfo">The decoded block location and size</param>
+		/// <param name="format">Compression type in the bundle</param>
 		/// <param name="payload">The bundle payload</param>
 		/// <returns>The decoded data</returns>
-		ReadOnlyMemory<byte> DecodePacket(BlobLocator locator, PacketInfo packetInfo, ReadOnlySequence<byte> payload)
+		ReadOnlyMemory<byte> DecodePacket(BlobLocator locator, PacketInfo packetInfo, BundleCompressionFormat format, ReadOnlySequence<byte> payload)
 		{
 			if (_cache == null)
 			{
-				return DecodePacketUncached(packetInfo, payload);
+				return DecodePacketUncached(packetInfo, format, payload);
 			}
 			else
 			{
 				string cacheKey = $"bundle-packet:{locator.BlobId}@{packetInfo.Offset}";
 				return _cache.GetOrCreate<ReadOnlyMemory<byte>>(cacheKey, entry =>
 				{
-					ReadOnlyMemory<byte> decodedPacket = DecodePacketUncached(packetInfo, payload);
+					ReadOnlyMemory<byte> decodedPacket = DecodePacketUncached(packetInfo, format, payload);
 					entry.SetSize(packetInfo.DecodedLength);
 					return decodedPacket;
 				});
 			}
 		}
 
-		static ReadOnlyMemory<byte> DecodePacketUncached(PacketInfo packetInfo, ReadOnlySequence<byte> payload)
+		static ReadOnlyMemory<byte> DecodePacketUncached(PacketInfo packetInfo, BundleCompressionFormat format, ReadOnlySequence<byte> payload)
 		{
 			ReadOnlyMemory<byte> encodedPacket = payload.Slice(packetInfo.Offset, packetInfo.EncodedLength).AsSingleSegment();
 
 			byte[] decodedPacket = new byte[packetInfo.DecodedLength];
-			LZ4Codec.Decode(encodedPacket.Span, decodedPacket);
+			Decode(format, encodedPacket.Span, decodedPacket);
 
 			return decodedPacket;
 		}
@@ -975,16 +976,6 @@ namespace EpicGames.Horde.Storage
 			}
 		}
 
-		static void CreateFreeSpace(ref byte[] buffer, int usedSize, int requiredSize)
-		{
-			if (requiredSize > buffer.Length)
-			{
-				byte[] newBuffer = new byte[requiredSize];
-				buffer.AsSpan(0, usedSize).CopyTo(newBuffer);
-				buffer = newBuffer;
-			}
-		}
-
 		#endregion
 	}
 }
