[UHT]
* Fixed so dll export is added to GetLifetimeReplicatedProps
* Fixed so UFUNCTION that produce _Validate function adds dll export if class is not dll exported but function is

diff --git a/Engine/Source/Programs/Shared/EpicGames.UHT/Exporters/CodeGen/UhtHeaderCodeGeneratorHFile.cs b/Engine/Source/Programs/Shared/EpicGames.UHT/Exporters/CodeGen/UhtHeaderCodeGeneratorHFile.cs
--- a/Engine/Source/Programs/Shared/EpicGames.UHT/Exporters/CodeGen/UhtHeaderCodeGeneratorHFile.cs
+++ b/Engine/Source/Programs/Shared/EpicGames.UHT/Exporters/CodeGen/UhtHeaderCodeGeneratorHFile.cs
@@ -349,10 +349,10 @@ namespace EpicGames.UHT.Exporters.CodeGen
 			bool hasEditorRpc = reversedFunctions.Any(x => IsRpcFunction(x, true));
 
 			// Output the RPC methods
-			AppendRpcFunctions(builder, classObj, reversedFunctions, false);
+			AppendRpcFunctions(builder, classObj, api, reversedFunctions, false);
 			if (hasEditorRpc)
 			{
-				AppendRpcFunctions(builder, classObj, reversedFunctions, true);
+				AppendRpcFunctions(builder, classObj, api, reversedFunctions, true);
 			}
 
 			// Output property accessors
@@ -606,13 +606,13 @@ namespace EpicGames.UHT.Exporters.CodeGen
 			return builder;
 		}
 
-		private StringBuilder AppendRpcFunctions(StringBuilder builder, UhtClass classObj, List<UhtFunction> reversedFunctions, bool editorOnly)
+		private StringBuilder AppendRpcFunctions(StringBuilder builder, UhtClass classObj, string api,  List<UhtFunction> reversedFunctions, bool editorOnly)
 		{
 			builder.AppendBeginEditorOnlyGuard(editorOnly);
 
 			using (UhtMacroCreator macro = new(builder, this, classObj, editorOnly ? EditorOnlyRpcWrappersMacroSuffix : RpcWrappersMacroSuffix))
 			{
-				AppendAutogeneratedBlueprintFunctionDeclarations(builder, reversedFunctions, editorOnly);
+				AppendAutogeneratedBlueprintFunctionDeclarations(builder, classObj, api, reversedFunctions, editorOnly);
 				AppendRpcWrappers(builder, reversedFunctions, editorOnly);
 			}
 
@@ -620,7 +620,7 @@ namespace EpicGames.UHT.Exporters.CodeGen
 			{
 				if (classObj.GeneratedCodeVersion <= EGeneratedCodeVersion.V1)
 				{
-					AppendAutogeneratedBlueprintFunctionDeclarationsOnlyNotDeclared(builder, reversedFunctions, editorOnly);
+					AppendAutogeneratedBlueprintFunctionDeclarationsOnlyNotDeclared(builder, classObj, api, reversedFunctions, editorOnly);
 				}
 				AppendRpcWrappers(builder, reversedFunctions, editorOnly);
 			}
@@ -662,7 +662,7 @@ namespace EpicGames.UHT.Exporters.CodeGen
 			return builder;
 		}
 
-		private StringBuilder AppendAutogeneratedBlueprintFunctionDeclarations(StringBuilder builder, List<UhtFunction> reversedFunctions, bool editorOnly)
+		private StringBuilder AppendAutogeneratedBlueprintFunctionDeclarations(StringBuilder builder, UhtClass classObj, string api, List<UhtFunction> reversedFunctions, bool editorOnly)
 		{
 			foreach (UhtFunction function in reversedFunctions)
 			{
@@ -670,13 +670,13 @@ namespace EpicGames.UHT.Exporters.CodeGen
 				{
 					continue;
 				}
-				AppendNetValidateDeclaration(builder, function);
+				AppendNetValidateDeclaration(builder, classObj, api, function);
 				AppendNativeFunctionHeader(builder, function, UhtPropertyTextType.ClassFunctionArgOrRetVal, true, null, null, UhtFunctionExportFlags.None, "; \\\r\n");
 			}
 			return builder;
 		}
 
-		private StringBuilder AppendAutogeneratedBlueprintFunctionDeclarationsOnlyNotDeclared(StringBuilder builder, List<UhtFunction> reversedFunctions, bool editorOnly)
+		private StringBuilder AppendAutogeneratedBlueprintFunctionDeclarationsOnlyNotDeclared(StringBuilder builder, UhtClass classObj, string api, List<UhtFunction> reversedFunctions, bool editorOnly)
 		{
 			foreach (UhtFunction function in reversedFunctions)
 			{
@@ -686,7 +686,7 @@ namespace EpicGames.UHT.Exporters.CodeGen
 				}
 				if (!function.FunctionExportFlags.HasAnyFlags(UhtFunctionExportFlags.ValidationImplFound))
 				{
-					AppendNetValidateDeclaration(builder, function);
+					AppendNetValidateDeclaration(builder, classObj, api, function);
 				}
 				if (!function.FunctionExportFlags.HasAnyFlags(UhtFunctionExportFlags.ImplFound))
 				{
@@ -696,11 +696,18 @@ namespace EpicGames.UHT.Exporters.CodeGen
 			return builder;
 		}
 
-		private static StringBuilder AppendNetValidateDeclaration(StringBuilder builder, UhtFunction function)
+		private static StringBuilder AppendNetValidateDeclaration(StringBuilder builder, UhtClass classObj, string api, UhtFunction function)
 		{
 			if (function.FunctionFlags.HasAnyFlags(EFunctionFlags.NetValidate))
 			{
 				builder.Append('\t');
+
+				if (!classObj.ClassFlags.HasAnyFlags(EClassFlags.RequiredAPI) && function.FunctionFlags.HasAnyFlags(EFunctionFlags.RequiredAPI))
+				{
+					builder.Append(api);
+				}
+
+				builder.Append(' ');
 				if (!function.FunctionFlags.HasAnyFlags(EFunctionFlags.Static) && !function.FunctionExportFlags.HasAnyFlags(UhtFunctionExportFlags.Final))
 				{
 					builder.Append("virtual");
@@ -1242,7 +1249,7 @@ namespace EpicGames.UHT.Exporters.CodeGen
 				// Default version autogenerates declarations.
 				if (classObj.GeneratedCodeVersion == EGeneratedCodeVersion.V1)
 				{
-					builder.Append("\tvoid GetLifetimeReplicatedProps(TArray<FLifetimeProperty>& OutLifetimeProps) const override; \\\r\n");
+					builder.Append('\t').Append(api).Append("void GetLifetimeReplicatedProps(TArray<FLifetimeProperty>& OutLifetimeProps) const override; \\\r\n");
 				}
 			}
 
diff --git a/Engine/Source/Programs/UnrealHeaderTool/Private/CodeGenerator.cpp b/Engine/Source/Programs/UnrealHeaderTool/Private/CodeGenerator.cpp
--- a/Engine/Source/Programs/UnrealHeaderTool/Private/CodeGenerator.cpp
+++ b/Engine/Source/Programs/UnrealHeaderTool/Private/CodeGenerator.cpp
@@ -297,7 +297,7 @@ namespace
 			// Default version autogenerates declarations.
 			if (ClassDef.GetGeneratedCodeVersion() == EGeneratedCodeVersion::V1)
 			{
-				Writer.Logf(TEXT("\tvoid GetLifetimeReplicatedProps(TArray<FLifetimeProperty>& OutLifetimeProps) const override;\r\n"));
+				Writer.Logf(TEXT("\t%s_API void GetLifetimeReplicatedProps(TArray<FLifetimeProperty>& OutLifetimeProps) const override;\r\n"), API);
 			}
 			else
 			{
@@ -4918,9 +4918,10 @@ void FNativeClassHeaderGenerator::ExportNativeFunctions(FOutputDevice& OutGenera
 			FFindDelcarationResults Implementation = FindDeclaration(ClassDef, FunctionData.CppImplName);
 			FFindDelcarationResults Validation = FindDeclaration(ClassDef, FunctionData.CppValidationImplName);
 
+			const FString& api = GetAPIString();
 			//Emit warning here if necessary
 			FUHTStringBuilder FunctionDeclaration;
-			ExportNativeFunctionHeader(FunctionDeclaration, OutReferenceGatherers.ForwardDeclarations, *FunctionDef, FunctionData, EExportFunctionType::Function, EExportFunctionHeaderStyle::Declaration, nullptr, *GetAPIString());
+			ExportNativeFunctionHeader(FunctionDeclaration, OutReferenceGatherers.ForwardDeclarations, *FunctionDef, FunctionData, EExportFunctionType::Function, EExportFunctionHeaderStyle::Declaration, nullptr, *api);
 			FunctionDeclaration.Log(TEXT(";\r\n"));
 
 			// Declare validation function if needed
@@ -4930,7 +4931,9 @@ void FNativeClassHeaderGenerator::ExportNativeFunctions(FOutputDevice& OutGenera
 
 				const TCHAR* Virtual = (!FunctionDef->HasAnyFunctionFlags(FUNC_Static) && !(FunctionData.FunctionExportFlags & FUNCEXPORT_Final)) ? TEXT("virtual") : TEXT("");
 				FStringOutputDevice ValidDecl;
-				ValidDecl.Logf(TEXT("\t%s bool %s(%s);\r\n"), Virtual, *FunctionData.CppValidationImplName, *ParameterList);
+				bool AddApi = !FunctionDef->GetOwnerClass()->HasAnyClassFlags(CLASS_RequiredAPI) && (FunctionData.FunctionExportFlags & FUNCEXPORT_RequiredAPI);
+
+				ValidDecl.Logf(TEXT("\t%s %s bool %s(%s);\r\n"), AddApi ? *api : TEXT(""), Virtual, *FunctionData.CppValidationImplName, *ParameterList);
 				FuncStringBuilders.AutogeneratedBlueprintFunctionDeclarations.Log(*ValidDecl);
 				if (!Validation.WasFound())
 				{
