Horde: Allow refs to point to specific exports within a bundle. Valuable when compacting a store, or rewriting old refs to use data in shared bundles.

diff --git a/Engine/Source/Programs/Horde/Horde.Agent/Commands/Bundles/CreateCommand.cs b/Engine/Source/Programs/Horde/Horde.Agent/Commands/Bundles/CreateCommand.cs
--- a/Engine/Source/Programs/Horde/Horde.Agent/Commands/Bundles/CreateCommand.cs
+++ b/Engine/Source/Programs/Horde/Horde.Agent/Commands/Bundles/CreateCommand.cs
@@ -2,6 +2,7 @@
 
 using System;
 using System.Buffers;
+using System.Globalization;
 using System.IO;
 using System.Threading;
 using System.Threading.Tasks;
@@ -71,25 +72,30 @@ namespace Horde.Agent.Commands.Bundles
 				throw new NotImplementedException();
 			}
 
-			public override async Task<BlobLocator> TryReadRefTargetAsync(RefName name, DateTime cacheTime = default, CancellationToken cancellationToken = default)
+			public override async Task<RefTarget?> TryReadRefTargetAsync(RefName name, DateTime cacheTime = default, CancellationToken cancellationToken = default)
 			{
 				FileReference file = GetRefFile(name);
-				if(!FileReference.Exists(file))
+				if (!FileReference.Exists(file))
 				{
-					return BlobLocator.Empty;
+					return null;
 				}
 
 				_logger.LogInformation("Reading {File}", file);
 				string text = await FileReference.ReadAllTextAsync(file, cancellationToken);
-				return new BlobLocator(text.Trim());
+
+				int hashIdx = text.IndexOf('#', StringComparison.Ordinal);
+				BlobLocator locator = new BlobLocator(text.Substring(0, hashIdx));
+				int exportIdx = Int32.Parse(text.Substring(hashIdx + 1), CultureInfo.InvariantCulture);
+
+				return new RefTarget(locator, exportIdx);
 			}
 
-			public override async Task WriteRefTargetAsync(RefName name, BlobLocator id, CancellationToken cancellationToken = default)
+			public override async Task WriteRefTargetAsync(RefName name, RefTarget target, CancellationToken cancellationToken = default)
 			{
 				FileReference file = GetRefFile(name);
 				DirectoryReference.CreateDirectory(file.Directory);
 				_logger.LogInformation("Writing {File}", file);
-				await FileReference.WriteAllTextAsync(file, id.ToString());
+				await FileReference.WriteAllTextAsync(file, $"{target.Locator}#{target.ExportIdx}");
 			}
 
 			#endregion
diff --git a/Engine/Source/Programs/Horde/Horde.Agent/Commands/Bundles/DumpCommand.cs b/Engine/Source/Programs/Horde/Horde.Agent/Commands/Bundles/DumpCommand.cs
--- a/Engine/Source/Programs/Horde/Horde.Agent/Commands/Bundles/DumpCommand.cs
+++ b/Engine/Source/Programs/Horde/Horde.Agent/Commands/Bundles/DumpCommand.cs
@@ -31,7 +31,7 @@ namespace Horde.Agent.Commands.Bundles
 			}
 			else
 			{
-				blob = await storage.ReadRefAsync(RefName);
+				blob = (await storage.ReadRefValueAsync(RefName)).Bundle;
 				logger.LogInformation("Summary for ref {RefId}", RefName);
 			}
 
diff --git a/Engine/Source/Programs/Horde/Horde.Build.Tests/BlobStoreTests.cs b/Engine/Source/Programs/Horde/Horde.Build.Tests/BlobStoreTests.cs
--- a/Engine/Source/Programs/Horde/Horde.Build.Tests/BlobStoreTests.cs
+++ b/Engine/Source/Programs/Horde/Horde.Build.Tests/BlobStoreTests.cs
@@ -112,15 +112,15 @@ namespace Horde.Build.Tests
 			for(int idx = 0; idx < 2; idx++)
 			{
 				RefName refName = new RefName("hello");
-				await store.WriteRefTargetAsync(refName, locator3);
-				BlobLocator refTargetId = await store.ReadRefTargetAsync(refName);
-				Assert.AreEqual(locator3, refTargetId);
+				await store.WriteRefTargetAsync(refName, new RefTarget(locator3, 0));
+				RefTarget refTarget = await store.ReadRefTargetAsync(refName);
+				Assert.AreEqual(locator3, refTarget.Locator);
 			}
 
 			RefName refName2 = new RefName("hello2");
 
-			BlobLocator refTargetId2 = await store.WriteRefAsync(refName2, CreateTestBundle(new ReadOnlySequence<byte>(input3), new BlobLocator[] { locator1, locator2 }));
-			Blob refTarget2 = await ReadBlobAsync(store, refTargetId2);
+			RefTarget refTargetId2 = await store.WriteRefValueAsync(refName2, CreateTestBundle(new ReadOnlySequence<byte>(input3), new BlobLocator[] { locator1, locator2 }));
+			Blob refTarget2 = await ReadBlobAsync(store, refTargetId2.Locator);
 			Assert.IsTrue(refTarget2.Data.Span.SequenceEqual(input3));
 			Assert.IsTrue(refTarget2.References.SequenceEqual(new BlobLocator[] { locator1, locator2 }));
 		}
diff --git a/Engine/Source/Programs/Horde/Horde.Build/Storage/IRefCollection.cs b/Engine/Source/Programs/Horde/Horde.Build/Storage/IRefCollection.cs
--- a/Engine/Source/Programs/Horde/Horde.Build/Storage/IRefCollection.cs
+++ b/Engine/Source/Programs/Horde/Horde.Build/Storage/IRefCollection.cs
@@ -25,17 +25,17 @@ namespace Horde.Build.Storage
 		/// <param name="name">The ref name</param>
 		/// <param name="cacheTime">Minimum age for a cached value to be returned</param>
 		/// <param name="cancellationToken">Cancellation token for the operation</param>
-		/// <returns>The ref value</returns>
-		Task<BlobLocator> TryReadRefTargetAsync(NamespaceId namespaceId, RefName name, DateTime cacheTime = default, CancellationToken cancellationToken = default);
+		/// <returns>The ref target</returns>
+		Task<RefTarget?> TryReadRefTargetAsync(NamespaceId namespaceId, RefName name, DateTime cacheTime = default, CancellationToken cancellationToken = default);
 
 		/// <summary>
 		/// Writes a ref to storage
 		/// </summary>
 		/// <param name="namespaceId">Namespace for the ref</param>
 		/// <param name="name">The ref name</param>
-		/// <param name="locator">Value for the ref</param>
+		/// <param name="target">Target for the ref</param>
 		/// <param name="cancellationToken">Cancellation token for the operation</param>
-		Task WriteRefTargetAsync(NamespaceId namespaceId, RefName name, BlobLocator locator, CancellationToken cancellationToken = default);
+		Task WriteRefTargetAsync(NamespaceId namespaceId, RefName name, RefTarget target, CancellationToken cancellationToken = default);
 
 		/// <summary>
 		/// Delete a ref from storage
diff --git a/Engine/Source/Programs/Horde/Horde.Build/Storage/RefCollection.cs b/Engine/Source/Programs/Horde/Horde.Build/Storage/RefCollection.cs
--- a/Engine/Source/Programs/Horde/Horde.Build/Storage/RefCollection.cs
+++ b/Engine/Source/Programs/Horde/Horde.Build/Storage/RefCollection.cs
@@ -45,33 +45,37 @@ namespace Horde.Build.Storage
 			[BsonElement("k")]
 			public RefName Name { get; set; }
 
-			[BsonElement("v")]
-			public BlobLocator Value { get; set; }
+			[BsonElement("l")]
+			public BlobLocator Locator { get; set; }
+
+			[BsonElement("e")]
+			public int ExportIdx { get; set; }
 
 			public RefDocument()
 			{
 				Name = RefName.Empty;
-				Value = BlobLocator.Empty;
+				Locator = BlobLocator.Empty;
 			}
 
-			public RefDocument(NamespaceId namespaceId, RefName name, BlobLocator value)
+			public RefDocument(NamespaceId namespaceId, RefName name, RefTarget target)
 			{
 				NamespaceId = namespaceId;
 				Name = name;
-				Value = value;
+				Locator = target.Locator;
+				ExportIdx = target.ExportIdx;
 			}
 		}
 
 		class CachedRefValue
 		{
 			public RefName Name { get; }
-			public BlobLocator Value { get; }
+			public RefTarget? Target { get; }
 			public DateTime Time { get; }
 
-			public CachedRefValue(RefName name, BlobLocator value)
+			public CachedRefValue(RefName name, RefTarget? target)
 			{
 				Name = name;
-				Value = value;
+				Target = target;
 				Time = DateTime.UtcNow;
 			}
 		}
@@ -90,13 +94,13 @@ namespace Horde.Build.Storage
 			_cache = cache;
 		}
 
-		CachedRefValue AddRefToCache(NamespaceId namespaceId, RefName name, BlobLocator value)
+		CachedRefValue AddRefToCache(NamespaceId namespaceId, RefName name, RefTarget? target)
 		{
-			CachedRefValue cacheValue = new CachedRefValue(name, value);
+			CachedRefValue cacheValue = new CachedRefValue(name, target);
 			using (ICacheEntry newEntry = _cache.CreateEntry(new RefKey(namespaceId, name)))
 			{
 				newEntry.Value = cacheValue;
-				newEntry.SetSize(name.Text.Length + value.Inner.Length);
+				newEntry.SetSize(name.Text.Length);
 				newEntry.SetSlidingExpiration(TimeSpan.FromMinutes(5.0));
 			}
 			return cacheValue;
@@ -106,28 +110,38 @@ namespace Horde.Build.Storage
 		public async Task DeleteRefAsync(NamespaceId namespaceId, RefName name, CancellationToken cancellationToken = default)
 		{
 			await _refs.DeleteOneAsync(x => x.NamespaceId == namespaceId && x.Name == name, cancellationToken);
-			AddRefToCache(namespaceId, name, BlobLocator.Empty);
+			AddRefToCache(namespaceId, name, null);
 		}
 
 		/// <inheritdoc/>
-		public async Task<BlobLocator> TryReadRefTargetAsync(NamespaceId namespaceId, RefName name, DateTime cacheTime = default, CancellationToken cancellationToken = default)
+		public async Task<RefTarget?> TryReadRefTargetAsync(NamespaceId namespaceId, RefName name, DateTime cacheTime = default, CancellationToken cancellationToken = default)
 		{
 			CachedRefValue entry;
 			if (!_cache.TryGetValue(name, out entry) || entry.Time < cacheTime)
 			{
 				RefDocument? refDocument = await _refs.Find(x => x.NamespaceId == namespaceId && x.Name == name).FirstOrDefaultAsync(cancellationToken);
-				BlobLocator value = (refDocument == null) ? BlobLocator.Empty : refDocument.Value;
-				entry = AddRefToCache(namespaceId, name, value);
+
+				RefTarget? target;
+				if (refDocument == null)
+				{
+					target = null;
+				}
+				else
+				{
+					target = new RefTarget(refDocument.Locator, refDocument.ExportIdx);
+				}
+
+				entry = AddRefToCache(namespaceId, name, target);
 			}
-			return entry.Value;
+			return entry.Target;
 		}
 
 		/// <inheritdoc/>
-		public async Task WriteRefTargetAsync(NamespaceId namespaceId, RefName name, BlobLocator locator, CancellationToken cancellationToken = default)
+		public async Task WriteRefTargetAsync(NamespaceId namespaceId, RefName name, RefTarget target, CancellationToken cancellationToken = default)
 		{
-			RefDocument refDocument = new RefDocument(namespaceId, name, locator);
+			RefDocument refDocument = new RefDocument(namespaceId, name, target);
 			await _refs.ReplaceOneAsync(x => x.NamespaceId == namespaceId && x.Name == name, refDocument, new ReplaceOptions { IsUpsert = true }, cancellationToken);
-			AddRefToCache(namespaceId, name, locator);
+			AddRefToCache(namespaceId, name, target);
 		}
 	}
 }
diff --git a/Engine/Source/Programs/Horde/Horde.Build/Storage/StorageService.cs b/Engine/Source/Programs/Horde/Horde.Build/Storage/StorageService.cs
--- a/Engine/Source/Programs/Horde/Horde.Build/Storage/StorageService.cs
+++ b/Engine/Source/Programs/Horde/Horde.Build/Storage/StorageService.cs
@@ -105,10 +105,10 @@ namespace Horde.Build.Storage
 			#region Refs
 
 			/// <inheritdoc/>
-			public override Task<BlobLocator> TryReadRefTargetAsync(RefName name, DateTime cacheTime = default, CancellationToken cancellationToken = default) => _refs.TryReadRefTargetAsync(NamespaceId, name, cacheTime, cancellationToken);
+			public override Task<RefTarget?> TryReadRefTargetAsync(RefName name, DateTime cacheTime = default, CancellationToken cancellationToken = default) => _refs.TryReadRefTargetAsync(NamespaceId, name, cacheTime, cancellationToken);
 
 			/// <inheritdoc/>
-			public override Task WriteRefTargetAsync(RefName name, BlobLocator locator, CancellationToken cancellationToken = default) => _refs.WriteRefTargetAsync(NamespaceId, name, locator, cancellationToken);
+			public override Task WriteRefTargetAsync(RefName name, RefTarget target, CancellationToken cancellationToken = default) => _refs.WriteRefTargetAsync(NamespaceId, name, target, cancellationToken);
 
 			/// <inheritdoc/>
 			public override Task DeleteRefAsync(RefName name, CancellationToken cancellationToken = default) => _refs.DeleteRefAsync(NamespaceId, name, cancellationToken);
diff --git a/Engine/Source/Programs/Shared/EpicGames.Horde.Tests/BundleTests.cs b/Engine/Source/Programs/Shared/EpicGames.Horde.Tests/BundleTests.cs
--- a/Engine/Source/Programs/Shared/EpicGames.Horde.Tests/BundleTests.cs
+++ b/Engine/Source/Programs/Shared/EpicGames.Horde.Tests/BundleTests.cs
@@ -172,9 +172,9 @@ namespace EpicGames.Horde.Tests
 			Assert.AreEqual(1, store.Blobs.Count);
 
 			// Check the ref
-			BundleRef bundle = await store.ReadRefAsync(refName);
-			Assert.AreEqual(0, bundle.Header.Imports.Count);
-			Assert.AreEqual(3, bundle.Header.Exports.Count);
+			RefValue refValue = await store.ReadRefValueAsync(refName);
+			Assert.AreEqual(0, refValue.Bundle.Header.Imports.Count);
+			Assert.AreEqual(3, refValue.Bundle.Header.Exports.Count);
 
 			// Create a new bundle and read it back in again
 			DirectoryNode newRoot = await store.ReadTreeAsync<DirectoryNode>(refName);
diff --git a/Engine/Source/Programs/Shared/EpicGames.Horde/Storage/Backends/InMemoryBlobStore.cs b/Engine/Source/Programs/Shared/EpicGames.Horde/Storage/Backends/InMemoryBlobStore.cs
--- a/Engine/Source/Programs/Shared/EpicGames.Horde/Storage/Backends/InMemoryBlobStore.cs
+++ b/Engine/Source/Programs/Shared/EpicGames.Horde/Storage/Backends/InMemoryBlobStore.cs
@@ -23,13 +23,13 @@ namespace EpicGames.Horde.Storage.Backends
 		/// <summary>
 		/// Map of ref name to ref data
 		/// </summary>
-		readonly ConcurrentDictionary<RefName, BlobLocator> _refs = new ConcurrentDictionary<RefName, BlobLocator>();
+		readonly ConcurrentDictionary<RefName, RefTarget> _refs = new ConcurrentDictionary<RefName, RefTarget>();
 
 		/// <inheritdoc cref="_blobs"/>
 		public IReadOnlyDictionary<BlobLocator, Bundle> Blobs => _blobs;
 
 		/// <inheritdoc cref="_refs"/>
-		public IReadOnlyDictionary<RefName, BlobLocator> Refs => _refs;
+		public IReadOnlyDictionary<RefName, RefTarget> Refs => _refs;
 
 		/// <summary>
 		/// Constructor
@@ -63,22 +63,17 @@ namespace EpicGames.Horde.Storage.Backends
 		public override Task DeleteRefAsync(RefName name, CancellationToken cancellationToken) => Task.FromResult(_refs.TryRemove(name, out _));
 
 		/// <inheritdoc/>
-		public override Task<BlobLocator> TryReadRefTargetAsync(RefName name, DateTime cacheTime = default, CancellationToken cancellationToken = default)
+		public override Task<RefTarget?> TryReadRefTargetAsync(RefName name, DateTime cacheTime = default, CancellationToken cancellationToken = default)
 		{
-			if (_refs.TryGetValue(name, out BlobLocator locator))
-			{
-				return Task.FromResult(locator);
-			}
-			else
-			{
-				return Task.FromResult(BlobLocator.Empty);
-			}
+			RefTarget? refTarget;
+			_refs.TryGetValue(name, out refTarget);
+			return Task.FromResult(refTarget);
 		}
 
 		/// <inheritdoc/>
-		public override Task WriteRefTargetAsync(RefName name, BlobLocator locator, CancellationToken cancellationToken = default)
+		public override Task WriteRefTargetAsync(RefName name, RefTarget target, CancellationToken cancellationToken = default)
 		{
-			_refs[name] = locator;
+			_refs[name] = target;
 			return Task.CompletedTask;
 		}
 
diff --git a/Engine/Source/Programs/Shared/EpicGames.Horde/Storage/IStorageClient.cs b/Engine/Source/Programs/Shared/EpicGames.Horde/Storage/IStorageClient.cs
--- a/Engine/Source/Programs/Shared/EpicGames.Horde/Storage/IStorageClient.cs
+++ b/Engine/Source/Programs/Shared/EpicGames.Horde/Storage/IStorageClient.cs
@@ -3,6 +3,7 @@
 using System;
 using System.Buffers;
 using System.Collections.Generic;
+using System.Diagnostics;
 using System.Threading;
 using System.Threading.Tasks;
 using EpicGames.Core;
@@ -32,22 +33,54 @@ namespace EpicGames.Horde.Storage
 	}
 
 	/// <summary>
-	/// Value for a ref
+	/// Target for a ref
 	/// </summary>
-	public class BundleRef : Bundle
+	[DebuggerDisplay("{Locator}#{Export}")]
+	public class RefTarget
 	{
 		/// <summary>
 		/// Location of the blob
 		/// </summary>
 		public BlobLocator Locator { get; }
 
+		/// <summary>
+		/// Index of the blob export
+		/// </summary>
+		public int ExportIdx { get; }
+
 		/// <summary>
 		/// Constructor
 		/// </summary>
-		public BundleRef(BlobLocator locator, Bundle bundle)
-			: base(bundle.Header, bundle.Payload)
+		public RefTarget(BlobLocator locator, int exportIdx)
 		{
 			Locator = locator;
+			ExportIdx = exportIdx;
+		}
+	}
+
+	/// <summary>
+	/// Value for a ref
+	/// </summary>
+	[DebuggerDisplay("{Target}")]
+	public class RefValue
+	{
+		/// <summary>
+		/// Target for the ref
+		/// </summary>
+		public RefTarget Target { get; }
+
+		/// <summary>
+		/// Bundle referenced as the target
+		/// </summary>
+		public Bundle Bundle { get; }
+
+		/// <summary>
+		/// Constructor
+		/// </summary>
+		public RefValue(RefTarget target, Bundle bundle)
+		{
+			Target = target;
+			Bundle = bundle;
 		}
 	}
 
@@ -107,7 +140,7 @@ namespace EpicGames.Horde.Storage
 		/// <param name="cacheTime">Minimum coherency for any cached value to be returned</param>
 		/// <param name="cancellationToken">Cancellation token for the operation</param>
 		/// <returns></returns>
-		Task<BundleRef?> TryReadRefAsync(RefName name, DateTime cacheTime = default, CancellationToken cancellationToken = default);
+		Task<RefValue?> TryReadRefValueAsync(RefName name, DateTime cacheTime = default, CancellationToken cancellationToken = default);
 
 		/// <summary>
 		/// Reads data for a ref from the store
@@ -116,26 +149,27 @@ namespace EpicGames.Horde.Storage
 		/// <param name="cacheTime">Minimum coherency for any cached value to be returned</param>
 		/// <param name="cancellationToken">Cancellation token for the operation</param>
 		/// <returns></returns>
-		Task<BlobLocator> TryReadRefTargetAsync(RefName name, DateTime cacheTime = default, CancellationToken cancellationToken = default);
+		Task<RefTarget?> TryReadRefTargetAsync(RefName name, DateTime cacheTime = default, CancellationToken cancellationToken = default);
 
 		/// <summary>
 		/// Writes a new ref to the store which points to a new blob
 		/// </summary>
 		/// <param name="name">Ref to write</param>
 		/// <param name="bundle">Bundle to write</param>
+		/// <param name="exportIdx">Index of the export to target for the ref</param>
 		/// <param name="prefix">Prefix for blob names.</param>
 		/// <param name="cancellationToken">Cancellation token for the operation</param>
 		/// <returns>Unique identifier for the blob</returns>
-		Task<BlobLocator> WriteRefAsync(RefName name, Bundle bundle, Utf8String prefix = default, CancellationToken cancellationToken = default);
+		Task<RefTarget> WriteRefValueAsync(RefName name, Bundle bundle, int exportIdx = 0, Utf8String prefix = default, CancellationToken cancellationToken = default);
 
 		/// <summary>
 		/// Writes a new ref to the store
 		/// </summary>
 		/// <param name="name">Ref to write</param>
-		/// <param name="locator">The target for the ref</param>
+		/// <param name="target">The target for the ref</param>
 		/// <param name="cancellationToken">Cancellation token for the operation</param>
 		/// <returns>Unique identifier for the blob</returns>
-		Task WriteRefTargetAsync(RefName name, BlobLocator locator, CancellationToken cancellationToken = default);
+		Task WriteRefTargetAsync(RefName name, RefTarget target, CancellationToken cancellationToken = default);
 
 		/// <summary>
 		/// Reads data for a ref from the store
@@ -191,16 +225,16 @@ namespace EpicGames.Horde.Storage
 			public Task DeleteRefAsync(RefName name, CancellationToken cancellationToken = default) => _inner.DeleteRefAsync(name, cancellationToken);
 
 			/// <inheritdoc/>
-			public Task<BundleRef?> TryReadRefAsync(RefName name, DateTime cacheTime = default, CancellationToken cancellationToken = default) => _inner.TryReadRefAsync(name, cacheTime, cancellationToken);
+			public Task<RefValue?> TryReadRefValueAsync(RefName name, DateTime cacheTime = default, CancellationToken cancellationToken = default) => _inner.TryReadRefValueAsync(name, cacheTime, cancellationToken);
 
 			/// <inheritdoc/>
-			public Task<BlobLocator> TryReadRefTargetAsync(RefName name, DateTime cacheTime = default, CancellationToken cancellationToken = default) => _inner.TryReadRefTargetAsync(name, cacheTime, cancellationToken);
+			public Task<RefTarget?> TryReadRefTargetAsync(RefName name, DateTime cacheTime = default, CancellationToken cancellationToken = default) => _inner.TryReadRefTargetAsync(name, cacheTime, cancellationToken);
 
 			/// <inheritdoc/>
-			public Task<BlobLocator> WriteRefAsync(RefName name, Bundle bundle, Utf8String prefix = default, CancellationToken cancellationToken = default) => _inner.WriteRefAsync(name, bundle, prefix, cancellationToken);
+			public Task<RefTarget> WriteRefValueAsync(RefName name, Bundle bundle, int exportIdx = 0, Utf8String prefix = default, CancellationToken cancellationToken = default) => _inner.WriteRefValueAsync(name, bundle, exportIdx, prefix, cancellationToken);
 
 			/// <inheritdoc/>
-			public Task WriteRefTargetAsync(RefName name, BlobLocator blobId, CancellationToken cancellationToken = default) => _inner.WriteRefTargetAsync(name, blobId, cancellationToken);
+			public Task WriteRefTargetAsync(RefName name, RefTarget target, CancellationToken cancellationToken = default) => _inner.WriteRefTargetAsync(name, target, cancellationToken);
 
 			#endregion
 		}
@@ -222,8 +256,8 @@ namespace EpicGames.Horde.Storage
 		/// <returns>True if the ref exists, false if it did not exist</returns>
 		public static async Task<bool> HasRefAsync(this IStorageClient store, RefName name, DateTime cacheTime = default, CancellationToken cancellationToken = default)
 		{
-			BlobLocator locator = await store.TryReadRefTargetAsync(name, cacheTime, cancellationToken);
-			return locator.IsValid();
+			RefTarget? locator = await store.TryReadRefTargetAsync(name, cacheTime, cancellationToken);
+			return locator != null;
 		}
 
 		/// <summary>
@@ -251,27 +285,27 @@ namespace EpicGames.Horde.Storage
 		}
 
 		/// <summary>
-		/// Reads a ref from the store, throwing an exception if it does not exist
+		/// Attempts to reads a ref from the store
 		/// </summary>
 		/// <param name="store">The store instance to read from</param>
-		/// <param name="name">Id for the ref</param>
+		/// <param name="name">Name of the ref</param>
 		/// <param name="maxAge">Maximum age of any cached ref</param>
 		/// <param name="cancellationToken">Cancellation token for the operation</param>
-		/// <returns>The blob instance</returns>
-		public static Task<BundleRef?> TryReadRefAsync(this IStorageClient store, RefName name, TimeSpan maxAge, CancellationToken cancellationToken = default)
+		/// <returns>The ref value</returns>
+		public static Task<RefValue?> TryReadRefValueAsync(this IStorageClient store, RefName name, TimeSpan maxAge, CancellationToken cancellationToken = default)
 		{
-			return store.TryReadRefAsync(name, DateTime.UtcNow - maxAge, cancellationToken);
+			return store.TryReadRefValueAsync(name, DateTime.UtcNow - maxAge, cancellationToken);
 		}
 
 		/// <summary>
-		/// Reads a ref from the store, throwing an exception if it does not exist
+		/// Attempts to reads a ref from the store
 		/// </summary>
 		/// <param name="store">The store instance to read from</param>
 		/// <param name="name">Id for the ref</param>
 		/// <param name="maxAge">Maximum age of any cached ref</param>
 		/// <param name="cancellationToken">Cancellation token for the operation</param>
-		/// <returns>The blob instance</returns>
-		public static Task<BlobLocator> TryReadRefTargetAsync(this IStorageClient store, RefName name, TimeSpan maxAge, CancellationToken cancellationToken = default)
+		/// <returns>The ref target</returns>
+		public static Task<RefTarget?> TryReadRefTargetAsync(this IStorageClient store, RefName name, TimeSpan maxAge, CancellationToken cancellationToken = default)
 		{
 			return store.TryReadRefTargetAsync(name, DateTime.UtcNow - maxAge, cancellationToken);
 		}
@@ -284,14 +318,14 @@ namespace EpicGames.Horde.Storage
 		/// <param name="cacheTime">Minimum coherency of any cached result</param>
 		/// <param name="cancellationToken">Cancellation token for the operation</param>
 		/// <returns>The blob instance</returns>
-		public static async Task<BundleRef> ReadRefAsync(this IStorageClient store, RefName name, DateTime cacheTime = default, CancellationToken cancellationToken = default)
+		public static async Task<RefValue> ReadRefValueAsync(this IStorageClient store, RefName name, DateTime cacheTime = default, CancellationToken cancellationToken = default)
 		{
-			BundleRef? bundleRef = await store.TryReadRefAsync(name, cacheTime, cancellationToken);
-			if (bundleRef == null)
+			RefValue? refValue = await store.TryReadRefValueAsync(name, cacheTime, cancellationToken);
+			if (refValue == null)
 			{
 				throw new RefNameNotFoundException(name);
 			}
-			return bundleRef;
+			return refValue;
 		}
 
 		/// <summary>
@@ -301,10 +335,10 @@ namespace EpicGames.Horde.Storage
 		/// <param name="name">Id for the ref</param>
 		/// <param name="maxAge">Maximum age for any cached result</param>
 		/// <param name="cancellationToken">Cancellation token for the operation</param>
-		/// <returns>The blob instance</returns>
-		public static Task<BundleRef> ReadRefAsync(this IStorageClient store, RefName name, TimeSpan maxAge, CancellationToken cancellationToken = default)
+		/// <returns>The ref value</returns>
+		public static Task<RefValue> ReadRefValueAsync(this IStorageClient store, RefName name, TimeSpan maxAge, CancellationToken cancellationToken = default)
 		{
-			return ReadRefAsync(store, name, DateTime.UtcNow - maxAge, cancellationToken);
+			return ReadRefValueAsync(store, name, DateTime.UtcNow - maxAge, cancellationToken);
 		}
 
 		/// <summary>
@@ -314,15 +348,15 @@ namespace EpicGames.Horde.Storage
 		/// <param name="name">Id for the ref</param>
 		/// <param name="cacheTime">Minimum coherency of any cached result</param>
 		/// <param name="cancellationToken">Cancellation token for the operation</param>
-		/// <returns>The blob instance</returns>
-		public static async Task<BlobLocator> ReadRefTargetAsync(this IStorageClient store, RefName name, DateTime cacheTime = default, CancellationToken cancellationToken = default)
+		/// <returns>The ref target</returns>
+		public static async Task<RefTarget> ReadRefTargetAsync(this IStorageClient store, RefName name, DateTime cacheTime = default, CancellationToken cancellationToken = default)
 		{
-			BlobLocator locator = await store.TryReadRefTargetAsync(name, cacheTime, cancellationToken);
-			if (!locator.IsValid())
+			RefTarget? refTarget = await store.TryReadRefTargetAsync(name, cacheTime, cancellationToken);
+			if (refTarget == null)
 			{
 				throw new RefNameNotFoundException(name);
 			}
-			return locator;
+			return refTarget;
 		}
 
 		/// <summary>
@@ -333,7 +367,7 @@ namespace EpicGames.Horde.Storage
 		/// <param name="maxAge">Maximum age for any cached result</param>
 		/// <param name="cancellationToken">Cancellation token for the operation</param>
 		/// <returns>The blob instance</returns>
-		public static Task<BlobLocator> ReadRefTargetAsync(this IStorageClient store, RefName name, TimeSpan maxAge, CancellationToken cancellationToken = default)
+		public static Task<RefTarget> ReadRefTargetAsync(this IStorageClient store, RefName name, TimeSpan maxAge, CancellationToken cancellationToken = default)
 		{
 			return ReadRefTargetAsync(store, name, DateTime.UtcNow - maxAge, cancellationToken);
 		}
diff --git a/Engine/Source/Programs/Shared/EpicGames.Horde/Storage/StorageClientBase.cs b/Engine/Source/Programs/Shared/EpicGames.Horde/Storage/StorageClientBase.cs
--- a/Engine/Source/Programs/Shared/EpicGames.Horde/Storage/StorageClientBase.cs
+++ b/Engine/Source/Programs/Shared/EpicGames.Horde/Storage/StorageClientBase.cs
@@ -60,32 +60,30 @@ namespace EpicGames.Horde.Storage
 		public abstract Task DeleteRefAsync(RefName name, CancellationToken cancellationToken = default);
 
 		/// <inheritdoc/>
-		public virtual async Task<BundleRef?> TryReadRefAsync(RefName name, DateTime cacheTime = default, CancellationToken cancellationToken = default)
+		public virtual async Task<RefValue?> TryReadRefValueAsync(RefName name, DateTime cacheTime = default, CancellationToken cancellationToken = default)
 		{
-			BlobLocator locator = await TryReadRefTargetAsync(name, cacheTime, cancellationToken);
-			if (locator.IsValid())
-			{
-				return new BundleRef(locator, await ReadBundleAsync(locator, cancellationToken));
-			}
-			else
+			RefTarget? refTarget = await TryReadRefTargetAsync(name, cacheTime, cancellationToken);
+			if (refTarget == null)
 			{
 				return null;
 			}
+			return new RefValue(refTarget, await ReadBundleAsync(refTarget.Locator, cancellationToken));
 		}
 
 		/// <inheritdoc/>
-		public abstract Task<BlobLocator> TryReadRefTargetAsync(RefName name, DateTime cacheTime = default, CancellationToken cancellationToken = default);
+		public abstract Task<RefTarget?> TryReadRefTargetAsync(RefName name, DateTime cacheTime = default, CancellationToken cancellationToken = default);
 
 		/// <inheritdoc/>
-		public virtual async Task<BlobLocator> WriteRefAsync(RefName name, Bundle bundle, Utf8String prefix = default, CancellationToken cancellationToken = default)
+		public virtual async Task<RefTarget> WriteRefValueAsync(RefName name, Bundle bundle, int exportIdx = 0, Utf8String prefix = default, CancellationToken cancellationToken = default)
 		{
 			BlobLocator locator = await WriteBundleAsync(bundle, prefix, cancellationToken);
-			await WriteRefTargetAsync(name, locator, cancellationToken);
-			return locator;
+			RefTarget target = new RefTarget(locator, exportIdx);
+			await WriteRefTargetAsync(name, target, cancellationToken);
+			return target;
 		}
 
 		/// <inheritdoc/>
-		public abstract Task WriteRefTargetAsync(RefName name, BlobLocator locator, CancellationToken cancellationToken = default);
+		public abstract Task WriteRefTargetAsync(RefName name, RefTarget target, CancellationToken cancellationToken = default);
 
 		#endregion
 	}
diff --git a/Engine/Source/Programs/Shared/EpicGames.Horde/Storage/TreeStore.cs b/Engine/Source/Programs/Shared/EpicGames.Horde/Storage/TreeStore.cs
--- a/Engine/Source/Programs/Shared/EpicGames.Horde/Storage/TreeStore.cs
+++ b/Engine/Source/Programs/Shared/EpicGames.Horde/Storage/TreeStore.cs
@@ -395,7 +395,7 @@ namespace EpicGames.Horde.Storage
 
 				// Write a reference to the blob containing this node
 				ImportedNodeState importedState = (ImportedNodeState)last.State;
-				await _owner._blobStore.WriteRefTargetAsync(refName, importedState.BundleInfo.Locator, cancellationToken);
+				await _owner._blobStore.WriteRefTargetAsync(refName, new RefTarget(importedState.BundleInfo.Locator, 0), cancellationToken);
 			}
 
 			async Task<NodeInfo> FlushInternalAsync(CancellationToken cancellationToken)
@@ -699,8 +699,8 @@ namespace EpicGames.Horde.Storage
 		public async Task<ITreeBlob?> TryReadTreeAsync(RefName name, TimeSpan maxAge = default, CancellationToken cancellationToken = default)
 		{
 			// Read the blob referenced by this ref
-			BundleRef? bundleRef = await _blobStore.TryReadRefAsync(name, maxAge, cancellationToken);
-			if (bundleRef == null)
+			RefValue? refValue = await _blobStore.TryReadRefValueAsync(name, maxAge, cancellationToken);
+			if (refValue == null)
 			{
 				return null;
 			}
@@ -709,8 +709,8 @@ namespace EpicGames.Horde.Storage
 			BundleContext context = new BundleContext();
 
 			// Add the new bundle to it
-			BundleInfo bundleInfo = context.FindOrAddBundle(bundleRef.Locator, bundleRef.Header.Exports.Count);
-			MountBundle(bundleInfo, bundleRef);
+			BundleInfo bundleInfo = context.FindOrAddBundle(refValue.Target.Locator, refValue.Bundle.Header.Exports.Count);
+			MountBundle(bundleInfo, refValue.Bundle);
 
 			// Return the last node from the bundle as the root
 			return await GetNodeAsync(bundleInfo.Exports[^1]!, cancellationToken);
