Fixed issue in C# UHT where if an old style enum was returned from the "Implementation" method, the zero value would not be wrapped in the TEnumAsByte.

diff --git a/Engine/Source/Programs/Shared/EpicGames.UHT/Exporters/CodeGen/UhtHeaderCodeGenerator.cs b/Engine/Source/Programs/Shared/EpicGames.UHT/Exporters/CodeGen/UhtHeaderCodeGenerator.cs
--- a/Engine/Source/Programs/Shared/EpicGames.UHT/Exporters/CodeGen/UhtHeaderCodeGenerator.cs
+++ b/Engine/Source/Programs/Shared/EpicGames.UHT/Exporters/CodeGen/UhtHeaderCodeGenerator.cs
@@ -253,6 +253,10 @@ namespace EpicGames.UHT.Exporters.CodeGen
 							{
 								builder.Append(" return TEnumAsByte<").Append(byteProperty.Enum.CppType).Append(">(").AppendNullConstructorArg(returnProperty, false).Append("); ");
 							}
+							else if (returnProperty is UhtEnumProperty enumProperty && enumProperty.Enum.CppForm != UhtEnumCppForm.EnumClass)
+							{
+								builder.Append(" return TEnumAsByte<").Append(enumProperty.Enum.CppType).Append(">(").AppendNullConstructorArg(returnProperty, false).Append("); ");
+							}
 							else
 							{
 								builder.Append(" return ").AppendNullConstructorArg(returnProperty, false).Append("; ");
