Horde: Add Stream wrappers around IMemoryReader and IMemoryWriter, and fix ByteArrayBuilder.GetMemory() always returning the minimum buffer size.

diff --git a/Engine/Source/Programs/Shared/EpicGames.BuildGraph/BgCompiler.cs b/Engine/Source/Programs/Shared/EpicGames.BuildGraph/BgCompiler.cs
--- a/Engine/Source/Programs/Shared/EpicGames.BuildGraph/BgCompiler.cs
+++ b/Engine/Source/Programs/Shared/EpicGames.BuildGraph/BgCompiler.cs
@@ -155,7 +155,7 @@ namespace EpicGames.BuildGraph
 			/// <inheritdoc/>
 			public override void WriteOpcode(BgOpcode opcode)
 			{
-				_builder.WriteByte((byte)opcode);
+				_builder.WriteUInt8((byte)opcode);
 			}
 
 			/// <inheritdoc/>
@@ -177,7 +177,7 @@ namespace EpicGames.BuildGraph
 				int textLength = Encoding.UTF8.GetByteCount(str);
 				int lengthLength = VarInt.MeasureUnsigned(textLength);
 
-				Span<byte> buffer = _builder.GetWritableSpan(lengthLength + textLength);
+				Span<byte> buffer = _builder.GetSpanAndAdvance(lengthLength + textLength);
 				VarInt.WriteUnsigned(buffer, textLength);
 				Encoding.UTF8.GetBytes(str, buffer.Slice(lengthLength));
 			}
diff --git a/Engine/Source/Programs/Shared/EpicGames.Core/ByteArrayBuilder.cs b/Engine/Source/Programs/Shared/EpicGames.Core/ByteArrayBuilder.cs
--- a/Engine/Source/Programs/Shared/EpicGames.Core/ByteArrayBuilder.cs
+++ b/Engine/Source/Programs/Shared/EpicGames.Core/ByteArrayBuilder.cs
@@ -57,36 +57,6 @@ namespace EpicGames.Core
 			_chunkSize = chunkSize;
 		}
 
-		/// <summary>
-		/// Appends a single byte to the sequence
-		/// </summary>
-		/// <param name="value">Byte to add</param>
-		public void WriteByte(byte value)
-		{
-			Span<byte> target = GetWritableSpan(1);
-			target[0] = value;
-		}
-
-		/// <summary>
-		/// Appends a span of bytes to the buffer
-		/// </summary>
-		/// <param name="span">Bytes to add</param>
-		public void Append(ReadOnlySpan<byte> span)
-		{
-			Span<byte> target = GetWritableSpan(span.Length);
-			span.CopyTo(target);
-		}
-
-		/// <summary>
-		/// Appends a sequence of bytes to the buffer
-		/// </summary>
-		/// <param name="sequence">Sequence to append</param>
-		public void Append(ReadOnlySequence<byte> sequence)
-		{
-			Span<byte> target = GetWritableSpan((int)sequence.Length);
-			sequence.CopyTo(target);
-		}
-
 		/// <inheritdoc/>
 		public Memory<byte> GetMemory(int minSize)
 		{
@@ -95,33 +65,19 @@ namespace EpicGames.Core
 				_currentChunk = new Chunk(_currentChunk.RunningIndex + _currentChunk.Length, Math.Max(minSize, _chunkSize));
 				_chunks.Add(_currentChunk);
 			}
-			return _currentChunk.Data.AsMemory(_currentChunk.Length, minSize);
+			return _currentChunk.Data.AsMemory(_currentChunk.Length);
 		}
 
 		/// <inheritdoc/>
 		public void Advance(int length)
 		{
-			_currentChunk.Length += length;
-			Length += length;
-		}
-
-		/// <summary>
-		/// Gets a span of bytes to write to, and increases the length of the buffer
-		/// </summary>
-		/// <param name="length">Length of the buffer</param>
-		/// <returns>Writable span of bytes</returns>
-		public Span<byte> GetWritableSpan(int length)
-		{
-			if (_currentChunk.Length + length > _currentChunk.Data.Length)
+			if (length < 0 || _currentChunk.Length + length > _currentChunk.Data.Length)
 			{
-				_currentChunk = new Chunk(_currentChunk.RunningIndex + _currentChunk.Length, Math.Max(length, _chunkSize));
-				_chunks.Add(_currentChunk);
+				throw new ArgumentException("Length exceeds possible size of data written to output buffer", nameof(length));
 			}
 
-			Span<byte> span = _currentChunk.Data.AsSpan(_currentChunk.Length, length);
 			_currentChunk.Length += length;
 			Length += length;
-			return span;
 		}
 
 		/// <summary>
diff --git a/Engine/Source/Programs/Shared/EpicGames.Core/MemoryReader.cs b/Engine/Source/Programs/Shared/EpicGames.Core/MemoryReader.cs
--- a/Engine/Source/Programs/Shared/EpicGames.Core/MemoryReader.cs
+++ b/Engine/Source/Programs/Shared/EpicGames.Core/MemoryReader.cs
@@ -15,7 +15,7 @@ namespace EpicGames.Core
 		/// </summary>
 		/// <param name="minSize">Minimum size of the returned data</param>
 		/// <returns>Memory of at least the given size</returns>
-		ReadOnlyMemory<byte> GetMemory(int minSize);
+		ReadOnlyMemory<byte> GetMemory(int minSize = 1);
 
 		/// <summary>
 		/// Updates the current position within the input buffer
diff --git a/Engine/Source/Programs/Shared/EpicGames.Core/MemoryReaderStream.cs b/Engine/Source/Programs/Shared/EpicGames.Core/MemoryReaderStream.cs
new file mode 100644
--- /dev/null
+++ b/Engine/Source/Programs/Shared/EpicGames.Core/MemoryReaderStream.cs
@@ -0,0 +1,75 @@
+// Copyright Epic Games, Inc. All Rights Reserved.
+
+using System;
+using System.IO;
+
+namespace EpicGames.Core
+{
+	/// <summary>
+	/// Implementation of <see cref="Stream"/> which reads data from a <see cref="IMemoryReader"/>
+	/// </summary>
+	public class MemoryReaderStream : Stream
+	{
+		readonly IMemoryReader _reader;
+
+		/// <summary>
+		/// Constructor
+		/// </summary>
+		public MemoryReaderStream(IMemoryReader reader)
+		{
+			_reader = reader;
+		}
+
+		/// <inheritdoc/>
+		public override bool CanRead => true;
+
+		/// <inheritdoc/>
+		public override bool CanSeek => false;
+
+		/// <inheritdoc/>
+		public override bool CanWrite => false;
+
+		/// <inheritdoc/>
+		public override long Length => throw new NotSupportedException();
+
+		/// <inheritdoc/>
+		public override long Position { get => throw new NotSupportedException(); set => throw new NotSupportedException(); }
+
+		/// <inheritdoc/>
+		public override void Flush() { }
+
+		/// <inheritdoc/>
+		public override int Read(Span<byte> buffer)
+		{
+			int sizeRead = 0;
+			while (buffer.Length > 0)
+			{
+				ReadOnlyMemory<byte> memory = _reader.GetMemory();
+				if (memory.Length == 0)
+				{
+					break;
+				}
+
+				int copyLength = Math.Min(buffer.Length, memory.Length);
+				memory.Slice(0, copyLength).Span.CopyTo(buffer);
+				_reader.Advance(copyLength);
+				sizeRead += copyLength;
+
+				buffer = buffer.Slice(copyLength);
+			}
+			return sizeRead;
+		}
+
+		/// <inheritdoc/>
+		public override int Read(byte[] buffer, int offset, int count) => Read(buffer.AsSpan(offset, count));
+
+		/// <inheritdoc/>
+		public override long Seek(long offset, SeekOrigin origin) => throw new NotSupportedException();
+
+		/// <inheritdoc/>
+		public override void SetLength(long value) => throw new NotSupportedException();
+
+		/// <inheritdoc/>
+		public override void Write(byte[] buffer, int offset, int count) => throw new NotSupportedException();
+	}
+}
diff --git a/Engine/Source/Programs/Shared/EpicGames.Core/MemoryWriter.cs b/Engine/Source/Programs/Shared/EpicGames.Core/MemoryWriter.cs
--- a/Engine/Source/Programs/Shared/EpicGames.Core/MemoryWriter.cs
+++ b/Engine/Source/Programs/Shared/EpicGames.Core/MemoryWriter.cs
@@ -4,6 +4,7 @@ using System;
 using System.Buffers;
 using System.Buffers.Binary;
 using System.Collections.Generic;
+using System.IO;
 using System.Text;
 
 namespace EpicGames.Core
@@ -22,8 +23,8 @@ namespace EpicGames.Core
 		/// Gets a block of memory with at least the given size
 		/// </summary>
 		/// <param name="minSize">Minimum size of the returned data</param>
-		/// <returns>Memory of at least the given size</returns>
-		Memory<byte> GetMemory(int minSize);
+		/// <returns>Memory of at least the given size. Note that the returned buffer may be larger than the requested size.</returns>
+		Memory<byte> GetMemory(int minSize = 1);
 
 		/// <summary>
 		/// Updates the current position within the input buffer
@@ -89,6 +90,25 @@ namespace EpicGames.Core
 		/// <param name="minSize">Minimum size for the returned span</param>
 		public static Span<byte> GetSpan(this IMemoryWriter writer, int minSize) => writer.GetMemory(minSize).Span;
 
+		/// <summary>
+		/// Writes a byte to memory
+		/// </summary>
+		/// <param name="writer">Writer to serialize to</param>
+		/// <param name="length">Length of the required span</param>
+		public static Memory<byte> GetMemoryAndAdvance(this IMemoryWriter writer, int length)
+		{
+			Memory<byte> memory = writer.GetMemory(length);
+			writer.Advance(length);
+			return memory.Slice(0, length); // Returned memory may be larger than requested
+		}
+
+		/// <summary>
+		/// Writes a byte to memory
+		/// </summary>
+		/// <param name="writer">Writer to serialize to</param>
+		/// <param name="length">Length of the required span</param>
+		public static Span<byte> GetSpanAndAdvance(this IMemoryWriter writer, int length) => GetMemoryAndAdvance(writer, length).Span;
+
 		/// <summary>
 		/// Writes a boolean to memory
 		/// </summary>
@@ -219,14 +239,25 @@ namespace EpicGames.Core
 		}
 
 		/// <summary>
-		/// Appends a sequence of bytes to the buffer
+		/// Appends a span of bytes to the buffer, without a length field.
+		/// </summary>
+		/// <param name="writer">Writer to serialize to</param>
+		/// <param name="span">Bytes to write</param>
+		public static void WriteFixedLengthBytes(this IMemoryWriter writer, ReadOnlySpan<byte> span)
+		{
+			Span<byte> target = GetSpanAndAdvance(writer, span.Length);
+			span.CopyTo(target);
+		}
+
+		/// <summary>
+		/// Appends a sequence of bytes to the buffer, without a length field.
 		/// </summary>
 		/// <param name="writer">Writer to serialize to</param>
 		/// <param name="sequence">Sequence to append</param>
-		public static void WriteSequence(this IMemoryWriter writer, ReadOnlySequence<byte> sequence)
+		public static void WriteFixedLengthBytes(this IMemoryWriter writer, ReadOnlySequence<byte> sequence)
 		{
-			Memory<byte> memory = writer.GetMemory((int)sequence.Length);
-			sequence.CopyTo(memory.Span);
+			Span<byte> span = GetSpanAndAdvance(writer, (int)sequence.Length);
+			sequence.CopyTo(span);
 		}
 
 		/// <summary>
@@ -237,10 +268,9 @@ namespace EpicGames.Core
 		public static void WriteVariableLengthBytes(this IMemoryWriter writer, ReadOnlySpan<byte> bytes)
 		{
 			int lengthBytes = VarInt.MeasureUnsigned(bytes.Length);
-			Span<byte> span = GetSpan(writer, lengthBytes + bytes.Length);
+			Span<byte> span = GetSpanAndAdvance(writer, lengthBytes + bytes.Length);
 			VarInt.WriteUnsigned(span, bytes.Length);
 			bytes.CopyTo(span[lengthBytes..]);
-			writer.Advance(lengthBytes + bytes.Length);
 		}
 
 		/// <summary>
@@ -254,18 +284,6 @@ namespace EpicGames.Core
 			writer.WriteFixedLengthBytes(bytes);
 		}
 
-		/// <summary>
-		/// Write a fixed-length sequence of bytes to the buffer
-		/// </summary>
-		/// <param name="writer">Writer to serialize to</param>
-		/// <param name="bytes">The bytes to write</param>
-		public static void WriteFixedLengthBytes(this IMemoryWriter writer, ReadOnlySpan<byte> bytes)
-		{
-			Span<byte> span = GetSpan(writer, bytes.Length);
-			bytes.CopyTo(span);
-			writer.Advance(bytes.Length);
-		}
-
 		/// <summary>
 		/// Writes a variable length array
 		/// </summary>
diff --git a/Engine/Source/Programs/Shared/EpicGames.Core/MemoryWriterStream.cs b/Engine/Source/Programs/Shared/EpicGames.Core/MemoryWriterStream.cs
new file mode 100644
--- /dev/null
+++ b/Engine/Source/Programs/Shared/EpicGames.Core/MemoryWriterStream.cs
@@ -0,0 +1,70 @@
+// Copyright Epic Games, Inc. All Rights Reserved.
+
+using System;
+using System.IO;
+
+namespace EpicGames.Core
+{
+	/// <summary>
+	/// Implementation of <see cref="Stream"/> which writes data to a <see cref="IMemoryWriter"/>
+	/// </summary>
+	public class MemoryWriterStream : Stream
+	{
+		readonly IMemoryWriter _writer;
+		long _length;
+
+		/// <summary>
+		/// Constructor
+		/// </summary>
+		public MemoryWriterStream(IMemoryWriter writer)
+		{
+			_writer = writer;
+		}
+
+		/// <inheritdoc/>
+		public override bool CanRead => false;
+
+		/// <inheritdoc/>
+		public override bool CanSeek => false;
+
+		/// <inheritdoc/>
+		public override bool CanWrite => true;
+
+		/// <inheritdoc/>
+		public override long Length => _length;
+
+		/// <inheritdoc/>
+		public override long Position { get => _length; set => throw new NotSupportedException(); }
+
+		/// <inheritdoc/>
+		public override void Flush() { }
+
+		/// <inheritdoc/>
+		public override int Read(byte[] buffer, int offset, int count) => throw new NotSupportedException();
+
+		/// <inheritdoc/>
+		public override long Seek(long offset, SeekOrigin origin) => throw new NotSupportedException();
+
+		/// <inheritdoc/>
+		public override void SetLength(long value) => throw new NotSupportedException();
+
+		/// <inheritdoc/>
+		public override void Write(ReadOnlySpan<byte> buffer)
+		{
+			while (buffer.Length > 0)
+			{
+				Memory<byte> current = _writer.GetMemory();
+
+				int copyLength = Math.Min(current.Length, buffer.Length);
+				buffer.Slice(0, copyLength).CopyTo(current.Span);
+				_writer.Advance(copyLength);
+				_length += copyLength;
+
+				buffer = buffer.Slice(copyLength);
+			}
+		}
+
+		/// <inheritdoc/>
+		public override void Write(byte[] buffer, int offset, int count) => Write(buffer.AsSpan(offset, count));
+	}
+}
diff --git a/Engine/Source/Programs/Shared/EpicGames.Horde/Storage/TreeStore.cs b/Engine/Source/Programs/Shared/EpicGames.Horde/Storage/TreeStore.cs
--- a/Engine/Source/Programs/Shared/EpicGames.Horde/Storage/TreeStore.cs
+++ b/Engine/Source/Programs/Shared/EpicGames.Horde/Storage/TreeStore.cs
@@ -648,8 +648,7 @@ namespace EpicGames.Horde.Storage
 				{
 					case BundleCompressionFormat.None:
 						{
-							builder.Append(inputSpan);
-							builder.Advance(inputSpan.Length);
+							builder.WriteFixedLengthBytes(inputSpan);
 							return inputSpan.Length;
 						}
 					case BundleCompressionFormat.LZ4:
