Horde: Encode data into bundles asynchonously while prefetching the next block of data.

diff --git a/Engine/Source/Programs/Shared/EpicGames.Horde/Storage/Nodes/DirectoryNode.cs b/Engine/Source/Programs/Shared/EpicGames.Horde/Storage/Nodes/DirectoryNode.cs
--- a/Engine/Source/Programs/Shared/EpicGames.Horde/Storage/Nodes/DirectoryNode.cs
+++ b/Engine/Source/Programs/Shared/EpicGames.Horde/Storage/Nodes/DirectoryNode.cs
@@ -124,15 +124,35 @@ namespace EpicGames.Horde.Storage.Nodes
 		/// <param name="cancellationToken">Cancellation token for the operation</param>
 		public async Task AppendAsync(Stream stream, ChunkingOptions options, ITreeWriter writer, CancellationToken cancellationToken)
 		{
-			byte[] buffer = new byte[4 * 1024];
+			const int BufferLength = 32 * 1024;
+
+			using IMemoryOwner<byte> owner = MemoryPool<byte>.Shared.Rent(BufferLength * 2);
+			Memory<byte> buffer = owner.Memory;
+
+			int readBufferOffset = 0;
+			Memory<byte> appendBuffer = Memory<byte>.Empty;
 			for (; ; )
 			{
-				int numBytes = await stream.ReadAsync(buffer, cancellationToken);
+				// Start a read into memory
+				Memory<byte> readBuffer = buffer.Slice(readBufferOffset, BufferLength);
+				Task<int> readTask = Task.Run(async () => await stream.ReadAsync(readBuffer, cancellationToken), cancellationToken);
+
+				// In the meantime, append the last data that was read to the tree
+				if (appendBuffer.Length > 0)
+				{
+					await AppendAsync(appendBuffer, options, writer, cancellationToken);
+				}
+
+				// Wait for the read to finish
+				int numBytes = await readTask;
 				if (numBytes == 0)
 				{
 					break;
 				}
-				await AppendAsync(buffer.AsMemory(0, numBytes), options, writer, cancellationToken);
+
+				// Switch the buffers around
+				appendBuffer = readBuffer.Slice(0, numBytes);
+				readBufferOffset ^= BufferLength;
 			}
 		}
 
