Horde: Add support for Gzip compression in bundles.

diff --git a/Engine/Source/Programs/Shared/EpicGames.Horde/Storage/Bundle.cs b/Engine/Source/Programs/Shared/EpicGames.Horde/Storage/Bundle.cs
--- a/Engine/Source/Programs/Shared/EpicGames.Horde/Storage/Bundle.cs
+++ b/Engine/Source/Programs/Shared/EpicGames.Horde/Storage/Bundle.cs
@@ -5,14 +5,11 @@ using System;
 using System.Buffers;
 using System.Buffers.Binary;
 using System.Collections.Generic;
-using System.Diagnostics.CodeAnalysis;
 using System.IO;
 using System.Linq;
-using System.Reflection.Metadata;
 using System.Text;
 using System.Threading;
 using System.Threading.Tasks;
-using static System.Net.WebRequestMethods;
 
 namespace EpicGames.Horde.Storage
 {
@@ -121,6 +118,11 @@ namespace EpicGames.Horde.Storage
 		/// LZ4 compression
 		/// </summary>
 		LZ4 = 1,
+
+		/// <summary>
+		/// Gzip compression
+		/// </summary>
+		Gzip = 2,
 	}
 
 	/// <summary>
diff --git a/Engine/Source/Programs/Shared/EpicGames.Horde/Storage/TreeStore.cs b/Engine/Source/Programs/Shared/EpicGames.Horde/Storage/TreeStore.cs
--- a/Engine/Source/Programs/Shared/EpicGames.Horde/Storage/TreeStore.cs
+++ b/Engine/Source/Programs/Shared/EpicGames.Horde/Storage/TreeStore.cs
@@ -6,6 +6,7 @@ using System.Collections.Concurrent;
 using System.Collections.Generic;
 using System.Diagnostics;
 using System.IO;
+using System.IO.Compression;
 using System.Linq;
 using System.Threading;
 using System.Threading.Tasks;
@@ -636,31 +637,38 @@ namespace EpicGames.Horde.Storage
 			{
 				if (inputData.Length > 0)
 				{
-					int encodedLength = Encode(format, inputData.Span, builder);
+					int encodedLength = Encode(format, inputData, builder);
 					Debug.Assert(encodedLength >= 0);
 					packets.Add(new BundlePacket(encodedLength, inputData.Length));
 				}
 			}
 
-			static int Encode(BundleCompressionFormat format, ReadOnlySpan<byte> inputSpan, ByteArrayBuilder builder)
+			static int Encode(BundleCompressionFormat format, ReadOnlyMemory<byte> input, ByteArrayBuilder builder)
 			{
 				switch (format)
 				{
 					case BundleCompressionFormat.None:
 						{
-							builder.WriteFixedLengthBytes(inputSpan);
-							return inputSpan.Length;
+							builder.WriteFixedLengthBytes(input.Span);
+							return input.Length;
 						}
 					case BundleCompressionFormat.LZ4:
 						{
-							int maxSize = LZ4Codec.MaximumOutputSize(inputSpan.Length);
+							int maxSize = LZ4Codec.MaximumOutputSize(input.Length);
 
 							Span<byte> outputSpan = builder.GetMemory(maxSize).Span;
-							int encodedLength = LZ4Codec.Encode(inputSpan, outputSpan);
+							int encodedLength = LZ4Codec.Encode(input.Span, outputSpan);
 							builder.Advance(encodedLength);
 
 							return encodedLength;
 						}
+					case BundleCompressionFormat.Gzip:
+						{
+							using MemoryWriterStream outputStream = new MemoryWriterStream(builder);
+							using GZipStream gzipStream = new GZipStream(outputStream, CompressionLevel.Fastest);
+							gzipStream.Write(input.Span);
+							return (int)outputStream.Length;
+						}
 					default:
 						throw new InvalidDataException($"Invalid compression format '{(int)format}'");
 				}
@@ -936,7 +944,7 @@ namespace EpicGames.Horde.Storage
 			ReadOnlyMemory<byte> encodedPacket = payload.Slice(packetInfo.Offset, packetInfo.EncodedLength).AsSingleSegment();
 
 			byte[] decodedPacket = new byte[packetInfo.DecodedLength];
-			Decode(format, encodedPacket.Span, decodedPacket);
+			Decode(format, encodedPacket, decodedPacket);
 
 			return decodedPacket;
 		}
@@ -960,16 +968,23 @@ namespace EpicGames.Horde.Storage
 			return IoHash.Compute(buffer);
 		}
 
-		static void Decode(BundleCompressionFormat format, ReadOnlySpan<byte> inputSpan, Span<byte> outputSpan)
+		static void Decode(BundleCompressionFormat format, ReadOnlyMemory<byte> input, Memory<byte> output)
 		{
 			switch (format)
 			{
 				case BundleCompressionFormat.None:
-					inputSpan.CopyTo(outputSpan);
+					input.CopyTo(output);
 					break;
 				case BundleCompressionFormat.LZ4:
-					LZ4Codec.Decode(inputSpan, outputSpan);
+					LZ4Codec.Decode(input.Span, output.Span);
 					break;
+				case BundleCompressionFormat.Gzip:
+					{
+						using ReadOnlyMemoryStream inputStream = new ReadOnlyMemoryStream(input);
+						using GZipStream outputStream = new GZipStream(new MemoryWriterStream(new MemoryWriter(output)), CompressionMode.Decompress, false);
+						inputStream.CopyTo(outputStream);
+						break;
+					}
 				default:
 					throw new InvalidDataException($"Invalid compression format '{(int)format}'");
 			}
