Horde: Add a C# wrapper library for Oodle data.

diff --git a/Engine/Source/Programs/Shared/EpicGames.Oodle/EpicGames.Oodle.csproj b/Engine/Source/Programs/Shared/EpicGames.Oodle/EpicGames.Oodle.csproj
new file mode 100644
--- /dev/null
+++ b/Engine/Source/Programs/Shared/EpicGames.Oodle/EpicGames.Oodle.csproj
@@ -0,0 +1,52 @@
+<Project Sdk="Microsoft.NET.Sdk">
+  <PropertyGroup>
+    <TargetFramework>net6.0</TargetFramework>
+    <Configuration Condition=" '$(Configuration)' == '' ">Development</Configuration>
+    <OutputType>Library</OutputType>
+    <Deterministic>true</Deterministic>
+    <RootNamespace>EpicGames.Oodle</RootNamespace>
+    <GenerateAssemblyInfo>false</GenerateAssemblyInfo>
+    <GenerateTargetFrameworkAttribute>false</GenerateTargetFrameworkAttribute>
+    <Configurations>Debug;Release;Development;Analyze</Configurations>
+    <Nullable>enable</Nullable>
+    <DebugType>pdbonly</DebugType>
+    <GenerateDocumentationFile>True</GenerateDocumentationFile>
+    <AllowUnsafeBlocks>True</AllowUnsafeBlocks>
+    <RunAnalyzersDuringBuild>false</RunAnalyzersDuringBuild>
+  </PropertyGroup>
+  <PropertyGroup Condition=" '$(Configuration)|$(Platform)' == 'Debug|AnyCPU' ">
+    <DocumentationFile></DocumentationFile>
+  </PropertyGroup>
+  <PropertyGroup Condition="'$(Configuration)|$(Platform)' == 'Development|AnyCPU'">
+    <DocumentationFile></DocumentationFile>
+    <Optimize>false</Optimize>
+  </PropertyGroup>
+  <PropertyGroup Condition="'$(Configuration)|$(Platform)' == 'Release|AnyCPU'">
+    <DocumentationFile></DocumentationFile>
+    <Optimize>true</Optimize>
+  </PropertyGroup>
+  <PropertyGroup Condition=" '$(Configuration)|$(Platform)' == 'Analyze|AnyCPU' ">
+    <DocumentationFile></DocumentationFile>
+    <GenerateDocumentationFile>True</GenerateDocumentationFile>
+    <RunAnalyzersDuringBuild>True</RunAnalyzersDuringBuild>
+  </PropertyGroup>
+  <ItemGroup>
+    <Compile Include="..\..\Shared\MetaData.cs">
+      <Link>Properties\MetaData.cs</Link>
+    </Compile>
+  </ItemGroup>
+  <ItemGroup>
+    <PackageReference Include="Microsoft.CodeAnalysis.CSharp.Workspaces" Version="4.2.0" />
+    <PackageReference Include="Microsoft.Extensions.Caching.Abstractions" Version="6.0.0" />
+    <PackageReference Include="Microsoft.Extensions.Caching.Memory" Version="6.0.1" />
+    <PackageReference Include="Microsoft.Extensions.DependencyInjection" Version="6.0.0" />
+    <PackageReference Include="Microsoft.Extensions.Http" Version="6.0.0" />
+    <PackageReference Include="Microsoft.Extensions.Logging" Version="6.0.0" />
+  </ItemGroup>
+  <ItemGroup>
+    <None Include="Sdk\2.9.3\linux\lib\liboo2corelinux64.so.9" Link="liboo2corelinux64.so.9" CopyToOutputDirectory="PreserveNewest" Visible="false" />
+    <None Include="Sdk\2.9.3\linux\lib\liboo2corelinux.so.9" Link="liboo2corelinux.so.9" CopyToOutputDirectory="PreserveNewest" Visible="false" />
+    <None Include="Sdk\2.9.3\win\redist\oo2core_9_win64.dll" Link="oo2core_9_win64.dll" CopyToOutputDirectory="PreserveNewest" Visible="false" />
+    <None Include="Sdk\2.9.3\win\redist\oo2core_9_win32.dll" Link="oo2core_9_win32.dll" CopyToOutputDirectory="PreserveNewest" Visible="false" />
+  </ItemGroup>
+</Project>
\ No newline at end of file
diff --git a/Engine/Source/Programs/Shared/EpicGames.Oodle/Oodle.cs b/Engine/Source/Programs/Shared/EpicGames.Oodle/Oodle.cs
new file mode 100644
--- /dev/null
+++ b/Engine/Source/Programs/Shared/EpicGames.Oodle/Oodle.cs
@@ -0,0 +1,330 @@
+// Copyright Epic Games, Inc. All Rights Reserved.
+
+using System;
+using System.Reflection;
+using System.Runtime.InteropServices;
+
+namespace EpicGames.Compression
+{
+#pragma warning disable CA1008 // Add a member that has a value of zero with a suggested name of None
+	/// <summary>
+	/// Compressor type to use
+	/// </summary>
+	public enum OodleCompressorType : int
+    {
+		/// <summary>
+		/// None = memcpy, pass through uncompressed bytes
+		/// </summary>
+		None = 3,
+
+		/// <summary>
+		/// Fast decompression and high compression ratios, amazing!
+		/// </summary>
+		Kraken = 8,
+
+		/// <summary>
+		/// Leviathan = Kraken's big brother with higher compression, slightly slower decompression.
+		/// </summary>
+		Leviathan = 13,
+
+		/// <summary>
+		/// Mermaid is between Kraken and Selkie - crazy fast, still decent compression.
+		/// </summary>
+		Mermaid = 9,
+
+		/// <summary>
+		/// Selkie is a super-fast relative of Mermaid.  For maximum decode speed.
+		/// </summary>
+		Selkie = 11,
+
+		/// <summary>
+		/// Hydra, the many-headed beast = Leviathan, Kraken, Mermaid, or Selkie (see $OodleLZ_About_Hydra)
+		/// </summary>
+		Hydra = 12,
+    };
+#pragma warning restore CA1008 // Add a member that has a value of zero with a suggested name of None
+
+	/// <summary>
+	/// Compression level
+	/// </summary>
+	public enum OodleCompressionLevel : int
+    {
+		/// <summary>
+		/// Don't compress, just copy raw bytes
+		/// </summary>
+		None = 0,
+
+		/// <summary>
+		/// Super fast mode, lower compression ratio
+		/// </summary>
+		SuperFast = 1,
+
+		/// <summary>
+		/// Fastest LZ mode with still decent compression ratio
+		/// </summary>
+		VeryFast = 2,
+
+		/// <summary>
+		/// Fast - good for daily use
+		/// </summary>
+		Fast = 3,
+
+		/// <summary>
+		/// Standard medium speed LZ mode
+		/// </summary>
+		Normal = 4,
+
+		/// <summary>
+		/// Optimal parse level 1 (faster optimal encoder)
+		/// </summary>
+		Optimal1 = 5,
+
+		/// <summary>
+		/// Optimal parse level 2 (recommended baseline optimal encoder)
+		/// </summary>
+		Optimal2 = 6,
+
+		/// <summary>
+		/// Optimal parse level 3 (slower optimal encoder)
+		/// </summary>
+		Optimal3 = 7,
+
+		/// <summary>
+		/// Optimal parse level 4 (very slow optimal encoder)
+		/// </summary>
+		Optimal4 = 8,
+
+		/// <summary>
+		/// Optimal parse level 5 (don't care about encode speed, maximum compression)
+		/// </summary>
+		Optimal5 = 9,
+
+		/// <summary>
+		/// Faster than SuperFast, less compression
+		/// </summary>
+		HyperFast1 = -1,
+
+		/// <summary>
+		/// Faster than HyperFast1, less compression
+		/// </summary>
+		HyperFast2 = -2,
+
+		/// <summary>
+		/// Faster than HyperFast2, less compression
+		/// </summary>
+		HyperFast3 = -3,
+
+		/// <summary>
+		/// Fastest, less compression
+		/// </summary>
+		HyperFast4 = -4,
+
+		#region Aliases
+
+		/// <summary>
+		/// Alias hyperfast base level
+		/// </summary>
+		HyperFast = HyperFast1,
+
+		/// <summary>
+		/// Alias optimal standard level
+		/// </summary>
+		Optimal = Optimal2,
+
+		/// <summary>
+		/// Maximum compression level
+		/// </summary>
+		Max = Optimal5,
+
+		/// <summary>
+		/// fastest compression level
+		/// </summary>
+		Min = HyperFast4,
+
+		#endregion
+    }
+
+	/// <summary>
+	/// Base class for oodle exceptions
+	/// </summary>
+	public sealed class OodleException : Exception
+	{
+		/// <summary>
+		/// Constructor
+		/// </summary>
+		public OodleException(string message) 
+			: base(message)
+		{
+		}
+	}
+
+	/// <summary>
+	/// Wraps an instance of the Oodle compressor
+	/// </summary>
+    public sealed class Oodle
+    {
+#pragma warning disable CA1712 // Do not prefix enum values with type name
+		enum OodleLZ_Decode_ThreadPhase
+		{
+			OodleLZ_Decode_ThreadPhase1 = 1,
+			OodleLZ_Decode_ThreadPhase2 = 2,
+			OodleLZ_Decode_ThreadPhaseAll = 3,
+			OodleLZ_Decode_Unthreaded = OodleLZ_Decode_ThreadPhaseAll
+		}
+
+		enum OodleLZ_FuzzSafe
+		{
+			OodleLZ_FuzzSafe_No = 0,
+			OodleLZ_FuzzSafe_Yes = 1
+		}
+
+		enum OodleLZ_CheckCRC
+		{
+			OodleLZ_CheckCRC_No = 0,
+			OodleLZ_CheckCRC_Yes = 1,
+			OodleLZ_CheckCRC_Force32 = 0x40000000
+		}
+
+		enum OodleLZ_Verbosity
+		{
+			OodleLZ_Verbosity_None = 0,
+			OodleLZ_Verbosity_Minimal = 1,
+			OodleLZ_Verbosity_Some = 2,
+			OodleLZ_Verbosity_Lots = 3,
+			OodleLZ_Verbosity_Force32 = 0x40000000
+		}
+
+#pragma warning restore CA1712 // Do not prefix enum values with type name
+
+		[StructLayout(LayoutKind.Sequential)]
+		class OodleConfigValues
+		{
+			public int m_OodleLZ_LW_LRM_step; // LZHLW LRM : bytes between LRM entries
+			public int m_OodleLZ_LW_LRM_hashLength; // LZHLW LRM : bytes hashed for each LRM entries
+			public int m_OodleLZ_LW_LRM_jumpbits; // LZHLW LRM : bits of hash used for jump table
+
+			public int m_OodleLZ_Decoder_Max_Stack_Size;  // if OodleLZ_Decompress needs to allocator a Decoder object, and it's smaller than this size, it's put on the stack instead of the heap
+			public int m_OodleLZ_Small_Buffer_LZ_Fallback_Size_Unused;  // deprecated
+			public int m_OodleLZ_BackwardsCompatible_MajorVersion; // if you need to encode streams that can be read with an older version of Oodle, set this to the Oodle2 MAJOR version number that you need compatibility with.  eg to be compatible with oodle 2.7.3 you would put 7 here
+
+			public int m_oodle_header_version;
+		}
+
+        [DllImport("oo2core")]
+        static extern unsafe long OodleLZ_Compress(OodleCompressorType compressor, byte* rawBuf, long rawLen, byte* compBuf, OodleCompressionLevel level, long pOptions, byte[]? dictionaryBase, byte[]? lrm, byte* scratchMem, long scratchSize);
+
+        [DllImport("oo2core")]
+        static extern long OodleLZ_GetCompressedBufferSizeNeeded(OodleCompressorType compressor, long rawSize);
+
+        [DllImport("oo2core")]
+        static extern int OodleLZDecoder_MemorySizeNeeded(OodleCompressorType compressor, long rawLen = -1);
+
+		[DllImport("oo2core")]
+		static extern int OodleLZ_GetCompressScratchMemBound(OodleCompressorType compressor, OodleCompressionLevel level, int rawLen, OodleConfigValues configValues);
+
+        [DllImport("oo2core")]
+        static extern unsafe long OodleLZ_Decompress(byte* compBuf, long compBufSize, byte* rawBuf, long rawLen, OodleLZ_FuzzSafe fuzzSafe = OodleLZ_FuzzSafe.OodleLZ_FuzzSafe_Yes, OodleLZ_CheckCRC checkCRC = OodleLZ_CheckCRC.OodleLZ_CheckCRC_No, OodleLZ_Verbosity verbosity = OodleLZ_Verbosity.OodleLZ_Verbosity_None, byte* decBufBase = null, long decBufSize = 0, long fpCallback = 0, long callbackUserData = 0, byte* decoderMemory = null, long decoderMemorySize = 0, OodleLZ_Decode_ThreadPhase threadPhase = OodleLZ_Decode_ThreadPhase.OodleLZ_Decode_Unthreaded);
+
+        [DllImport("oo2core")]
+        static extern void Oodle_GetConfigValues(OodleConfigValues configValues);
+
+        [DllImport("oo2core")]
+        static extern void Oodle_SetConfigValues(OodleConfigValues configValues);
+
+		static readonly OodleConfigValues s_configValues = new OodleConfigValues();
+
+		/// <summary>
+		/// Static constructor. Registers the import resolver for the native Oodle library.
+		/// </summary>
+		static Oodle()
+		{
+			NativeLibrary.SetDllImportResolver(typeof(Oodle).Assembly, ImportResolver);
+
+			Oodle_GetConfigValues(s_configValues);
+			s_configValues.m_OodleLZ_BackwardsCompatible_MajorVersion = 9;
+            Oodle_SetConfigValues(s_configValues);
+		}
+
+        private static IntPtr ImportResolver(string libraryName, Assembly assembly, DllImportSearchPath? searchpath)
+        {
+            IntPtr libHandle = IntPtr.Zero;
+            if (libraryName == "oo2core")
+            {
+                // we manually load the assembly as the file name is not consistent between platforms thus the automatic searching will not work
+                string assemblyFilename;
+                bool is64Bit = Environment.Is64BitProcess;
+
+                if (RuntimeInformation.IsOSPlatform(OSPlatform.Windows))
+                {
+                    assemblyFilename = $"oo2core_9_{(is64Bit ? "win64" : "win32")}.dll";
+                }
+				else if (RuntimeInformation.IsOSPlatform(OSPlatform.Linux))
+                {
+                    assemblyFilename = $"liboo2core{(is64Bit ? "linux64" : "linux")}.so.9";
+                }
+                else
+                {
+                    throw new NotImplementedException();
+                }
+                IntPtr handle = NativeLibrary.Load(assemblyFilename, assembly, searchpath);
+                libHandle = handle;
+            }
+            return libHandle;
+        }
+
+		/// <summary>
+		/// Compress a block of data
+		/// </summary>
+		/// <param name="compressor">Compressor to use</param>
+		/// <param name="inputData">Data to be compressed</param>
+		/// <param name="outputData">Buffer for output data</param>
+		/// <param name="compressionLevel">Desired compression level</param>
+		/// <returns></returns>
+		public static unsafe int Compress(OodleCompressorType compressor, ReadOnlySpan<byte> inputData, Span<byte> outputData, OodleCompressionLevel compressionLevel = OodleCompressionLevel.Normal)
+        {
+			fixed (byte* rawBufPtr = inputData)
+			fixed (byte* compBufPtr = outputData)
+			{
+				long result = OodleLZ_Compress(compressor, rawBufPtr, inputData.Length, compBufPtr, compressionLevel, 0, null, null, null, 0);
+				if (result == 0)
+				{
+					throw new OodleException($"Unable to compress data (result: {result})");
+				}
+				return (int)result;
+			}
+        }
+
+		/// <summary>
+		/// Determines the max size of the compressed buffer
+		/// </summary>
+		/// <param name="compressor">Compressor type to use</param>
+		/// <param name="uncompressedLength">Length of the input data</param>
+		/// <returns>Size of the buffer required for output data</returns>
+        public static int MaximumOutputSize(OodleCompressorType compressor, int uncompressedLength)
+        {
+            long bufferSizeNeeded = OodleLZ_GetCompressedBufferSizeNeeded(compressor, uncompressedLength);
+            return (int)bufferSizeNeeded;
+        }
+
+		/// <summary>
+		/// Decompresses a block of data
+		/// </summary>
+		/// <param name="compressor">Type of the compressor</param>
+		/// <param name="inputData">The compressed buffer</param>
+		/// <param name="outputData">Output buffer for decompressed data</param>
+		/// <returns></returns>
+        public static unsafe int Decompress(OodleCompressorType compressor, ReadOnlySpan<byte> inputData, Span<byte> outputData)
+        {
+			fixed (byte* inputPtr = inputData)
+			fixed (byte* outputPtr = outputData)
+			{
+				long result = OodleLZ_Decompress(inputPtr, inputData.Length, outputPtr, outputData.Length);
+				if (result == 0)
+				{
+					throw new OodleException($"Unable to decompress data (result: {result})");
+				}
+				return (int)result;
+			}
+        }
+    }
+}
