Fix wrong API macro being used in C# UHT
Cleaned up formatting in both C++ and C# UHT

diff --git a/Engine/Source/Programs/Shared/EpicGames.UHT/Exporters/CodeGen/UhtHeaderCodeGeneratorHFile.cs b/Engine/Source/Programs/Shared/EpicGames.UHT/Exporters/CodeGen/UhtHeaderCodeGeneratorHFile.cs
--- a/Engine/Source/Programs/Shared/EpicGames.UHT/Exporters/CodeGen/UhtHeaderCodeGeneratorHFile.cs
+++ b/Engine/Source/Programs/Shared/EpicGames.UHT/Exporters/CodeGen/UhtHeaderCodeGeneratorHFile.cs
@@ -349,10 +349,10 @@ namespace EpicGames.UHT.Exporters.CodeGen
 			bool hasEditorRpc = reversedFunctions.Any(x => IsRpcFunction(x, true));
 
 			// Output the RPC methods
-			AppendRpcFunctions(builder, classObj, api, reversedFunctions, false);
+			AppendRpcFunctions(builder, classObj, reversedFunctions, false);
 			if (hasEditorRpc)
 			{
-				AppendRpcFunctions(builder, classObj, api, reversedFunctions, true);
+				AppendRpcFunctions(builder, classObj, reversedFunctions, true);
 			}
 
 			// Output property accessors
@@ -606,13 +606,13 @@ namespace EpicGames.UHT.Exporters.CodeGen
 			return builder;
 		}
 
-		private StringBuilder AppendRpcFunctions(StringBuilder builder, UhtClass classObj, string api,  List<UhtFunction> reversedFunctions, bool editorOnly)
+		private StringBuilder AppendRpcFunctions(StringBuilder builder, UhtClass classObj, List<UhtFunction> reversedFunctions, bool editorOnly)
 		{
 			builder.AppendBeginEditorOnlyGuard(editorOnly);
 
 			using (UhtMacroCreator macro = new(builder, this, classObj, editorOnly ? EditorOnlyRpcWrappersMacroSuffix : RpcWrappersMacroSuffix))
 			{
-				AppendAutogeneratedBlueprintFunctionDeclarations(builder, classObj, api, reversedFunctions, editorOnly);
+				AppendAutogeneratedBlueprintFunctionDeclarations(builder, classObj, reversedFunctions, editorOnly);
 				AppendRpcWrappers(builder, reversedFunctions, editorOnly);
 			}
 
@@ -620,7 +620,7 @@ namespace EpicGames.UHT.Exporters.CodeGen
 			{
 				if (classObj.GeneratedCodeVersion <= EGeneratedCodeVersion.V1)
 				{
-					AppendAutogeneratedBlueprintFunctionDeclarationsOnlyNotDeclared(builder, classObj, api, reversedFunctions, editorOnly);
+					AppendAutogeneratedBlueprintFunctionDeclarationsOnlyNotDeclared(builder, classObj, reversedFunctions, editorOnly);
 				}
 				AppendRpcWrappers(builder, reversedFunctions, editorOnly);
 			}
@@ -662,7 +662,7 @@ namespace EpicGames.UHT.Exporters.CodeGen
 			return builder;
 		}
 
-		private StringBuilder AppendAutogeneratedBlueprintFunctionDeclarations(StringBuilder builder, UhtClass classObj, string api, List<UhtFunction> reversedFunctions, bool editorOnly)
+		private StringBuilder AppendAutogeneratedBlueprintFunctionDeclarations(StringBuilder builder, UhtClass classObj, List<UhtFunction> reversedFunctions, bool editorOnly)
 		{
 			foreach (UhtFunction function in reversedFunctions)
 			{
@@ -670,13 +670,13 @@ namespace EpicGames.UHT.Exporters.CodeGen
 				{
 					continue;
 				}
-				AppendNetValidateDeclaration(builder, classObj, api, function);
+				AppendNetValidateDeclaration(builder, classObj, function);
 				AppendNativeFunctionHeader(builder, function, UhtPropertyTextType.ClassFunctionArgOrRetVal, true, null, null, UhtFunctionExportFlags.None, "; \\\r\n");
 			}
 			return builder;
 		}
 
-		private StringBuilder AppendAutogeneratedBlueprintFunctionDeclarationsOnlyNotDeclared(StringBuilder builder, UhtClass classObj, string api, List<UhtFunction> reversedFunctions, bool editorOnly)
+		private StringBuilder AppendAutogeneratedBlueprintFunctionDeclarationsOnlyNotDeclared(StringBuilder builder, UhtClass classObj, List<UhtFunction> reversedFunctions, bool editorOnly)
 		{
 			foreach (UhtFunction function in reversedFunctions)
 			{
@@ -686,7 +686,7 @@ namespace EpicGames.UHT.Exporters.CodeGen
 				}
 				if (!function.FunctionExportFlags.HasAnyFlags(UhtFunctionExportFlags.ValidationImplFound))
 				{
-					AppendNetValidateDeclaration(builder, classObj, api, function);
+					AppendNetValidateDeclaration(builder, classObj, function);
 				}
 				if (!function.FunctionExportFlags.HasAnyFlags(UhtFunctionExportFlags.ImplFound))
 				{
@@ -696,7 +696,7 @@ namespace EpicGames.UHT.Exporters.CodeGen
 			return builder;
 		}
 
-		private static StringBuilder AppendNetValidateDeclaration(StringBuilder builder, UhtClass classObj, string api, UhtFunction function)
+		private StringBuilder AppendNetValidateDeclaration(StringBuilder builder, UhtClass classObj, UhtFunction function)
 		{
 			if (function.FunctionFlags.HasAnyFlags(EFunctionFlags.NetValidate))
 			{
@@ -704,10 +704,9 @@ namespace EpicGames.UHT.Exporters.CodeGen
 
 				if (!classObj.ClassFlags.HasAnyFlags(EClassFlags.RequiredAPI) && function.FunctionFlags.HasAnyFlags(EFunctionFlags.RequiredAPI))
 				{
-					builder.Append(api);
+					builder.Append(PackageApi);
 				}
 
-				builder.Append(' ');
 				if (!function.FunctionFlags.HasAnyFlags(EFunctionFlags.Static) && !function.FunctionExportFlags.HasAnyFlags(UhtFunctionExportFlags.Final))
 				{
 					builder.Append("virtual");
diff --git a/Engine/Source/Programs/UnrealHeaderTool/Private/CodeGenerator.cpp b/Engine/Source/Programs/UnrealHeaderTool/Private/CodeGenerator.cpp
--- a/Engine/Source/Programs/UnrealHeaderTool/Private/CodeGenerator.cpp
+++ b/Engine/Source/Programs/UnrealHeaderTool/Private/CodeGenerator.cpp
@@ -4917,11 +4917,11 @@ void FNativeClassHeaderGenerator::ExportNativeFunctions(FOutputDevice& OutGenera
 
 			FFindDelcarationResults Implementation = FindDeclaration(ClassDef, FunctionData.CppImplName);
 			FFindDelcarationResults Validation = FindDeclaration(ClassDef, FunctionData.CppValidationImplName);
+			const FString& API = GetAPIString();
 
-			const FString& api = GetAPIString();
 			//Emit warning here if necessary
 			FUHTStringBuilder FunctionDeclaration;
-			ExportNativeFunctionHeader(FunctionDeclaration, OutReferenceGatherers.ForwardDeclarations, *FunctionDef, FunctionData, EExportFunctionType::Function, EExportFunctionHeaderStyle::Declaration, nullptr, *api);
+			ExportNativeFunctionHeader(FunctionDeclaration, OutReferenceGatherers.ForwardDeclarations, *FunctionDef, FunctionData, EExportFunctionType::Function, EExportFunctionHeaderStyle::Declaration, nullptr, *API);
 			FunctionDeclaration.Log(TEXT(";\r\n"));
 
 			// Declare validation function if needed
@@ -4929,11 +4929,15 @@ void FNativeClassHeaderGenerator::ExportNativeFunctions(FOutputDevice& OutGenera
 			{
 				FString ParameterList = GetFunctionParameterString(*FunctionDef, OutReferenceGatherers);
 
-				const TCHAR* Virtual = (!FunctionDef->HasAnyFunctionFlags(FUNC_Static) && !(FunctionData.FunctionExportFlags & FUNCEXPORT_Final)) ? TEXT("virtual") : TEXT("");
 				FStringOutputDevice ValidDecl;
-				bool AddApi = !FunctionDef->GetOwnerClass()->HasAnyClassFlags(CLASS_RequiredAPI) && (FunctionData.FunctionExportFlags & FUNCEXPORT_RequiredAPI);
+				ValidDecl.Logf(TEXT("\t"));
+				if (!FunctionDef->GetOwnerClass()->HasAnyClassFlags(CLASS_RequiredAPI) && (FunctionData.FunctionExportFlags & FUNCEXPORT_RequiredAPI))
+				{
+					ValidDecl.Logf(TEXT("%s"), *API);
+				}
+				const TCHAR* Virtual = (!FunctionDef->HasAnyFunctionFlags(FUNC_Static) && !(FunctionData.FunctionExportFlags & FUNCEXPORT_Final)) ? TEXT("virtual") : TEXT("");
+				ValidDecl.Logf(TEXT("%s bool %s(%s);\r\n"), Virtual, * FunctionData.CppValidationImplName, * ParameterList);
 
-				ValidDecl.Logf(TEXT("\t%s %s bool %s(%s);\r\n"), AddApi ? *api : TEXT(""), Virtual, *FunctionData.CppValidationImplName, *ParameterList);
 				FuncStringBuilders.AutogeneratedBlueprintFunctionDeclarations.Log(*ValidDecl);
 				if (!Validation.WasFound())
 				{
