Horde: Add support for using Oodle to compress bundle packets.

diff --git a/Engine/Source/Programs/Shared/EpicGames.Horde/Storage/Bundle.cs b/Engine/Source/Programs/Shared/EpicGames.Horde/Storage/Bundle.cs
--- a/Engine/Source/Programs/Shared/EpicGames.Horde/Storage/Bundle.cs
+++ b/Engine/Source/Programs/Shared/EpicGames.Horde/Storage/Bundle.cs
@@ -123,6 +123,11 @@ namespace EpicGames.Horde.Storage
 		/// Gzip compression
 		/// </summary>
 		Gzip = 2,
+
+		/// <summary>
+		/// Oodle compression (Kraken)
+		/// </summary>
+		Oodle = 3,
 	}
 
 	/// <summary>
diff --git a/Engine/Source/Programs/Shared/EpicGames.Horde/Storage/TreeStore.cs b/Engine/Source/Programs/Shared/EpicGames.Horde/Storage/TreeStore.cs
--- a/Engine/Source/Programs/Shared/EpicGames.Horde/Storage/TreeStore.cs
+++ b/Engine/Source/Programs/Shared/EpicGames.Horde/Storage/TreeStore.cs
@@ -14,6 +14,10 @@ using EpicGames.Core;
 using K4os.Compression.LZ4;
 using Microsoft.Extensions.Caching.Memory;
 
+#if WITH_OODLE
+using EpicGames.Compression;
+#endif
+
 namespace EpicGames.Horde.Storage
 {
 	/// <summary>
@@ -656,7 +660,7 @@ namespace EpicGames.Horde.Storage
 						{
 							int maxSize = LZ4Codec.MaximumOutputSize(input.Length);
 
-							Span<byte> outputSpan = builder.GetMemory(maxSize).Span;
+							Span<byte> outputSpan = builder.GetSpan(maxSize);
 							int encodedLength = LZ4Codec.Encode(input.Span, outputSpan);
 							builder.Advance(encodedLength);
 
@@ -669,6 +673,20 @@ namespace EpicGames.Horde.Storage
 							gzipStream.Write(input.Span);
 							return (int)outputStream.Length;
 						}
+					case BundleCompressionFormat.Oodle:
+						{
+#if WITH_OODLE
+							int maxSize = Oodle.MaximumOutputSize(OodleCompressorType.Selkie, input.Length);
+
+							Span<byte> outputSpan = builder.GetSpan(maxSize);
+							int encodedLength = Oodle.Compress(OodleCompressorType.Selkie, input.Span, outputSpan, OodleCompressionLevel.HyperFast);
+							builder.Advance(encodedLength);
+
+							return encodedLength;
+#else
+							throw new NotSupportedException("Oodle is not compiled into this build.");
+#endif
+						}
 					default:
 						throw new InvalidDataException($"Invalid compression format '{(int)format}'");
 				}
@@ -985,6 +1003,13 @@ namespace EpicGames.Horde.Storage
 						inputStream.CopyTo(outputStream);
 						break;
 					}
+				case BundleCompressionFormat.Oodle:
+#if WITH_OODLE
+					Oodle.Decompress(input.Span, output.Span);
+					break;
+#else
+					throw new NotSupportedException("Oodle is not compiled into this build.");
+#endif
 				default:
 					throw new InvalidDataException($"Invalid compression format '{(int)format}'");
 			}
diff --git a/Engine/Source/Programs/Shared/EpicGames.Oodle/Oodle.cs b/Engine/Source/Programs/Shared/EpicGames.Oodle/Oodle.cs
--- a/Engine/Source/Programs/Shared/EpicGames.Oodle/Oodle.cs
+++ b/Engine/Source/Programs/Shared/EpicGames.Oodle/Oodle.cs
@@ -309,11 +309,10 @@ namespace EpicGames.Compression
 		/// <summary>
 		/// Decompresses a block of data
 		/// </summary>
-		/// <param name="compressor">Type of the compressor</param>
 		/// <param name="inputData">The compressed buffer</param>
 		/// <param name="outputData">Output buffer for decompressed data</param>
 		/// <returns></returns>
-        public static unsafe int Decompress(OodleCompressorType compressor, ReadOnlySpan<byte> inputData, Span<byte> outputData)
+        public static unsafe int Decompress(ReadOnlySpan<byte> inputData, Span<byte> outputData)
         {
 			fixed (byte* inputPtr = inputData)
 			fixed (byte* outputPtr = outputData)
