Fix turnkey valid SDK filtering

diff --git a/Engine/Source/Programs/Shared/EpicGames.Build/UEBuildPlatformSDK.cs b/Engine/Source/Programs/Shared/EpicGames.Build/UEBuildPlatformSDK.cs
--- a/Engine/Source/Programs/Shared/EpicGames.Build/UEBuildPlatformSDK.cs
+++ b/Engine/Source/Programs/Shared/EpicGames.Build/UEBuildPlatformSDK.cs
@@ -154,7 +154,7 @@ namespace EpicGames.Core
 		public SDKCollection(string? Min, string? Max, UEBuildPlatformSDK PlatformSDK)
 			: this(PlatformSDK)
 		{
-			Sdks.Add(new SDKDescriptor(DefaultName, Min, Max));
+			Sdks.Add(new SDKDescriptor(DefaultName, Min, Max, null, null, this));
 		}
 
 		// convenience for single unnamed Sdk current value
@@ -163,7 +163,7 @@ namespace EpicGames.Core
 		{
 			if (Current != null)
 			{
-				Sdks.Add(new SDKDescriptor(DefaultName, Current));
+				Sdks.Add(new SDKDescriptor(DefaultName, null, null, Current, null, this));
 			}
 		}
 
