Fixed C# UHT issue where unterminated strings were not detected properly and could result in an infinite loop.

diff --git a/Engine/Source/Programs/Shared/EpicGames.UHT/Tokenizer/UhtTokenBufferReader.cs b/Engine/Source/Programs/Shared/EpicGames.UHT/Tokenizer/UhtTokenBufferReader.cs
--- a/Engine/Source/Programs/Shared/EpicGames.UHT/Tokenizer/UhtTokenBufferReader.cs
+++ b/Engine/Source/Programs/Shared/EpicGames.UHT/Tokenizer/UhtTokenBufferReader.cs
@@ -625,22 +625,21 @@ Restart:
 				for (; ; )
 				{
 					char nextChar = InternalGetChar(span);
-					if (nextChar == '\r' || nextChar == '\n')
+					if (nextChar == '\r' || nextChar == '\n' || nextChar == 0)
 					{
-						// throw
+						throw new UhtException(this, "Unterminated character constant");
 					}
 
-					if (nextChar == '\\')
+					else if (nextChar == '\\')
 					{
 						nextChar = InternalGetChar(span);
-						if (nextChar == '\r' || nextChar == '\n')
+						if (nextChar == '\r' || nextChar == '\n' || nextChar == 0)
 						{
-							// throw
+							throw new UhtException(this, "Unterminated character constant");
 						}
-						nextChar = InternalGetChar(span);
 					}
 
-					if (nextChar == '"')
+					else if (nextChar == '"')
 					{
 						break;
 					}
