Modified C# UHT to generate errors when extra ';' are detected.  C++ UHT will generate an error.

#rnx

diff --git a/Engine/Source/Programs/Shared/EpicGames.UHT/Parsers/UhtClassParser.cs b/Engine/Source/Programs/Shared/EpicGames.UHT/Parsers/UhtClassParser.cs
--- a/Engine/Source/Programs/Shared/EpicGames.UHT/Parsers/UhtClassParser.cs
+++ b/Engine/Source/Programs/Shared/EpicGames.UHT/Parsers/UhtClassParser.cs
@@ -48,7 +48,8 @@ namespace EpicGames.UHT.Parsers
 
 			classObj.GeneratedBodyLineNumber = topScope.TokenReader.InputLine;
 
-			topScope.TokenReader.Optional(';');
+			// C++ UHT TODO - In the C++ version, we don't skip any trailing ';'.  ParseStatement will generate an error
+			//topScope.TokenReader.Optional(';');
 			return UhtParseResult.Handled;
 		}
 		#endregion
diff --git a/Engine/Source/Programs/Shared/EpicGames.UHT/Parsers/UhtHeaderFileParser.cs b/Engine/Source/Programs/Shared/EpicGames.UHT/Parsers/UhtHeaderFileParser.cs
--- a/Engine/Source/Programs/Shared/EpicGames.UHT/Parsers/UhtHeaderFileParser.cs
+++ b/Engine/Source/Programs/Shared/EpicGames.UHT/Parsers/UhtHeaderFileParser.cs
@@ -4,6 +4,7 @@ using System;
 using System.Collections.Generic;
 using System.Diagnostics.CodeAnalysis;
 using System.IO;
+using System.Runtime.CompilerServices;
 using EpicGames.Core;
 using EpicGames.UHT.Tables;
 using EpicGames.UHT.Tokenizer;
@@ -420,9 +421,19 @@ namespace EpicGames.UHT.Parsers
 					break;
 
 				case UhtTokenType.Symbol:
-					// Ignore any extra semicolons
+					// C++ UHT TODO - Generate errors when extra ';' are found.  
+					// UPROPERTY has code to remove extra ';'
 					if (token.IsSymbol(';'))
 					{
+						UhtToken nextToken = topScope.TokenReader.PeekToken();
+						if (nextToken.TokenType.IsEndType())
+						{
+							topScope.TokenReader.LogError("Extra ';' before end of file");
+						}
+						else
+						{
+							topScope.TokenReader.LogError($"Extra ';' before '{nextToken}'");
+						}
 						return true;
 					}
 					parseResult = DispatchKeyword(topScope, ref token);
diff --git a/Engine/Source/Programs/Shared/EpicGames.UHT/Parsers/UhtPropertyParser.cs b/Engine/Source/Programs/Shared/EpicGames.UHT/Parsers/UhtPropertyParser.cs
--- a/Engine/Source/Programs/Shared/EpicGames.UHT/Parsers/UhtPropertyParser.cs
+++ b/Engine/Source/Programs/Shared/EpicGames.UHT/Parsers/UhtPropertyParser.cs
@@ -1251,6 +1251,18 @@ namespace EpicGames.UHT.Parsers
 			UhtPropertyParseOptions options = UhtPropertyParseOptions.ParseLayoutMacro | UhtPropertyParseOptions.List | UhtPropertyParseOptions.AddModuleRelativePath;
 			topScope.HeaderParser.GetCachedPropertyParser().Parse(topScope, EPropertyFlags.ParmFlags, options, UhtPropertyCategory.Member, s_propertyDelegate);
 			topScope.TokenReader.Require(';');
+
+			// C++ UHT TODO - Skip any extra ';'.  This can be removed if we remove UhtHeaderfileParser.ParserStatement generating errors
+			// when extra ';' are found.  Oddly, UPROPERTY specifically skips extra ';'
+			while (true)
+			{
+				UhtToken nextToken = topScope.TokenReader.PeekToken();
+				if (!nextToken.IsSymbol(';'))
+				{
+					break;
+				}
+				topScope.TokenReader.ConsumeToken();
+			}
 			return UhtParseResult.Handled;
 		}
 		#endregion
