Horde: Add support for ranged reads of blobs.

diff --git a/Engine/Source/Programs/Shared/EpicGames.Horde/Storage/IStorageClient.cs b/Engine/Source/Programs/Shared/EpicGames.Horde/Storage/IStorageClient.cs
--- a/Engine/Source/Programs/Shared/EpicGames.Horde/Storage/IStorageClient.cs
+++ b/Engine/Source/Programs/Shared/EpicGames.Horde/Storage/IStorageClient.cs
@@ -95,40 +95,29 @@ namespace EpicGames.Horde.Storage
 		/// <summary>
 		/// Reads data for a blob from the store
 		/// </summary>
-		/// <param name="locator">The blob identifier</param>
+		/// <param name="locator">The blob locator</param>
 		/// <param name="cancellationToken">Cancellation token for the operation</param>
-		/// <returns></returns>
+		/// <returns>Stream containing the data</returns>
 		Task<Stream> ReadBlobAsync(BlobLocator locator, CancellationToken cancellationToken = default);
 
 		/// <summary>
-		/// Writes a new blob to the store
+		/// Reads part of a blob from storage
 		/// </summary>
-		/// <param name="stream">Blob data</param>
+		/// <param name="locator">The blob locator</param>
+		/// <param name="offset">Offset within the blob to read from</param>
+		/// <param name="length">Maximum length of data to read. The stream will be truncated if the blob is shorter than this.</param>
 		/// <param name="cancellationToken">Cancellation token for the operation</param>
-		/// <param name="prefix">Prefix for blob names. While the returned BlobId is guaranteed to be unique, this name can be used as a prefix to aid debugging.</param>
-		/// <returns>Unique identifier for the blob</returns>
-		Task<BlobLocator> WriteBlobAsync(Stream stream, Utf8String prefix = default, CancellationToken cancellationToken = default);
-
-		#endregion
-
-		#region Bundles
+		/// <returns>Stream containing the data</returns>
+		Task<Stream> ReadBlobRangeAsync(BlobLocator locator, int offset = 0, int length = Int32.MaxValue, CancellationToken cancellationToken = default);
 
 		/// <summary>
-		/// Reads data for a bundle from the store
-		/// </summary>
-		/// <param name="locator">The blob identifier</param>
-		/// <param name="cancellationToken">Cancellation token for the operation</param>
-		/// <returns></returns>
-		Task<Bundle> ReadBundleAsync(BlobLocator locator, CancellationToken cancellationToken = default);
-
-		/// <summary>
-		/// Writes a new bundle to the store
+		/// Writes a new blob to the store
 		/// </summary>
-		/// <param name="bundle">Bundle data</param>
+		/// <param name="stream">Blob data</param>
 		/// <param name="cancellationToken">Cancellation token for the operation</param>
 		/// <param name="prefix">Prefix for blob names. While the returned BlobId is guaranteed to be unique, this name can be used as a prefix to aid debugging.</param>
 		/// <returns>Unique identifier for the blob</returns>
-		Task<BlobLocator> WriteBundleAsync(Bundle bundle, Utf8String prefix = default, CancellationToken cancellationToken = default);
+		Task<BlobLocator> WriteBlobAsync(Stream stream, Utf8String prefix = default, CancellationToken cancellationToken = default);
 
 		#endregion
 
@@ -227,17 +216,10 @@ namespace EpicGames.Horde.Storage
 			public Task<Stream> ReadBlobAsync(BlobLocator locator, CancellationToken cancellationToken = default) => _inner.ReadBlobAsync(locator, cancellationToken);
 
 			/// <inheritdoc/>
-			public Task<BlobLocator> WriteBlobAsync(Stream stream, Utf8String prefix = default, CancellationToken cancellationToken = default) => _inner.WriteBlobAsync(stream, prefix, cancellationToken);
-
-			#endregion
-
-			#region Bundles
-
-			/// <inheritdoc/>
-			public Task<Bundle> ReadBundleAsync(BlobLocator id, CancellationToken cancellationToken = default) => _inner.ReadBundleAsync(id, cancellationToken);
+			public Task<Stream> ReadBlobRangeAsync(BlobLocator locator, int offset, int length, CancellationToken cancellationToken = default) => _inner.ReadBlobRangeAsync(locator, offset, length, cancellationToken);
 
 			/// <inheritdoc/>
-			public Task<BlobLocator> WriteBundleAsync(Bundle bundle, Utf8String prefix = default, CancellationToken cancellationToken = default) => _inner.WriteBundleAsync(bundle, prefix, cancellationToken);
+			public Task<BlobLocator> WriteBlobAsync(Stream stream, Utf8String prefix = default, CancellationToken cancellationToken = default) => _inner.WriteBlobAsync(stream, prefix, cancellationToken);
 
 			#endregion
 
@@ -278,6 +260,33 @@ namespace EpicGames.Horde.Storage
 		/// <returns></returns>
 		public static IStorageClient<T> ForType<T>(this IStorageClient blobStore) => new TypedStorageClient<T>(blobStore);
 
+		/// <summary>
+		/// Reads data for a bundle from the store
+		/// </summary>
+		/// <param name="store">The store instance to read from</param>
+		/// <param name="locator">The blob identifier</param>
+		/// <param name="cancellationToken">Cancellation token for the operation</param>
+		/// <returns></returns>
+		public static async Task<Bundle> ReadBundleAsync(this IStorageClient store, BlobLocator locator, CancellationToken cancellationToken = default)
+		{
+			using Stream stream = await store.ReadBlobAsync(locator, cancellationToken);
+			return await Bundle.FromStreamAsync(stream, cancellationToken);
+		}
+
+		/// <summary>
+		/// Writes a new bundle to the store
+		/// </summary>
+		/// <param name="store">The store instance to read from</param>
+		/// <param name="bundle">Bundle data</param>
+		/// <param name="cancellationToken">Cancellation token for the operation</param>
+		/// <param name="prefix">Prefix for blob names. While the returned BlobId is guaranteed to be unique, this name can be used as a prefix to aid debugging.</param>
+		/// <returns>Unique identifier for the blob</returns>
+		public static async Task<BlobLocator> WriteBundleAsync(this IStorageClient store, Bundle bundle, Utf8String prefix = default, CancellationToken cancellationToken = default)
+		{
+			using ReadOnlySequenceStream stream = new ReadOnlySequenceStream(bundle.AsSequence());
+			return await store.WriteBlobAsync(stream, prefix, cancellationToken);
+		}
+
 		/// <summary>
 		/// Checks if the given ref exists
 		/// </summary>
diff --git a/Engine/Source/Programs/Shared/EpicGames.Horde/Storage/StorageClientBase.cs b/Engine/Source/Programs/Shared/EpicGames.Horde/Storage/StorageClientBase.cs
--- a/Engine/Source/Programs/Shared/EpicGames.Horde/Storage/StorageClientBase.cs
+++ b/Engine/Source/Programs/Shared/EpicGames.Horde/Storage/StorageClientBase.cs
@@ -18,6 +18,82 @@ namespace EpicGames.Horde.Storage
 	/// </summary>
 	public abstract class StorageClientBase : IStorageClient
 	{
+		class RangedReadStream : Stream, IDisposable
+		{
+			readonly Stream _inner;
+			readonly int _length;
+
+			int _position;
+
+			public RangedReadStream(Stream inner, int length)
+			{
+				_inner = inner;
+				_length = length;
+			}
+
+			/// <inheritdoc/>
+			public override bool CanRead => true;
+
+			/// <inheritdoc/>
+			public override bool CanSeek => false;
+
+			/// <inheritdoc/>
+			public override bool CanWrite => false;
+
+			/// <inheritdoc/>
+			public override long Length => throw new NotSupportedException();
+
+			/// <inheritdoc/>
+			public override long Position { get => _position; set => throw new NotSupportedException(); }
+
+			/// <inheritdoc/>
+			protected override void Dispose(bool disposing)
+			{
+				base.Dispose(disposing);
+
+				if (disposing)
+				{
+					_inner.Dispose();
+				}
+			}
+
+			/// <inheritdoc/>
+			public override void Flush() => _inner.Flush();
+
+			/// <inheritdoc/>
+			public override int Read(Span<byte> buffer)
+			{
+				buffer = buffer.Slice(0, Math.Min(buffer.Length, _length - _position));
+				int read = _inner.Read(buffer);
+				_position += read;
+				return read;
+			}
+
+			/// <inheritdoc/>
+			public override int Read(byte[] buffer, int offset, int count) => Read(buffer.AsSpan(offset, count));
+
+			/// <inheritdoc/>
+			public override async ValueTask<int> ReadAsync(Memory<byte> buffer, CancellationToken cancellationToken = default)
+			{
+				buffer = buffer.Slice(0, Math.Min(buffer.Length, _length - _position));
+				int read = await _inner.ReadAsync(buffer, cancellationToken);
+				_position += read;
+				return read;
+			}
+
+			/// <inheritdoc/>
+			public override async Task<int> ReadAsync(byte[] buffer, int offset, int count, CancellationToken cancellationToken) => await ReadAsync(buffer.AsMemory(offset, count), cancellationToken);
+
+			/// <inheritdoc/>
+			public override long Seek(long offset, SeekOrigin origin) => throw new NotSupportedException();
+
+			/// <inheritdoc/>
+			public override void SetLength(long value) => throw new NotSupportedException();
+
+			/// <inheritdoc/>
+			public override void Write(byte[] buffer, int offset, int count) => throw new NotSupportedException();
+		}
+
 		readonly TreeStore _treeStore;
 
 		/// <summary>
@@ -35,25 +111,15 @@ namespace EpicGames.Horde.Storage
 		public abstract Task<Stream> ReadBlobAsync(BlobLocator locator, CancellationToken cancellationToken = default);
 
 		/// <inheritdoc/>
-		public abstract Task<BlobLocator> WriteBlobAsync(Stream stream, Utf8String prefix = default, CancellationToken cancellationToken = default);
-
-		#endregion
-
-		#region Bundles
-
-		/// <inheritdoc/>
-		public virtual async Task<Bundle> ReadBundleAsync(BlobLocator locator, CancellationToken cancellationToken = default)
+		public virtual async Task<Stream> ReadBlobRangeAsync(BlobLocator locator, int offset, int length, CancellationToken cancellationToken = default)
 		{
-			using Stream stream = await ReadBlobAsync(locator, cancellationToken);
-			return await Bundle.FromStreamAsync(stream, cancellationToken);
+			Stream stream = await ReadBlobAsync(locator, cancellationToken);
+			stream.Seek(offset, SeekOrigin.Begin);
+			return new RangedReadStream(stream, Math.Min((int)(stream.Length - offset), length));
 		}
 
 		/// <inheritdoc/>
-		public virtual async Task<BlobLocator> WriteBundleAsync(Bundle bundle, Utf8String prefix = default, CancellationToken cancellationToken = default)
-		{
-			using ReadOnlySequenceStream stream = new ReadOnlySequenceStream(bundle.AsSequence());
-			return await WriteBlobAsync(stream, prefix, cancellationToken);
-		}
+		public abstract Task<BlobLocator> WriteBlobAsync(Stream stream, Utf8String prefix = default, CancellationToken cancellationToken = default);
 
 		#endregion
 
@@ -86,7 +152,7 @@ namespace EpicGames.Horde.Storage
 			{
 				return null;
 			}
-			return new RefValue(refTarget, await ReadBundleAsync(refTarget.Locator, cancellationToken));
+			return new RefValue(refTarget, await this.ReadBundleAsync(refTarget.Locator, cancellationToken));
 		}
 
 		/// <inheritdoc/>
@@ -95,7 +161,7 @@ namespace EpicGames.Horde.Storage
 		/// <inheritdoc/>
 		public virtual async Task<RefTarget> WriteRefValueAsync(RefName name, Bundle bundle, int exportIdx = 0, Utf8String prefix = default, CancellationToken cancellationToken = default)
 		{
-			BlobLocator locator = await WriteBundleAsync(bundle, prefix, cancellationToken);
+			BlobLocator locator = await this.WriteBundleAsync(bundle, prefix, cancellationToken);
 			RefTarget target = new RefTarget(locator, exportIdx);
 			await WriteRefTargetAsync(name, target, cancellationToken);
 			return target;
