Horde: Store bundle packets in separate memory blocks, and move encode/decode logic out of TreeStore.

diff --git a/Engine/Source/Programs/Horde/Horde.Agent/Commands/Bundles/DumpCommand.cs b/Engine/Source/Programs/Horde/Horde.Agent/Commands/Bundles/DumpCommand.cs
--- a/Engine/Source/Programs/Horde/Horde.Agent/Commands/Bundles/DumpCommand.cs
+++ b/Engine/Source/Programs/Horde/Horde.Agent/Commands/Bundles/DumpCommand.cs
@@ -35,7 +35,6 @@ namespace Horde.Agent.Commands.Bundles
 				logger.LogInformation("Summary for ref {RefId}", RefName);
 			}
 
-			ReadOnlyMemory<byte> data = blob.Payload.AsSingleSegment();
 			BundleHeader header = blob.Header;
 			int packetStart = 0;
 
diff --git a/Engine/Source/Programs/Horde/Horde.Build.Tests/BlobStoreTests.cs b/Engine/Source/Programs/Horde/Horde.Build.Tests/BlobStoreTests.cs
--- a/Engine/Source/Programs/Horde/Horde.Build.Tests/BlobStoreTests.cs
+++ b/Engine/Source/Programs/Horde/Horde.Build.Tests/BlobStoreTests.cs
@@ -34,7 +34,7 @@ namespace Horde.Build.Tests
 			public List<BlobLocator> References { get; set; } = new List<BlobLocator>();
 		}
 
-		static Bundle CreateTestBundle(ReadOnlySequence<byte> data, IReadOnlyList<BlobLocator> refs)
+		static Bundle CreateTestBundle(ReadOnlyMemory<byte> data, IReadOnlyList<BlobLocator> refs)
 		{
 			List<BundleImport> imports = new List<BundleImport>();
 
@@ -51,13 +51,13 @@ namespace Horde.Build.Tests
 			}
 
 			List<BundleExport> exports = new List<BundleExport>();
-			exports.Add(new BundleExport(IoHash.Compute(data), (int)data.Length, refs.Select(x => locatorToIndex[x]).ToArray()));
+			exports.Add(new BundleExport(IoHash.Compute(data.Span), (int)data.Length, refs.Select(x => locatorToIndex[x]).ToArray()));
 
 			List<BundlePacket> packets = new List<BundlePacket>();
 			packets.Add(new BundlePacket((int)data.Length, (int)data.Length));
 
 			BundleHeader header = new BundleHeader(BundleCompressionFormat.None, imports, exports, packets);
-			return new Bundle(header, data);
+			return new Bundle(header, new[] { data });
 		}
 
 		static async Task<Blob> ReadBlobAsync(IStorageClient store, BlobLocator locator)
@@ -69,7 +69,7 @@ namespace Horde.Build.Tests
 		static Blob ExtractBlobFromBundle(Bundle bundle)
 		{
 			Blob blob = new Blob();
-			blob.Data = bundle.Payload.AsSingleSegment();
+			blob.Data = bundle.Packets[0];
 			blob.References = bundle.Header.Exports[0].References.Select(x => bundle.Header.Imports[x].Locator).ToList();
 			return blob;
 		}
@@ -80,10 +80,10 @@ namespace Horde.Build.Tests
 			IStorageClient store = CreateBlobStore();
 
 			byte[] input = CreateTestData(256, 0);
-			BlobLocator locator = await store.WriteBundleAsync(CreateTestBundle(new ReadOnlySequence<byte>(input), Array.Empty<BlobLocator>()));
+			BlobLocator locator = await store.WriteBundleAsync(CreateTestBundle(input, Array.Empty<BlobLocator>()));
 
 			Bundle outputBundle = await store.ReadBundleAsync(locator);
-			Assert.IsTrue(outputBundle.Payload.AsSingleSegment().Span.SequenceEqual(input));
+			Assert.IsTrue(outputBundle.Packets[0].Span.SequenceEqual(input));
 		}
 
 		[TestMethod]
@@ -92,19 +92,19 @@ namespace Horde.Build.Tests
 			IStorageClient store = CreateBlobStore();
 
 			byte[] input1 = CreateTestData(256, 1);
-			BlobLocator locator1 = await store.WriteBundleAsync(CreateTestBundle(new ReadOnlySequence<byte>(input1), Array.Empty<BlobLocator>()));
+			BlobLocator locator1 = await store.WriteBundleAsync(CreateTestBundle(input1, Array.Empty<BlobLocator>()));
 			Blob blob1 = await ReadBlobAsync(store, locator1);
 			Assert.IsTrue(blob1.Data.Span.SequenceEqual(input1));
 			Assert.IsTrue(blob1.References.SequenceEqual(Array.Empty<BlobLocator>()));
 
 			byte[] input2 = CreateTestData(256, 2);
-			BlobLocator locator2 = await store.WriteBundleAsync(CreateTestBundle(new ReadOnlySequence<byte>(input2), new BlobLocator[] { locator1 }));
+			BlobLocator locator2 = await store.WriteBundleAsync(CreateTestBundle(input2, new BlobLocator[] { locator1 }));
 			Blob blob2 = await ReadBlobAsync(store, locator2);
 			Assert.IsTrue(blob2.Data.Span.SequenceEqual(input2));
 			Assert.IsTrue(blob2.References.SequenceEqual(new BlobLocator[] { locator1 }));
 
 			byte[] input3 = CreateTestData(256, 3);
-			BlobLocator locator3 = await store.WriteBundleAsync(CreateTestBundle(new ReadOnlySequence<byte>(input3), new BlobLocator[] { locator1, locator2, locator1 }));
+			BlobLocator locator3 = await store.WriteBundleAsync(CreateTestBundle(input3, new BlobLocator[] { locator1, locator2, locator1 }));
 			Blob blob3 = await ReadBlobAsync(store, locator3);
 			Assert.IsTrue(blob3.Data.Span.SequenceEqual(input3));
 			Assert.IsTrue(blob3.References.SequenceEqual(new BlobLocator[] { locator1, locator2, locator1 }));
@@ -119,7 +119,7 @@ namespace Horde.Build.Tests
 
 			RefName refName2 = new RefName("hello2");
 
-			RefTarget refTargetId2 = await store.WriteRefValueAsync(refName2, CreateTestBundle(new ReadOnlySequence<byte>(input3), new BlobLocator[] { locator1, locator2 }));
+			RefTarget refTargetId2 = await store.WriteRefValueAsync(refName2, CreateTestBundle(input3, new BlobLocator[] { locator1, locator2 }));
 			Blob refTarget2 = await ReadBlobAsync(store, refTargetId2.Locator);
 			Assert.IsTrue(refTarget2.Data.Span.SequenceEqual(input3));
 			Assert.IsTrue(refTarget2.References.SequenceEqual(new BlobLocator[] { locator1, locator2 }));
diff --git a/Engine/Source/Programs/Shared/EpicGames.Core/StreamUtils.cs b/Engine/Source/Programs/Shared/EpicGames.Core/StreamUtils.cs
new file mode 100644
--- /dev/null
+++ b/Engine/Source/Programs/Shared/EpicGames.Core/StreamUtils.cs
@@ -0,0 +1,34 @@
+// Copyright Epic Games, Inc. All Rights Reserved.
+
+using System;
+using System.IO;
+using System.Threading;
+using System.Threading.Tasks;
+
+namespace EpicGames.Core
+{
+	/// <summary>
+	/// Utility methods for streams
+	/// </summary>
+	public static class StreamUtils
+	{
+		/// <summary>
+		/// Reads a fixed amount of data from a stream, throwing an exception if the entire buffer cannot be read.
+		/// </summary>
+		/// <param name="stream">Stream to read from</param>
+		/// <param name="buffer">Buffer to receive the output data</param>
+		/// <param name="cancellationToken">Cancellation token for the operation</param>
+		public static async Task ReadFixedLengthBytesAsync(this Stream stream, Memory<byte> buffer, CancellationToken cancellationToken = default)
+		{
+			for (int offset = 0; offset < buffer.Length;)
+			{
+				int readLength = await stream.ReadAsync(buffer.Slice(offset), cancellationToken);
+				if (readLength == 0)
+				{
+					throw new EndOfStreamException($"Unexpected end of stream while trying to read {buffer.Length} bytes.");
+				}
+				offset += readLength;
+			}
+		}
+	}
+}
diff --git a/Engine/Source/Programs/Shared/EpicGames.Horde/Storage/Bundle.cs b/Engine/Source/Programs/Shared/EpicGames.Horde/Storage/Bundle.cs
--- a/Engine/Source/Programs/Shared/EpicGames.Horde/Storage/Bundle.cs
+++ b/Engine/Source/Programs/Shared/EpicGames.Horde/Storage/Bundle.cs
@@ -1,11 +1,13 @@
 // Copyright Epic Games, Inc. All Rights Reserved.
 
 using EpicGames.Core;
+using K4os.Compression.LZ4;
 using System;
 using System.Buffers;
 using System.Buffers.Binary;
 using System.Collections.Generic;
 using System.IO;
+using System.IO.Compression;
 using System.Linq;
 using System.Text;
 using System.Threading;
@@ -50,15 +52,15 @@ namespace EpicGames.Horde.Storage
 		/// <summary>
 		/// Packet data as described in the header
 		/// </summary>
-		public ReadOnlySequence<byte> Payload { get; }
+		public IReadOnlyList<ReadOnlyMemory<byte>> Packets { get; }
 
 		/// <summary>
 		/// Constructor
 		/// </summary>
-		public Bundle(BundleHeader header, ReadOnlySequence<byte> payload)
+		public Bundle(BundleHeader header, IReadOnlyList<ReadOnlyMemory<byte>> packets)
 		{
 			Header = header;
-			Payload = payload;
+			Packets = packets;
 		}
 
 		/// <summary>
@@ -67,8 +69,15 @@ namespace EpicGames.Horde.Storage
 		public Bundle(IMemoryReader reader)
 		{
 			Header = new BundleHeader(reader);
-			int length = (int)reader.ReadUnsignedVarInt();
-			Payload = new ReadOnlySequence<byte>(reader.ReadFixedLengthBytes(length));
+
+			ReadOnlyMemory<byte>[] packets = new ReadOnlyMemory<byte>[Header.Packets.Count];
+			for (int idx = 0; idx < Header.Packets.Count; idx++)
+			{
+				BundlePacket packet = Header.Packets[idx];
+				packets[idx] = reader.ReadFixedLengthBytes(packet.EncodedLength);
+			}
+
+			Packets = packets;
 		}
 
 		/// <summary>
@@ -79,11 +88,20 @@ namespace EpicGames.Horde.Storage
 		/// <returns>Bundle that was read</returns>
 		public static async Task<Bundle> FromStreamAsync(Stream stream, CancellationToken cancellationToken)
 		{
-			using (MemoryStream memoryStream = new MemoryStream())
+			BundleHeader header = await BundleHeader.FromStreamAsync(stream, cancellationToken);
+
+			ReadOnlyMemory<byte>[] packets = new ReadOnlyMemory<byte>[header.Packets.Count];
+			for (int idx = 0; idx < header.Packets.Count; idx++)
 			{
-				await stream.CopyToAsync(memoryStream, cancellationToken);
-				return new Bundle(new MemoryReader(memoryStream.ToArray()));
+				BundlePacket packet = header.Packets[idx];
+
+				byte[] data = new byte[packet.EncodedLength];
+				await stream.ReadFixedLengthBytesAsync(data, cancellationToken);
+
+				packets[idx] = data;
 			}
+
+			return new Bundle(header, packets);
 		}
 
 		/// <summary>
@@ -94,11 +112,14 @@ namespace EpicGames.Horde.Storage
 		{
 			ByteArrayBuilder builder = new ByteArrayBuilder();
 			Header.Write(builder);
-			builder.WriteUnsignedVarInt((ulong)Payload.Length);
 
 			ReadOnlySequenceBuilder<byte> sequence = new ReadOnlySequenceBuilder<byte>();
 			sequence.Append(builder.AsSequence());
-			sequence.Append(Payload);
+
+			foreach (ReadOnlyMemory<byte> packet in Packets)
+			{
+				sequence.Append(packet);
+			}
 
 			return sequence.Construct();
 		}
@@ -212,29 +233,19 @@ namespace EpicGames.Horde.Storage
 		/// <param name="stream">Stream to read from</param>
 		/// <param name="cancellationToken">Cancellation token for the stream</param>
 		/// <returns>New header</returns>
-		public static async Task<BundleHeader> ReadAsync(Stream stream, CancellationToken cancellationToken)
+		public static async Task<BundleHeader> FromStreamAsync(Stream stream, CancellationToken cancellationToken)
 		{
 			byte[] prelude = new byte[8];
-
-			int length = await stream.ReadAsync(prelude, cancellationToken);
-			if (length != 8)
-			{
-				throw new InvalidDataException("Unexpected end of stream while reading prelude data");
-			}
+			await stream.ReadFixedLengthBytesAsync(prelude, cancellationToken);
 			CheckPrelude(prelude);
 
 			int headerLength = BinaryPrimitives.ReadInt32BigEndian(prelude.AsSpan(4));
 			using (IMemoryOwner<byte> owner = MemoryPool<byte>.Shared.Rent(headerLength))
 			{
-				Memory<byte> header = owner.Memory.Slice(headerLength);
+				Memory<byte> header = owner.Memory.Slice(0, headerLength);
 				prelude.CopyTo(header);
 
-				int readLength = await stream.ReadAsync(header.Slice(8), cancellationToken);
-				if (readLength != headerLength)
-				{
-					throw new InvalidDataException("Unexpected end of file in stream");
-				}
-
+				await stream.ReadFixedLengthBytesAsync(header.Slice(prelude.Length), cancellationToken);
 				return new BundleHeader(new MemoryReader(header));
 			}
 		}
@@ -360,12 +371,12 @@ namespace EpicGames.Horde.Storage
 		/// <summary>
 		/// Encoded length of the packet
 		/// </summary>
-		public int EncodedLength { get; set; }
+		public int EncodedLength { get; }
 
 		/// <summary>
 		/// Decoded length of the packet
 		/// </summary>
-		public int DecodedLength { get; set; }
+		public int DecodedLength { get; }
 
 		/// <summary>
 		/// Constructor
@@ -467,4 +478,94 @@ namespace EpicGames.Horde.Storage
 			}
 		}
 	}
+
+	/// <summary>
+	/// Utility methods for bundles
+	/// </summary>
+	public static class BundleData
+	{
+		/// <summary>
+		/// Compress a data packet
+		/// </summary>
+		/// <param name="format"></param>
+		/// <param name="input"></param>
+		/// <returns>The compressed data</returns>
+		public static ReadOnlyMemory<byte> Compress(BundleCompressionFormat format, ReadOnlyMemory<byte> input)
+		{
+			switch (format)
+			{
+				case BundleCompressionFormat.None:
+					{
+						return input;
+					}
+				case BundleCompressionFormat.LZ4:
+					{
+						int maxSize = LZ4Codec.MaximumOutputSize(input.Length);
+						using (IMemoryOwner<byte> buffer = MemoryPool<byte>.Shared.Rent(maxSize))
+						{
+							int encodedLength = LZ4Codec.Encode(input.Span, buffer.Memory.Span);
+							return buffer.Memory.Slice(0, encodedLength).ToArray();
+						}
+					}
+				case BundleCompressionFormat.Gzip:
+					{
+						using MemoryStream outputStream = new MemoryStream(input.Length);
+						using GZipStream gzipStream = new GZipStream(outputStream, CompressionLevel.Fastest);
+						gzipStream.Write(input.Span);
+						return outputStream.ToArray();
+					}
+				case BundleCompressionFormat.Oodle:
+					{
+#if WITH_OODLE
+						int maxSize = Oodle.MaximumOutputSize(OodleCompressorType.Selkie, input.Length);
+
+						Span<byte> outputSpan = builder.GetSpan(maxSize);
+						int encodedLength = Oodle.Compress(OodleCompressorType.Selkie, input.Span, outputSpan, OodleCompressionLevel.HyperFast);
+						builder.Advance(encodedLength);
+
+						return encodedLength;
+#else
+						throw new NotSupportedException("Oodle is not compiled into this build.");
+#endif
+					}
+				default:
+					throw new InvalidDataException($"Invalid compression format '{(int)format}'");
+			}
+		}
+
+		/// <summary>
+		/// Decompress a packet of data
+		/// </summary>
+		/// <param name="format">Format of the compressed data</param>
+		/// <param name="input">Compressed data</param>
+		/// <param name="output">Buffer to receive the decompressed data</param>
+		public static void Decompress(BundleCompressionFormat format, ReadOnlyMemory<byte> input, Memory<byte> output)
+		{
+			switch (format)
+			{
+				case BundleCompressionFormat.None:
+					input.CopyTo(output);
+					break;
+				case BundleCompressionFormat.LZ4:
+					LZ4Codec.Decode(input.Span, output.Span);
+					break;
+				case BundleCompressionFormat.Gzip:
+					{
+						using ReadOnlyMemoryStream inputStream = new ReadOnlyMemoryStream(input);
+						using GZipStream outputStream = new GZipStream(new MemoryWriterStream(new MemoryWriter(output)), CompressionMode.Decompress, false);
+						inputStream.CopyTo(outputStream);
+						break;
+					}
+				case BundleCompressionFormat.Oodle:
+#if WITH_OODLE
+					Oodle.Decompress(input.Span, output.Span);
+					break;
+#else
+					throw new NotSupportedException("Oodle is not compiled into this build.");
+#endif
+				default:
+					throw new InvalidDataException($"Invalid compression format '{(int)format}'");
+			}
+		}
+	}
 }
diff --git a/Engine/Source/Programs/Shared/EpicGames.Horde/Storage/TreeStore.cs b/Engine/Source/Programs/Shared/EpicGames.Horde/Storage/TreeStore.cs
--- a/Engine/Source/Programs/Shared/EpicGames.Horde/Storage/TreeStore.cs
+++ b/Engine/Source/Programs/Shared/EpicGames.Horde/Storage/TreeStore.cs
@@ -94,12 +94,14 @@ namespace EpicGames.Horde.Storage
 		/// </summary>
 		class PacketInfo
 		{
+			public readonly int Index;
 			public readonly int Offset;
 			public readonly int DecodedLength;
 			public readonly int EncodedLength;
 
-			public PacketInfo(int offset, BundlePacket packet)
+			public PacketInfo(int index, int offset, BundlePacket packet)
 			{
+				Index = index;
 				Offset = offset;
 				DecodedLength = packet.DecodedLength;
 				EncodedLength = packet.EncodedLength;
@@ -460,9 +462,10 @@ namespace EpicGames.Horde.Storage
 				// Update the node states to reference the written bundle
 				int exportIdx = 0;
 				int packetOffset = 0;
-				foreach (BundlePacket packet in bundle.Header.Packets)
+				for (int packetIdx = 0; packetIdx < bundle.Header.Packets.Count; packetIdx++)
 				{
-					PacketInfo packetInfo = new PacketInfo(packetOffset, packet);
+					BundlePacket packet = bundle.Header.Packets[packetIdx];
+					PacketInfo packetInfo = new PacketInfo(packetIdx, packetOffset, packet);
 
 					int nodeOffset = 0;
 					for (; exportIdx < header.Exports.Count && nodeOffset + bundle.Header.Exports[exportIdx].Length <= packet.DecodedLength; exportIdx++)
@@ -552,9 +555,6 @@ namespace EpicGames.Horde.Storage
 					imports.Add(new BundleImport(bundleInfo.Locator, bundleInfo.Exports.Length, exportInfos));
 				}
 
-				// Preallocate data in the encoded buffer to reduce fragmentation if we have to resize
-				ByteArrayBuilder builder = new ByteArrayBuilder(options.MinCompressionPacketSize * 2);
-
 				// Create the export list
 				List<BundleExport> exports = new List<BundleExport>();
 				List<BundlePacket> packets = new List<BundlePacket>();
@@ -562,6 +562,9 @@ namespace EpicGames.Horde.Storage
 				// Size of data currently stored in the block buffer
 				int blockSize = 0;
 
+				// List of data packets
+				List<ReadOnlyMemory<byte>> packetBlocks = new List<ReadOnlyMemory<byte>>();
+
 				// Segments of data in the current block
 				List<ReadOnlyMemory<byte>> blockSegments = new List<ReadOnlyMemory<byte>>();
 
@@ -574,7 +577,7 @@ namespace EpicGames.Horde.Storage
 					// If we can't fit this data into the current block, flush the contents of it first
 					if (blockSize > 0 && blockSize + nodeData.Length > options.MinCompressionPacketSize)
 					{
-						FlushPacket(_options.CompressionFormat, blockSegments, blockSize, builder, packets);
+						FlushPacket(_options.CompressionFormat, blockSegments, blockSize, packets, packetBlocks);
 						blockSize = 0;
 					}
 
@@ -591,14 +594,14 @@ namespace EpicGames.Horde.Storage
 					}
 					else
 					{
-						FlushPacket(options.CompressionFormat, nodeData.First, builder, packets);
+						FlushPacket(options.CompressionFormat, nodeData.First, packets, packetBlocks);
 					}
 				}
-				FlushPacket(options.CompressionFormat, blockSegments, blockSize, builder, packets);
+				FlushPacket(options.CompressionFormat, blockSegments, blockSize, packets, packetBlocks);
 
 				// Flush the data
 				BundleHeader header = new BundleHeader(options.CompressionFormat, imports, exports, packets);
-				return new Bundle(header, builder.AsSequence());
+				return new Bundle(header, packetBlocks);
 			}
 
 			static int AddSegments(ReadOnlySequence<byte> sequence, List<ReadOnlyMemory<byte>> segments)
@@ -612,13 +615,13 @@ namespace EpicGames.Horde.Storage
 				return size;
 			}
 
-			static void FlushPacket(BundleCompressionFormat format, List<ReadOnlyMemory<byte>> blockSegments, int blockSize, ByteArrayBuilder builder, List<BundlePacket> packets)
+			static void FlushPacket(BundleCompressionFormat format, List<ReadOnlyMemory<byte>> blockSegments, int blockSize, List<BundlePacket> packets, List<ReadOnlyMemory<byte>> packetBlocks)
 			{
 				if (blockSize > 0)
 				{
 					if (blockSegments.Count == 1)
 					{
-						FlushPacket(format, blockSegments[0], builder, packets);
+						FlushPacket(format, blockSegments[0], packets, packetBlocks);
 					}
 					else
 					{
@@ -631,64 +634,19 @@ namespace EpicGames.Horde.Storage
 							output = output.Slice(blockSegment.Length);
 						}
 
-						FlushPacket(format, buffer.Memory.Slice(0, blockSize), builder, packets);
+						FlushPacket(format, buffer.Memory.Slice(0, blockSize), packets, packetBlocks);
 					}
 				}
 				blockSegments.Clear();
 			}
 
-			static void FlushPacket(BundleCompressionFormat format, ReadOnlyMemory<byte> inputData, ByteArrayBuilder builder, List<BundlePacket> packets)
+			static void FlushPacket(BundleCompressionFormat format, ReadOnlyMemory<byte> inputData, List<BundlePacket> packets, List<ReadOnlyMemory<byte>> packetBlocks)
 			{
 				if (inputData.Length > 0)
 				{
-					int encodedLength = Encode(format, inputData, builder);
-					Debug.Assert(encodedLength >= 0);
-					packets.Add(new BundlePacket(encodedLength, inputData.Length));
-				}
-			}
-
-			static int Encode(BundleCompressionFormat format, ReadOnlyMemory<byte> input, ByteArrayBuilder builder)
-			{
-				switch (format)
-				{
-					case BundleCompressionFormat.None:
-						{
-							builder.WriteFixedLengthBytes(input.Span);
-							return input.Length;
-						}
-					case BundleCompressionFormat.LZ4:
-						{
-							int maxSize = LZ4Codec.MaximumOutputSize(input.Length);
-
-							Span<byte> outputSpan = builder.GetSpan(maxSize);
-							int encodedLength = LZ4Codec.Encode(input.Span, outputSpan);
-							builder.Advance(encodedLength);
-
-							return encodedLength;
-						}
-					case BundleCompressionFormat.Gzip:
-						{
-							using MemoryWriterStream outputStream = new MemoryWriterStream(builder);
-							using GZipStream gzipStream = new GZipStream(outputStream, CompressionLevel.Fastest);
-							gzipStream.Write(input.Span);
-							return (int)outputStream.Length;
-						}
-					case BundleCompressionFormat.Oodle:
-						{
-#if WITH_OODLE
-							int maxSize = Oodle.MaximumOutputSize(OodleCompressorType.Selkie, input.Length);
-
-							Span<byte> outputSpan = builder.GetSpan(maxSize);
-							int encodedLength = Oodle.Compress(OodleCompressorType.Selkie, input.Span, outputSpan, OodleCompressionLevel.HyperFast);
-							builder.Advance(encodedLength);
-
-							return encodedLength;
-#else
-							throw new NotSupportedException("Oodle is not compiled into this build.");
-#endif
-						}
-					default:
-						throw new InvalidDataException($"Invalid compression format '{(int)format}'");
+					ReadOnlyMemory<byte> encodedData = BundleData.Compress(format, inputData);
+					packets.Add(new BundlePacket(encodedData.Length, inputData.Length));
+					packetBlocks.Add(encodedData);
 				}
 			}
 		}
@@ -757,7 +715,7 @@ namespace EpicGames.Horde.Storage
 			ExportedNodeState exportedState = await GetExportedStateAsync(node, cancellationToken);
 
 			Bundle bundle = await ReadBundleAsync(exportedState.BundleInfo.Locator, cancellationToken);
-			ReadOnlyMemory<byte> packetData = DecodePacket(exportedState.BundleInfo.Locator, exportedState.PacketInfo, bundle.Header.CompressionFormat, bundle.Payload);
+			ReadOnlyMemory<byte> packetData = DecodePacket(exportedState.BundleInfo.Locator, exportedState.PacketInfo, bundle.Header.CompressionFormat, bundle.Packets[exportedState.PacketInfo.Index]);
 			ReadOnlyMemory<byte> nodeData = packetData.Slice(exportedState.Offset, exportedState.Length);
 
 			return new InMemoryNodeState(new ReadOnlySequence<byte>(nodeData), exportedState.References);
@@ -835,7 +793,8 @@ namespace EpicGames.Horde.Storage
 			{
 				using (ICacheEntry entry = _cache.CreateEntry(cacheKey))
 				{
-					entry.SetSize(bundle.Payload.Length);
+					int length = bundle.Packets.Sum(x => x.Length);
+					entry.SetSize(length);
 					entry.SetValue(bundle);
 				}
 			}
@@ -879,9 +838,10 @@ namespace EpicGames.Horde.Storage
 			// Create the exported nodes, or update the state of any imported nodes to exported
 			int exportIdx = 0;
 			int packetOffset = 0;
-			foreach (BundlePacket packet in header.Packets)
+			for (int packetIdx = 0; packetIdx < header.Packets.Count; packetIdx++)
 			{
-				PacketInfo packetInfo = new PacketInfo(packetOffset, packet);
+				BundlePacket packet = header.Packets[packetIdx];
+				PacketInfo packetInfo = new PacketInfo(packetIdx, packetOffset, packet);
 
 				int nodeOffset = 0;
 				for (; exportIdx < header.Exports.Count && nodeOffset + header.Exports[exportIdx].Length <= packet.DecodedLength; exportIdx++)
@@ -939,7 +899,7 @@ namespace EpicGames.Horde.Storage
 		/// <param name="format">Compression type in the bundle</param>
 		/// <param name="payload">The bundle payload</param>
 		/// <returns>The decoded data</returns>
-		ReadOnlyMemory<byte> DecodePacket(BlobLocator locator, PacketInfo packetInfo, BundleCompressionFormat format, ReadOnlySequence<byte> payload)
+		ReadOnlyMemory<byte> DecodePacket(BlobLocator locator, PacketInfo packetInfo, BundleCompressionFormat format, ReadOnlyMemory<byte> payload)
 		{
 			if (_cache == null)
 			{
@@ -957,13 +917,10 @@ namespace EpicGames.Horde.Storage
 			}
 		}
 
-		static ReadOnlyMemory<byte> DecodePacketUncached(PacketInfo packetInfo, BundleCompressionFormat format, ReadOnlySequence<byte> payload)
+		static ReadOnlyMemory<byte> DecodePacketUncached(PacketInfo packetInfo, BundleCompressionFormat format, ReadOnlyMemory<byte> encodedPacket)
 		{
-			ReadOnlyMemory<byte> encodedPacket = payload.Slice(packetInfo.Offset, packetInfo.EncodedLength).AsSingleSegment();
-
 			byte[] decodedPacket = new byte[packetInfo.DecodedLength];
-			Decode(format, encodedPacket, decodedPacket);
-
+			BundleData.Decompress(format, encodedPacket, decodedPacket);
 			return decodedPacket;
 		}
 
@@ -986,35 +943,6 @@ namespace EpicGames.Horde.Storage
 			return IoHash.Compute(buffer);
 		}
 
-		static void Decode(BundleCompressionFormat format, ReadOnlyMemory<byte> input, Memory<byte> output)
-		{
-			switch (format)
-			{
-				case BundleCompressionFormat.None:
-					input.CopyTo(output);
-					break;
-				case BundleCompressionFormat.LZ4:
-					LZ4Codec.Decode(input.Span, output.Span);
-					break;
-				case BundleCompressionFormat.Gzip:
-					{
-						using ReadOnlyMemoryStream inputStream = new ReadOnlyMemoryStream(input);
-						using GZipStream outputStream = new GZipStream(new MemoryWriterStream(new MemoryWriter(output)), CompressionMode.Decompress, false);
-						inputStream.CopyTo(outputStream);
-						break;
-					}
-				case BundleCompressionFormat.Oodle:
-#if WITH_OODLE
-					Oodle.Decompress(input.Span, output.Span);
-					break;
-#else
-					throw new NotSupportedException("Oodle is not compiled into this build.");
-#endif
-				default:
-					throw new InvalidDataException($"Invalid compression format '{(int)format}'");
-			}
-		}
-
 		#endregion
 	}
 }
