Horde: Move cache management and partial reads of bundles into the implementation of IStorageClient.

Implementation of TreeStore will likely be the same for all implementations, while there are many different strategies we might use for bulk data.

diff --git a/Engine/Source/Programs/Horde/Horde.Agent/Commands/Bundles/DumpCommand.cs b/Engine/Source/Programs/Horde/Horde.Agent/Commands/Bundles/DumpCommand.cs
--- a/Engine/Source/Programs/Horde/Horde.Agent/Commands/Bundles/DumpCommand.cs
+++ b/Engine/Source/Programs/Horde/Horde.Agent/Commands/Bundles/DumpCommand.cs
@@ -23,19 +23,18 @@ namespace Horde.Agent.Commands.Bundles
 		{
 			IStorageClient storage = CreateStorageClient(logger);
 
-			Bundle blob;
+			BundleHeader header;
 			if (BlobId != null)
 			{
-				blob = await storage.ReadBundleAsync(BlobId.Value);
+				header = await storage.ReadBundleHeaderAsync(BlobId.Value);
 				logger.LogInformation("Summary for blob {BlobId}", BlobId.Value);
 			}
 			else
 			{
-				blob = (await storage.ReadRefValueAsync(RefName)).Bundle;
+				header = (await storage.ReadRefValueAsync(RefName)).Bundle.Header;
 				logger.LogInformation("Summary for ref {RefId}", RefName);
 			}
 
-			BundleHeader header = blob.Header;
 			int packetStart = 0;
 
 			logger.LogInformation("");
diff --git a/Engine/Source/Programs/Shared/EpicGames.Horde/Storage/IStorageClient.cs b/Engine/Source/Programs/Shared/EpicGames.Horde/Storage/IStorageClient.cs
--- a/Engine/Source/Programs/Shared/EpicGames.Horde/Storage/IStorageClient.cs
+++ b/Engine/Source/Programs/Shared/EpicGames.Horde/Storage/IStorageClient.cs
@@ -101,23 +101,51 @@ namespace EpicGames.Horde.Storage
 		Task<Stream> ReadBlobAsync(BlobLocator locator, CancellationToken cancellationToken = default);
 
 		/// <summary>
-		/// Reads part of a blob from storage
+		/// Writes a new blob to the store
 		/// </summary>
-		/// <param name="locator">The blob locator</param>
-		/// <param name="offset">Offset within the blob to read from</param>
-		/// <param name="length">Maximum length of data to read. The stream will be truncated if the blob is shorter than this.</param>
+		/// <param name="stream">Blob data</param>
 		/// <param name="cancellationToken">Cancellation token for the operation</param>
-		/// <returns>Stream containing the data</returns>
-		Task<Stream> ReadBlobRangeAsync(BlobLocator locator, int offset = 0, int length = Int32.MaxValue, CancellationToken cancellationToken = default);
+		/// <param name="prefix">Prefix for blob names. While the returned BlobId is guaranteed to be unique, this name can be used as a prefix to aid debugging.</param>
+		/// <returns>Unique identifier for the blob</returns>
+		Task<BlobLocator> WriteBlobAsync(Stream stream, Utf8String prefix = default, CancellationToken cancellationToken = default);
+
+		#endregion
+
+		#region Bundles
 
 		/// <summary>
-		/// Writes a new blob to the store
+		/// Reads data for a bundle from the store
 		/// </summary>
-		/// <param name="stream">Blob data</param>
+		/// <param name="locator">The blob identifier</param>
+		/// <param name="cancellationToken">Cancellation token for the operation</param>
+		/// <returns></returns>
+		Task<Bundle> ReadBundleAsync(BlobLocator locator, CancellationToken cancellationToken = default);
+
+		/// <summary>
+		/// Reads the header for a bundle from the store
+		/// </summary>
+		/// <param name="locator">Locator for the blob</param>
+		/// <param name="cancellationToken">Cancellation token for the operation</param>
+		/// <returns></returns>
+		Task<BundleHeader> ReadBundleHeaderAsync(BlobLocator locator, CancellationToken cancellationToken = default);
+
+		/// <summary>
+		/// Reads a bundle packet from the store
+		/// </summary>
+		/// <param name="locator">Locator for the blob</param>
+		/// <param name="packetIdx">Index of the packet</param>
+		/// <param name="cancellationToken">Cancellation token for the operation</param>
+		/// <returns>Decoded data for the packet</returns>
+		Task<ReadOnlyMemory<byte>> ReadBundlePacketAsync(BlobLocator locator, int packetIdx, CancellationToken cancellationToken = default);
+
+		/// <summary>
+		/// Writes a new bundle to the store
+		/// </summary>
+		/// <param name="bundle">Bundle data</param>
 		/// <param name="cancellationToken">Cancellation token for the operation</param>
 		/// <param name="prefix">Prefix for blob names. While the returned BlobId is guaranteed to be unique, this name can be used as a prefix to aid debugging.</param>
 		/// <returns>Unique identifier for the blob</returns>
-		Task<BlobLocator> WriteBlobAsync(Stream stream, Utf8String prefix = default, CancellationToken cancellationToken = default);
+		Task<BlobLocator> WriteBundleAsync(Bundle bundle, Utf8String prefix = default, CancellationToken cancellationToken = default);
 
 		#endregion
 
@@ -216,10 +244,23 @@ namespace EpicGames.Horde.Storage
 			public Task<Stream> ReadBlobAsync(BlobLocator locator, CancellationToken cancellationToken = default) => _inner.ReadBlobAsync(locator, cancellationToken);
 
 			/// <inheritdoc/>
-			public Task<Stream> ReadBlobRangeAsync(BlobLocator locator, int offset, int length, CancellationToken cancellationToken = default) => _inner.ReadBlobRangeAsync(locator, offset, length, cancellationToken);
+			public Task<BlobLocator> WriteBlobAsync(Stream stream, Utf8String prefix = default, CancellationToken cancellationToken = default) => _inner.WriteBlobAsync(stream, prefix, cancellationToken);
+
+			#endregion
+
+			#region Bundles
 
 			/// <inheritdoc/>
-			public Task<BlobLocator> WriteBlobAsync(Stream stream, Utf8String prefix = default, CancellationToken cancellationToken = default) => _inner.WriteBlobAsync(stream, prefix, cancellationToken);
+			public Task<Bundle> ReadBundleAsync(BlobLocator locator, CancellationToken cancellationToken = default) => _inner.ReadBundleAsync(locator, cancellationToken);
+
+			/// <inheritdoc/>
+			public Task<BundleHeader> ReadBundleHeaderAsync(BlobLocator locator, CancellationToken cancellationToken = default) => _inner.ReadBundleHeaderAsync(locator, cancellationToken);
+
+			/// <inheritdoc/>
+			public Task<ReadOnlyMemory<byte>> ReadBundlePacketAsync(BlobLocator locator, int packetIdx, CancellationToken cancellationToken = default) => _inner.ReadBundlePacketAsync(locator, packetIdx, cancellationToken);
+
+			/// <inheritdoc/>
+			public Task<BlobLocator> WriteBundleAsync(Bundle bundle, Utf8String prefix = default, CancellationToken cancellationToken = default) => _inner.WriteBundleAsync(bundle, prefix, cancellationToken);
 
 			#endregion
 
@@ -260,33 +301,6 @@ namespace EpicGames.Horde.Storage
 		/// <returns></returns>
 		public static IStorageClient<T> ForType<T>(this IStorageClient blobStore) => new TypedStorageClient<T>(blobStore);
 
-		/// <summary>
-		/// Reads data for a bundle from the store
-		/// </summary>
-		/// <param name="store">The store instance to read from</param>
-		/// <param name="locator">The blob identifier</param>
-		/// <param name="cancellationToken">Cancellation token for the operation</param>
-		/// <returns></returns>
-		public static async Task<Bundle> ReadBundleAsync(this IStorageClient store, BlobLocator locator, CancellationToken cancellationToken = default)
-		{
-			using Stream stream = await store.ReadBlobAsync(locator, cancellationToken);
-			return await Bundle.FromStreamAsync(stream, cancellationToken);
-		}
-
-		/// <summary>
-		/// Writes a new bundle to the store
-		/// </summary>
-		/// <param name="store">The store instance to read from</param>
-		/// <param name="bundle">Bundle data</param>
-		/// <param name="cancellationToken">Cancellation token for the operation</param>
-		/// <param name="prefix">Prefix for blob names. While the returned BlobId is guaranteed to be unique, this name can be used as a prefix to aid debugging.</param>
-		/// <returns>Unique identifier for the blob</returns>
-		public static async Task<BlobLocator> WriteBundleAsync(this IStorageClient store, Bundle bundle, Utf8String prefix = default, CancellationToken cancellationToken = default)
-		{
-			using ReadOnlySequenceStream stream = new ReadOnlySequenceStream(bundle.AsSequence());
-			return await store.WriteBlobAsync(stream, prefix, cancellationToken);
-		}
-
 		/// <summary>
 		/// Checks if the given ref exists
 		/// </summary>
diff --git a/Engine/Source/Programs/Shared/EpicGames.Horde/Storage/StorageClientBase.cs b/Engine/Source/Programs/Shared/EpicGames.Horde/Storage/StorageClientBase.cs
--- a/Engine/Source/Programs/Shared/EpicGames.Horde/Storage/StorageClientBase.cs
+++ b/Engine/Source/Programs/Shared/EpicGames.Horde/Storage/StorageClientBase.cs
@@ -94,7 +94,13 @@ namespace EpicGames.Horde.Storage
 			public override void Write(byte[] buffer, int offset, int count) => throw new NotSupportedException();
 		}
 
+		readonly IMemoryCache? _cache;
 		readonly TreeStore _treeStore;
+		readonly object _lockObject = new object();
+
+		// Active read tasks at any moment. If a BundleObject is not available in the cache, we start a read and add an entry to this dictionary
+		// so that other threads can also await it.
+		readonly Dictionary<string, Task<Bundle>> _readTasks = new Dictionary<string, Task<Bundle>>();
 
 		/// <summary>
 		/// Constructor
@@ -102,7 +108,8 @@ namespace EpicGames.Horde.Storage
 		/// <param name="cache"></param>
 		protected StorageClientBase(IMemoryCache? cache)
 		{
-			_treeStore = new TreeStore(this, cache);
+			_cache = cache;
+			_treeStore = new TreeStore(this);
 		}
 
 		#region Blobs
@@ -123,6 +130,146 @@ namespace EpicGames.Horde.Storage
 
 		#endregion
 
+		#region Bundles
+
+		static string GetBundleCacheKey(BlobId blobId) => $"bundle:{blobId}";
+
+		/// <summary>
+		/// Adds a bundle to the cache
+		/// </summary>
+		/// <param name="blobId">Blob id for the bundle</param>
+		/// <param name="bundle">The bundle data</param>
+		void AddBundleToCache(BlobId blobId, Bundle bundle)
+		{
+			if (_cache != null)
+			{
+				string cacheKey = GetBundleCacheKey(blobId);
+				using (ICacheEntry entry = _cache.CreateEntry(cacheKey))
+				{
+					int length = bundle.Packets.Sum(x => x.Length);
+					entry.SetSize(length);
+					entry.SetValue(bundle);
+				}
+			}
+		}
+
+		/// <summary>
+		/// Reads a bundle from the given blob id, or retrieves it from the cache
+		/// </summary>
+		/// <param name="locator"></param>
+		/// <param name="cancellationToken"></param>
+		/// <returns></returns>
+		public async Task<Bundle> ReadBundleAsync(BlobLocator locator, CancellationToken cancellationToken = default)
+		{
+			string cacheKey = GetBundleCacheKey(locator.BlobId);
+			if (_cache == null || !_cache.TryGetValue<Bundle>(cacheKey, out Bundle? bundle))
+			{
+				Task<Bundle>? readTask;
+				lock (_lockObject)
+				{
+					if (!_readTasks.TryGetValue(cacheKey, out readTask))
+					{
+						readTask = Task.Run(() => ReadBundleInternalAsync(cacheKey, locator, cancellationToken), cancellationToken);
+						_readTasks.Add(cacheKey, readTask);
+					}
+				}
+				bundle = await readTask;
+			}
+			return bundle;
+		}
+
+		async Task<Bundle> ReadBundleInternalAsync(string cacheKey, BlobLocator locator, CancellationToken cancellationToken)
+		{
+			// Perform another (sequenced) check whether an object has been added to the cache, to counteract the race between a read task being added and a task completing.
+			lock (_lockObject)
+			{
+				if (_cache != null && _cache.TryGetValue<Bundle>(cacheKey, out Bundle? cachedObject))
+				{
+					return cachedObject;
+				}
+			}
+
+			// Read the data from storage
+			Bundle bundle;
+			using (Stream stream = await ReadBlobAsync(locator, cancellationToken))
+			{
+				bundle = await Bundle.FromStreamAsync(stream, cancellationToken);
+			}
+
+			// Add the object to the cache
+			AddBundleToCache(locator.BlobId, bundle);
+
+			// Remove this object from the list of read tasks
+			lock (_lockObject)
+			{
+				_readTasks.Remove(cacheKey);
+			}
+			return bundle;
+		}
+
+		/// <inheritdoc/>
+		public Task<Stream> CreateBundleStreamAsync(BlobLocator locator, CancellationToken cancellationToken = default)
+		{
+			return ReadBlobAsync(locator, cancellationToken);
+		}
+
+		/// <inheritdoc/>
+		public async Task<BundleHeader> ReadBundleHeaderAsync(BlobLocator locator, CancellationToken cancellationToken = default)
+		{
+			Bundle bundle = await ReadBundleAsync(locator, cancellationToken);
+			return bundle.Header;
+		}
+
+		/// <inheritdoc/>
+		public async Task<ReadOnlyMemory<byte>> ReadBundlePacketAsync(BlobLocator locator, int packetIdx, CancellationToken cancellationToken = default)
+		{
+			Bundle bundle = await ReadBundleAsync(locator, cancellationToken);
+			return DecodePacket(locator, packetIdx, bundle.Header.Packets[packetIdx], bundle.Header.CompressionFormat, bundle.Packets[packetIdx]);
+		}
+
+		/// <summary>
+		/// Gets a decoded block from the store
+		/// </summary>
+		/// <param name="locator">Information about the bundle</param>
+		/// <param name="packetIdx">Index of the packet</param>
+		/// <param name="packet">The decoded block location and size</param>
+		/// <param name="format">Compression type in the bundle</param>
+		/// <param name="encodedPacket">The encoded packet data</param>
+		/// <returns>The decoded data</returns>
+		ReadOnlyMemory<byte> DecodePacket(BlobLocator locator, int packetIdx, BundlePacket packet, BundleCompressionFormat format, ReadOnlyMemory<byte> encodedPacket)
+		{
+			if (_cache == null)
+			{
+				return DecodePacketUncached(packet, format, encodedPacket);
+			}
+			else
+			{
+				string cacheKey = $"bundle-packet:{locator.BlobId}#{packetIdx}";
+				return _cache.GetOrCreate<ReadOnlyMemory<byte>>(cacheKey, entry =>
+				{
+					ReadOnlyMemory<byte> decodedPacket = DecodePacketUncached(packet, format, encodedPacket);
+					entry.SetSize(packet.DecodedLength);
+					return decodedPacket;
+				});
+			}
+		}
+
+		static ReadOnlyMemory<byte> DecodePacketUncached(BundlePacket packet, BundleCompressionFormat format, ReadOnlyMemory<byte> encodedPacket)
+		{
+			byte[] decodedPacket = new byte[packet.DecodedLength];
+			BundleData.Decompress(format, encodedPacket, decodedPacket);
+			return decodedPacket;
+		}
+
+		/// <inheritdoc/>
+		public async Task<BlobLocator> WriteBundleAsync(Bundle bundle, Utf8String prefix = default, CancellationToken cancellationToken = default)
+		{
+			using ReadOnlySequenceStream stream = new ReadOnlySequenceStream(bundle.AsSequence());
+			return await WriteBlobAsync(stream, prefix, cancellationToken);
+		}
+
+		#endregion
+
 		#region Trees
 
 		/// <inheritdoc/>
@@ -152,7 +299,7 @@ namespace EpicGames.Horde.Storage
 			{
 				return null;
 			}
-			return new RefValue(refTarget, await this.ReadBundleAsync(refTarget.Locator, cancellationToken));
+			return new RefValue(refTarget, await ReadBundleAsync(refTarget.Locator, cancellationToken));
 		}
 
 		/// <inheritdoc/>
diff --git a/Engine/Source/Programs/Shared/EpicGames.Horde/Storage/TreeStore.cs b/Engine/Source/Programs/Shared/EpicGames.Horde/Storage/TreeStore.cs
--- a/Engine/Source/Programs/Shared/EpicGames.Horde/Storage/TreeStore.cs
+++ b/Engine/Source/Programs/Shared/EpicGames.Horde/Storage/TreeStore.cs
@@ -44,11 +44,6 @@ namespace EpicGames.Horde.Storage
 		/// Number of nodes to retain in the working set after performing a partial flush
 		/// </summary>
 		public int TrimIgnoreCount { get; set; } = 150;
-
-		/// <summary>
-		/// Size of the bundle/decode cache
-		/// </summary>
-		public long CacheSize { get; set; }
 	}
 
 	/// <summary>
@@ -449,8 +444,6 @@ namespace EpicGames.Horde.Storage
 
 				// Write it to storage
 				BlobLocator locator = await _owner._blobStore.WriteBundleAsync(bundle, _prefix, cancellationToken);
-				string cacheKey = GetBundleCacheKey(locator);
-				_owner.AddBundleToCache(cacheKey, bundle);
 
 				// Create a BundleInfo for it
 				BundleInfo bundleInfo = _context.FindOrAddBundle(locator, writeNodes.Length);
@@ -652,21 +645,13 @@ namespace EpicGames.Horde.Storage
 		}
 
 		readonly IStorageClient _blobStore;
-		readonly IMemoryCache? _cache;
-
-		readonly object _lockObject = new object();
-
-		// Active read tasks at any moment. If a BundleObject is not available in the cache, we start a read and add an entry to this dictionary
-		// so that other threads can also await it.
-		readonly Dictionary<string, Task<Bundle>> _readTasks = new Dictionary<string, Task<Bundle>>();
 
 		/// <summary>
 		/// Constructor
 		/// </summary>
-		public TreeStore(IStorageClient blobStore, IMemoryCache? cache)
+		public TreeStore(IStorageClient blobStore)
 		{
 			_blobStore = blobStore;
-			_cache = cache;
 		}
 
 		/// <inheritdoc/>
@@ -693,7 +678,7 @@ namespace EpicGames.Horde.Storage
 
 			// Add the new bundle to it
 			BundleInfo bundleInfo = context.FindOrAddBundle(refValue.Target.Locator, refValue.Bundle.Header.Exports.Count);
-			MountBundle(bundleInfo, refValue.Bundle);
+			MountBundle(bundleInfo, refValue.Bundle.Header);
 
 			// Return the last node from the bundle as the root
 			return await GetNodeAsync(bundleInfo.Exports[^1]!, cancellationToken);
@@ -714,8 +699,7 @@ namespace EpicGames.Horde.Storage
 
 			ExportedNodeState exportedState = await GetExportedStateAsync(node, cancellationToken);
 
-			Bundle bundle = await ReadBundleAsync(exportedState.BundleInfo.Locator, cancellationToken);
-			ReadOnlyMemory<byte> packetData = DecodePacket(exportedState.BundleInfo.Locator, exportedState.PacketInfo, bundle.Header.CompressionFormat, bundle.Packets[exportedState.PacketInfo.Index]);
+			ReadOnlyMemory<byte> packetData = await _blobStore.ReadBundlePacketAsync(exportedState.BundleInfo.Locator, exportedState.PacketInfo.Index, cancellationToken);
 			ReadOnlyMemory<byte> nodeData = packetData.Slice(exportedState.Offset, exportedState.Length);
 
 			return new InMemoryNodeState(new ReadOnlySequence<byte>(nodeData), exportedState.References);
@@ -735,71 +719,6 @@ namespace EpicGames.Horde.Storage
 
 		#region Reading bundles
 
-		/// <summary>
-		/// Reads a bundle from the given blob id, or retrieves it from the cache
-		/// </summary>
-		/// <param name="locator"></param>
-		/// <param name="cancellationToken"></param>
-		/// <returns></returns>
-		async Task<Bundle> ReadBundleAsync(BlobLocator locator, CancellationToken cancellationToken = default)
-		{
-			string cacheKey = GetBundleCacheKey(locator);
-			if (_cache == null || !_cache.TryGetValue<Bundle>(cacheKey, out Bundle? bundle))
-			{
-				Task<Bundle>? readTask;
-				lock (_lockObject)
-				{
-					if (!_readTasks.TryGetValue(cacheKey, out readTask))
-					{
-						readTask = Task.Run(() => ReadBundleInternalAsync(cacheKey, locator, cancellationToken), cancellationToken);
-						_readTasks.Add(cacheKey, readTask);
-					}
-				}
-				bundle = await readTask;
-			}
-			return bundle;
-		}
-
-		async Task<Bundle> ReadBundleInternalAsync(string cacheKey, BlobLocator locator, CancellationToken cancellationToken)
-		{
-			// Perform another (sequenced) check whether an object has been added to the cache, to counteract the race between a read task being added and a task completing.
-			lock (_lockObject)
-			{
-				if (_cache != null && _cache.TryGetValue<Bundle>(cacheKey, out Bundle? cachedObject))
-				{
-					return cachedObject;
-				}
-			}
-
-			// Read the data from storage
-			Bundle bundle = await _blobStore.ReadBundleAsync(locator, cancellationToken);
-
-			// Add the object to the cache
-			AddBundleToCache(cacheKey, bundle);
-
-			// Remove this object from the list of read tasks
-			lock (_lockObject)
-			{
-				_readTasks.Remove(cacheKey);
-			}
-			return bundle;
-		}
-
-		static string GetBundleCacheKey(BlobLocator blobId) => $"bundle:{blobId}";
-
-		void AddBundleToCache(string cacheKey, Bundle bundle)
-		{
-			if (_cache != null)
-			{
-				using (ICacheEntry entry = _cache.CreateEntry(cacheKey))
-				{
-					int length = bundle.Packets.Sum(x => x.Length);
-					entry.SetSize(length);
-					entry.SetValue(bundle);
-				}
-			}
-		}
-
 		async Task MountBundleAsync(BundleInfo bundleInfo, CancellationToken cancellationToken)
 		{
 			if (!bundleInfo.Mounted)
@@ -823,15 +742,13 @@ namespace EpicGames.Horde.Storage
 			// If it didn't succeed, try again
 			if (!bundleInfo.Mounted)
 			{
-				Bundle bundle = await ReadBundleAsync(bundleInfo.Locator, cancellationToken);
-				MountBundle(bundleInfo, bundle);
+				BundleHeader header = await _blobStore.ReadBundleHeaderAsync(bundleInfo.Locator, cancellationToken);
+				MountBundle(bundleInfo, header);
 			}
 		}
 
-		void MountBundle(BundleInfo bundleInfo, Bundle bundle)
+		void MountBundle(BundleInfo bundleInfo, BundleHeader header)
 		{
-			BundleHeader header = bundle.Header;
-
 			// Create all the imported nodes
 			List<NodeInfo> nodes = CreateImports(bundleInfo.Context, header);
 
@@ -891,38 +808,6 @@ namespace EpicGames.Horde.Storage
 			}
 			return indexedNodes;
 		}
-		/// <summary>
-		/// Gets a decoded block from the store
-		/// </summary>
-		/// <param name="locator">Information about the bundle</param>
-		/// <param name="packetInfo">The decoded block location and size</param>
-		/// <param name="format">Compression type in the bundle</param>
-		/// <param name="payload">The bundle payload</param>
-		/// <returns>The decoded data</returns>
-		ReadOnlyMemory<byte> DecodePacket(BlobLocator locator, PacketInfo packetInfo, BundleCompressionFormat format, ReadOnlyMemory<byte> payload)
-		{
-			if (_cache == null)
-			{
-				return DecodePacketUncached(packetInfo, format, payload);
-			}
-			else
-			{
-				string cacheKey = $"bundle-packet:{locator.BlobId}@{packetInfo.Offset}";
-				return _cache.GetOrCreate<ReadOnlyMemory<byte>>(cacheKey, entry =>
-				{
-					ReadOnlyMemory<byte> decodedPacket = DecodePacketUncached(packetInfo, format, payload);
-					entry.SetSize(packetInfo.DecodedLength);
-					return decodedPacket;
-				});
-			}
-		}
-
-		static ReadOnlyMemory<byte> DecodePacketUncached(PacketInfo packetInfo, BundleCompressionFormat format, ReadOnlyMemory<byte> encodedPacket)
-		{
-			byte[] decodedPacket = new byte[packetInfo.DecodedLength];
-			BundleData.Decompress(format, encodedPacket, decodedPacket);
-			return decodedPacket;
-		}
 
 		#endregion
 
