EpicGames.Perforce: Fix incorrect path copying native PDB to output directory.

diff --git a/Engine/Source/Programs/Shared/EpicGames.Perforce/EpicGames.Perforce.csproj b/Engine/Source/Programs/Shared/EpicGames.Perforce/EpicGames.Perforce.csproj
--- a/Engine/Source/Programs/Shared/EpicGames.Perforce/EpicGames.Perforce.csproj
+++ b/Engine/Source/Programs/Shared/EpicGames.Perforce/EpicGames.Perforce.csproj
@@ -47,7 +47,7 @@
     <None Include="..\..\..\..\Binaries\DotNET\EpicGames.Perforce.Native\win-x64\$(WindowsNativeConfig)\EpicGames.Perforce.Native.dll" Visible="false">
       <CopyToOutputDirectory>PreserveNewest</CopyToOutputDirectory>
     </None>
-    <None Include="..\..\..\..\Binaries\DotNET\EpicGames.Perforce.Native\win-x64\$(WindowsNativeConfig)\EpicGames.Perforce.Native.pdb" Visible="false" Condition="Exists('..\..\..\..\Binaries\DotNET\EpicGames.Perforce.Native\bin\win-x64\$(WindowsNativeConfig)\EpicGames.Perforce.Native.pdb')">
+    <None Include="..\..\..\..\Binaries\DotNET\EpicGames.Perforce.Native\win-x64\$(WindowsNativeConfig)\EpicGames.Perforce.Native.pdb" Visible="false" Condition="Exists('..\..\..\..\Binaries\DotNET\EpicGames.Perforce.Native\win-x64\$(WindowsNativeConfig)\EpicGames.Perforce.Native.pdb')">
       <CopyToOutputDirectory>PreserveNewest</CopyToOutputDirectory>
     </None>
     <None Include="..\..\..\..\Binaries\DotNET\EpicGames.Perforce.Native\linux-x64\Release\EpicGames.Perforce.Native.so" Visible="false">
