Horde: Refactor IStorageClient.

* Now supports partial reads of blobs, and caches them inside StorageClientBase.
* IStorageClient includes methods for reading/writing TreeNode-derived classes, and caches bundle information and encoded/decoded packets internally.
* TreeWriter class can be used for packing data together into bundles. Supports creating async tasks to fence nodes being flushed rather than an explicit concept of forking child writers.

diff --git a/Engine/Source/Programs/Horde/Horde.Agent/Commands/Bundles/CreateCommand.cs b/Engine/Source/Programs/Horde/Horde.Agent/Commands/Bundles/CreateCommand.cs
--- a/Engine/Source/Programs/Horde/Horde.Agent/Commands/Bundles/CreateCommand.cs
+++ b/Engine/Source/Programs/Horde/Horde.Agent/Commands/Bundles/CreateCommand.cs
@@ -8,6 +8,7 @@ using System.Threading.Tasks;
 using EpicGames.Core;
 using EpicGames.Horde.Storage;
 using EpicGames.Horde.Storage.Nodes;
+using Microsoft.Extensions.Caching.Memory;
 using Microsoft.Extensions.Logging;
 
 namespace Horde.Agent.Commands.Bundles
@@ -19,8 +20,8 @@ namespace Horde.Agent.Commands.Bundles
 			readonly DirectoryReference _rootDir;
 			readonly ILogger _logger;
 
-			public FileBlobStore(DirectoryReference rootDir, ILogger logger)
-				: base(null)
+			public FileBlobStore(DirectoryReference rootDir, IMemoryCache cache, ILogger logger)
+				: base(cache)
 			{
 				_rootDir = rootDir;
 				_logger = logger;
@@ -64,12 +65,12 @@ namespace Horde.Agent.Commands.Bundles
 				throw new NotImplementedException();
 			}
 
-			public override async Task<RefTarget?> TryReadRefTargetAsync(RefName name, DateTime cacheTime = default, CancellationToken cancellationToken = default)
+			public override async Task<NodeLocator> TryReadRefTargetAsync(RefName name, DateTime cacheTime = default, CancellationToken cancellationToken = default)
 			{
 				FileReference file = GetRefFile(name);
 				if (!FileReference.Exists(file))
 				{
-					return null;
+					return default;
 				}
 
 				_logger.LogInformation("Reading {File}", file);
@@ -79,15 +80,15 @@ namespace Horde.Agent.Commands.Bundles
 				BlobLocator locator = new BlobLocator(text.Substring(0, hashIdx));
 				int exportIdx = Int32.Parse(text.Substring(hashIdx + 1), CultureInfo.InvariantCulture);
 
-				return new RefTarget(locator, exportIdx);
+				return new NodeLocator(locator, exportIdx);
 			}
 
-			public override async Task WriteRefTargetAsync(RefName name, RefTarget target, CancellationToken cancellationToken = default)
+			public override async Task WriteRefTargetAsync(RefName name, NodeLocator target, CancellationToken cancellationToken = default)
 			{
 				FileReference file = GetRefFile(name);
 				DirectoryReference.CreateDirectory(file.Directory);
 				_logger.LogInformation("Writing {File}", file);
-				await FileReference.WriteAllTextAsync(file, $"{target.Locator}#{target.ExportIdx}");
+				await FileReference.WriteAllTextAsync(file, $"{target.Blob}#{target.ExportIdx}");
 			}
 
 			#endregion
@@ -98,12 +99,9 @@ namespace Horde.Agent.Commands.Bundles
 		[CommandLine("-StorageDir=", Description = "Overrides the default storage server with a local directory")]
 		public DirectoryReference StorageDir { get; set; } = DirectoryReference.Combine(Program.AppDir, "bundles");
 
-		IStorageClient? _storageClient;
-
-		protected IStorageClient CreateStorageClient(ILogger logger)
+		protected IStorageClient CreateStorageClient(IMemoryCache cache, ILogger logger)
 		{
-			_storageClient ??= new FileBlobStore(StorageDir, logger);
-			return _storageClient;
+			return new FileBlobStore(StorageDir, cache, logger);
 		}
 	}
 
@@ -118,8 +116,9 @@ namespace Horde.Agent.Commands.Bundles
 
 		public override async Task<int> ExecuteAsync(ILogger logger)
 		{
-			IStorageClient store = CreateStorageClient(logger);
-			ITreeWriter writer = store.CreateTreeWriter(RefName.Text);
+			using IMemoryCache cache = new MemoryCache(new MemoryCacheOptions());
+			IStorageClient store = CreateStorageClient(cache, logger);
+			TreeWriter writer = new TreeWriter(store, prefix: RefName.Text);
 
 			DirectoryNode node = new DirectoryNode(DirectoryFlags.None);
 			await node.CopyFromDirectoryAsync(InputDir.ToDirectoryInfo(), new ChunkingOptions(), writer, logger, CancellationToken.None);
diff --git a/Engine/Source/Programs/Horde/Horde.Agent/Commands/Bundles/DumpCommand.cs b/Engine/Source/Programs/Horde/Horde.Agent/Commands/Bundles/DumpCommand.cs
--- a/Engine/Source/Programs/Horde/Horde.Agent/Commands/Bundles/DumpCommand.cs
+++ b/Engine/Source/Programs/Horde/Horde.Agent/Commands/Bundles/DumpCommand.cs
@@ -6,6 +6,7 @@ using System.Linq;
 using System.Threading.Tasks;
 using EpicGames.Core;
 using EpicGames.Horde.Storage;
+using Microsoft.Extensions.Caching.Memory;
 using Microsoft.Extensions.Logging;
 
 namespace Horde.Agent.Commands.Bundles
@@ -21,20 +22,19 @@ namespace Horde.Agent.Commands.Bundles
 
 		public override async Task<int> ExecuteAsync(ILogger logger)
 		{
-			IStorageClient storage = CreateStorageClient(logger);
+			using IMemoryCache cache = new MemoryCache(new MemoryCacheOptions());
+			IStorageClient store = CreateStorageClient(cache, logger);
 
-			BundleHeader header;
-			if (BlobId != null)
+			if (BlobId == null)
 			{
-				header = await storage.ReadBundleHeaderAsync(BlobId.Value);
-				logger.LogInformation("Summary for blob {BlobId}", BlobId.Value);
-			}
-			else
-			{
-				header = (await storage.ReadRefValueAsync(RefName)).Bundle.Header;
-				logger.LogInformation("Summary for ref {RefId}", RefName);
+				NodeLocator locator = await store.ReadRefTargetAsync(RefName);
+				BlobId = locator.Blob;
 			}
 
+			logger.LogInformation("Summary for blob {BlobId}", BlobId.Value);
+			Bundle bundle = await store.ReadBundleAsync(BlobId.Value);
+
+			BundleHeader header = bundle.Header;
 			int packetStart = 0;
 
 			logger.LogInformation("");
@@ -43,7 +43,7 @@ namespace Horde.Agent.Commands.Bundles
 			int refIdx = 0;
 			foreach (BundleImport import in header.Imports)
 			{
-				logger.LogInformation("  From blob {BlobId} ({NumExports}/{TotalExports} nodes)", import.Locator, import.Exports.Count, import.ExportCount);
+				logger.LogInformation("  From blob {BlobId} ({NumExports} nodes)", import.Locator, import.Exports.Count);
 				foreach ((int exportIdx, IoHash exportHash) in import.Exports)
 				{
 					logger.LogInformation("    [{Index}] IMP {BlobId}:{ExportIdx} = {ExportHash}", refIdx, import.Locator, exportIdx, exportHash);
diff --git a/Engine/Source/Programs/Horde/Horde.Agent/Commands/Bundles/ExtractCommand.cs b/Engine/Source/Programs/Horde/Horde.Agent/Commands/Bundles/ExtractCommand.cs
--- a/Engine/Source/Programs/Horde/Horde.Agent/Commands/Bundles/ExtractCommand.cs
+++ b/Engine/Source/Programs/Horde/Horde.Agent/Commands/Bundles/ExtractCommand.cs
@@ -5,6 +5,7 @@ using System.Threading.Tasks;
 using EpicGames.Core;
 using EpicGames.Horde.Storage;
 using EpicGames.Horde.Storage.Nodes;
+using Microsoft.Extensions.Caching.Memory;
 using Microsoft.Extensions.Logging;
 
 namespace Horde.Agent.Commands.Bundles
@@ -20,9 +21,10 @@ namespace Horde.Agent.Commands.Bundles
 
 		public override async Task<int> ExecuteAsync(ILogger logger)
 		{
-			IStorageClient storage = CreateStorageClient(logger);
+			using IMemoryCache cache = new MemoryCache(new MemoryCacheOptions());
+			IStorageClient store = CreateStorageClient(cache, logger);
 
-			DirectoryNode node = await storage.ReadTreeAsync<DirectoryNode>(RefName);
+			DirectoryNode node = await store.ReadNodeAsync<DirectoryNode>(RefName);
 			await node.CopyToDirectoryAsync(OutputDir.ToDirectoryInfo(), logger, CancellationToken.None);
 
 			return 0;
diff --git a/Engine/Source/Programs/Horde/Horde.Build.Tests/BlobStoreTests.cs b/Engine/Source/Programs/Horde/Horde.Build.Tests/BlobStoreTests.cs
--- a/Engine/Source/Programs/Horde/Horde.Build.Tests/BlobStoreTests.cs
+++ b/Engine/Source/Programs/Horde/Horde.Build.Tests/BlobStoreTests.cs
@@ -10,15 +10,31 @@ using MongoDB.Driver;
 using EpicGames.Horde.Storage;
 using System.Buffers;
 using EpicGames.Horde.Storage.Backends;
+using Microsoft.Extensions.Caching.Memory;
 
 namespace Horde.Build.Tests
 {
 	[TestClass]
 	public class BlobStoreTests : TestSetup
 	{
-		static IStorageClient CreateBlobStore()
+		IMemoryCache? _cache;
+		InMemoryBlobStore? _store;
+
+		IStorageClient CreateBlobStore()
+		{
+			_cache ??= new MemoryCache(new MemoryCacheOptions());
+			_store ??= new InMemoryBlobStore(_cache);
+			return _store;
+		}
+
+		protected override void Dispose(bool disposing)
 		{
-			return new InMemoryBlobStore();
+			base.Dispose(disposing);
+
+			if (disposing)
+			{
+				_cache?.Dispose();
+			}
 		}
 
 		static byte[] CreateTestData(int length, int seed)
@@ -44,7 +60,7 @@ namespace Horde.Build.Tests
 				int index;
 				if (!locatorToIndex.TryGetValue(locator, out index))
 				{
-					BundleImport import = new BundleImport(locator, 1, new List<(int, IoHash)> { (0, IoHash.Zero) });
+					BundleImport import = new BundleImport(locator, new List<(int, IoHash)> { (0, IoHash.Zero) });
 					imports.Add(import);
 					locatorToIndex.Add(locator, imports.Count - 1);
 				}
@@ -112,15 +128,15 @@ namespace Horde.Build.Tests
 			for(int idx = 0; idx < 2; idx++)
 			{
 				RefName refName = new RefName("hello");
-				await store.WriteRefTargetAsync(refName, new RefTarget(locator3, 0));
-				RefTarget refTarget = await store.ReadRefTargetAsync(refName);
-				Assert.AreEqual(locator3, refTarget.Locator);
+				await store.WriteRefTargetAsync(refName, new NodeLocator(locator3, 0));
+				NodeLocator refTarget = await store.ReadRefTargetAsync(refName);
+				Assert.AreEqual(locator3, refTarget.Blob);
 			}
 
 			RefName refName2 = new RefName("hello2");
 
-			RefTarget refTargetId2 = await store.WriteRefValueAsync(refName2, CreateTestBundle(input3, new BlobLocator[] { locator1, locator2 }));
-			Blob refTarget2 = await ReadBlobAsync(store, refTargetId2.Locator);
+			NodeLocator refTargetId2 = await store.WriteRefAsync(refName2, CreateTestBundle(input3, new BlobLocator[] { locator1, locator2 }), 0);
+			Blob refTarget2 = await ReadBlobAsync(store, refTargetId2.Blob);
 			Assert.IsTrue(refTarget2.Data.Span.SequenceEqual(input3));
 			Assert.IsTrue(refTarget2.References.SequenceEqual(new BlobLocator[] { locator1, locator2 }));
 		}
diff --git a/Engine/Source/Programs/Horde/Horde.Build/Commands/Bundles/CreateCommand.cs b/Engine/Source/Programs/Horde/Horde.Build/Commands/Bundles/CreateCommand.cs
--- a/Engine/Source/Programs/Horde/Horde.Build/Commands/Bundles/CreateCommand.cs
+++ b/Engine/Source/Programs/Horde/Horde.Build/Commands/Bundles/CreateCommand.cs
@@ -35,14 +35,12 @@ namespace Horde.Build.Commands.Bundles
 			using ServiceProvider serviceProvider = Startup.CreateServiceProvider(_configuration, _loggerProvider);
 
 			IStorageClient store = serviceProvider.GetRequiredService<IStorageClient<ReplicationService>>();
-			ITreeWriter writer = store.CreateTreeWriter(RefName.Text);
+			TreeWriter writer = new TreeWriter(store, prefix: RefName.Text);
 
 			DirectoryNode node = new DirectoryNode(DirectoryFlags.None);
 			await node.CopyFromDirectoryAsync(InputDir.ToDirectoryInfo(), new ChunkingOptions(), writer, logger, CancellationToken.None);
 
-			TreeNodeRef<DirectoryNode> root = new TreeNodeRef<DirectoryNode>(null!, node);
-			await root.CollapseAsync(writer, CancellationToken.None);
-
+			await writer.WriteRefAsync(RefName, node, CancellationToken.None);
 			return 0;
 		}
 	}
diff --git a/Engine/Source/Programs/Horde/Horde.Build/Commands/Bundles/ExtractCommand.cs b/Engine/Source/Programs/Horde/Horde.Build/Commands/Bundles/ExtractCommand.cs
--- a/Engine/Source/Programs/Horde/Horde.Build/Commands/Bundles/ExtractCommand.cs
+++ b/Engine/Source/Programs/Horde/Horde.Build/Commands/Bundles/ExtractCommand.cs
@@ -37,7 +37,7 @@ namespace Horde.Build.Commands.Bundles
 
 			IStorageClient store = serviceProvider.GetRequiredService<IStorageClient<ReplicationService>>();
 
-			ReplicationNode node = await store.ReadTreeAsync<ReplicationNode>(RefName);
+			ReplicationNode node = await store.ReadNodeAsync<ReplicationNode>(RefName);
 			DirectoryNode root = await node.Contents.ExpandAsync();
 			await root.CopyToDirectoryAsync(OutputDir.ToDirectoryInfo(), logger, CancellationToken.None);
 
diff --git a/Engine/Source/Programs/Horde/Horde.Build/Commands/Bundles/PerforceCommand.cs b/Engine/Source/Programs/Horde/Horde.Build/Commands/Bundles/PerforceCommand.cs
--- a/Engine/Source/Programs/Horde/Horde.Build/Commands/Bundles/PerforceCommand.cs
+++ b/Engine/Source/Programs/Horde/Horde.Build/Commands/Bundles/PerforceCommand.cs
@@ -82,11 +82,11 @@ namespace Horde.Build.Commands.Bundles
 			ReplicationNode baseContents;
 			if (BaseChange == 0)
 			{
-				baseContents = new ReplicationNode(new DirectoryNode(DirectoryFlags.WithGitHashes));
+				baseContents = new ReplicationNode(new DirectoryNode());
 			}
 			else
 			{
-				baseContents = await replicationService.ReadCommitTreeAsync(storage, stream, BaseChange, Filter, RevisionsOnly, CancellationToken.None);
+				baseContents = await ReplicationService.ReadCommitTreeAsync(storage, stream, BaseChange, Filter, RevisionsOnly, CancellationToken.None);
 			}
 
 			await foreach (ICommit commit in commitService.GetCollection(stream).FindAsync(Change, null, Count))
diff --git a/Engine/Source/Programs/Horde/Horde.Build/Commands/Bundles/RepoCommand.cs b/Engine/Source/Programs/Horde/Horde.Build/Commands/Bundles/RepoCommand.cs
--- a/Engine/Source/Programs/Horde/Horde.Build/Commands/Bundles/RepoCommand.cs
+++ b/Engine/Source/Programs/Horde/Horde.Build/Commands/Bundles/RepoCommand.cs
@@ -13,6 +13,7 @@ using Microsoft.Extensions.Configuration;
 using Microsoft.Extensions.DependencyInjection;
 using Microsoft.Extensions.Logging;
 
+#if false
 namespace Horde.Build.Commands.Bundles
 {
 	[Command("bundle", "repo", "Synthesizes a Git repo from a bundle to the local hard drive")]
@@ -39,7 +40,7 @@ namespace Horde.Build.Commands.Bundles
 
 			IStorageClient store = serviceProvider.GetRequiredService<IStorageClient<ReplicationService>>();
 
-			ReplicationNode node = await store.ReadTreeAsync<ReplicationNode>(RefName);
+			ReplicationNode node = await store.ReadRefAsync<ReplicationNode>(RefName);
 			DirectoryNode root = await node.Contents.ExpandAsync();
 
 			Sha1Hash tree = await WriteTreeRecursiveAsync(root, CancellationToken.None);
@@ -133,3 +134,4 @@ namespace Horde.Build.Commands.Bundles
 		}
 	}
 }
+#endif
diff --git a/Engine/Source/Programs/Horde/Horde.Build/Perforce/ReplicationService.cs b/Engine/Source/Programs/Horde/Horde.Build/Perforce/ReplicationService.cs
--- a/Engine/Source/Programs/Horde/Horde.Build/Perforce/ReplicationService.cs
+++ b/Engine/Source/Programs/Horde/Horde.Build/Perforce/ReplicationService.cs
@@ -5,6 +5,7 @@ using System.Collections.Generic;
 using System.Diagnostics;
 using System.IO;
 using System.Linq;
+using System.Reflection.Metadata;
 using System.Security.Cryptography;
 using System.Text;
 using System.Threading;
@@ -62,7 +63,6 @@ namespace Horde.Build.Perforce
 	/// <summary>
 	/// Root node for a commit snapshot
 	/// </summary>
-	[TreeSerializer(typeof(ReplicationNodeSerializer))]
 	public class ReplicationNode : TreeNode
 	{
 		/// <summary>
@@ -79,7 +79,7 @@ namespace Horde.Build.Perforce
 		/// Constructor
 		/// </summary>
 		public ReplicationNode()
-			: this(new DirectoryNode(DirectoryFlags.WithGitHashes))
+			: this(new DirectoryNode())
 		{
 		}
 
@@ -89,50 +89,27 @@ namespace Horde.Build.Perforce
 		public ReplicationNode(DirectoryNode contents)
 		{
 			Paths = new List<Utf8String>();
-			Contents = new TreeNodeRef<DirectoryNode>(this, contents);
+			Contents = new TreeNodeRef<DirectoryNode>(contents);
 		}
 
-		private ReplicationNode(IEnumerable<Utf8String> paths, ITreeBlobRef contents)
-		{
-			Paths = new List<Utf8String>(paths);
-			Contents = new TreeNodeRef<DirectoryNode>(this, contents);
-		}
-
-		/// <inheritdoc/>
-		public override async Task<ITreeBlob> SerializeAsync(ITreeWriter writer, CancellationToken cancellationToken)
-		{
-			List<ITreeBlobRef> references = new List<ITreeBlobRef>();
-			references.Add(await Contents.CollapseAsync(writer, cancellationToken));
-
-			ByteArrayBuilder builder = new ByteArrayBuilder();
-			builder.WriteVariableLengthArray(Paths, x => builder.WriteUtf8String(x));
-
-			return new NewTreeBlob(builder.AsSequence(), references);
-		}
-
-		internal static ReplicationNode Deserialize(ITreeBlob blob)
+		/// <summary>
+		/// Deserializing constructor
+		/// </summary>
+		/// <param name="reader"></param>
+		public ReplicationNode(ITreeNodeReader reader)
 		{
-			MemoryReader reader = new MemoryReader(blob.Data.AsSingleSegment());
-			Utf8String[] paths = reader.ReadVariableLengthArray(() => reader.ReadUtf8String());
-
-			return new ReplicationNode(paths, blob.Refs[0]);
+			Paths = reader.ReadList(() => reader.ReadUtf8String());
+			Contents = reader.ReadRef<DirectoryNode>();
 		}
 
 		/// <inheritdoc/>
-		public override IReadOnlyList<TreeNodeRef> GetReferences()
+		public override void Serialize(ITreeNodeWriter writer)
 		{
-			List<TreeNodeRef> refs = new List<TreeNodeRef>();
-			refs.Add(Contents);
-			return refs;
+			writer.WriteList(Paths, x => writer.WriteUtf8String(x));
+			writer.WriteRef(Contents);
 		}
 	}
 
-	class ReplicationNodeSerializer : TreeNodeSerializer<ReplicationNode>
-	{
-		/// <inheritdoc/>
-		public override ReplicationNode Deserialize(ITreeBlob node) => ReplicationNode.Deserialize(node);
-	}
-
 	/// <summary>
 	/// Service which replicates content from Perforce
 	/// </summary>
@@ -420,7 +397,7 @@ namespace Horde.Build.Perforce
 				{
 					RefName prevRefName = GetRefName(stream, values[0]);
 
-					ReplicationNode? contents = await storage.TryReadTreeAsync<ReplicationNode>(prevRefName, cancellationToken: cancellationToken);
+					ReplicationNode? contents = await storage.TryReadNodeAsync<ReplicationNode>(prevRefName, cancellationToken: cancellationToken);
 					if (contents == null)
 					{
 						_logger.LogInformation("No content for CL {Change}; creating full snapshot", values[0]);
@@ -498,10 +475,10 @@ namespace Horde.Build.Perforce
 		/// <param name="revisionsOnly"></param>
 		/// <param name="cancellationToken"></param>
 		/// <returns></returns>
-		public async Task<ReplicationNode> ReadCommitTreeAsync(IStorageClient storage, IStream stream, int change, string? filter, bool revisionsOnly, CancellationToken cancellationToken)
+		public static async Task<ReplicationNode> ReadCommitTreeAsync(IStorageClient storage, IStream stream, int change, string? filter, bool revisionsOnly, CancellationToken cancellationToken)
 		{
 			RefName name = GetRefName(stream.Id, change, filter, revisionsOnly);
-			return await storage.ReadTreeAsync<ReplicationNode>(name, cancellationToken: cancellationToken);
+			return await storage.ReadNodeAsync<ReplicationNode>(name, cancellationToken: cancellationToken);
 		}
 
 		/// <summary>
@@ -551,20 +528,23 @@ namespace Horde.Build.Perforce
 
 			public void Dispose() => _hash.Dispose();
 
-			public async Task AppendAsync(ReadOnlyMemory<byte> data, ChunkingOptions options, ITreeWriter writer, CancellationToken cancellationToken)
+			public async Task AppendAsync(ReadOnlyMemory<byte> data, ChunkingOptions options, TreeWriter writer, CancellationToken cancellationToken)
 			{
 				_hash.AppendData(data.Span);
 				await _entry.AppendAsync(data, options, writer, cancellationToken);
 			}
 
-			public async Task FinishAsync(ITreeWriter writer, CancellationToken cancellationToken)
+			public async Task FinishAsync(TreeWriter writer, CancellationToken cancellationToken)
 			{
 				if (_entry.Length != _size)
 				{
 					throw new ReplicationException($"Invalid size for replicated file '{_entry.Name}'. Expected {_size}, got {_entry.Length}.");
 				}
 				UpdateHash();
-				await _entry.CollapseAsync(writer, cancellationToken);
+				if (_entry.Target != null)
+				{
+					await writer.WriteAsync(_entry.Target, cancellationToken);
+				}
 			}
 
 			void UpdateHash()
@@ -606,11 +586,10 @@ namespace Horde.Build.Perforce
 
 			RefName refName = GetRefName(stream.Id, change, filter, revisionsOnly);
 			RefName incRefName = new RefName($"{refName}-inc");
-			ITreeWriter directoryWriter = storage.CreateTreeWriter(refName.Text);
-			ITreeWriter fileWriter = directoryWriter.CreateChildWriter();
+			TreeWriter writer = new TreeWriter(storage, prefix: refName.Text);
 
 			// Get the current sync state for this change
-			ReplicationNode? syncNode = await storage.TryReadTreeAsync<ReplicationNode>(incRefName, cancellationToken: cancellationToken);
+			ReplicationNode? syncNode = await storage.TryReadNodeAsync<ReplicationNode>(incRefName, cancellationToken: cancellationToken);
 			bool deleteIncRef = syncNode != null;
 			if (syncNode == null)
 			{
@@ -643,7 +622,7 @@ namespace Horde.Build.Perforce
 					using ReadOnlyMemoryStream dataStream = new ReadOnlyMemoryStream(data);
 
 					FileEntry entry = await root.AddFileByPathAsync(path, FileEntryFlags.PerforceDepotPathAndRevision, cancellationToken);
-					await entry.AppendAsync(dataStream, Options.Chunking, directoryWriter, cancellationToken);
+					await entry.AppendAsync(dataStream, Options.Chunking, writer, cancellationToken);
 				}
 			}
 			else
@@ -708,7 +687,7 @@ namespace Horde.Build.Perforce
 					if (syncedSize > 0)
 					{
 						Stopwatch flushTimer = Stopwatch.StartNew();
-						await directoryWriter.WriteRefAsync(incRefName, syncNode, cancellationToken);
+						await writer.WriteRefAsync(incRefName, syncNode, cancellationToken);
 						flushTimer.Stop();
 						deleteIncRef = true;
 					}
@@ -798,7 +777,7 @@ namespace Horde.Build.Perforce
 								else if (io.Command == PerforceIoCommand.Write)
 								{
 									FileWriter file = handles[io.File];
-									await file.AppendAsync(io.Payload, Options.Chunking, fileWriter, cancellationToken);
+									await file.AppendAsync(io.Payload, Options.Chunking, writer, cancellationToken);
 								}
 								else if (io.Command == PerforceIoCommand.Close)
 								{
@@ -807,7 +786,7 @@ namespace Horde.Build.Perforce
 									{
 										if (handles.Remove(io.File, out file))
 										{
-											await file.FinishAsync(fileWriter, cancellationToken);
+											await file.FinishAsync(writer, cancellationToken);
 										}
 									}
 									finally
@@ -872,13 +851,13 @@ namespace Horde.Build.Perforce
 					}
 				}
 			}
-			await directoryWriter.WriteNodeAsync(syncNode, cancellationToken);
+			await writer.WriteAsync(syncNode, cancellationToken);
 
 			clientInfo.Change = change;
 
 			// Return the new root object
 			_logger.LogInformation("Writing ref {RefId} for {StreamId} change {Change}", refName, stream.Id, change);
-			await directoryWriter.WriteRefAsync(refName, syncNode, cancellationToken);
+			await writer.WriteRefAsync(refName, syncNode, cancellationToken);
 
 			// Delete the incremental state
 			if (deleteIncRef)
diff --git a/Engine/Source/Programs/Horde/Horde.Build/Storage/IRefCollection.cs b/Engine/Source/Programs/Horde/Horde.Build/Storage/IRefCollection.cs
--- a/Engine/Source/Programs/Horde/Horde.Build/Storage/IRefCollection.cs
+++ b/Engine/Source/Programs/Horde/Horde.Build/Storage/IRefCollection.cs
@@ -3,13 +3,7 @@
 using System;
 using System.Threading;
 using System.Threading.Tasks;
-using Amazon.Runtime.Internal.Util;
 using EpicGames.Horde.Storage;
-using Horde.Build.Server;
-using Microsoft.Extensions.Caching.Memory;
-using MongoDB.Bson;
-using MongoDB.Bson.Serialization.Attributes;
-using MongoDB.Driver;
 
 namespace Horde.Build.Storage
 {
@@ -26,7 +20,7 @@ namespace Horde.Build.Storage
 		/// <param name="cacheTime">Minimum age for a cached value to be returned</param>
 		/// <param name="cancellationToken">Cancellation token for the operation</param>
 		/// <returns>The ref target</returns>
-		Task<RefTarget?> TryReadRefTargetAsync(NamespaceId namespaceId, RefName name, DateTime cacheTime = default, CancellationToken cancellationToken = default);
+		Task<NodeLocator> TryReadRefTargetAsync(NamespaceId namespaceId, RefName name, DateTime cacheTime = default, CancellationToken cancellationToken = default);
 
 		/// <summary>
 		/// Writes a ref to storage
@@ -35,7 +29,7 @@ namespace Horde.Build.Storage
 		/// <param name="name">The ref name</param>
 		/// <param name="target">Target for the ref</param>
 		/// <param name="cancellationToken">Cancellation token for the operation</param>
-		Task WriteRefTargetAsync(NamespaceId namespaceId, RefName name, RefTarget target, CancellationToken cancellationToken = default);
+		Task WriteRefTargetAsync(NamespaceId namespaceId, RefName name, NodeLocator target, CancellationToken cancellationToken = default);
 
 		/// <summary>
 		/// Delete a ref from storage
diff --git a/Engine/Source/Programs/Horde/Horde.Build/Storage/RefCollection.cs b/Engine/Source/Programs/Horde/Horde.Build/Storage/RefCollection.cs
--- a/Engine/Source/Programs/Horde/Horde.Build/Storage/RefCollection.cs
+++ b/Engine/Source/Programs/Horde/Horde.Build/Storage/RefCollection.cs
@@ -45,8 +45,8 @@ namespace Horde.Build.Storage
 			[BsonElement("k")]
 			public RefName Name { get; set; }
 
-			[BsonElement("l")]
-			public BlobLocator Locator { get; set; }
+			[BsonElement("b")]
+			public BlobLocator Blob { get; set; }
 
 			[BsonElement("e")]
 			public int ExportIdx { get; set; }
@@ -54,14 +54,14 @@ namespace Horde.Build.Storage
 			public RefDocument()
 			{
 				Name = RefName.Empty;
-				Locator = BlobLocator.Empty;
+				Blob = BlobLocator.Empty;
 			}
 
-			public RefDocument(NamespaceId namespaceId, RefName name, RefTarget target)
+			public RefDocument(NamespaceId namespaceId, RefName name, NodeLocator target)
 			{
 				NamespaceId = namespaceId;
 				Name = name;
-				Locator = target.Locator;
+				Blob = target.Blob;
 				ExportIdx = target.ExportIdx;
 			}
 		}
@@ -69,10 +69,10 @@ namespace Horde.Build.Storage
 		class CachedRefValue
 		{
 			public RefName Name { get; }
-			public RefTarget? Target { get; }
+			public NodeLocator Target { get; }
 			public DateTime Time { get; }
 
-			public CachedRefValue(RefName name, RefTarget? target)
+			public CachedRefValue(RefName name, NodeLocator target)
 			{
 				Name = name;
 				Target = target;
@@ -94,7 +94,7 @@ namespace Horde.Build.Storage
 			_cache = cache;
 		}
 
-		CachedRefValue AddRefToCache(NamespaceId namespaceId, RefName name, RefTarget? target)
+		CachedRefValue AddRefToCache(NamespaceId namespaceId, RefName name, NodeLocator target)
 		{
 			CachedRefValue cacheValue = new CachedRefValue(name, target);
 			using (ICacheEntry newEntry = _cache.CreateEntry(new RefKey(namespaceId, name)))
@@ -110,25 +110,25 @@ namespace Horde.Build.Storage
 		public async Task DeleteRefAsync(NamespaceId namespaceId, RefName name, CancellationToken cancellationToken = default)
 		{
 			await _refs.DeleteOneAsync(x => x.NamespaceId == namespaceId && x.Name == name, cancellationToken);
-			AddRefToCache(namespaceId, name, null);
+			AddRefToCache(namespaceId, name, default);
 		}
 
 		/// <inheritdoc/>
-		public async Task<RefTarget?> TryReadRefTargetAsync(NamespaceId namespaceId, RefName name, DateTime cacheTime = default, CancellationToken cancellationToken = default)
+		public async Task<NodeLocator> TryReadRefTargetAsync(NamespaceId namespaceId, RefName name, DateTime cacheTime = default, CancellationToken cancellationToken = default)
 		{
 			CachedRefValue entry;
 			if (!_cache.TryGetValue(name, out entry) || entry.Time < cacheTime)
 			{
 				RefDocument? refDocument = await _refs.Find(x => x.NamespaceId == namespaceId && x.Name == name).FirstOrDefaultAsync(cancellationToken);
 
-				RefTarget? target;
+				NodeLocator target;
 				if (refDocument == null)
 				{
-					target = null;
+					target = default;
 				}
 				else
 				{
-					target = new RefTarget(refDocument.Locator, refDocument.ExportIdx);
+					target = new NodeLocator(refDocument.Blob, refDocument.ExportIdx);
 				}
 
 				entry = AddRefToCache(namespaceId, name, target);
@@ -137,7 +137,7 @@ namespace Horde.Build.Storage
 		}
 
 		/// <inheritdoc/>
-		public async Task WriteRefTargetAsync(NamespaceId namespaceId, RefName name, RefTarget target, CancellationToken cancellationToken = default)
+		public async Task WriteRefTargetAsync(NamespaceId namespaceId, RefName name, NodeLocator target, CancellationToken cancellationToken = default)
 		{
 			RefDocument refDocument = new RefDocument(namespaceId, name, target);
 			await _refs.ReplaceOneAsync(x => x.NamespaceId == namespaceId && x.Name == name, refDocument, new ReplaceOptions { IsUpsert = true }, cancellationToken);
diff --git a/Engine/Source/Programs/Horde/Horde.Build/Storage/StorageController.cs b/Engine/Source/Programs/Horde/Horde.Build/Storage/StorageController.cs
--- a/Engine/Source/Programs/Horde/Horde.Build/Storage/StorageController.cs
+++ b/Engine/Source/Programs/Horde/Horde.Build/Storage/StorageController.cs
@@ -49,7 +49,7 @@ namespace Horde.Build.Storage
 		/// <summary>
 		/// Locator for the target blob
 		/// </summary>
-		public BlobLocator Locator { get; set; }
+		public BlobLocator Blob { get; set; }
 
 		/// <summary>
 		/// Export index for the ref
@@ -59,9 +59,9 @@ namespace Horde.Build.Storage
 		/// <summary>
 		/// Constructor
 		/// </summary>
-		public ReadRefResponse(RefTarget target)
+		public ReadRefResponse(NodeLocator target)
 		{
-			Locator = target.Locator;
+			Blob = target.Blob;
 			ExportIdx = target.ExportIdx;
 		}
 	}
@@ -148,7 +148,7 @@ namespace Horde.Build.Storage
 			}
 
 			IStorageClient client = await _storageService.GetClientAsync(namespaceId, cancellationToken);
-			RefTarget target = new RefTarget(request.Locator, request.ExportIdx);
+			NodeLocator target = new NodeLocator(request.Locator, request.ExportIdx);
 			await client.WriteRefTargetAsync(refName, target, cancellationToken);
 
 			return Ok();
@@ -171,8 +171,8 @@ namespace Horde.Build.Storage
 
 			IStorageClient client = await _storageService.GetClientAsync(namespaceId, cancellationToken);
 
-			RefTarget? target = await client.TryReadRefTargetAsync(refName, cancellationToken: cancellationToken);
-			if (target == null)
+			NodeLocator target = await client.TryReadRefTargetAsync(refName, cancellationToken: cancellationToken);
+			if (!target.IsValid())
 			{
 				return NotFound();
 			}
diff --git a/Engine/Source/Programs/Horde/Horde.Build/Storage/StorageService.cs b/Engine/Source/Programs/Horde/Horde.Build/Storage/StorageService.cs
--- a/Engine/Source/Programs/Horde/Horde.Build/Storage/StorageService.cs
+++ b/Engine/Source/Programs/Horde/Horde.Build/Storage/StorageService.cs
@@ -53,7 +53,7 @@ namespace Horde.Build.Storage
 			readonly IStorageBackend _backend;
 			readonly IRefCollection _refs;
 
-			public StorageClient(NamespaceConfig config, IStorageBackend backend, IRefCollection refs, IMemoryCache? cache)
+			public StorageClient(NamespaceConfig config, IStorageBackend backend, IRefCollection refs, IMemoryCache cache)
 				: base(cache)
 			{
 				Config = config;
@@ -102,10 +102,10 @@ namespace Horde.Build.Storage
 			#region Refs
 
 			/// <inheritdoc/>
-			public override Task<RefTarget?> TryReadRefTargetAsync(RefName name, DateTime cacheTime = default, CancellationToken cancellationToken = default) => _refs.TryReadRefTargetAsync(NamespaceId, name, cacheTime, cancellationToken);
+			public override Task<NodeLocator> TryReadRefTargetAsync(RefName name, DateTime cacheTime = default, CancellationToken cancellationToken = default) => _refs.TryReadRefTargetAsync(NamespaceId, name, cacheTime, cancellationToken);
 
 			/// <inheritdoc/>
-			public override Task WriteRefTargetAsync(RefName name, RefTarget target, CancellationToken cancellationToken = default) => _refs.WriteRefTargetAsync(NamespaceId, name, target, cancellationToken);
+			public override Task WriteRefTargetAsync(RefName name, NodeLocator target, CancellationToken cancellationToken = default) => _refs.WriteRefTargetAsync(NamespaceId, name, target, cancellationToken);
 
 			/// <inheritdoc/>
 			public override Task DeleteRefAsync(RefName name, CancellationToken cancellationToken = default) => _refs.DeleteRefAsync(NamespaceId, name, cancellationToken);
diff --git a/Engine/Source/Programs/Shared/EpicGames.Core/ByteArrayBuilder.cs b/Engine/Source/Programs/Shared/EpicGames.Core/ByteArrayBuilder.cs
--- a/Engine/Source/Programs/Shared/EpicGames.Core/ByteArrayBuilder.cs
+++ b/Engine/Source/Programs/Shared/EpicGames.Core/ByteArrayBuilder.cs
@@ -57,6 +57,17 @@ namespace EpicGames.Core
 			_chunkSize = chunkSize;
 		}
 
+		/// <summary>
+		/// Clear the current builder
+		/// </summary>
+		public void Clear()
+		{
+			_currentChunk = _chunks[0];
+			_currentChunk.Length = 0;
+			_chunks.RemoveRange(1, _chunks.Count - 1);
+			Length = 0;
+		}
+
 		/// <inheritdoc/>
 		public Memory<byte> GetMemory(int minSize)
 		{
diff --git a/Engine/Source/Programs/Shared/EpicGames.Core/MemoryReader.cs b/Engine/Source/Programs/Shared/EpicGames.Core/MemoryReader.cs
--- a/Engine/Source/Programs/Shared/EpicGames.Core/MemoryReader.cs
+++ b/Engine/Source/Programs/Shared/EpicGames.Core/MemoryReader.cs
@@ -2,6 +2,7 @@
 
 using System;
 using System.Buffers.Binary;
+using System.Collections.Generic;
 
 namespace EpicGames.Core
 {
@@ -243,6 +244,24 @@ namespace EpicGames.Core
 			return bytes;
 		}
 
+		/// <summary>
+		/// Reads a variable length list
+		/// </summary>
+		/// <param name="reader">Reader to deserialize from</param>
+		/// <param name="readItem">Delegate to write an individual item</param>
+		public static List<T> ReadList<T>(this IMemoryReader reader, Func<T> readItem)
+		{
+			int length = (int)reader.ReadUnsignedVarInt();
+
+			List<T> list = new List<T>(length);
+			for (int idx = 0; idx < length; idx++)
+			{
+				list.Add(readItem());
+			}
+
+			return list;
+		}
+
 		/// <summary>
 		/// Reads a variable length array
 		/// </summary>
diff --git a/Engine/Source/Programs/Shared/EpicGames.Core/MemoryWriter.cs b/Engine/Source/Programs/Shared/EpicGames.Core/MemoryWriter.cs
--- a/Engine/Source/Programs/Shared/EpicGames.Core/MemoryWriter.cs
+++ b/Engine/Source/Programs/Shared/EpicGames.Core/MemoryWriter.cs
@@ -48,6 +48,11 @@ namespace EpicGames.Core
 		/// </summary>
 		int _length;
 
+		/// <summary>
+		/// Returns the memory that was written to
+		/// </summary>
+		public ReadOnlyMemory<byte> WrittenMemory => _memory.Slice(0, _length);
+
 		/// <summary>
 		/// Constructor
 		/// </summary>
@@ -78,6 +83,60 @@ namespace EpicGames.Core
 		public void Advance(int length) => _length += length;
 	}
 
+	/// <summary>
+	/// Writes into an expandable memory block
+	/// </summary>
+	public class ArrayMemoryWriter : IMemoryWriter
+	{
+		byte[] _data;
+		int _length;
+
+		public ReadOnlyMemory<byte> WrittenMemory => _data.AsMemory(0, _length);
+
+		/// <summary>
+		/// Constructor
+		/// </summary>
+		/// <param name="initialSize">Initial size of the allocated data</param>
+		public ArrayMemoryWriter(int initialSize)
+		{
+			_data = new byte[initialSize];
+		}
+
+		/// <inheritdoc/>
+		public int Length => _length;
+
+		public void Clear()
+		{
+			_length = 0;
+		}
+
+		public void Resize(int newLength)
+		{
+			if (newLength > _length)
+			{
+				Array.Resize(ref _data, newLength);
+			}
+			_length = newLength;
+		}
+
+		/// <inheritdoc/>
+		public Memory<byte> GetMemory(int minSize = 1)
+		{
+			int requiredLength = _length + minSize;
+			if (_data.Length < requiredLength)
+			{
+				Array.Resize(ref _data, requiredLength + 4096);
+			}
+			return _data.AsMemory(_length);
+		}
+
+		/// <inheritdoc/>
+		public void Advance(int length)
+		{
+			_length += length;
+		}
+	}
+
 	/// <summary>
 	/// Extension methods for <see cref="IMemoryWriter"/>
 	/// </summary>
@@ -260,6 +319,17 @@ namespace EpicGames.Core
 			sequence.CopyTo(span);
 		}
 
+		/// <summary>
+		/// Writes a variable length array
+		/// </summary>
+		/// <param name="writer">Writer to serialize to</param>
+		/// <param name="list">The array to write</param>
+		/// <param name="writeItem">Delegate to write an individual item</param>
+		public static void WriteList<T>(this IMemoryWriter writer, IReadOnlyList<T> list, Action<T> writeItem)
+		{
+			WriteVariableLengthArray(writer, list, writeItem);
+		}
+
 		/// <summary>
 		/// Writes a variable length span of bytes
 		/// </summary>
@@ -273,6 +343,19 @@ namespace EpicGames.Core
 			bytes.CopyTo(span[lengthBytes..]);
 		}
 
+		/// <summary>
+		/// Appends a sequence of bytes to the buffer, without a length field.
+		/// </summary>
+		/// <param name="writer">Writer to serialize to</param>
+		/// <param name="sequence">Sequence to append</param>
+		public static void WriteVariableLengthBytes(this IMemoryWriter writer, ReadOnlySequence<byte> sequence)
+		{
+			int lengthBytes = VarInt.MeasureUnsigned((int)sequence.Length);
+			Span<byte> span = GetSpanAndAdvance(writer, (int)(lengthBytes + sequence.Length));
+			VarInt.WriteUnsigned(span, (int)sequence.Length);
+			sequence.CopyTo(span[lengthBytes..]);
+		}
+
 		/// <summary>
 		/// Writes a variable length span of bytes
 		/// </summary>
diff --git a/Engine/Source/Programs/Shared/EpicGames.Horde.Tests/BundleStoreTests.cs b/Engine/Source/Programs/Shared/EpicGames.Horde.Tests/BundleStoreTests.cs
--- a/Engine/Source/Programs/Shared/EpicGames.Horde.Tests/BundleStoreTests.cs
+++ b/Engine/Source/Programs/Shared/EpicGames.Horde.Tests/BundleStoreTests.cs
@@ -14,6 +14,7 @@ using System.Text;
 using System.IO;
 using System.Security.Cryptography;
 using EpicGames.Core;
+using Microsoft.Extensions.Caching.Memory;
 
 namespace EpicGames.Horde.Tests
 {
@@ -23,7 +24,9 @@ namespace EpicGames.Horde.Tests
 		[TestMethod]
 		public async Task TestTreeAsync()
 		{
-			InMemoryBlobStore blobStore = new InMemoryBlobStore();
+			using IMemoryCache cache = new MemoryCache(new MemoryCacheOptions());
+
+			InMemoryBlobStore blobStore = new InMemoryBlobStore(cache);
 			await TestTreeAsync(blobStore, new TreeOptions { MaxBlobSize = 1024 * 1024 });
 
 			Assert.AreEqual(1, blobStore.Blobs.Count);
@@ -33,89 +36,116 @@ namespace EpicGames.Horde.Tests
 		[TestMethod]
 		public async Task TestTreeSeparateBlobsAsync()
 		{
-			InMemoryBlobStore blobStore = new InMemoryBlobStore();
+			using IMemoryCache cache = new MemoryCache(new MemoryCacheOptions());
+
+			InMemoryBlobStore blobStore = new InMemoryBlobStore(cache);
 			await TestTreeAsync(blobStore, new TreeOptions { MaxBlobSize = 1 });
 
 			Assert.AreEqual(5, blobStore.Blobs.Count);
 			Assert.AreEqual(1, blobStore.Refs.Count);
 		}
 
+		class SimpleNode : TreeNode
+		{
+			public ReadOnlySequence<byte> Data { get; }
+			public IReadOnlyList<TreeNodeRef<SimpleNode>> Refs { get; }
+
+			public SimpleNode(ReadOnlySequence<byte> data, IReadOnlyList<TreeNodeRef<SimpleNode>> refs)
+			{
+				Data = data;
+				Refs = refs;
+			}
+
+			public SimpleNode(ITreeNodeReader reader)
+			{
+				Data = new ReadOnlySequence<byte>(reader.ReadVariableLengthBytes());
+				Refs = reader.ReadVariableLengthArray(() => reader.ReadRef<SimpleNode>());
+			}
+
+			public override void Serialize(ITreeNodeWriter writer)
+			{
+				writer.WriteVariableLengthBytes(Data);
+				writer.WriteVariableLengthArray(Refs, x => writer.WriteRef(x));
+			}
+		}
+
 		static async Task TestTreeAsync(InMemoryBlobStore store, TreeOptions options)
 		{
 			// Generate a tree
 			{
-				ITreeWriter writer = store.CreateTreeWriter(options, "test");
+				TreeWriter writer = new TreeWriter(store, options, "test");
 
-				ITreeBlobRef node1 = await writer.WriteNodeAsync(new ReadOnlySequence<byte>(new byte[] { 1 }), Array.Empty<ITreeBlobRef>(), CancellationToken.None);
-				ITreeBlobRef node2 = await writer.WriteNodeAsync(new ReadOnlySequence<byte>(new byte[] { 2 }), new[] { node1 }, CancellationToken.None);
-				ITreeBlobRef node3 = await writer.WriteNodeAsync(new ReadOnlySequence<byte>(new byte[] { 3 }), new[] { node2 }, CancellationToken.None);
-				ITreeBlobRef node4 = await writer.WriteNodeAsync(new ReadOnlySequence<byte>(new byte[] { 4 }), Array.Empty<ITreeBlobRef>(), CancellationToken.None);
+				SimpleNode node1 = new SimpleNode(new ReadOnlySequence<byte>(new byte[] { 1 }), Array.Empty<TreeNodeRef<SimpleNode>>());
+				SimpleNode node2 = new SimpleNode(new ReadOnlySequence<byte>(new byte[] { 2 }), new[] { new TreeNodeRef<SimpleNode>(node1) });
+				SimpleNode node3 = new SimpleNode(new ReadOnlySequence<byte>(new byte[] { 3 }), new[] { new TreeNodeRef<SimpleNode>(node2) });
+				SimpleNode node4 = new SimpleNode(new ReadOnlySequence<byte>(new byte[] { 4 }), Array.Empty<TreeNodeRef<SimpleNode>>());
 
-				ITreeBlob root = TreeBlob.Create(new ReadOnlySequence<byte>(new byte[] { 5 }), new[] { node4, node3 });
-				await writer.WriteRefAsync(new RefName("test"), root.Data, root.Refs);
+				SimpleNode root = new SimpleNode(new ReadOnlySequence<byte>(new byte[] { 5 }), new[] { new TreeNodeRef<SimpleNode>(node4), new TreeNodeRef<SimpleNode>(node3) });
+
+				await writer.WriteRefAsync(new RefName("test"), root, CancellationToken.None);
 
 				await CheckTree(root);
 			}
 
 			// Check we can read it back in
 			{
-				ITreeBlob? root = await store.TryReadTreeAsync(new RefName("test"));
-				Assert.IsNotNull(root);
+				SimpleNode root = await store.ReadNodeAsync<SimpleNode>(new RefName("test"));
 				await CheckTree(root!);
 			}
 		}
 
-		static async Task CheckTree(ITreeBlob root)
+		static async Task CheckTree(SimpleNode root)
 		{
-			ITreeBlob node5 = root;
+			SimpleNode node5 = root;
 			byte[] data5 = node5.Data.ToArray();
 			Assert.IsTrue(data5.SequenceEqual(new byte[] { 5 }));
-			IReadOnlyList<ITreeBlobRef> refs5 = node5.Refs;
+			IReadOnlyList<TreeNodeRef<SimpleNode>> refs5 = node5.Refs;
 			Assert.AreEqual(2, refs5.Count);
 
-			ITreeBlob node4 = await refs5[0].GetTargetAsync();
+			SimpleNode node4 = await refs5[0].ExpandAsync();
 			byte[] data4 = node4.Data.ToArray();
 			Assert.IsTrue(data4.SequenceEqual(new byte[] { 4 }));
-			IReadOnlyList<ITreeBlobRef> refs4 = node4.Refs;
+			IReadOnlyList<TreeNodeRef<SimpleNode>> refs4 = node4.Refs;
 			Assert.AreEqual(0, refs4.Count);
 
-			ITreeBlob node3 = await refs5[1].GetTargetAsync();
+			SimpleNode node3 = await refs5[1].ExpandAsync();
 			byte[] data3 = node3.Data.ToArray();
 			Assert.IsTrue(data3.SequenceEqual(new byte[] { 3 }));
-			IReadOnlyList<ITreeBlobRef> refs3 = node3.Refs;
+			IReadOnlyList<TreeNodeRef<SimpleNode>> refs3 = node3.Refs;
 			Assert.AreEqual(1, refs3.Count);
 
-			ITreeBlob node2 = await refs3[0].GetTargetAsync();
+			SimpleNode node2 = await refs3[0].ExpandAsync();
 			byte[] data2 = node2.Data.ToArray();
 			Assert.IsTrue(data2.SequenceEqual(new byte[] { 2 }));
-			IReadOnlyList<ITreeBlobRef> refs2 = node2.Refs;
+			IReadOnlyList<TreeNodeRef<SimpleNode>> refs2 = node2.Refs;
 			Assert.AreEqual(1, refs2.Count);
 
-			ITreeBlob node1 = await refs2[0].GetTargetAsync();
+			SimpleNode node1 = await refs2[0].ExpandAsync();
 			byte[] data1 = node1.Data.ToArray();
 			Assert.IsTrue(data1.SequenceEqual(new byte[] { 1 }));
-			IReadOnlyList<ITreeBlobRef> refs1 = node1.Refs;
+			IReadOnlyList<TreeNodeRef<SimpleNode>> refs1 = node1.Refs;
 			Assert.AreEqual(0, refs1.Count);
 		}
 
 		[TestMethod]
 		public async Task DirectoryNodesAsync()
 		{
-			InMemoryBlobStore store = new InMemoryBlobStore();
+			using IMemoryCache cache = new MemoryCache(new MemoryCacheOptions());
+			InMemoryBlobStore store = new InMemoryBlobStore(cache);
 
 			// Generate a tree
 			{
 				DirectoryNode root = new DirectoryNode(DirectoryFlags.None);
 				DirectoryNode hello = root.AddDirectory("hello");
 				DirectoryNode world = hello.AddDirectory("world");
-				await store.WriteTreeAsync(new RefName("test"), root, CancellationToken.None);
+				await store.WriteNodeAsync(new RefName("test"), root);
 
 				await CheckDirectoryTreeAsync(root);
 			}
 
 			// Check we can read it back in
 			{
-				DirectoryNode root = await store.ReadTreeAsync<DirectoryNode>(new RefName("test"));
+				DirectoryNode root = await store.ReadNodeAsync<DirectoryNode>(new RefName("test"));
 				await CheckDirectoryTreeAsync(root);
 			}
 		}
@@ -136,11 +166,12 @@ namespace EpicGames.Horde.Tests
 		[TestMethod]
 		public async Task FileNodesAsync()
 		{
-			InMemoryBlobStore store = new InMemoryBlobStore();
+			using IMemoryCache cache = new MemoryCache(new MemoryCacheOptions());
+			InMemoryBlobStore store = new InMemoryBlobStore(cache);
 
 			// Generate a tree
 			{
-				ITreeWriter writer = store.CreateTreeWriter();
+				TreeWriter writer = new TreeWriter(store, new TreeOptions());
 
 				DirectoryNode root = new DirectoryNode(DirectoryFlags.None);
 				DirectoryNode hello = root.AddDirectory("hello");
@@ -155,7 +186,7 @@ namespace EpicGames.Horde.Tests
 
 			// Check we can read it back in
 			{
-				DirectoryNode root = await store.ReadTreeAsync<DirectoryNode>(new RefName("test"));
+				DirectoryNode root = await store.ReadNodeAsync<DirectoryNode>(new RefName("test"));
 				await CheckFileTreeAsync(root);
 			}
 		}
@@ -179,8 +210,10 @@ namespace EpicGames.Horde.Tests
 		[TestMethod]
 		public async Task ChunkingTests()
 		{
-			InMemoryBlobStore store = new InMemoryBlobStore();
-			ITreeWriter writer = store.CreateTreeWriter(new TreeOptions());
+			using IMemoryCache cache = new MemoryCache(new MemoryCacheOptions());
+
+			InMemoryBlobStore store = new InMemoryBlobStore(cache);
+			TreeWriter writer = new TreeWriter(store, new TreeOptions());
 
 			byte[] chunk = RandomNumberGenerator.GetBytes(256);
 
@@ -213,11 +246,13 @@ namespace EpicGames.Horde.Tests
 			long uniqueSize = uniqueChunks.Sum(x => x.Value);
 			Assert.IsTrue(uniqueSize < data.Length / 3);
 		}
-
+		/*
 		[TestMethod]
 		public async Task ChildWriterAsync()
 		{
-			InMemoryBlobStore store = new InMemoryBlobStore();
+			using IMemoryCache cache = new MemoryCache(new MemoryCacheOptions());
+
+			InMemoryBlobStore store = new InMemoryBlobStore(cache);
 
 			{
 				ITreeWriter writer = store.CreateTreeWriter(new TreeOptions());
@@ -248,11 +283,12 @@ namespace EpicGames.Horde.Tests
 				Assert.AreEqual(0, childBlob2.Refs.Count);
 			}
 		}
-
+		*/
 		[TestMethod]
 		public async Task LargeFileTestAsync()
 		{
-			InMemoryBlobStore store = new InMemoryBlobStore();
+			using IMemoryCache cache = new MemoryCache(new MemoryCacheOptions());
+			InMemoryBlobStore store = new InMemoryBlobStore(cache);
 
 			const int length = 1024;
 			const int copies = 4096;
@@ -267,10 +303,11 @@ namespace EpicGames.Horde.Tests
 			}
 
 			// Generate a tree
+			DirectoryNode root;
 			{
-				ITreeWriter writer = store.CreateTreeWriter(new TreeOptions { MaxBlobSize = 1024 });
+				TreeWriter writer = new TreeWriter(store, new TreeOptions { MaxBlobSize = 1024 });
 
-				DirectoryNode root = new DirectoryNode(DirectoryFlags.None);
+				root = new DirectoryNode(DirectoryFlags.None);
 
 				ChunkingOptions options = new ChunkingOptions();
 				options.LeafOptions.MinSize = 1;
@@ -287,19 +324,62 @@ namespace EpicGames.Horde.Tests
 
 			// Check we can read it back in
 			{
-				DirectoryNode root = await store.ReadTreeAsync<DirectoryNode>(new RefName("test"));
+				DirectoryNode newRoot = await store.ReadNodeAsync<DirectoryNode>(new RefName("test"));
+				await CompareTrees(root, newRoot);
 				await CheckLargeFileTreeAsync(root, data);
 
 				TreeNodeRef<FileNode> file = root.GetFileEntry("test");
 
-				Dictionary<IoHash, long> uniqueBlobs = new Dictionary<IoHash, long>();
-				await FindUniqueBlobs(file.Target!, uniqueBlobs);
-
-				long uniqueSize = uniqueBlobs.Sum(x => x.Value);
+				long uniqueSize = store.Blobs.Values.SelectMany(x => x.Header.Packets).Sum(x => x.DecodedLength);
 				Assert.IsTrue(uniqueSize < data.Length / 3); // random fraction meaning "lots of dedupe happened"
 			}
 		}
 
+		static async Task CompareTrees(DirectoryNode oldNode, DirectoryNode newNode)
+		{
+			Assert.AreEqual(oldNode.Length, newNode.Length);
+			Assert.AreEqual(oldNode.Files.Count, newNode.Files.Count);
+			Assert.AreEqual(oldNode.Directories.Count, newNode.Directories.Count);
+
+			foreach ((FileEntry oldFileEntry, FileEntry newFileEntry) in oldNode.Files.Zip(newNode.Files))
+			{
+				FileNode oldFile = await oldFileEntry.ExpandAsync();
+				FileNode newFile = await newFileEntry.ExpandAsync();
+				await CompareTrees(oldFile, newFile);
+			}
+		}
+
+		static async Task CompareTrees(FileNode oldNode, FileNode newNode)
+		{
+			if (oldNode is InteriorFileNode oldInteriorNode)
+			{
+				InteriorFileNode newInteriorNode = (InteriorFileNode)newNode;
+
+				Assert.AreEqual(oldInteriorNode.Length, newInteriorNode.Length);
+				Assert.AreEqual(oldInteriorNode.Children.Count, newInteriorNode.Children.Count);
+				Assert.AreEqual(oldInteriorNode.RollingHash, newInteriorNode.RollingHash);
+
+				int index = 0;
+				foreach ((TreeNodeRef<FileNode> oldFileRef, TreeNodeRef<FileNode> newFileRef) in oldInteriorNode.Children.Zip(newInteriorNode.Children))
+				{
+					FileNode oldFile = await oldFileRef.ExpandAsync();
+					FileNode newFile = await newFileRef.ExpandAsync();
+					await CompareTrees(oldFile, newFile);
+					index++;
+				}
+			}
+			else if (oldNode is LeafFileNode oldLeafNode)
+			{
+				LeafFileNode newLeafNode = (LeafFileNode)newNode;
+				Assert.AreEqual(oldLeafNode.Length, newLeafNode.Length);
+				Assert.IsTrue(oldLeafNode.Data.AsSingleSegment().Span.SequenceEqual(newLeafNode.Data.AsSingleSegment().Span));
+			}
+			else
+			{
+				throw new NotImplementedException();
+			}
+		}
+
 		static async Task CheckLargeFileTreeAsync(DirectoryNode root, byte[] data)
 		{
 			Assert.AreEqual(0, root.Directories.Count);
@@ -312,20 +392,6 @@ namespace EpicGames.Horde.Tests
 			Assert.IsTrue(worldData.SequenceEqual(data));
 		}
 
-		static async Task FindUniqueBlobs(ITreeBlobRef blobRef, Dictionary<IoHash, long> uniqueBlobs)
-		{
-			ITreeBlob blob = await blobRef.GetTargetAsync();
-
-			ReadOnlySequence<byte> data = blob.Data;
-			uniqueBlobs[blobRef.Hash] = data.Length;
-
-			IReadOnlyList<ITreeBlobRef> references = blob.Refs;
-			foreach (ITreeBlobRef reference in references)
-			{
-				await FindUniqueBlobs(reference, uniqueBlobs);
-			}
-		}
-
 		static async Task<byte[]> GetFileDataAsync(FileNode fileNode)
 		{
 			using (MemoryStream stream = new MemoryStream())
diff --git a/Engine/Source/Programs/Shared/EpicGames.Horde.Tests/BundleTests.cs b/Engine/Source/Programs/Shared/EpicGames.Horde.Tests/BundleTests.cs
--- a/Engine/Source/Programs/Shared/EpicGames.Horde.Tests/BundleTests.cs
+++ b/Engine/Source/Programs/Shared/EpicGames.Horde.Tests/BundleTests.cs
@@ -4,6 +4,7 @@ using EpicGames.Core;
 using EpicGames.Horde.Storage;
 using EpicGames.Horde.Storage.Backends;
 using EpicGames.Horde.Storage.Nodes;
+using Microsoft.Extensions.Caching.Memory;
 using Microsoft.VisualStudio.TestTools.UnitTesting;
 using System;
 using System.Collections.Generic;
@@ -15,13 +16,20 @@ using System.Threading.Tasks;
 namespace EpicGames.Horde.Tests
 {
 	[TestClass]
-	public sealed class BundleTests
+	public sealed class BundleTests : IDisposable
 	{
+		readonly IMemoryCache _cache;
 		readonly InMemoryBlobStore _storage;
 
 		public BundleTests()
 		{
-			_storage = new InMemoryBlobStore();
+			_cache = new MemoryCache(new MemoryCacheOptions());
+			_storage = new InMemoryBlobStore(_cache);
+		}
+
+		public void Dispose()
+		{
+			_cache.Dispose();
 		}
 
 		[TestMethod]
@@ -57,7 +65,7 @@ namespace EpicGames.Horde.Tests
 		public async Task BasicChunkingTests()
 		{
 			RefName refName = new RefName("test");
-			ITreeWriter writer = _storage.CreateTreeWriter(refName.Text);
+			TreeWriter writer = new TreeWriter(_storage, new TreeOptions(), refName.Text);
 
 			ChunkingOptions options = new ChunkingOptions();
 			options.LeafOptions = new ChunkingOptionsForNodeType(8, 8, 8);
@@ -75,14 +83,14 @@ namespace EpicGames.Horde.Tests
 			node = new LeafFileNode();
 			node = await node.AppendAsync(new byte[9], options, writer, CancellationToken.None);
 			Assert.IsTrue(node is InteriorFileNode);
-			Assert.AreEqual(2, node.GetReferences().Count);
+			Assert.AreEqual(2, ((InteriorFileNode)node).Children.Count);
 
-			FileNode? childNode1 = await ((TreeNodeRef<FileNode>)node.GetReferences()[0]).ExpandAsync();
+			FileNode? childNode1 = await ((TreeNodeRef<FileNode>)((InteriorFileNode)node).Children[0]).ExpandAsync();
 			Assert.IsNotNull(childNode1);
 			Assert.IsTrue(childNode1 is LeafFileNode);
 			Assert.AreEqual(8, ((LeafFileNode)childNode1!).Data.Length);
 
-			FileNode? childNode2 = await ((TreeNodeRef<FileNode>)node.GetReferences()[1]).ExpandAsync();
+			FileNode? childNode2 = await ((TreeNodeRef<FileNode>)((InteriorFileNode)node).Children[1]).ExpandAsync();
 			Assert.IsNotNull(childNode2);
 			Assert.IsTrue(childNode2 is LeafFileNode);
 			Assert.AreEqual(1, ((LeafFileNode)childNode2!).Data.Length);
@@ -94,8 +102,8 @@ namespace EpicGames.Horde.Tests
 			BundleHeader oldHeader;
 			{
 				List<BundleImport> imports = new List<BundleImport>();
-				imports.Add(new BundleImport(new BlobLocator("import1"), 10, new (int, IoHash)[] { (5, IoHash.Compute(Encoding.UTF8.GetBytes("blob1"))) }));
-				imports.Add(new BundleImport(new BlobLocator("import2"), 20, new (int, IoHash)[] { (6, IoHash.Compute(Encoding.UTF8.GetBytes("blob2"))) }));
+				imports.Add(new BundleImport(new BlobLocator("import1"), new (int, IoHash)[] { (5, IoHash.Compute(Encoding.UTF8.GetBytes("blob1"))) }));
+				imports.Add(new BundleImport(new BlobLocator("import2"), new (int, IoHash)[] { (6, IoHash.Compute(Encoding.UTF8.GetBytes("blob2"))) }));
 
 				List<BundleExport> exports = new List<BundleExport>();
 				exports.Add(new BundleExport(IoHash.Compute(Encoding.UTF8.GetBytes("export1")), 2, new int[] { 1, 2 }));
@@ -121,7 +129,6 @@ namespace EpicGames.Horde.Tests
 				BundleImport newImport = newHeader.Imports[idx];
 
 				Assert.AreEqual(oldImport.Locator, newImport.Locator);
-				Assert.AreEqual(oldImport.ExportCount, newImport.ExportCount);
 				Assert.AreEqual(oldImport.Exports.Count, newImport.Exports.Count);
 
 				for (int importIdx = 0; importIdx < oldImport.Exports.Count; importIdx++)
@@ -158,26 +165,27 @@ namespace EpicGames.Horde.Tests
 		[TestMethod]
 		public async Task BasicTestDirectory()
 		{
-			InMemoryBlobStore store = new InMemoryBlobStore();
+			InMemoryBlobStore store = _storage;
 
 			DirectoryNode root = new DirectoryNode(DirectoryFlags.None);
 			DirectoryNode node = root.AddDirectory("hello");
 			DirectoryNode node2 = node.AddDirectory("world");
 
 			RefName refName = new RefName("testref");
-			await store.WriteTreeAsync(refName, root);
+			await store.WriteNodeAsync(refName, root);
 
 			// Should be stored inline
 			Assert.AreEqual(1, store.Refs.Count);
 			Assert.AreEqual(1, store.Blobs.Count);
 
 			// Check the ref
-			RefValue refValue = await store.ReadRefValueAsync(refName);
-			Assert.AreEqual(0, refValue.Bundle.Header.Imports.Count);
-			Assert.AreEqual(3, refValue.Bundle.Header.Exports.Count);
+			NodeLocator refTarget = await store.ReadRefTargetAsync(refName);
+			Bundle bundle = await store.ReadBundleAsync(refTarget.Blob);
+			Assert.AreEqual(0, bundle.Header.Imports.Count);
+			Assert.AreEqual(3, bundle.Header.Exports.Count);
 
 			// Create a new bundle and read it back in again
-			DirectoryNode newRoot = await store.ReadTreeAsync<DirectoryNode>(refName);
+			DirectoryNode newRoot = await store.ReadNodeAsync<DirectoryNode>(refName);
 
 			Assert.AreEqual(0, newRoot.Files.Count);
 			Assert.AreEqual(1, newRoot.Directories.Count);
@@ -199,21 +207,19 @@ namespace EpicGames.Horde.Tests
 			TreeOptions options = new TreeOptions();
 			options.MaxBlobSize = 1;
 
-			InMemoryBlobStore storage = new InMemoryBlobStore();
-
 			DirectoryNode root = new DirectoryNode(DirectoryFlags.None);
 			root.AddDirectory("node1");
 			root.AddDirectory("node2");
 			root.AddDirectory("node3");
 
-			Assert.AreEqual(0, storage.Refs.Count);
-			Assert.AreEqual(0, storage.Blobs.Count);
+			Assert.AreEqual(0, _storage.Refs.Count);
+			Assert.AreEqual(0, _storage.Blobs.Count);
 
 			RefName refName = new RefName("ref");
-			await storage.WriteTreeAsync(refName, root, options);
+			await _storage.WriteNodeAsync(refName, root, options);
 
-			Assert.AreEqual(1, storage.Refs.Count);
-			Assert.AreEqual(2, storage.Blobs.Count);
+			Assert.AreEqual(1, _storage.Refs.Count);
+			Assert.AreEqual(2, _storage.Blobs.Count);
 		}
 
 		[TestMethod]
@@ -232,14 +238,14 @@ namespace EpicGames.Horde.Tests
 				DirectoryNode node3 = node2.AddDirectory("node3");
 				DirectoryNode node4 = node3.AddDirectory("node4");
 
-				await _storage.WriteTreeAsync(refName, root, options);
+				await _storage.WriteNodeAsync(refName, root, options);
 
 				Assert.AreEqual(1, _storage.Refs.Count);
 				Assert.AreEqual(5, _storage.Blobs.Count);
 			}
 
 			{
-				DirectoryNode root = await _storage.ReadTreeAsync<DirectoryNode>(refName);
+				DirectoryNode root = await _storage.ReadNodeAsync<DirectoryNode>(refName);
 
 				DirectoryNode? newNode1 = await root.FindDirectoryAsync("node1", CancellationToken.None);
 				Assert.IsNotNull(newNode1);
@@ -327,7 +333,7 @@ namespace EpicGames.Horde.Tests
 		[TestMethod]
 		public async Task CoreAppendTest()
 		{
-			ITreeWriter writer = _storage.CreateTreeWriter();
+			TreeWriter writer = new TreeWriter(_storage);
 
 			byte[] data = new byte[4096];
 			new Random(0).NextBytes(data);
@@ -367,7 +373,7 @@ namespace EpicGames.Horde.Tests
 
 		async Task ChunkingTests(ChunkingOptions options)
 		{
-			ITreeWriter writer = _storage.CreateTreeWriter(new TreeOptions());
+			TreeWriter writer = new TreeWriter(_storage);
 
 			byte[] data = new byte[4096];
 			new Random(0).NextBytes(data);
@@ -420,6 +426,7 @@ namespace EpicGames.Horde.Tests
 			}
 		}
 
+#if false
 		[TestMethod]
 		public async Task SpillTestAsync()
 		{
@@ -481,5 +488,6 @@ namespace EpicGames.Horde.Tests
 			}
 			return size;
 		}
+#endif
 	}
 }
diff --git a/Engine/Source/Programs/Shared/EpicGames.Horde.Tests/ChunkingTests.cs b/Engine/Source/Programs/Shared/EpicGames.Horde.Tests/ChunkingTests.cs
--- a/Engine/Source/Programs/Shared/EpicGames.Horde.Tests/ChunkingTests.cs
+++ b/Engine/Source/Programs/Shared/EpicGames.Horde.Tests/ChunkingTests.cs
@@ -10,6 +10,7 @@ using System.IO;
 using System.Threading;
 using EpicGames.Horde.Storage;
 using EpicGames.Horde.Storage.Backends;
+using Microsoft.Extensions.Caching.Memory;
 
 namespace EpicGames.Horde.Tests
 {
@@ -48,37 +49,40 @@ namespace EpicGames.Horde.Tests
 		[TestMethod]
 		public async Task BasicChunkingTests()
 		{
-			InMemoryBlobStore store = new InMemoryBlobStore();
-			ITreeWriter writer = store.CreateTreeWriter("test");
+			using MemoryCache cache = new MemoryCache(new MemoryCacheOptions());
+
+			InMemoryBlobStore store = new InMemoryBlobStore(cache);
+			TreeWriter writer = new TreeWriter(store, new TreeOptions(), "test");
 
 			ChunkingOptions options = new ChunkingOptions();
 			options.LeafOptions = new ChunkingOptionsForNodeType(8, 8, 8);
 
 			FileNode node = new LeafFileNode();
-			await node.AppendAsync(new byte[7], options, writer, CancellationToken.None);
+			node = await node.AppendAsync(new byte[7], options, writer, CancellationToken.None);
 			Assert.AreEqual(7, node.Length);
-			Assert.AreEqual(0, node.GetReferences().Count);
+			Assert.IsTrue(node is LeafFileNode);
 			Assert.AreEqual(7, (await node.ToByteArrayAsync(CancellationToken.None)).Length);
 
 			node = new LeafFileNode();
-			await node.AppendAsync(new byte[8], options, writer, CancellationToken.None);
+			node = await node.AppendAsync(new byte[8], options, writer, CancellationToken.None);
 			Assert.AreEqual(8, node.Length);
-			Assert.AreEqual(0, node.GetReferences().Count);
+			Assert.IsTrue(node is LeafFileNode);
 			Assert.AreEqual(8, (await node.ToByteArrayAsync(CancellationToken.None)).Length);
 
 			node = new LeafFileNode();
 			node = await node.AppendAsync(new byte[9], options, writer, CancellationToken.None);
 			Assert.AreEqual(9, node.Length);
-			Assert.AreEqual(2, node.GetReferences().Count);
+			Assert.IsTrue(node is InteriorFileNode);
+			Assert.AreEqual(2, ((InteriorFileNode)node).Children.Count);
 
-			FileNode? childNode1 = await ((TreeNodeRef<FileNode>)node.GetReferences()[0]).ExpandAsync();
+			FileNode? childNode1 = await ((TreeNodeRef<FileNode>)((InteriorFileNode)node).Children[0]).ExpandAsync();
 			Assert.IsNotNull(childNode1);
-			Assert.AreEqual(0, childNode1!.GetReferences().Count);
+			Assert.IsTrue(childNode1 is LeafFileNode);
 			Assert.AreEqual(8, (await childNode1!.ToByteArrayAsync(CancellationToken.None)).Length);
 
-			FileNode? childNode2 = await ((TreeNodeRef<FileNode>)node.GetReferences()[1]).ExpandAsync();
+			FileNode? childNode2 = await ((TreeNodeRef<FileNode>)((InteriorFileNode)node).Children[1]).ExpandAsync();
 			Assert.IsNotNull(childNode2);
-			Assert.AreEqual(0, childNode2!.GetReferences().Count);
+			Assert.IsTrue(childNode2 is LeafFileNode);
 			Assert.AreEqual(1, (await childNode2!.ToByteArrayAsync(CancellationToken.None)).Length);
 		}
 
@@ -104,8 +108,10 @@ namespace EpicGames.Horde.Tests
 
 		static async Task TestChunkingAsync(ChunkingOptions options)
 		{
-			InMemoryBlobStore store = new InMemoryBlobStore();
-			ITreeWriter writer = store.CreateTreeWriter();
+			using MemoryCache cache = new MemoryCache(new MemoryCacheOptions());
+
+			InMemoryBlobStore store = new InMemoryBlobStore(cache);
+			TreeWriter writer = new TreeWriter(store, new TreeOptions());
 
 			byte[] data = new byte[4096];
 			new Random(0).NextBytes(data);
diff --git a/Engine/Source/Programs/Shared/EpicGames.Horde/Storage/Backends/InMemoryBlobStore.cs b/Engine/Source/Programs/Shared/EpicGames.Horde/Storage/Backends/InMemoryBlobStore.cs
--- a/Engine/Source/Programs/Shared/EpicGames.Horde/Storage/Backends/InMemoryBlobStore.cs
+++ b/Engine/Source/Programs/Shared/EpicGames.Horde/Storage/Backends/InMemoryBlobStore.cs
@@ -8,6 +8,7 @@ using System.IO;
 using System.Threading;
 using System.Threading.Tasks;
 using EpicGames.Core;
+using Microsoft.Extensions.Caching.Memory;
 
 namespace EpicGames.Horde.Storage.Backends
 {
@@ -24,19 +25,19 @@ namespace EpicGames.Horde.Storage.Backends
 		/// <summary>
 		/// Map of ref name to ref data
 		/// </summary>
-		readonly ConcurrentDictionary<RefName, RefTarget> _refs = new ConcurrentDictionary<RefName, RefTarget>();
+		readonly ConcurrentDictionary<RefName, NodeLocator> _refs = new ConcurrentDictionary<RefName, NodeLocator>();
 
 		/// <inheritdoc cref="_blobs"/>
 		public IReadOnlyDictionary<BlobLocator, Bundle> Blobs => _blobs;
 
 		/// <inheritdoc cref="_refs"/>
-		public IReadOnlyDictionary<RefName, RefTarget> Refs => _refs;
+		public IReadOnlyDictionary<RefName, NodeLocator> Refs => _refs;
 
 		/// <summary>
 		/// Constructor
 		/// </summary>
-		public InMemoryBlobStore() 
-			: base(null)
+		public InMemoryBlobStore(IMemoryCache cache) 
+			: base(cache)
 		{
 		}
 
@@ -65,15 +66,15 @@ namespace EpicGames.Horde.Storage.Backends
 		public override Task DeleteRefAsync(RefName name, CancellationToken cancellationToken) => Task.FromResult(_refs.TryRemove(name, out _));
 
 		/// <inheritdoc/>
-		public override Task<RefTarget?> TryReadRefTargetAsync(RefName name, DateTime cacheTime = default, CancellationToken cancellationToken = default)
+		public override Task<NodeLocator> TryReadRefTargetAsync(RefName name, DateTime cacheTime = default, CancellationToken cancellationToken = default)
 		{
-			RefTarget? refTarget;
+			NodeLocator refTarget;
 			_refs.TryGetValue(name, out refTarget);
 			return Task.FromResult(refTarget);
 		}
 
 		/// <inheritdoc/>
-		public override Task WriteRefTargetAsync(RefName name, RefTarget target, CancellationToken cancellationToken = default)
+		public override Task WriteRefTargetAsync(RefName name, NodeLocator target, CancellationToken cancellationToken = default)
 		{
 			_refs[name] = target;
 			return Task.CompletedTask;
diff --git a/Engine/Source/Programs/Shared/EpicGames.Horde/Storage/Bundle.cs b/Engine/Source/Programs/Shared/EpicGames.Horde/Storage/Bundle.cs
--- a/Engine/Source/Programs/Shared/EpicGames.Horde/Storage/Bundle.cs
+++ b/Engine/Source/Programs/Shared/EpicGames.Horde/Storage/Bundle.cs
@@ -161,6 +161,11 @@ namespace EpicGames.Horde.Storage
 		/// </summary>
 		public static ReadOnlyMemory<byte> Signature { get; } = Encoding.UTF8.GetBytes("UEBN");
 
+		/// <summary>
+		/// Length of the prelude data
+		/// </summary>
+		public const int PreludeLength = 8;
+
 		/// <summary>
 		/// Compression format for the bundle
 		/// </summary>
@@ -227,6 +232,17 @@ namespace EpicGames.Horde.Storage
 			}
 		}
 
+		/// <summary>
+		/// Reads the prelude bytes from a bundle, and returns the header size.
+		/// </summary>
+		/// <param name="prelude">The prelude data</param>
+		/// <returns>Size of the header, including the prelude data</returns>
+		public static int ReadPrelude(ReadOnlyMemory<byte> prelude)
+		{
+			CheckPrelude(prelude);
+			return BinaryPrimitives.ReadInt32BigEndian(prelude.Span.Slice(4));
+		}
+
 		/// <summary>
 		/// Reads a bundle header from a stream
 		/// </summary>
@@ -302,11 +318,6 @@ namespace EpicGames.Horde.Storage
 		/// </summary>
 		public BlobLocator Locator { get; }
 
-		/// <summary>
-		/// Number of exports from this blob.
-		/// </summary>
-		public int ExportCount { get; }
-
 		/// <summary>
 		/// Indexes of referenced nodes exported from this bundle
 		/// </summary>
@@ -315,10 +326,9 @@ namespace EpicGames.Horde.Storage
 		/// <summary>
 		/// Constructor
 		/// </summary>
-		public BundleImport(BlobLocator locator, int exportCount, IReadOnlyList<(int, IoHash)> exports)
+		public BundleImport(BlobLocator locator, IReadOnlyList<(int, IoHash)> exports)
 		{
 			Locator = locator;
-			ExportCount = exportCount;
 			Exports = exports;
 		}
 
@@ -330,8 +340,6 @@ namespace EpicGames.Horde.Storage
 		{
 			Locator = reader.ReadBlobLocator();
 
-			ExportCount = (int)reader.ReadUnsignedVarInt();
-
 			int count = (int)reader.ReadUnsignedVarInt();
 
 			(int, IoHash)[] exports = new (int, IoHash)[count];
@@ -351,7 +359,6 @@ namespace EpicGames.Horde.Storage
 		public void Write(IMemoryWriter writer)
 		{
 			writer.WriteBlobLocator(Locator);
-			writer.WriteUnsignedVarInt(ExportCount);
 
 			writer.WriteUnsignedVarInt(Exports.Count);
 			for(int idx = 0; idx < Exports.Count; idx++)
diff --git a/Engine/Source/Programs/Shared/EpicGames.Horde/Storage/IStorageClient.cs b/Engine/Source/Programs/Shared/EpicGames.Horde/Storage/IStorageClient.cs
--- a/Engine/Source/Programs/Shared/EpicGames.Horde/Storage/IStorageClient.cs
+++ b/Engine/Source/Programs/Shared/EpicGames.Horde/Storage/IStorageClient.cs
@@ -1,14 +1,14 @@
 // Copyright Epic Games, Inc. All Rights Reserved.
 
 using System;
-using System.Buffers;
 using System.Collections.Generic;
 using System.Diagnostics;
 using System.IO;
+using System.Reflection;
+using System.Reflection.Metadata;
 using System.Threading;
 using System.Threading.Tasks;
 using EpicGames.Core;
-using Microsoft.CodeAnalysis.Rename;
 
 namespace EpicGames.Horde.Storage
 {
@@ -34,32 +34,6 @@
 	}
 
 	/// <summary>
-	/// Target for a ref
-	/// </summary>
-	[DebuggerDisplay("{Locator}#{Export}")]
-	public class RefTarget
-	{
-		/// <summary>
-		/// Location of the blob
-		/// </summary>
-		public BlobLocator Locator { get; }
-
-		/// <summary>
-		/// Index of the blob export
-		/// </summary>
-		public int ExportIdx { get; }
-
-		/// <summary>
-		/// Constructor
-		/// </summary>
-		public RefTarget(BlobLocator locator, int exportIdx)
-		{
-			Locator = locator;
-			ExportIdx = exportIdx;
-		}
-	}
-
-	/// <summary>
 	/// Value for a ref
 	/// </summary>
 	[DebuggerDisplay("{Target}")]
@@ -68,7 +42,7 @@
 		/// <summary>
 		/// Target for the ref
 		/// </summary>
-		public RefTarget Target { get; }
+		public NodeLocator Target { get; }
 
 		/// <summary>
 		/// Bundle referenced as the target
@@ -78,7 +52,7 @@
 		/// <summary>
 		/// Constructor
 		/// </summary>
-		public RefValue(RefTarget target, Bundle bundle)
+		public RefValue(NodeLocator target, Bundle bundle)
 		{
 			Target = target;
 			Bundle = bundle;
@@ -86,3 +60,33 @@
 	}
 
 	/// <summary>
+	/// Locates a node in storage
+	/// </summary>
+	public struct NodeLocator
+	{
+		/// <summary>
+		/// Location of the blob containing this node
+		/// </summary>
+		public BlobLocator Blob { get; }
+
+		/// <summary>
+		/// Index of the export within the blob
+		/// </summary>
+		public int ExportIdx { get; }
+
+		/// <summary>
+		/// Constructor
+		/// </summary>
+		public NodeLocator(BlobLocator blob, int exportIdx)
+		{
+			Blob = blob;
+			ExportIdx = exportIdx;
+		}
+
+		/// <summary>
+		/// Determines if this locator points to a valid entry
+		/// </summary>
+		public bool IsValid() => Blob.IsValid();
+	}
+
+	/// <summary>
@@ -93,7 +97,7 @@
 		#region Blobs
 
 		/// <summary>
-		/// Reads data for a blob from the store
+		/// Reads raw data for a blob from the store
 		/// </summary>
 		/// <param name="locator">The blob locator</param>
 		/// <param name="cancellationToken">Cancellation token for the operation</param>
@@ -101,5 +105,14 @@
 		Task<Stream> ReadBlobAsync(BlobLocator locator, CancellationToken cancellationToken = default);
 
 		/// <summary>
+		/// Reads a ranged chunk from a blob
+		/// </summary>
+		/// <param name="locator">Locator for the blob</param>
+		/// <param name="offset">Starting offset for the data to read</param>
+		/// <param name="length">Length of the data</param>
+		/// <param name="cancellationToken">Cancellation token for the operation</param>
+		Task<Stream> ReadBlobRangeAsync(BlobLocator locator, int offset, int length, CancellationToken cancellationToken = default);
+
+		/// <summary>
 		/// Writes a new blob to the store
 		/// </summary>
@@ -122,22 +135,5 @@
 		Task<Bundle> ReadBundleAsync(BlobLocator locator, CancellationToken cancellationToken = default);
 
 		/// <summary>
-		/// Reads the header for a bundle from the store
-		/// </summary>
-		/// <param name="locator">Locator for the blob</param>
-		/// <param name="cancellationToken">Cancellation token for the operation</param>
-		/// <returns></returns>
-		Task<BundleHeader> ReadBundleHeaderAsync(BlobLocator locator, CancellationToken cancellationToken = default);
-
-		/// <summary>
-		/// Reads a bundle packet from the store
-		/// </summary>
-		/// <param name="locator">Locator for the blob</param>
-		/// <param name="packetIdx">Index of the packet</param>
-		/// <param name="cancellationToken">Cancellation token for the operation</param>
-		/// <returns>Decoded data for the packet</returns>
-		Task<ReadOnlyMemory<byte>> ReadBundlePacketAsync(BlobLocator locator, int packetIdx, CancellationToken cancellationToken = default);
-
-		/// <summary>
 		/// Writes a new bundle to the store
 		/// </summary>
@@ -149,57 +145,48 @@
 
 		#endregion
 
-		#region Trees
+		#region Nodes
 
 		/// <summary>
-		/// Creates a new context for reading trees.
+		/// Reads a node from a bundle
 		/// </summary>
-		/// <param name="options">Options for the tree writer</param>
-		/// <param name="prefix">Prefix for blob names.</param>
-		/// <returns>New context instance</returns>
-		ITreeWriter CreateTreeWriter(TreeOptions options, Utf8String prefix = default);
+		/// <param name="locator">Locator for the node</param>
+		/// <param name="cancellationToken">Cancellation token for the operation</param>
+		/// <returns>Node data read from the given bundle</returns>
+		ValueTask<TNode> ReadNodeAsync<TNode>(NodeLocator locator, CancellationToken cancellationToken = default) where TNode : TreeNode;
 
 		/// <summary>
-		/// Attempts to read a tree with the given name
+		/// Reads data for a ref from the store, along with the node's contents.
 		/// </summary>
-		/// <param name="name">Name of the ref used to store the tree</param>
-		/// <param name="maxAge"></param>
+		/// <param name="name">The ref name</param>
+		/// <param name="cacheTime">Minimum coherency for any cached value to be returned</param>
 		/// <param name="cancellationToken">Cancellation token for the operation</param>
-		/// <returns>Reference to the root of the tree. Null if the ref does not exist.</returns>
-		Task<ITreeBlob?> TryReadTreeAsync(RefName name, TimeSpan maxAge = default, CancellationToken cancellationToken = default);
+		/// <returns>Node for the given ref, or null if it does not exist</returns>
+		Task<TNode?> TryReadNodeAsync<TNode>(RefName name, DateTime cacheTime = default, CancellationToken cancellationToken = default) where TNode : TreeNode;
 
 		#endregion
 
 		#region Refs
 
 		/// <summary>
-		/// Reads data for a ref from the store, along with the blob contents.
-		/// </summary>
-		/// <param name="name">The ref name</param>
-		/// <param name="cacheTime">Minimum coherency for any cached value to be returned</param>
-		/// <param name="cancellationToken">Cancellation token for the operation</param>
-		/// <returns></returns>
-		Task<RefValue?> TryReadRefValueAsync(RefName name, DateTime cacheTime = default, CancellationToken cancellationToken = default);
-
-		/// <summary>
 		/// Reads data for a ref from the store
 		/// </summary>
 		/// <param name="name">The ref name</param>
 		/// <param name="cacheTime">Minimum coherency for any cached value to be returned</param>
 		/// <param name="cancellationToken">Cancellation token for the operation</param>
-		/// <returns></returns>
-		Task<RefTarget?> TryReadRefTargetAsync(RefName name, DateTime cacheTime = default, CancellationToken cancellationToken = default);
+		/// <returns>Node pointed to by the ref</returns>
+		Task<NodeLocator> TryReadRefTargetAsync(RefName name, DateTime cacheTime = default, CancellationToken cancellationToken = default);
 
 		/// <summary>
 		/// Writes a new ref to the store which points to a new blob
 		/// </summary>
 		/// <param name="name">Ref to write</param>
-		/// <param name="bundle">Bundle to write</param>
-		/// <param name="exportIdx">Index of the export to target for the ref</param>
+		/// <param name="bundle">The bundle to write</param>
+		/// <param name="exportIdx">Index of the export in the bundle to be the root of the tree</param>
 		/// <param name="prefix">Prefix for blob names.</param>
 		/// <param name="cancellationToken">Cancellation token for the operation</param>
 		/// <returns>Unique identifier for the blob</returns>
-		Task<RefTarget> WriteRefValueAsync(RefName name, Bundle bundle, int exportIdx = 0, Utf8String prefix = default, CancellationToken cancellationToken = default);
+		Task<NodeLocator> WriteRefAsync(RefName name, Bundle bundle, int exportIdx, Utf8String prefix = default, CancellationToken cancellationToken = default);
 
 		/// <summary>
 		/// Writes a new ref to the store
@@ -208,7 +195,7 @@
 		/// <param name="target">The target for the ref</param>
 		/// <param name="cancellationToken">Cancellation token for the operation</param>
 		/// <returns>Unique identifier for the blob</returns>
-		Task WriteRefTargetAsync(RefName name, RefTarget target, CancellationToken cancellationToken = default);
+		Task WriteRefTargetAsync(RefName name, NodeLocator target, CancellationToken cancellationToken = default);
 
 		/// <summary>
 		/// Reads data for a ref from the store
@@ -244,5 +231,8 @@
 			public Task<Stream> ReadBlobAsync(BlobLocator locator, CancellationToken cancellationToken = default) => _inner.ReadBlobAsync(locator, cancellationToken);
 
 			/// <inheritdoc/>
+			public Task<Stream> ReadBlobRangeAsync(BlobLocator locator, int offset, int length, CancellationToken cancellationToken = default) => _inner.ReadBlobRangeAsync(locator, offset, length, cancellationToken);
+
+			/// <inheritdoc/>
 			public Task<BlobLocator> WriteBlobAsync(Stream stream, Utf8String prefix = default, CancellationToken cancellationToken = default) => _inner.WriteBlobAsync(stream, prefix, cancellationToken);
 
@@ -254,11 +244,5 @@
 			public Task<Bundle> ReadBundleAsync(BlobLocator locator, CancellationToken cancellationToken = default) => _inner.ReadBundleAsync(locator, cancellationToken);
 
 			/// <inheritdoc/>
-			public Task<BundleHeader> ReadBundleHeaderAsync(BlobLocator locator, CancellationToken cancellationToken = default) => _inner.ReadBundleHeaderAsync(locator, cancellationToken);
-
-			/// <inheritdoc/>
-			public Task<ReadOnlyMemory<byte>> ReadBundlePacketAsync(BlobLocator locator, int packetIdx, CancellationToken cancellationToken = default) => _inner.ReadBundlePacketAsync(locator, packetIdx, cancellationToken);
-
-			/// <inheritdoc/>
 			public Task<BlobLocator> WriteBundleAsync(Bundle bundle, Utf8String prefix = default, CancellationToken cancellationToken = default) => _inner.WriteBundleAsync(bundle, prefix, cancellationToken);
 
@@ -267,10 +251,7 @@
 			#region Trees
 
 			/// <inheritdoc/>
-			public ITreeWriter CreateTreeWriter(TreeOptions options, Utf8String prefix) => _inner.CreateTreeWriter(options, prefix);
+			public ValueTask<TNode> ReadNodeAsync<TNode>(NodeLocator locator, CancellationToken cancellationToken = default) where TNode : TreeNode => _inner.ReadNodeAsync<TNode>(locator, cancellationToken);
 
-			/// <inheritdoc/>
-			public Task<ITreeBlob?> TryReadTreeAsync(RefName name, TimeSpan maxAge = default, CancellationToken cancellationToken = default) => _inner.TryReadTreeAsync(name, maxAge, cancellationToken);
-
 			#endregion
 
@@ -280,16 +261,16 @@
 			public Task DeleteRefAsync(RefName name, CancellationToken cancellationToken = default) => _inner.DeleteRefAsync(name, cancellationToken);
 
 			/// <inheritdoc/>
-			public Task<RefValue?> TryReadRefValueAsync(RefName name, DateTime cacheTime = default, CancellationToken cancellationToken = default) => _inner.TryReadRefValueAsync(name, cacheTime, cancellationToken);
+			public Task<TNode?> TryReadNodeAsync<TNode>(RefName name, DateTime cacheTime = default, CancellationToken cancellationToken = default) where TNode : TreeNode => _inner.TryReadNodeAsync<TNode>(name, cacheTime, cancellationToken);
 
 			/// <inheritdoc/>
-			public Task<RefTarget?> TryReadRefTargetAsync(RefName name, DateTime cacheTime = default, CancellationToken cancellationToken = default) => _inner.TryReadRefTargetAsync(name, cacheTime, cancellationToken);
+			public Task<NodeLocator> TryReadRefTargetAsync(RefName name, DateTime cacheTime = default, CancellationToken cancellationToken = default) => _inner.TryReadRefTargetAsync(name, cacheTime, cancellationToken);
 
 			/// <inheritdoc/>
-			public Task<RefTarget> WriteRefValueAsync(RefName name, Bundle bundle, int exportIdx = 0, Utf8String prefix = default, CancellationToken cancellationToken = default) => _inner.WriteRefValueAsync(name, bundle, exportIdx, prefix, cancellationToken);
+			public Task<NodeLocator> WriteRefAsync(RefName name, Bundle bundle, int exportIdx = 0, Utf8String prefix = default, CancellationToken cancellationToken = default) => _inner.WriteRefAsync(name, bundle, exportIdx, prefix, cancellationToken);
 
 			/// <inheritdoc/>
-			public Task WriteRefTargetAsync(RefName name, RefTarget target, CancellationToken cancellationToken = default) => _inner.WriteRefTargetAsync(name, target, cancellationToken);
+			public Task WriteRefTargetAsync(RefName name, NodeLocator target, CancellationToken cancellationToken = default) => _inner.WriteRefTargetAsync(name, target, cancellationToken);
 
 			#endregion
 		}
@@ -301,43 +282,7 @@
 		/// <returns></returns>
 		public static IStorageClient<T> ForType<T>(this IStorageClient blobStore) => new TypedStorageClient<T>(blobStore);
 
-		/// <summary>
-		/// Checks if the given ref exists
-		/// </summary>
-		/// <param name="store">The store instance to read from</param>
-		/// <param name="name">Name of the reference to look for</param>
-		/// <param name="cacheTime">Minimum coherency for any cached value to be returned</param>
-		/// <param name="cancellationToken">Cancellation token for the operation</param>
-		/// <returns>True if the ref exists, false if it did not exist</returns>
-		public static async Task<bool> HasRefAsync(this IStorageClient store, RefName name, DateTime cacheTime = default, CancellationToken cancellationToken = default)
-		{
-			RefTarget? locator = await store.TryReadRefTargetAsync(name, cacheTime, cancellationToken);
-			return locator != null;
-		}
-
-		/// <summary>
-		/// Checks if the given ref exists
-		/// </summary>
-		/// <param name="store">The store instance to read from</param>
-		/// <param name="name">Name of the reference to look for</param>
-		/// <param name="maxAge">Maximum age of any cached ref</param>
-		/// <param name="cancellationToken">Cancellation token for the operation</param>
-		/// <returns>True if the ref exists, false if it did not exist</returns>
-		public static Task<bool> HasRefAsync(this IStorageClient store, RefName name, TimeSpan maxAge, CancellationToken cancellationToken = default)
-		{
-			return HasRefAsync(store, name, DateTime.UtcNow - maxAge, cancellationToken);
-		}
-
-		/// <summary>
-		/// Creates a new tree writer with default options
-		/// </summary>
-		/// <param name="store">Store instance</param>
-		/// <param name="prefix"></param>
-		/// <returns></returns>
-		public static ITreeWriter CreateTreeWriter(this IStorageClient store, Utf8String prefix = default)
-		{
-			return store.CreateTreeWriter(new TreeOptions(), prefix);
-		}
+		#region Nodes
 
 		/// <summary>
 		/// Attempts to reads a ref from the store
@@ -347,22 +292,9 @@
 		/// <param name="maxAge">Maximum age of any cached ref</param>
 		/// <param name="cancellationToken">Cancellation token for the operation</param>
 		/// <returns>The ref value</returns>
-		public static Task<RefValue?> TryReadRefValueAsync(this IStorageClient store, RefName name, TimeSpan maxAge, CancellationToken cancellationToken = default)
+		public static Task<TNode?> TryReadNodeAsync<TNode>(this IStorageClient store, RefName name, TimeSpan maxAge, CancellationToken cancellationToken = default) where TNode : TreeNode
 		{
-			return store.TryReadRefValueAsync(name, DateTime.UtcNow - maxAge, cancellationToken);
-		}
-
-		/// <summary>
-		/// Attempts to reads a ref from the store
-		/// </summary>
-		/// <param name="store">The store instance to read from</param>
-		/// <param name="name">Id for the ref</param>
-		/// <param name="maxAge">Maximum age of any cached ref</param>
-		/// <param name="cancellationToken">Cancellation token for the operation</param>
-		/// <returns>The ref target</returns>
-		public static Task<RefTarget?> TryReadRefTargetAsync(this IStorageClient store, RefName name, TimeSpan maxAge, CancellationToken cancellationToken = default)
-		{
-			return store.TryReadRefTargetAsync(name, DateTime.UtcNow - maxAge, cancellationToken);
+			return store.TryReadNodeAsync<TNode>(name, DateTime.UtcNow - maxAge, cancellationToken);
 		}
 
 		/// <summary>
@@ -373,9 +305,9 @@
 		/// <param name="cacheTime">Minimum coherency of any cached result</param>
 		/// <param name="cancellationToken">Cancellation token for the operation</param>
 		/// <returns>The blob instance</returns>
-		public static async Task<RefValue> ReadRefValueAsync(this IStorageClient store, RefName name, DateTime cacheTime = default, CancellationToken cancellationToken = default)
+		public static async Task<TNode> ReadNodeAsync<TNode>(this IStorageClient store, RefName name, DateTime cacheTime = default, CancellationToken cancellationToken = default) where TNode : TreeNode
 		{
-			RefValue? refValue = await store.TryReadRefValueAsync(name, cacheTime, cancellationToken);
+			TNode? refValue = await store.TryReadNodeAsync<TNode>(name, cacheTime, cancellationToken);
 			if (refValue == null)
 			{
 				throw new RefNameNotFoundException(name);
@@ -391,9 +323,69 @@
 		/// <param name="maxAge">Maximum age for any cached result</param>
 		/// <param name="cancellationToken">Cancellation token for the operation</param>
 		/// <returns>The ref value</returns>
-		public static Task<RefValue> ReadRefValueAsync(this IStorageClient store, RefName name, TimeSpan maxAge, CancellationToken cancellationToken = default)
+		public static Task<TNode> ReadNodeAsync<TNode>(this IStorageClient store, RefName name, TimeSpan maxAge, CancellationToken cancellationToken = default) where TNode : TreeNode
 		{
-			return ReadRefValueAsync(store, name, DateTime.UtcNow - maxAge, cancellationToken);
+			return ReadNodeAsync<TNode>(store, name, DateTime.UtcNow - maxAge, cancellationToken);
 		}
 
 		/// <summary>
+		/// Writes a node to storage
+		/// </summary>
+		/// <param name="store">Store instance to write to</param>
+		/// <param name="name">Name of the ref containing this node</param>
+		/// <param name="node">Node to be written</param>
+		/// <param name="options">Options for the node writer</param>
+		/// <param name="prefix">Prefix for uploaded blobs</param>
+		/// <param name="cancellationToken">Cancellation token for the operation</param>
+		/// <returns>Location of node targetted by the ref</returns>
+		public static async Task<NodeLocator> WriteNodeAsync(this IStorageClient store, RefName name, TreeNode node, TreeOptions? options = null, Utf8String prefix = default, CancellationToken cancellationToken = default)
+		{
+			TreeWriter writer = new TreeWriter(store, options, prefix.IsEmpty ? name.Text : prefix);
+			return await writer.WriteRefAsync(name, node, cancellationToken);
+		}
+
+		#endregion
+
+		#region Refs
+
+		/// <summary>
+		/// Checks if the given ref exists
+		/// </summary>
+		/// <param name="store">The store instance to read from</param>
+		/// <param name="name">Name of the reference to look for</param>
+		/// <param name="cacheTime">Minimum coherency for any cached value to be returned</param>
+		/// <param name="cancellationToken">Cancellation token for the operation</param>
+		/// <returns>True if the ref exists, false if it did not exist</returns>
+		public static async Task<bool> HasRefAsync(this IStorageClient store, RefName name, DateTime cacheTime = default, CancellationToken cancellationToken = default)
+		{
+			NodeLocator locator = await store.TryReadRefTargetAsync(name, cacheTime, cancellationToken);
+			return locator.IsValid();
+		}
+
+		/// <summary>
+		/// Checks if the given ref exists
+		/// </summary>
+		/// <param name="store">The store instance to read from</param>
+		/// <param name="name">Name of the reference to look for</param>
+		/// <param name="maxAge">Maximum age of any cached ref</param>
+		/// <param name="cancellationToken">Cancellation token for the operation</param>
+		/// <returns>True if the ref exists, false if it did not exist</returns>
+		public static Task<bool> HasRefAsync(this IStorageClient store, RefName name, TimeSpan maxAge, CancellationToken cancellationToken = default)
+		{
+			return HasRefAsync(store, name, DateTime.UtcNow - maxAge, cancellationToken);
+		}
+
+		/// <summary>
+		/// Attempts to reads a ref from the store
+		/// </summary>
+		/// <param name="store">The store instance to read from</param>
+		/// <param name="name">Id for the ref</param>
+		/// <param name="maxAge">Maximum age of any cached ref</param>
+		/// <param name="cancellationToken">Cancellation token for the operation</param>
+		/// <returns>The ref target</returns>
+		public static Task<NodeLocator> TryReadRefTargetAsync(this IStorageClient store, RefName name, TimeSpan maxAge, CancellationToken cancellationToken = default)
+		{
+			return store.TryReadRefTargetAsync(name, DateTime.UtcNow - maxAge, cancellationToken);
+		}
+
+		/// <summary>
@@ -404,10 +396,10 @@
 		/// <param name="cacheTime">Minimum coherency of any cached result</param>
 		/// <param name="cancellationToken">Cancellation token for the operation</param>
 		/// <returns>The ref target</returns>
-		public static async Task<RefTarget> ReadRefTargetAsync(this IStorageClient store, RefName name, DateTime cacheTime = default, CancellationToken cancellationToken = default)
+		public static async Task<NodeLocator> ReadRefTargetAsync(this IStorageClient store, RefName name, DateTime cacheTime = default, CancellationToken cancellationToken = default)
 		{
-			RefTarget? refTarget = await store.TryReadRefTargetAsync(name, cacheTime, cancellationToken);
-			if (refTarget == null)
+			NodeLocator refTarget = await store.TryReadRefTargetAsync(name, cacheTime, cancellationToken);
+			if (!refTarget.IsValid())
 			{
 				throw new RefNameNotFoundException(name);
 			}
@@ -422,9 +414,11 @@
 		/// <param name="maxAge">Maximum age for any cached result</param>
 		/// <param name="cancellationToken">Cancellation token for the operation</param>
 		/// <returns>The blob instance</returns>
-		public static Task<RefTarget> ReadRefTargetAsync(this IStorageClient store, RefName name, TimeSpan maxAge, CancellationToken cancellationToken = default)
+		public static Task<NodeLocator> ReadRefTargetAsync(this IStorageClient store, RefName name, TimeSpan maxAge, CancellationToken cancellationToken = default)
 		{
 			return ReadRefTargetAsync(store, name, DateTime.UtcNow - maxAge, cancellationToken);
 		}
+
+		#endregion
 	}
 }
diff --git a/Engine/Source/Programs/Shared/EpicGames.Horde/Storage/ITreeStore.cs b/Engine/Source/Programs/Shared/EpicGames.Horde/Storage/ITreeStore.cs
deleted file mode 100644
diff --git a/Engine/Source/Programs/Shared/EpicGames.Horde/Storage/Nodes/DirectoryNode.cs b/Engine/Source/Programs/Shared/EpicGames.Horde/Storage/Nodes/DirectoryNode.cs
--- a/Engine/Source/Programs/Shared/EpicGames.Horde/Storage/Nodes/DirectoryNode.cs
+++ b/Engine/Source/Programs/Shared/EpicGames.Horde/Storage/Nodes/DirectoryNode.cs
@@ -2,7 +2,6 @@
 
 using EpicGames.Core;
 using EpicGames.Horde.Storage.Git;
-using EpicGames.Serialization;
 using Microsoft.Extensions.Logging;
 using System;
 using System.Buffers;
@@ -11,7 +10,6 @@ using System.Diagnostics;
 using System.Diagnostics.CodeAnalysis;
 using System.IO;
 using System.Linq;
-using System.Security.Cryptography;
 using System.Threading;
 using System.Threading.Tasks;
 
@@ -52,7 +50,7 @@ namespace EpicGames.Horde.Storage.Nodes
 	/// <summary>
 	/// Entry for a file within a directory node
 	/// </summary>
-	public class FileEntry : TreeNodeRef<FileNode>
+	public sealed class FileEntry : TreeNodeRef<FileNode>
 	{
 		/// <summary>
 		/// Name of this file
@@ -67,7 +65,7 @@ namespace EpicGames.Horde.Storage.Nodes
 		/// <summary>
 		/// Length of this entry
 		/// </summary>
-		public long Length => (Node == null) ? _cachedLength : Node.Length;
+		public long Length => (Target == null) ? _cachedLength : Target.Length;
 
 		/// <summary>
 		/// SHA1 hash of this file, with Git prefix
@@ -82,23 +80,37 @@ namespace EpicGames.Horde.Storage.Nodes
 		/// <summary>
 		/// Constructor
 		/// </summary>
-		public FileEntry(DirectoryNode owner, Utf8String name, FileEntryFlags flags, FileNode node)
-			: base(owner, node)
+		public FileEntry(Utf8String name, FileEntryFlags flags, FileNode node)
+			: base(node)
 		{
 			Name = name;
 			Flags = flags;
 		}
 
 		/// <summary>
-		/// Constructor
+		/// Deserialize from a buffer
 		/// </summary>
-		public FileEntry(DirectoryNode owner, Utf8String name, FileEntryFlags flags, long length, ITreeBlobRef target)
-			: base(owner, target)
+		/// <param name="reader"></param>
+		public FileEntry(ITreeNodeReader reader)
+			: base(reader)
 		{
-			Name = name;
-			Flags = flags;
-			
-			_cachedLength = length;
+			Name = reader.ReadUtf8String();
+			Flags = (FileEntryFlags)reader.ReadUnsignedVarInt();
+
+			_cachedLength = (long)reader.ReadUnsignedVarInt();
+		}
+
+		/// <summary>
+		/// Serialize this entry
+		/// </summary>
+		/// <param name="writer"></param>
+		public new void Serialize(ITreeNodeWriter writer)
+		{
+			base.Serialize(writer);
+
+			writer.WriteUtf8String(Name);
+			writer.WriteUnsignedVarInt((ulong)Flags);
+			writer.WriteUnsignedVarInt((ulong)Length);
 		}
 
 		/// <summary>
@@ -108,11 +120,10 @@ namespace EpicGames.Horde.Storage.Nodes
 		/// <param name="options">Options for chunking the data</param>
 		/// <param name="writer">Writer for new node data</param>
 		/// <param name="cancellationToken">Cancellation token for the operation</param>
-		public async Task AppendAsync(ReadOnlyMemory<byte> data, ChunkingOptions options, ITreeWriter writer, CancellationToken cancellationToken)
+		public async Task AppendAsync(ReadOnlyMemory<byte> data, ChunkingOptions options, TreeWriter writer, CancellationToken cancellationToken)
 		{
 			FileNode node = await ExpandAsync(cancellationToken);
-			Node = await node.AppendAsync(data, options, writer, cancellationToken);
-			MarkAsDirty();
+			Target = await node.AppendAsync(data, options, writer, cancellationToken);
 		}
 
 		/// <summary>
@@ -122,7 +133,7 @@ namespace EpicGames.Horde.Storage.Nodes
 		/// <param name="options">Options for chunking the data</param>
 		/// <param name="writer">Writer for new node data</param>
 		/// <param name="cancellationToken">Cancellation token for the operation</param>
-		public async Task AppendAsync(Stream stream, ChunkingOptions options, ITreeWriter writer, CancellationToken cancellationToken)
+		public async Task AppendAsync(Stream stream, ChunkingOptions options, TreeWriter writer, CancellationToken cancellationToken)
 		{
 			const int BufferLength = 32 * 1024;
 
@@ -157,13 +168,7 @@ namespace EpicGames.Horde.Storage.Nodes
 		}
 
 		/// <inheritdoc/>
-		protected override void OnCollapse()
-		{
-			_cachedLength = Node!.Length;
-		}
-
-		/// <inheritdoc/>
-		public override string ToString() => IsDirty() ? $"{Name} (*)" : Name.ToString();
+		public override string ToString() => Name.ToString();
 	}
 
 	/// <summary>
@@ -179,50 +184,47 @@ namespace EpicGames.Horde.Storage.Nodes
 		/// <summary>
 		/// Length of this directory tree
 		/// </summary>
-		public long Length => (Node == null) ? _cachedLength : Node.Length;
-
-		/// <summary>
-		/// Optional Git hash for the referenced directory object.
-		/// </summary>
-		public Sha1Hash GitHash { get; set; }
+		public long Length => (Target == null) ? _cachedLength : Target.Length;
 
 		/// <summary>
 		/// Cached value for the length of this tree
 		/// </summary>
-		long _cachedLength;
+		readonly long _cachedLength;
 
 		/// <summary>
 		/// Constructor
 		/// </summary>
-		public DirectoryEntry(DirectoryNode owner, Utf8String name, DirectoryNode node)
-			: base(owner, node)
+		public DirectoryEntry(Utf8String name, DirectoryNode node)
+			: base(node)
 		{
 			Name = name;
 		}
 
 		/// <summary>
-		/// Constructor
+		/// 
 		/// </summary>
-		public DirectoryEntry(DirectoryNode owner, Utf8String name, long length, ITreeBlobRef target)
-			: base(owner, target)
+		/// <param name="reader"></param>
+		public DirectoryEntry(ITreeNodeReader reader)
+			: base(reader)
 		{
-			Name = name;
-
-			_cachedLength = length;
+			Name = reader.ReadUtf8String();
+			_cachedLength = (long)reader.ReadUnsignedVarInt();
 		}
 
-		/// <inheritdoc/>
-		protected override void OnCollapse()
+		/// <summary>
+		/// Serialize this directory entry to disk
+		/// </summary>
+		/// <param name="writer"></param>
+		public new void Serialize(ITreeNodeWriter writer)
 		{
-			_cachedLength = Node!.Length;
-			if (_owner is DirectoryNode directoryNode && (directoryNode.Flags & DirectoryFlags.WithGitHashes) != 0)
-			{
-				GitHash = Node!.AsGitTree().GetHash();
-			}
+			base.Serialize(writer);
+
+			writer.WriteUtf8String(Name);
+			writer.WriteUnsignedVarInt((ulong)Length);
 		}
 
 		/// <inheritdoc/>
-		public override string ToString() => IsDirty() ? $"{Name} (*)" : Name.ToString();
+		public override string ToString() => Name.ToString();
 	}
 
 	/// <summary>
@@ -234,17 +236,11 @@ namespace EpicGames.Horde.Storage.Nodes
 		/// No flags specified
 		/// </summary>
 		None = 0,
-
-		/// <summary>
-		/// Includes hashes for all entries
-		/// </summary>
-		WithGitHashes = 1,
 	}
 
 	/// <summary>
 	/// A directory node
 	/// </summary>
-	[TreeSerializer(typeof(DirectoryNodeSerializer))]
 	public class DirectoryNode : TreeNode
 	{
 		internal const byte TypeId = (byte)'d';
@@ -276,11 +272,63 @@ namespace EpicGames.Horde.Storage.Nodes
 		/// Constructor
 		/// </summary>
 		/// <param name="flags"></param>
-		public DirectoryNode(DirectoryFlags flags)
+		public DirectoryNode(DirectoryFlags flags = DirectoryFlags.None)
 		{
 			Flags = flags;
 		}
 
+		/// <summary>
+		/// Deserialization constructor
+		/// </summary>
+		/// <param name="reader">Reader to deserialize from</param>
+		public DirectoryNode(ITreeNodeReader reader)
+		{
+			byte typeId = reader.ReadUInt8();
+			if (typeId != TypeId)
+			{
+				throw new InvalidOperationException("Invalid signature byte for directory");
+			}
+
+			Flags = (DirectoryFlags)reader.ReadUnsignedVarInt();
+
+			int fileCount = (int)reader.ReadUnsignedVarInt();
+			_nameToFileEntry.EnsureCapacity(fileCount);
+
+			for (int idx = 0; idx < fileCount; idx++)
+			{
+				FileEntry entry = new FileEntry(reader);
+				_nameToFileEntry[entry.Name] = entry;
+			}
+
+			int directoryCount = (int)reader.ReadUnsignedVarInt();
+			_nameToDirectoryEntry.EnsureCapacity(directoryCount);
+			
+			for (int idx = 0; idx < directoryCount; idx++)
+			{
+				DirectoryEntry entry = new DirectoryEntry(reader);
+				_nameToDirectoryEntry[entry.Name] = entry;
+			}
+		}
+
+		/// <inheritdoc/>
+		public override void Serialize(ITreeNodeWriter writer)
+		{
+			writer.WriteUInt8(TypeId);
+			writer.WriteUnsignedVarInt((ulong)Flags);
+
+			writer.WriteUnsignedVarInt(Files.Count);
+			foreach (FileEntry fileEntry in Files)
+			{
+				fileEntry.Serialize(writer);
+			}
+
+			writer.WriteUnsignedVarInt(Directories.Count);
+			foreach (DirectoryEntry directoryEntry in Directories)
+			{
+				directoryEntry.Serialize(writer);
+			}
+		}
+
 		/// <summary>
 		/// Clear the contents of this directory
 		/// </summary>
@@ -298,16 +346,7 @@ namespace EpicGames.Horde.Storage.Nodes
 		/// <returns>True if the entry exists</returns>
 		public bool Contains(Utf8String name) => TryGetFileEntry(name, out _) || TryGetDirectoryEntry(name, out _);
 
-		/// <inheritdoc/>
-		public override IReadOnlyList<TreeNodeRef> GetReferences()
-		{
-			List<TreeNodeRef> refs = new List<TreeNodeRef>();
-			refs.AddRange(Files);
-			refs.AddRange(Directories);
-			return refs;
-		}
-
-		#region File operations
+#region File operations
 
 		/// <summary>
 		/// Adds a new directory with the given name
@@ -324,7 +363,7 @@ namespace EpicGames.Horde.Storage.Nodes
 
 			FileNode newNode = new LeafFileNode();
 
-			FileEntry entry = new FileEntry(this, name, flags, newNode);
+			FileEntry entry = new FileEntry(name, flags, newNode);
 			_nameToFileEntry[name] = entry;
 			MarkAsDirty();
 
@@ -341,7 +380,7 @@ namespace EpicGames.Horde.Storage.Nodes
 		/// <param name="writer">Writer for new node data</param>
 		/// <param name="cancellationToken">Cancellation token for the operation</param>
 		/// <returns>The new directory object</returns>
-		public async Task<FileEntry> AddFileAsync(Utf8String name, FileEntryFlags flags, Stream stream, ChunkingOptions options, ITreeWriter writer, CancellationToken cancellationToken = default)
+		public async Task<FileEntry> AddFileAsync(Utf8String name, FileEntryFlags flags, Stream stream, ChunkingOptions options, TreeWriter writer, CancellationToken cancellationToken = default)
 		{
 			FileEntry entry = AddFile(name, flags);
 			await entry.AppendAsync(stream, options, writer, cancellationToken);
@@ -445,9 +484,9 @@ namespace EpicGames.Horde.Storage.Nodes
 			return false;
 		}
 
-		#endregion
+#endregion
 
-		#region Directory operations
+#region Directory operations
 
 		/// <summary>
 		/// Adds a new directory with the given name
@@ -463,7 +502,7 @@ namespace EpicGames.Horde.Storage.Nodes
 
 			DirectoryNode node = new DirectoryNode(Flags);
 
-			DirectoryEntry entry = new DirectoryEntry(this, name, node);
+			DirectoryEntry entry = new DirectoryEntry(name, node);
 			_nameToDirectoryEntry.Add(name, entry);
 			MarkAsDirty();
 
@@ -534,180 +573,7 @@ namespace EpicGames.Horde.Storage.Nodes
 			return false;
 		}
 
-		#endregion
-
-		/// <inheritdoc/>
-		public override async Task<ITreeBlob> SerializeAsync(ITreeWriter writer, CancellationToken cancellationToken)
-		{
-			List<ITreeBlobRef> refs = new List<ITreeBlobRef>();
-
-			List<FileEntry> fileEntries = _nameToFileEntry.Values.OrderBy(x => x.Name).ToList();
-			foreach (FileEntry fileEntry in fileEntries)
-			{
-				refs.Add(await fileEntry.CollapseAsync(writer, cancellationToken));
-			}
-
-			List<DirectoryEntry> directoryEntries = _nameToDirectoryEntry.Values.OrderBy(x => x.Name).ToList();
-			foreach (DirectoryEntry directoryEntry in directoryEntries)
-			{
-				refs.Add(await directoryEntry.CollapseAsync(writer, cancellationToken));
-			}
-
-			ReadOnlySequence<byte> data = SerializeData(Flags, fileEntries, directoryEntries);
-			return new NewTreeBlob(data, refs);
-		}
-
-		static ReadOnlySequence<byte> SerializeData(DirectoryFlags flags, List<FileEntry> fileEntries, List<DirectoryEntry> directoryEntries)
-		{
-			// Measure the required size of the write buffer
-			int size = 1;
-			size += VarInt.MeasureUnsigned((ulong)flags);
-			size += VarInt.MeasureUnsigned(fileEntries.Count);
-			foreach (FileEntry fileEntry in fileEntries)
-			{
-				size += (fileEntry.Name.Length + 1) + 1 + VarInt.MeasureUnsigned((ulong)fileEntry.Length);
-			}
-			size += VarInt.MeasureUnsigned(directoryEntries.Count);
-			foreach (DirectoryEntry directoryEntry in directoryEntries)
-			{
-				size += (directoryEntry.Name.Length + 1) + VarInt.MeasureUnsigned((ulong)directoryEntry.Length);
-			}
-			if ((flags & DirectoryFlags.WithGitHashes) != 0)
-			{
-				size += (directoryEntries.Count + fileEntries.Count) * IoHash.NumBytes;
-			}
-
-			// Allocate the buffer and copy the node to it
-			byte[] data = new byte[size];
-			Span<byte> span = data.AsSpan();
-
-			span[0] = TypeId;
-			span = span.Slice(1);
-
-			int directoryFlagBytes = VarInt.WriteUnsigned(span, (ulong)flags);
-			span = span.Slice(directoryFlagBytes);
-
-			int fileCountBytes = VarInt.WriteUnsigned(span, fileEntries.Count);
-			span = span.Slice(fileCountBytes);
-
-			foreach (FileEntry fileEntry in fileEntries)
-			{
-				fileEntry.Name.Span.CopyTo(span);
-				span = span.Slice(fileEntry.Name.Length + 1);
-
-				span[0] = (byte)fileEntry.Flags;
-				span = span.Slice(1);
-
-				int lengthBytes = VarInt.WriteUnsigned(span, fileEntry.Length);
-				span = span.Slice(lengthBytes);
-
-				if ((flags & DirectoryFlags.WithGitHashes) != 0)
-				{
-					fileEntry.GitHash.CopyTo(span);
-					span = span.Slice(Sha1Hash.NumBytes);
-				}
-			}
-
-			int directoryCountBytes = VarInt.WriteUnsigned(span, directoryEntries.Count);
-			span = span.Slice(directoryCountBytes);
-
-			foreach (DirectoryEntry directoryEntry in directoryEntries)
-			{
-				directoryEntry.Name.Span.CopyTo(span);
-				span = span.Slice(directoryEntry.Name.Length + 1);
-
-				int lengthBytes = VarInt.WriteUnsigned(span, directoryEntry.Length);
-				span = span.Slice(lengthBytes);
-
-				if ((flags & DirectoryFlags.WithGitHashes) != 0)
-				{
-					directoryEntry.GitHash.CopyTo(span);
-					span = span.Slice(Sha1Hash.NumBytes);
-				}
-			}
-
-			Debug.Assert(span.Length == 0);
-			return new ReadOnlySequence<byte>(data);
-		}
-
-		/// <summary>
-		/// Deserialize a directory node from data
-		/// </summary>
-		/// <param name="source">The source node to deserialize</param>
-		/// <returns>The deserialized directory node</returns>
-		public static DirectoryNode Deserialize(ITreeBlob source)
-		{
-			ReadOnlySpan<byte> span = source.Data.AsSingleSegment().Span;
-			if (span[0] != TypeId)
-			{
-				throw new InvalidOperationException("Invalid signature byte for directory");
-			}
-
-			span = span.Slice(1);
-
-			DirectoryFlags directoryFlags = (DirectoryFlags)VarInt.ReadUnsigned(span, out int directoryFlagBytes);
-			span = span.Slice(directoryFlagBytes);
-
-			DirectoryNode node = new DirectoryNode(directoryFlags);
-
-			int fileCount = (int)VarInt.ReadUnsigned(span, out int fileCountBytes);
-			span = span.Slice(fileCountBytes);
-
-			int childIdx = 0;
-
-			node._nameToFileEntry.EnsureCapacity(fileCount);
-			for (int idx = 0; idx < fileCount; idx++)
-			{
-				int nameLen = span.IndexOf((byte)0);
-				Utf8String name = new Utf8String(span.Slice(0, nameLen).ToArray());
-				span = span.Slice(nameLen + 1);
-
-				FileEntryFlags flags = (FileEntryFlags)VarInt.ReadUnsigned(span, out int flagBytes);
-				span = span.Slice(flagBytes);
-
-				long length = (long)VarInt.ReadUnsigned(span, out int lengthBytes);
-				span = span.Slice(lengthBytes);
-
-				ITreeBlobRef child = source.Refs[childIdx++];
-				FileEntry entry = new FileEntry(node, name, flags, length, child);
-
-				if ((directoryFlags & DirectoryFlags.WithGitHashes) != 0)
-				{
-					entry.GitHash = new Sha1Hash(span);
-					span = span.Slice(Sha1Hash.NumBytes);
-				}
-
-				node._nameToFileEntry[name] = entry;
-			}
-
-			int directoryCount = (int)VarInt.ReadUnsigned(span, out int directoryCountBytes);
-			span = span.Slice(directoryCountBytes);
-
-			node._nameToDirectoryEntry.EnsureCapacity(directoryCount);
-			for (int idx = 0; idx < directoryCount; idx++)
-			{
-				int nameLen = span.IndexOf((byte)0);
-				Utf8String name = new Utf8String(span.Slice(0, nameLen).ToArray());
-				span = span.Slice(nameLen + 1);
-
-				long length = (long)VarInt.ReadUnsigned(span, out int lengthBytes);
-				span = span.Slice(lengthBytes);
-
-				ITreeBlobRef child = source.Refs[childIdx++];
-				DirectoryEntry entry = new DirectoryEntry(node, name, length, child);
-
-				if ((directoryFlags & DirectoryFlags.WithGitHashes) != 0)
-				{
-					entry.GitHash = new Sha1Hash(span);
-					span = span.Slice(Sha1Hash.NumBytes);
-				}
-
-				node._nameToDirectoryEntry[name] = entry;
-			}
-
-			Debug.Assert(span.Length == 0);
-			return node;
-		}
+#endregion
 
 		/// <summary>
 		/// Adds files from a directory on disk
@@ -718,7 +584,7 @@ namespace EpicGames.Horde.Storage.Nodes
 		/// <param name="logger"></param>
 		/// <param name="cancellationToken">Cancellation token for the operation</param>
 		/// <returns></returns>
-		public async Task CopyFromDirectoryAsync(DirectoryInfo directoryInfo, ChunkingOptions options, ITreeWriter writer, ILogger logger, CancellationToken cancellationToken)
+		public async Task CopyFromDirectoryAsync(DirectoryInfo directoryInfo, ChunkingOptions options, TreeWriter writer, ILogger logger, CancellationToken cancellationToken)
 		{
 			foreach (DirectoryInfo subDirectoryInfo in directoryInfo.EnumerateDirectories())
 			{
@@ -763,25 +629,6 @@ namespace EpicGames.Horde.Storage.Nodes
 
 			await Task.WhenAll(tasks);
 		}
-
-		/// <summary>
-		/// Creates a Git tree object for this directory
-		/// </summary>
-		/// <returns>Tree object</returns>
-		public GitTree AsGitTree()
-		{
-			GitTree tree = new GitTree();
-			foreach (FileEntry fileEntry in _nameToFileEntry.Values)
-			{
-				tree.Entries.Add(new GitTreeEntry(GitFileMode.File, fileEntry.Name, fileEntry.GitHash));
-			}
-			foreach (DirectoryEntry directoryEntry in _nameToDirectoryEntry.Values)
-			{
-				tree.Entries.Add(new GitTreeEntry(GitFileMode.Tree, directoryEntry.Name, directoryEntry.GitHash));
-			}
-			tree.Entries.SortBy(x => x.Name);
-			return tree;
-		}
 	}
 
 	/// <summary>
@@ -790,7 +637,7 @@ namespace EpicGames.Horde.Storage.Nodes
 	public class DirectoryNodeSerializer : TreeNodeSerializer<DirectoryNode>
 	{
 		/// <inheritdoc/>
-		public override DirectoryNode Deserialize(ITreeBlob blob) => DirectoryNode.Deserialize(blob);
+		public override DirectoryNode Deserialize(ITreeNodeReader reader) => new DirectoryNode(reader);
 	}
 
 	/// <summary>
diff --git a/Engine/Source/Programs/Shared/EpicGames.Horde/Storage/Nodes/FileNode.cs b/Engine/Source/Programs/Shared/EpicGames.Horde/Storage/Nodes/FileNode.cs
--- a/Engine/Source/Programs/Shared/EpicGames.Horde/Storage/Nodes/FileNode.cs
+++ b/Engine/Source/Programs/Shared/EpicGames.Horde/Storage/Nodes/FileNode.cs
@@ -8,6 +8,7 @@ using System.Buffers.Binary;
 using System.Collections.Generic;
 using System.Diagnostics;
 using System.IO;
+using System.Linq;
 using System.Threading;
 using System.Threading.Tasks;
 
@@ -87,6 +88,12 @@ namespace EpicGames.Horde.Storage.Nodes
 		/// </summary>
 		public abstract uint RollingHash { get; }
 
+		/// <summary>
+		/// Whether the node can have data appended to it
+		/// </summary>
+		/// <returns>True if the node is read only</returns>
+		public abstract bool IsReadOnly();
+
 		/// <summary>
 		/// Creates a file node from a stream
 		/// </summary>
@@ -95,7 +102,7 @@ namespace EpicGames.Horde.Storage.Nodes
 		/// <param name="writer">Writer for new tree nodes</param>
 		/// <param name="cancellationToken">Cancellation token for the operation</param>
 		/// <returns>New file node</returns>
-		public static async Task<FileNode> CreateAsync(Stream stream, ChunkingOptions options, ITreeWriter writer, CancellationToken cancellationToken)
+		public static async Task<FileNode> CreateAsync(Stream stream, ChunkingOptions options, TreeWriter writer, CancellationToken cancellationToken)
 		{
 			FileNode node = new LeafFileNode();
 
@@ -120,12 +127,12 @@ namespace EpicGames.Horde.Storage.Nodes
 		/// <param name="options">Settings for chunking the data</param>
 		/// <param name="writer">Writer for new tree nodes</param>
 		/// <param name="cancellationToken">Cancellation token for the operation</param>
-		public ValueTask<FileNode> AppendAsync(ReadOnlyMemory<byte> input, ChunkingOptions options, ITreeWriter writer, CancellationToken cancellationToken)
+		public ValueTask<FileNode> AppendAsync(ReadOnlyMemory<byte> input, ChunkingOptions options, TreeWriter writer, CancellationToken cancellationToken)
 		{
 			return AppendAsync(this, input, options, writer, cancellationToken);
 		}
 
-		static async ValueTask<FileNode> AppendAsync(FileNode root, ReadOnlyMemory<byte> input, ChunkingOptions options, ITreeWriter writer, CancellationToken cancellationToken)
+		static async ValueTask<FileNode> AppendAsync(FileNode root, ReadOnlyMemory<byte> input, ChunkingOptions options, TreeWriter writer, CancellationToken cancellationToken)
 		{
 			for (; ; )
 			{
@@ -142,7 +149,7 @@ namespace EpicGames.Horde.Storage.Nodes
 			return root;
 		}
 
-		private async Task<ReadOnlyMemory<byte>> AppendToNodeAsync(ReadOnlyMemory<byte> appendData, ChunkingOptions options, ITreeWriter writer, CancellationToken cancellationToken)
+		private async Task<ReadOnlyMemory<byte>> AppendToNodeAsync(ReadOnlyMemory<byte> appendData, ChunkingOptions options, TreeWriter writer, CancellationToken cancellationToken)
 		{
 			if (appendData.Length == 0 || IsReadOnly())
 			{
@@ -162,7 +169,7 @@ namespace EpicGames.Horde.Storage.Nodes
 		/// <param name="writer">Writer for new tree nodes</param>
 		/// <param name="cancellationToken">Cancellation token for the operation</param>
 		/// <returns>Remaining data in the buffer</returns>
-		public abstract ValueTask<ReadOnlyMemory<byte>> AppendDataAsync(ReadOnlyMemory<byte> newData, ChunkingOptions options, ITreeWriter writer, CancellationToken cancellationToken);
+		public abstract ValueTask<ReadOnlyMemory<byte>> AppendDataAsync(ReadOnlyMemory<byte> newData, ChunkingOptions options, TreeWriter writer, CancellationToken cancellationToken);
 
 		/// <summary>
 		/// Copies the contents of this node and its children to the given output stream
@@ -224,12 +231,9 @@ namespace EpicGames.Horde.Storage.Nodes
 		/// </summary>
 		public const byte TypeId = (byte)'l';
 
-		const int HeaderLength = 1 + sizeof(uint);
-
 		bool _isReadOnly;
 		uint _rollingHash;
-		readonly byte[] _header;
-		readonly DataSegment _firstSegment;
+		DataSegment _firstSegment;
 		DataSegment _lastSegment;
 		ReadOnlySequence<byte> _writtenSequence; // Payload described by _firstSegment -> _lastSegment
 
@@ -238,41 +242,38 @@ namespace EpicGames.Horde.Storage.Nodes
 		/// </summary>
 		public LeafFileNode()
 		{
-			_header = new byte[HeaderLength];
-			_header[0] = TypeId;
-
-			_firstSegment = new DataSegment(0, _header);
+			_firstSegment = new DataSegment(0, ReadOnlyMemory<byte>.Empty);
 			_lastSegment = _firstSegment;
 
-			_writtenSequence = CreateSequence(_firstSegment, _lastSegment);
+			_writtenSequence = ReadOnlySequence<byte>.Empty;
 		}
 
 		/// <summary>
 		/// Create a leaf node from the given serialized data
 		/// </summary>
-		/// <param name="data">Data to construct from</param>
-		public LeafFileNode(ReadOnlySequence<byte> data)
+		public LeafFileNode(ITreeNodeReader reader)
 		{
-			_header = Array.Empty<byte>();
-
-			ReadOnlySpan<byte> span = data.Slice(0, HeaderLength).AsSingleSegment().Span;
-			if (span[0] != TypeId)
+			byte typeId = reader.ReadUInt8();
+			if (typeId != TypeId)
 			{
 				throw new InvalidDataException($"Invalid type id for {nameof(LeafFileNode)}");
 			}
 
 			_isReadOnly = true;
-			_rollingHash = BinaryPrimitives.ReadUInt32LittleEndian(span.Slice(1, sizeof(uint)));
+			_rollingHash = reader.ReadUInt32();
 
-			_firstSegment = null!;
-			_lastSegment = null!;
+			_firstSegment = new DataSegment(0, reader.ReadVariableLengthBytes());
+			_lastSegment = _firstSegment;
 
-			_writtenSequence = data;
+			_writtenSequence = new ReadOnlySequence<byte>(_firstSegment.Memory);
 		}
 
-		static ReadOnlySequence<byte> CreateSequence(DataSegment firstSegment, DataSegment lastSegment)
+		/// <inheritdoc/>
+		public override void Serialize(ITreeNodeWriter writer)
 		{
-			return new ReadOnlySequence<byte>(firstSegment, 0, lastSegment, lastSegment.Memory.Length);
+			writer.WriteUInt8(TypeId);
+			writer.WriteUInt32(_rollingHash);
+			writer.WriteVariableLengthBytes(_writtenSequence);
 		}
 
 		/// <inheritdoc/>
@@ -284,19 +285,13 @@ namespace EpicGames.Horde.Storage.Nodes
 		/// <summary>
 		/// Gets the data for this node
 		/// </summary>
-		public ReadOnlySequence<byte> Data => _writtenSequence.Slice(HeaderLength);
+		public ReadOnlySequence<byte> Data => _writtenSequence;
 
 		/// <inheritdoc/>
-		public override long Length => _writtenSequence.Length - HeaderLength;
+		public override long Length => _writtenSequence.Length;
 
 		/// <inheritdoc/>
-		public override IReadOnlyList<TreeNodeRef> GetReferences() => Array.Empty<TreeNodeRef>();
-
-		/// <inheritdoc/>
-		public override Task<ITreeBlob> SerializeAsync(ITreeWriter writer, CancellationToken cancellationToken) => Task.FromResult<ITreeBlob>(new NewTreeBlob(_writtenSequence, Array.Empty<ITreeBlobRef>()));
-
-		/// <inheritdoc/>
-		public override ValueTask<ReadOnlyMemory<byte>> AppendDataAsync(ReadOnlyMemory<byte> newData, ChunkingOptions options, ITreeWriter writer, CancellationToken cancellationToken)
+		public override ValueTask<ReadOnlyMemory<byte>> AppendDataAsync(ReadOnlyMemory<byte> newData, ChunkingOptions options, TreeWriter writer, CancellationToken cancellationToken)
 		{
 			return new ValueTask<ReadOnlyMemory<byte>>(AppendData(newData, options));
 		}
@@ -390,19 +385,25 @@ namespace EpicGames.Horde.Storage.Nodes
 		private void AppendLeafData(ReadOnlySpan<byte> leafData, uint newRollingHash)
 		{
 			_rollingHash = newRollingHash;
-			BinaryPrimitives.WriteUInt32LittleEndian(_header.AsSpan(1, sizeof(uint)), newRollingHash);
 
 			DataSegment segment = new DataSegment(_lastSegment.RunningIndex + _lastSegment.Memory.Length, leafData.ToArray());
-			_lastSegment.SetNext(segment);
+			if (_writtenSequence.Length == 0)
+			{
+				_firstSegment = segment;
+			}
+			else
+			{
+				_lastSegment.SetNext(segment);
+			}
 			_lastSegment = segment;
 
-			_writtenSequence = CreateSequence(_firstSegment, _lastSegment);
+			_writtenSequence = new ReadOnlySequence<byte>(_firstSegment, 0, _lastSegment, _lastSegment.Memory.Length);
 		}
 
 		/// <inheritdoc/>
 		public override async Task CopyToStreamAsync(Stream outputStream, CancellationToken cancellationToken)
 		{
-			foreach (ReadOnlyMemory<byte> segment in _writtenSequence.Slice(HeaderLength))
+			foreach (ReadOnlyMemory<byte> segment in _writtenSequence)
 			{
 				await outputStream.WriteAsync(segment, cancellationToken);
 			}
@@ -443,70 +444,42 @@ namespace EpicGames.Horde.Storage.Nodes
 		public InteriorFileNode(long length, FileNode child)
 		{
 			_length = length;
-			_children.Add(new TreeNodeRef<FileNode>(this, child));
+			_children.Add(new TreeNodeRef<FileNode>(child));
 		}
 
 		/// <summary>
 		/// Constructor
 		/// </summary>
-		/// <param name="data"></param>
-		/// <param name="children"></param>
-		public InteriorFileNode(ReadOnlySequence<byte> data, IReadOnlyList<ITreeBlobRef> children)
+		public InteriorFileNode(ITreeNodeReader reader)
 		{
-			ReadOnlySpan<byte> span = data.AsSingleSegment().Span;
-			if (span[0] != TypeId)
+			byte typeId = reader.ReadUInt8();
+			if (typeId != TypeId)
 			{
 				throw new InvalidDataException($"Invalid type id for {nameof(InteriorFileNode)}");
 			}
 
-			int pos = 1;
-
-			_rollingHash = BinaryPrimitives.ReadUInt32LittleEndian(span[pos..]);
-			pos += sizeof(uint);
-
-			_length = (int)VarInt.ReadUnsigned(span[pos..], out int lengthBytes);
-			pos += lengthBytes;
-
-			_isReadOnly = span[pos] != 0;
-			_children = children.ConvertAll(x => new TreeNodeRef<FileNode>(this, x));
+			_rollingHash = reader.ReadUInt32();
+			_length = (long)reader.ReadUnsignedVarInt();
+			_isReadOnly = reader.ReadBoolean();
+			_children = reader.ReadList(() => reader.ReadRef<FileNode>());
 		}
 
 		/// <inheritdoc/>
 		public override bool IsReadOnly() => _isReadOnly;
 
 		/// <inheritdoc/>
-		public override IReadOnlyList<TreeNodeRef> GetReferences() => _children;
-
-		/// <inheritdoc/>
-		public override async Task<ITreeBlob> SerializeAsync(ITreeWriter writer, CancellationToken cancellationToken)
+		public override void Serialize(ITreeNodeWriter writer)
 		{
-			byte[] data = new byte[1 + sizeof(uint) + VarInt.MeasureUnsigned((ulong)_length) + 1];
-			data[0] = TypeId;
-
-			int pos = 1;
-
-			BinaryPrimitives.WriteUInt32LittleEndian(data.AsSpan(pos), _rollingHash);
-			pos += sizeof(uint);
-
-			int lengthBytes = VarInt.WriteUnsigned(data.AsSpan(pos), _length);
-			pos += lengthBytes;
-
-			data[pos] = _isReadOnly ? (byte)1 : (byte)0;
-			pos++;
-
-			Debug.Assert(pos == data.Length);
-
-			List<ITreeBlobRef> childRefs = new List<ITreeBlobRef>();
-			foreach (TreeNodeRef typedRef in _children)
-			{
-				childRefs.Add(await typedRef.CollapseAsync(writer, cancellationToken));
-			}
+			writer.WriteUInt8(TypeId);
 
-			return new NewTreeBlob(data, childRefs);
+			writer.WriteUInt32(_rollingHash);
+			writer.WriteUnsignedVarInt((ulong)_length);
+			writer.WriteBoolean(_isReadOnly);
+			writer.WriteList(_children, x => writer.WriteRef(x));
 		}
 
 		/// <inheritdoc/>
-		public override async ValueTask<ReadOnlyMemory<byte>> AppendDataAsync(ReadOnlyMemory<byte> newData, ChunkingOptions options, ITreeWriter writer, CancellationToken cancellationToken)
+		public override async ValueTask<ReadOnlyMemory<byte>> AppendDataAsync(ReadOnlyMemory<byte> newData, ChunkingOptions options, TreeWriter writer, CancellationToken cancellationToken)
 		{
 			for (; ; )
 			{
@@ -515,7 +488,7 @@ namespace EpicGames.Horde.Storage.Nodes
 				// Try to write to the last node
 				if (_children.Count > 0)
 				{
-					FileNode? lastNode = _children[^1].Node;
+					FileNode? lastNode = _children[^1].Target;
 					if (lastNode != null)
 					{
 						// Update the length to match the new node
@@ -545,12 +518,12 @@ namespace EpicGames.Horde.Storage.Nodes
 						}
 
 						// Collapse the final node
-						await Children[^1].CollapseAsync(writer, cancellationToken);
+						await writer.WriteAsync(Children[^1].Target!, cancellationToken);
 					}
 				}
 
 				// Add a new child node
-				_children.Add(new TreeNodeRef<FileNode>(this, new LeafFileNode()));
+				_children.Add(new TreeNodeRef<FileNode>(new LeafFileNode()));
 			}
 		}
 
@@ -582,15 +555,15 @@ namespace EpicGames.Horde.Storage.Nodes
 	public class FileNodeSerializer : TreeNodeSerializer<FileNode>
 	{
 		/// <inheritdoc/>
-		public override FileNode Deserialize(ITreeBlob blob)
+		public override FileNode Deserialize(ITreeNodeReader reader)
 		{
-			ReadOnlySequence<byte> data = blob.Data;
-			switch (data.FirstSpan[0])
+			ReadOnlySpan<byte> data = reader.GetSpan(1);
+			switch (data[0])
 			{
 				case LeafFileNode.TypeId:
-					return new LeafFileNode(data);
+					return new LeafFileNode(reader);
 				case InteriorFileNode.TypeId:
-					return new InteriorFileNode(data, blob.Refs);
+					return new InteriorFileNode(reader);
 				default:
 					throw new InvalidDataException("Unknown type id while deserializing file node");
 			};
diff --git a/Engine/Source/Programs/Shared/EpicGames.Horde/Storage/StorageClientBase.cs b/Engine/Source/Programs/Shared/EpicGames.Horde/Storage/StorageClientBase.cs
--- a/Engine/Source/Programs/Shared/EpicGames.Horde/Storage/StorageClientBase.cs
+++ b/Engine/Source/Programs/Shared/EpicGames.Horde/Storage/StorageClientBase.cs
@@ -1,11 +1,10 @@
 // Copyright Epic Games, Inc. All Rights Reserved.
 
 using System;
+using System.Buffers;
 using System.Collections.Generic;
 using System.IO;
 using System.Linq;
-using System.Runtime.CompilerServices;
-using System.Text;
 using System.Threading;
 using System.Threading.Tasks;
 using EpicGames.Core;
@@ -94,22 +93,76 @@ namespace EpicGames.Horde.Storage
 			public override void Write(byte[] buffer, int offset, int count) => throw new NotSupportedException();
 		}
 
-		readonly IMemoryCache? _cache;
-		readonly TreeStore _treeStore;
-		readonly object _lockObject = new object();
+		/// <summary>
+		/// Describes a bundle export
+		/// </summary>
+		/// <param name="PacketIdx">Index of the packet within a bundle</param>
+		/// <param name="Offset">Offset within the packet</param>
+		record struct ExportInfo(int PacketIdx, int Offset);
+
+		/// <summary>
+		/// Computed information about a bundle
+		/// </summary>
+		class BundleInfo
+		{
+			public readonly BlobLocator Locator;
+			public readonly BundleHeader Header;
+			public readonly int[] PacketOffsets;
+			public readonly ExportInfo[] Exports;
+			public readonly (IoHash, NodeLocator)[] References;
+
+			public BundleInfo(BlobLocator locator, BundleHeader header, int headerLength)
+			{
+				Locator = locator;
+				Header = header;
+
+				PacketOffsets = new int[header.Packets.Count];
+				Exports = new ExportInfo[header.Exports.Count];
+
+				int exportIdx = 0;
+				int packetOffset = headerLength;
+				for (int packetIdx = 0; packetIdx < header.Packets.Count; packetIdx++)
+				{
+					BundlePacket packet = header.Packets[packetIdx];
+					PacketOffsets[packetIdx] = packetOffset;
+
+					int nodeOffset = 0;
+					for (; exportIdx < header.Exports.Count && nodeOffset + header.Exports[exportIdx].Length <= packet.DecodedLength; exportIdx++)
+					{
+						Exports[exportIdx] = new ExportInfo(packetIdx, nodeOffset);
+						nodeOffset += header.Exports[exportIdx].Length;
+					}
+
+					packetOffset += packet.EncodedLength;
+				}
+
+				References = new (IoHash, NodeLocator)[header.Imports.Sum(x => x.Exports.Count) + header.Exports.Count];
+
+				int referenceIdx = 0;
+				foreach (BundleImport import in header.Imports)
+				{
+					foreach ((int importExportIdx, IoHash importExportHash) in import.Exports)
+					{
+						NodeLocator importLocator = new NodeLocator(import.Locator, importExportIdx);
+						References[referenceIdx++] = (importExportHash, importLocator);
+					}
+				}
+				for(int idx = 0; idx < header.Exports.Count; idx++)
+				{
+					References[referenceIdx++] = (header.Exports[idx].Hash, new NodeLocator(locator, idx));
+				}
+			}
+		}
 
-		// Active read tasks at any moment. If a BundleObject is not available in the cache, we start a read and add an entry to this dictionary
-		// so that other threads can also await it.
-		readonly Dictionary<string, Task<Bundle>> _readTasks = new Dictionary<string, Task<Bundle>>();
+		readonly IMemoryCache _cache;
 
 		/// <summary>
 		/// Constructor
 		/// </summary>
 		/// <param name="cache"></param>
-		protected StorageClientBase(IMemoryCache? cache)
+		protected StorageClientBase(IMemoryCache cache)
 		{
 			_cache = cache;
-			_treeStore = new TreeStore(this);
 		}
 
 		#region Blobs
@@ -125,6 +178,24 @@ namespace EpicGames.Horde.Storage
 			return new RangedReadStream(stream, Math.Min((int)(stream.Length - offset), length));
 		}
 
+		async Task<Memory<byte>> ReadBlobRangeAsync(BlobLocator locator, int offset, Memory<byte> memory, CancellationToken cancellationToken = default)
+		{
+			using (Stream stream = await ReadBlobRangeAsync(locator, offset, memory.Length, cancellationToken))
+			{
+				int length = 0;
+				while (length < memory.Length)
+				{
+					int readBytes = await stream.ReadAsync(memory.Slice(length), cancellationToken);
+					if (readBytes == 0)
+					{
+						break;
+					}
+					length += readBytes;
+				}
+				return memory.Slice(0, length);
+			}
+		}
+
 		/// <inheritdoc/>
 		public abstract Task<BlobLocator> WriteBlobAsync(Stream stream, Utf8String prefix = default, CancellationToken cancellationToken = default);
 
@@ -132,24 +203,16 @@ namespace EpicGames.Horde.Storage
 
 		#region Bundles
 
-		static string GetBundleCacheKey(BlobId blobId) => $"bundle:{blobId}";
+		static string GetBundleInfoCacheKey(BlobId blobId) => $"bundle:{blobId}";
+		static string GetEncodedPacketCacheKey(BlobId blobId, int packetIdx) => $"encoded-packet:{blobId}#{packetIdx}";
+		static string GetDecodedPacketCacheKey(BlobId blobId, int packetIdx) => $"decoded-packet:{blobId}#{packetIdx}";
 
-		/// <summary>
-		/// Adds a bundle to the cache
-		/// </summary>
-		/// <param name="blobId">Blob id for the bundle</param>
-		/// <param name="bundle">The bundle data</param>
-		void AddBundleToCache(BlobId blobId, Bundle bundle)
+		void AddToCache(string cacheKey, object value, int size)
 		{
-			if (_cache != null)
+			using (ICacheEntry entry = _cache.CreateEntry(cacheKey))
 			{
-				string cacheKey = GetBundleCacheKey(blobId);
-				using (ICacheEntry entry = _cache.CreateEntry(cacheKey))
-				{
-					int length = bundle.Packets.Sum(x => x.Length);
-					entry.SetSize(length);
-					entry.SetValue(bundle);
-				}
+				entry.SetValue(value);
+				entry.SetSize(size);
 			}
 		}
 
@@ -161,103 +224,97 @@ namespace EpicGames.Horde.Storage
 		/// <returns></returns>
 		public async Task<Bundle> ReadBundleAsync(BlobLocator locator, CancellationToken cancellationToken = default)
 		{
-			string cacheKey = GetBundleCacheKey(locator.BlobId);
-			if (_cache == null || !_cache.TryGetValue<Bundle>(cacheKey, out Bundle? bundle))
+			using (Stream stream = await ReadBlobAsync(locator, cancellationToken))
 			{
-				Task<Bundle>? readTask;
-				lock (_lockObject)
-				{
-					if (!_readTasks.TryGetValue(cacheKey, out readTask))
-					{
-						readTask = Task.Run(() => ReadBundleInternalAsync(cacheKey, locator, cancellationToken), cancellationToken);
-						_readTasks.Add(cacheKey, readTask);
-					}
-				}
-				bundle = await readTask;
+				return await Bundle.FromStreamAsync(stream, cancellationToken);
 			}
-			return bundle;
 		}
 
-		async Task<Bundle> ReadBundleInternalAsync(string cacheKey, BlobLocator locator, CancellationToken cancellationToken)
+		/// <summary>
+		/// Reads the header and structural metadata about the bundle
+		/// </summary>
+		/// <param name="locator">The bundle location</param>
+		/// <param name="cancellationToken">Cancellation token for the operation</param>
+		/// <returns>Information about the bundle</returns>
+		async ValueTask<BundleInfo> GetBundleInfoAsync(BlobLocator locator, CancellationToken cancellationToken = default)
 		{
-			// Perform another (sequenced) check whether an object has been added to the cache, to counteract the race between a read task being added and a task completing.
-			lock (_lockObject)
+			string cacheKey = GetBundleInfoCacheKey(locator.BlobId);
+			if (!_cache.TryGetValue(cacheKey, out BundleInfo bundleInfo))
 			{
-				if (_cache != null && _cache.TryGetValue<Bundle>(cacheKey, out Bundle? cachedObject))
+				int prefetchSize = 1024 * 1024;
+				for (; ; )
 				{
-					return cachedObject;
+					using (IMemoryOwner<byte> owner = MemoryPool<byte>.Shared.Rent(prefetchSize))
+					{
+						// Read the prefetch size from the blob
+						Memory<byte> memory = owner.Memory.Slice(0, prefetchSize);
+						memory = await ReadBlobRangeAsync(locator, 0, memory, cancellationToken);
+
+						// Make sure it's large enough to hold the header
+						int headerSize = BundleHeader.ReadPrelude(memory);
+						if (headerSize > prefetchSize)
+						{
+							prefetchSize = headerSize;
+							continue;
+						}
+
+						// Parse the header and construct the bundle info from it
+						BundleHeader header = new BundleHeader(new MemoryReader(memory));
+						bundleInfo = new BundleInfo(locator, header, headerSize);
+
+						// Add the info to the cache
+						AddToCache(cacheKey, bundleInfo, headerSize);
+
+						// Also add any encoded packets we prefetched
+						int packetOffset = headerSize;
+						for (int packetIdx = 0; packetIdx < header.Packets.Count && packetOffset + header.Packets[packetIdx].EncodedLength < memory.Length; packetIdx++)
+						{
+							int packetLength = header.Packets[packetIdx].EncodedLength;
+							ReadOnlyMemory<byte> packetData = memory.Slice(packetOffset, packetLength).ToArray();
+							AddToCache(GetEncodedPacketCacheKey(locator.BlobId, packetIdx), packetData, packetData.Length);
+							packetOffset += packetLength;
+						}
+
+						break;
+					}
 				}
 			}
-
-			// Read the data from storage
-			Bundle bundle;
-			using (Stream stream = await ReadBlobAsync(locator, cancellationToken))
-			{
-				bundle = await Bundle.FromStreamAsync(stream, cancellationToken);
-			}
-
-			// Add the object to the cache
-			AddBundleToCache(locator.BlobId, bundle);
-
-			// Remove this object from the list of read tasks
-			lock (_lockObject)
-			{
-				_readTasks.Remove(cacheKey);
-			}
-			return bundle;
-		}
-
-		/// <inheritdoc/>
-		public Task<Stream> CreateBundleStreamAsync(BlobLocator locator, CancellationToken cancellationToken = default)
-		{
-			return ReadBlobAsync(locator, cancellationToken);
-		}
-
-		/// <inheritdoc/>
-		public async Task<BundleHeader> ReadBundleHeaderAsync(BlobLocator locator, CancellationToken cancellationToken = default)
-		{
-			Bundle bundle = await ReadBundleAsync(locator, cancellationToken);
-			return bundle.Header;
-		}
-
-		/// <inheritdoc/>
-		public async Task<ReadOnlyMemory<byte>> ReadBundlePacketAsync(BlobLocator locator, int packetIdx, CancellationToken cancellationToken = default)
-		{
-			Bundle bundle = await ReadBundleAsync(locator, cancellationToken);
-			return DecodePacket(locator, packetIdx, bundle.Header.Packets[packetIdx], bundle.Header.CompressionFormat, bundle.Packets[packetIdx]);
+			return bundleInfo;
 		}
 
 		/// <summary>
 		/// Gets a decoded block from the store
 		/// </summary>
-		/// <param name="locator">Information about the bundle</param>
+		/// <param name="bundleInfo">Information about the bundle</param>
 		/// <param name="packetIdx">Index of the packet</param>
-		/// <param name="packet">The decoded block location and size</param>
-		/// <param name="format">Compression type in the bundle</param>
-		/// <param name="encodedPacket">The encoded packet data</param>
+		/// <param name="cancellationToken">Cancellation token for the operation</param>
 		/// <returns>The decoded data</returns>
-		ReadOnlyMemory<byte> DecodePacket(BlobLocator locator, int packetIdx, BundlePacket packet, BundleCompressionFormat format, ReadOnlyMemory<byte> encodedPacket)
+		async ValueTask<ReadOnlyMemory<byte>> ReadBundlePacketAsync(BundleInfo bundleInfo, int packetIdx, CancellationToken cancellationToken)
 		{
-			if (_cache == null)
-			{
-				return DecodePacketUncached(packet, format, encodedPacket);
-			}
-			else
+			string cacheKey = GetDecodedPacketCacheKey(bundleInfo.Locator.BlobId, packetIdx);
+			if (!_cache.TryGetValue(cacheKey, out ReadOnlyMemory<byte> decodedPacket))
 			{
-				string cacheKey = $"bundle-packet:{locator.BlobId}#{packetIdx}";
-				return _cache.GetOrCreate<ReadOnlyMemory<byte>>(cacheKey, entry =>
+				BundlePacket packet = bundleInfo.Header.Packets[packetIdx];
+				byte[] decodedPacketBuffer = new byte[packet.DecodedLength];
+
+				ReadOnlyMemory<byte> encodedPacket;
+				if (_cache.TryGetValue(GetEncodedPacketCacheKey(bundleInfo.Locator.BlobId, packetIdx), out encodedPacket))
 				{
-					ReadOnlyMemory<byte> decodedPacket = DecodePacketUncached(packet, format, encodedPacket);
-					entry.SetSize(packet.DecodedLength);
-					return decodedPacket;
-				});
-			}
-		}
+					BundleData.Decompress(bundleInfo.Header.CompressionFormat, encodedPacket, decodedPacketBuffer);
+				}
+				else
+				{
+					using IMemoryOwner<byte> owner = MemoryPool<byte>.Shared.Rent(packet.EncodedLength);
 
-		static ReadOnlyMemory<byte> DecodePacketUncached(BundlePacket packet, BundleCompressionFormat format, ReadOnlyMemory<byte> encodedPacket)
-		{
-			byte[] decodedPacket = new byte[packet.DecodedLength];
-			BundleData.Decompress(format, encodedPacket, decodedPacket);
+					Memory<byte> buffer = owner.Memory.Slice(0, packet.EncodedLength);
+					buffer = await ReadBlobRangeAsync(bundleInfo.Locator, bundleInfo.PacketOffsets[packetIdx], buffer, cancellationToken);
+
+					BundleData.Decompress(bundleInfo.Header.CompressionFormat, buffer, decodedPacketBuffer);
+				}
+
+				decodedPacket = decodedPacketBuffer;
+				AddToCache(cacheKey, decodedPacket, decodedPacket.Length);
+			}
 			return decodedPacket;
 		}
 
@@ -270,18 +327,44 @@ namespace EpicGames.Horde.Storage
 
 		#endregion
 
-		#region Trees
+		#region Nodes
 
-		/// <inheritdoc/>
-		public ITreeWriter CreateTreeWriter(TreeOptions options, Utf8String prefix = default)
+		class NodeReader : MemoryReader, ITreeNodeReader
 		{
-			return _treeStore.CreateTreeWriter(options, prefix);
+			readonly IStorageClient _store;
+			readonly IReadOnlyList<(IoHash, NodeLocator)> _refs;
+
+			public NodeReader(IStorageClient store, ReadOnlyMemory<byte> data, IReadOnlyList<(IoHash, NodeLocator)> refs)
+				: base(data)
+			{
+				_store = store;
+				_refs = refs;
+			}
+
+			public void ReadRef(TreeNodeRef treeNodeRef)
+			{
+				treeNodeRef._store = _store;
+				(treeNodeRef._hash, treeNodeRef._locator) = _refs[(int)this.ReadUnsignedVarInt()];
+			}
 		}
 
 		/// <inheritdoc/>
-		public virtual Task<ITreeBlob?> TryReadTreeAsync(RefName name, TimeSpan maxAge = default, CancellationToken cancellationToken = default)
+		public async ValueTask<TNode> ReadNodeAsync<TNode>(NodeLocator locator, CancellationToken cancellationToken = default) where TNode : TreeNode
 		{
-			return _treeStore.TryReadTreeAsync(name, maxAge, cancellationToken);
+			BundleInfo bundleInfo = await GetBundleInfoAsync(locator.Blob, cancellationToken);
+			BundleExport export = bundleInfo.Header.Exports[locator.ExportIdx];
+
+			ExportInfo exportInfo = bundleInfo.Exports[locator.ExportIdx];
+			ReadOnlyMemory<byte> packetData = await ReadBundlePacketAsync(bundleInfo, exportInfo.PacketIdx, cancellationToken);
+
+			(IoHash, NodeLocator)[] refs = new (IoHash, NodeLocator)[export.References.Count];
+			for (int idx = 0; idx < export.References.Count; idx++)
+			{
+				refs[idx] = bundleInfo.References[export.References[idx]];
+			}
+
+			ReadOnlyMemory<byte> nodeData = packetData.Slice(exportInfo.Offset, export.Length);
+			return TreeNode.Deserialize<TNode>(new NodeReader(this, nodeData, refs));
 		}
 
 		#endregion
@@ -292,30 +375,30 @@ namespace EpicGames.Horde.Storage
 		public abstract Task DeleteRefAsync(RefName name, CancellationToken cancellationToken = default);
 
 		/// <inheritdoc/>
-		public virtual async Task<RefValue?> TryReadRefValueAsync(RefName name, DateTime cacheTime = default, CancellationToken cancellationToken = default)
+		public virtual async Task<TNode?> TryReadNodeAsync<TNode>(RefName name, DateTime cacheTime = default, CancellationToken cancellationToken = default) where TNode : TreeNode
 		{
-			RefTarget? refTarget = await TryReadRefTargetAsync(name, cacheTime, cancellationToken);
-			if (refTarget == null)
+			NodeLocator refTarget = await TryReadRefTargetAsync(name, cacheTime, cancellationToken);
+			if (!refTarget.IsValid())
 			{
 				return null;
 			}
-			return new RefValue(refTarget, await ReadBundleAsync(refTarget.Locator, cancellationToken));
+			return await ReadNodeAsync<TNode>(refTarget, cancellationToken);
 		}
 
 		/// <inheritdoc/>
-		public abstract Task<RefTarget?> TryReadRefTargetAsync(RefName name, DateTime cacheTime = default, CancellationToken cancellationToken = default);
+		public abstract Task<NodeLocator> TryReadRefTargetAsync(RefName name, DateTime cacheTime = default, CancellationToken cancellationToken = default);
 
 		/// <inheritdoc/>
-		public virtual async Task<RefTarget> WriteRefValueAsync(RefName name, Bundle bundle, int exportIdx = 0, Utf8String prefix = default, CancellationToken cancellationToken = default)
+		public virtual async Task<NodeLocator> WriteRefAsync(RefName name, Bundle bundle, int exportIdx = 0, Utf8String prefix = default, CancellationToken cancellationToken = default)
 		{
-			BlobLocator locator = await this.WriteBundleAsync(bundle, prefix, cancellationToken);
-			RefTarget target = new RefTarget(locator, exportIdx);
+			BlobLocator locator = await WriteBundleAsync(bundle, prefix, cancellationToken);
+			NodeLocator target = new NodeLocator(locator, exportIdx);
 			await WriteRefTargetAsync(name, target, cancellationToken);
 			return target;
 		}
 
 		/// <inheritdoc/>
-		public abstract Task WriteRefTargetAsync(RefName name, RefTarget target, CancellationToken cancellationToken = default);
+		public abstract Task WriteRefTargetAsync(RefName name, NodeLocator target, CancellationToken cancellationToken = default);
 
 		#endregion
 	}
diff --git a/Engine/Source/Programs/Shared/EpicGames.Horde/Storage/TreeNode.cs b/Engine/Source/Programs/Shared/EpicGames.Horde/Storage/TreeNode.cs
--- a/Engine/Source/Programs/Shared/EpicGames.Horde/Storage/TreeNode.cs
+++ b/Engine/Source/Programs/Shared/EpicGames.Horde/Storage/TreeNode.cs
@@ -5,11 +5,14 @@ using System.Buffers;
 using System.Collections.Generic;
 using System.Diagnostics;
 using System.Reflection;
+using System.Reflection.Emit;
 using System.Reflection.Metadata;
 using System.Threading;
 using System.Threading.Tasks;
 using System.Xml.Linq;
+using EpicGames.Core;
 using EpicGames.Serialization;
+using JetBrains.Annotations;
 
 namespace EpicGames.Horde.Storage
 {
@@ -19,90 +22,93 @@ namespace EpicGames.Horde.Storage
 	public abstract class TreeNode
 	{
 		/// <summary>
-		/// Cached incoming reference to the owner of this node.
+		/// Whether the node is dirty or not
 		/// </summary>
-		internal TreeNodeRef? IncomingRef { get; set; }
+		public bool IsDirty { get; private set; }
 
 		/// <summary>
-		/// Queries if the node in its current state is read-only. Once we know that nodes are no longer going to be modified, they are favored for spilling to persistent storage.
+		/// Default constructor
 		/// </summary>
-		/// <returns>True if the node is read-only.</returns>
-		public virtual bool IsReadOnly() => false;
+		protected TreeNode()
+		{
+			IsDirty = true;
+		}
+
+		/// <summary>
+		/// Serialization constructor. Clears the dirty flag by default.
+		/// </summary>
+		/// <param name="reader"></param>
+		protected TreeNode(IMemoryReader reader)
+		{
+			_ = reader;
+			IsDirty = false;
+		}
 
 		/// <summary>
 		/// Mark this node as dirty
 		/// </summary>
-		protected void MarkAsDirty() => IncomingRef?.MarkAsDirty();
+		protected void MarkAsDirty()
+		{
+			IsDirty = true;
+		}
 
 		/// <summary>
-		/// Enumerates all the child references from this node
+		/// Serialize the contents of this node
 		/// </summary>
-		/// <returns>Children of this node</returns>
-		public abstract IReadOnlyList<TreeNodeRef> GetReferences();
+		/// <returns>Data for the node</returns>
+		public abstract void Serialize(ITreeNodeWriter writer);
 
 		/// <summary>
 		/// Static instance of the serializer for a particular <see cref="TreeNode"/> type.
 		/// </summary>
 		static class SerializerInstance<T> where T : TreeNode
 		{
-			static readonly TreeSerializerAttribute _attribute = typeof(T).GetCustomAttribute<TreeSerializerAttribute>()!;
-			public static TreeNodeSerializer<T> Serializer { get; } = (TreeNodeSerializer<T>)Activator.CreateInstance(_attribute.Type)!;
-		}
+			public static readonly TreeNodeSerializer<T> Serializer = CreateSerializer();
 
-		/// <summary>
-		/// Serialize a node to a block of memory
-		/// </summary>
-		/// <returns>New data to be stored into a blob</returns>
-		public abstract Task<ITreeBlob> SerializeAsync(ITreeWriter writer, CancellationToken cancellationToken);
+			static TreeNodeSerializer<T> CreateSerializer()
+			{
+				TreeSerializerAttribute? _attribute = typeof(T).GetCustomAttribute<TreeSerializerAttribute>()!;
+				if (_attribute == null)
+				{
+					return new DefaultTreeNodeSerializer<T>();
+				}
+				else
+				{
+					return (TreeNodeSerializer<T>)Activator.CreateInstance(_attribute.Type)!;
+				}
+			}
+		}
 
 		/// <summary>
 		/// Deserialize a node from data
 		/// </summary>
 		/// <typeparam name="T"></typeparam>
-		/// <param name="blob">Blob to deserialize</param>
+		/// <param name="reader">Reader to deserialize data from</param>
 		/// <returns></returns>
-		public static T Deserialize<T>(ITreeBlob blob) where T : TreeNode
-		{
-			return SerializerInstance<T>.Serializer.Deserialize(blob);
-		}
+		public static T Deserialize<T>(ITreeNodeReader reader) where T : TreeNode => SerializerInstance<T>.Serializer.Deserialize(reader);
 	}
 
 	/// <summary>
-	/// Data to be stored in a tree blob
+	/// Reader for tree nodes
 	/// </summary>
-	public struct NewTreeBlob : ITreeBlob
+	public interface ITreeNodeReader : IMemoryReader
 	{
 		/// <summary>
-		/// The opaque data payload
-		/// </summary>
-		public ReadOnlySequence<byte> Data { get; }
-
-		/// <summary>
-		/// References to other tree nodes
+		/// Reads a reference to another node
 		/// </summary>
-		public IReadOnlyList<ITreeBlobRef> Refs { get; }
-
-		/// <summary>
-		/// Constructor
-		/// </summary>
-		/// <param name="data">Data payload</param>
-		/// <param name="refs">References to other nodes</param>
-		public NewTreeBlob(ReadOnlyMemory<byte> data, IReadOnlyList<ITreeBlobRef> refs)
-		{
-			Data = new ReadOnlySequence<byte>(data);
-			Refs = refs;
-		}
+		void ReadRef(TreeNodeRef target);
+	}
 
+	/// <summary>
+	/// Writer for tree nodes
+	/// </summary>
+	public interface ITreeNodeWriter : IMemoryWriter
+	{
 		/// <summary>
-		/// Constructor
+		/// Write a reference to another node
 		/// </summary>
-		/// <param name="data">Data payload</param>
-		/// <param name="refs">References to other nodes</param>
-		public NewTreeBlob(ReadOnlySequence<byte> data, IReadOnlyList<ITreeBlobRef> refs)
-		{
-			Data = data;
-			Refs = refs;
-		}
+		/// <param name="target">Target of the reference</param>
+		void WriteRef(TreeNodeRef target);
 	}
 
 	/// <summary>
@@ -114,9 +120,46 @@ namespace EpicGames.Horde.Storage
 		/// <summary>
 		/// Deserializes data from the given data
 		/// </summary>
-		/// <param name="blob">The typed blob</param>
+		/// <param name="reader">Reader to deserialize data from</param>
 		/// <returns>New node parsed from the data</returns>
-		public abstract T Deserialize(ITreeBlob blob);
+		public abstract T Deserialize(ITreeNodeReader reader);
+	}
+
+	/// <summary>
+	/// Default implementation of <see cref="TreeNodeSerializer{T}"/> which calls a constructor taking an <see cref="ITreeNodeReader"/> argument.
+	/// </summary>
+	/// <typeparam name="T"></typeparam>
+	public class DefaultTreeNodeSerializer<T> : TreeNodeSerializer<T> where T : TreeNode
+	{
+		static readonly Type[] s_signature = new[] { typeof(ITreeNodeReader) };
+
+		readonly Func<ITreeNodeReader, T> _constructor;
+
+		/// <summary>
+		/// Constructor
+		/// </summary>
+		public DefaultTreeNodeSerializer()
+		{
+			Type type = typeof(T);
+
+			ConstructorInfo? constructorInfo = type.GetConstructor(s_signature);
+			if (constructorInfo == null)
+			{
+				throw new InvalidOperationException($"Type {type.Name} does not have a constructor taking an {typeof(ITreeNodeReader).Name} as parameter.");
+			}
+
+			DynamicMethod method = new DynamicMethod($"Create_{type.Name}", type, s_signature, true);
+
+			ILGenerator generator = method.GetILGenerator();
+			generator.Emit(OpCodes.Ldarg_0);
+			generator.Emit(OpCodes.Newobj, constructorInfo);
+			generator.Emit(OpCodes.Ret);
+
+			_constructor = (Func<ITreeNodeReader, T>)method.CreateDelegate(typeof(Func<ITreeNodeReader, T>));
+		}
+
+		/// <inheritdoc/>
+		public override T Deserialize(ITreeNodeReader reader) => (T)_constructor(reader);
 	}
 
 	/// <summary>
@@ -142,63 +185,14 @@ namespace EpicGames.Horde.Storage
 	public static class TreeNodeExtensions
 	{
 		/// <summary>
-		/// Writes a node to storage
-		/// </summary>
-		/// <param name="writer">Writer to output the nodes to</param>
-		/// <param name="node">Root node to serialize</param>
-		/// <param name="cancellationToken">Cancellation token for the operation</param>
-		public static async Task<ITreeBlobRef> WriteNodeAsync(this ITreeWriter writer, TreeNode node, CancellationToken cancellationToken = default)
-		{
-			ITreeBlob blob = await node.SerializeAsync(writer, cancellationToken);
-			return await writer.WriteNodeAsync(blob.Data, blob.Refs, cancellationToken);
-		}
-
-		/// <summary>
-		/// Flushes a tree to storage using the given root node
+		/// 
 		/// </summary>
-		/// <param name="writer">Writer to output the nodes to</param>
-		/// <param name="name">Name of the ref to write</param>
-		/// <param name="node">Root node to serialize</param>
-		/// <param name="cancellationToken">Cancellation token for the operation</param>
-		public static async Task WriteRefAsync(this ITreeWriter writer, RefName name, TreeNode node, CancellationToken cancellationToken = default)
-		{
-			ITreeBlobRef root = await writer.WriteNodeAsync(node, cancellationToken);
-			await writer.WriteRefAsync(name, root, cancellationToken);
-		}
-
-		/// <inheritdoc/>
-		public static async Task<T?> TryReadTreeAsync<T>(this IStorageClient store, RefName name, TimeSpan maxAge = default, CancellationToken cancellationToken = default) where T : TreeNode
-		{
-			ITreeBlob? root = await store.TryReadTreeAsync(name, maxAge, cancellationToken);
-			if (root == null)
-			{
-				return null;
-			}
-			return TreeNode.Deserialize<T>(root);
-		}
-
-		/// <inheritdoc/>
-		public static async Task<T> ReadTreeAsync<T>(this IStorageClient store, RefName name, TimeSpan maxAge = default, CancellationToken cancellationToken = default) where T : TreeNode
-		{
-			T? result = await store.TryReadTreeAsync<T>(name, maxAge, cancellationToken);
-			if (result == null)
-			{
-				throw new RefNameNotFoundException(name);
-			}
-			return result;
-		}
-
-		/// <inheritdoc/>
-		public static Task WriteTreeAsync(this IStorageClient store, RefName name, TreeNode root, CancellationToken cancellationToken = default)
-		{
-			return WriteTreeAsync(store, name, root, new TreeOptions(), cancellationToken);
-		}
-
-		/// <inheritdoc/>
-		public static async Task WriteTreeAsync(this IStorageClient store, RefName name, TreeNode root, TreeOptions options, CancellationToken cancellationToken = default)
+		/// <typeparam name="T"></typeparam>
+		/// <param name="reader"></param>
+		/// <returns></returns>
+		public static TreeNodeRef<T> ReadRef<T>(this ITreeNodeReader reader) where T : TreeNode
 		{
-			ITreeWriter writer = store.CreateTreeWriter(options);
-			await writer.WriteRefAsync(name, root, cancellationToken);
+			return new TreeNodeRef<T>(reader);
 		}
 	}
 }
diff --git a/Engine/Source/Programs/Shared/EpicGames.Horde/Storage/TreeNodeRef.cs b/Engine/Source/Programs/Shared/EpicGames.Horde/Storage/TreeNodeRef.cs
--- a/Engine/Source/Programs/Shared/EpicGames.Horde/Storage/TreeNodeRef.cs
+++ b/Engine/Source/Programs/Shared/EpicGames.Horde/Storage/TreeNodeRef.cs
@@ -1,199 +1,81 @@
 // Copyright Epic Games, Inc. All Rights Reserved.
 
 using System;
-using System.Collections.Generic;
-using System.Diagnostics;
 using System.Reflection;
 using System.Threading;
 using System.Threading.Tasks;
-using EpicGames.Serialization;
+using EpicGames.Core;
 
 namespace EpicGames.Horde.Storage
 {
 	/// <summary>
 	/// Stores a reference from a parent to child node, which can be resurrected after the child node is flushed to storage if subsequently modified.
 	/// </summary>
-	public class TreeNodeRef
+	public abstract class TreeNodeRef
 	{
 		/// <summary>
-		/// Cached reference to the parent node
+		/// Store containing the node data. May be null for nodes in memory.
 		/// </summary>
-		internal TreeNode? _owner;
+		internal IStorageClient? _store;
 
 		/// <summary>
-		/// Strong reference to the current node
+		/// Hash of the referenced node. Invalid for nodes in memory.
 		/// </summary>
-		internal TreeNode? _strongRef;
+		internal IoHash _hash;
 
 		/// <summary>
-		/// Weak reference to the child node. Maintained after the object has been flushed.
+		/// Locator for the blob containing this node. Invalid for nodes in memory.
 		/// </summary>
-		internal WeakReference<TreeNode>? _weakRef;
+		internal NodeLocator _locator;
 
 		/// <summary>
-		/// Last time that the node was modified
+		/// The target node, or null if the node is not resident in memory.
 		/// </summary>
-		internal long _lastModifiedTime;
+		public TreeNode? Target { get; set; }
 
 		/// <summary>
-		/// Handle to the node in storage. May be set to null when the node is in memory.
+		/// Creates a reference to a node in memory.
 		/// </summary>
-		internal ITreeBlobRef? _target;
-
-		/// <summary>
-		/// Reference to the node in storage, if unmodified.
-		/// </summary>
-		public ITreeBlobRef? Target => ReferenceEquals(_strongRef, null) ? null : _target;
-
-		/// <summary>
-		/// Reference to the node.
-		/// </summary>
-		public TreeNode? Node
+		/// <param name="target">Node to reference</param>
+		protected TreeNodeRef(TreeNode target)
 		{
-			get => _strongRef;
-			set
-			{
-				_strongRef = value;
-				_weakRef = null;
-			}
+			Target = target;
 		}
 
 		/// <summary>
 		/// Creates a reference to a node with the given hash
 		/// </summary>
-		/// <param name="persistedNode">Hash of the referenced node</param>
-		public TreeNodeRef(ITreeBlobRef persistedNode)
-		{
-			_target = persistedNode;
-		}
-
-		/// <summary>
-		/// Creates a reference to a node with the given hash
-		/// </summary>
-		/// <param name="owner">The node which owns the reference</param>
-		/// <param name="persistedNode">Hash of the referenced node</param>
-		public TreeNodeRef(TreeNode owner, ITreeBlobRef persistedNode)
-		{
-			_owner = owner;
-			_target = persistedNode;
-		}
-
-		/// <summary>
-		/// Creates a reference to the given node
-		/// </summary>
-		/// <param name="owner">The node which owns the reference</param>
-		/// <param name="node">The referenced node</param>
-		public TreeNodeRef(TreeNode owner, TreeNode node)
+		/// <param name="store">Storage client containing the data</param>
+		/// <param name="hash">Hash of the referenced node</param>
+		/// <param name="locator">Locator for the node</param>
+		internal TreeNodeRef(IStorageClient store, IoHash hash, NodeLocator locator)
 		{
-			_owner = owner;
-			_strongRef = node;
-			_lastModifiedTime = Stopwatch.GetTimestamp();
-
-			node.IncomingRef = this;
+			_store = store;
+			_hash = hash;
+			_locator = locator;
 		}
 
 		/// <summary>
-		/// Reparent this reference to a new owner.
+		/// 
 		/// </summary>
-		public void Reparent(TreeNode newOwner)
+		/// <param name="reader"></param>
+		public TreeNodeRef(ITreeNodeReader reader)
 		{
-			MarkAsDirty();
-			_owner = newOwner;
-			MarkAsDirty();
+			reader.ReadRef(this);
 		}
 
 		/// <summary>
-		/// Figure out whether this ref is dirty
+		/// Serialize the node to the given writer
 		/// </summary>
-		/// <returns></returns>
-		public bool IsDirty() => _lastModifiedTime != 0;
+		/// <param name="writer"></param>
+		public void Serialize(ITreeNodeWriter writer) => writer.WriteRef(this);
 
 		/// <summary>
-		/// Converts the node in this reference from a weak to strong reference.
-		/// </summary>
-		/// <returns>True if the ref contains a valid node on return</returns>
-		protected TreeNode? ConvertToStrongRef()
-		{
-			// Check if it's already a strong ref
-			if (_strongRef != null)
-			{
-				return _strongRef;
-			}
-
-			// Otherwise check if it's a weak ref that hasn't been GC'd yet
-			if (_weakRef != null && _weakRef.TryGetTarget(out _strongRef))
-			{
-				_weakRef = null;
-				return _strongRef;
-			}
-
-			return null;
-		}
-
-		/// <summary>
-		/// Marks this reference as dirty
-		/// </summary>
-		internal void MarkAsDirty()
-		{
-			if (_lastModifiedTime == 0)
-			{
-				_target = null;
-
-				_lastModifiedTime = Stopwatch.GetTimestamp();
-				if (_strongRef == null)
-				{
-					if (_weakRef == null || !_weakRef.TryGetTarget(out _strongRef))
-					{
-						throw new InvalidOperationException("Unable to resolve weak reference to node");
-					}
-					_weakRef = null;
-				}
-				_owner?.IncomingRef?.MarkAsDirty();
-			}
-		}
-
-		/// <summary>
-		/// Marks this reference as clean
-		/// </summary>
-		internal void MarkAsClean(ITreeBlobRef persistedNode)
-		{
-			_target = persistedNode;
-
-			if (_strongRef != null)
-			{
-				OnCollapse();
-
-				_weakRef = new WeakReference<TreeNode>(_strongRef);
-				_strongRef = null;
-			}
-
-			_lastModifiedTime = 0;
-		}
-
-		/// <summary>
-		/// Collapse this reference and write its node to storage
+		/// Resolve this reference to a concrete node
 		/// </summary>
-		/// <param name="writer">Writer to output the node to</param>
 		/// <param name="cancellationToken">Cancellation token for the operation</param>
 		/// <returns></returns>
-		public async ValueTask<ITreeBlobRef> CollapseAsync(ITreeWriter writer, CancellationToken cancellationToken)
-		{
-			ITreeBlobRef? target = _target;
-			if (target == null)
-			{
-				ITreeBlob blob = await Node!.SerializeAsync(writer, cancellationToken);
-				target = await writer.WriteNodeAsync(blob.Data, blob.Refs, cancellationToken);
-				MarkAsClean(target);
-			}
-			return target;
-		}
-
-		/// <summary>
-		/// Callback for when a reference is collapsed to a hash value. At the time of calling, both the hash and node data will be valid.
-		/// </summary>
-		protected virtual void OnCollapse()
-		{
-		}
+		public abstract ValueTask<TreeNode> ExpandBaseAsync(CancellationToken cancellationToken = default);
 	}
 
 	/// <summary>
@@ -202,45 +84,51 @@ namespace EpicGames.Horde.Storage
 	/// <typeparam name="T">Type of the node</typeparam>
 	public class TreeNodeRef<T> : TreeNodeRef where T : TreeNode
 	{
-		static readonly TreeNodeSerializer<T> Serializer = CreateSerializer();
+		/// <summary>
+		/// Accessor for the target node
+		/// </summary>
+		public new T? Target
+		{
+			get => (T?)base.Target;
+			set => base.Target = value;
+		}
 
-		/// <inheritdoc cref="TreeNodeRef.Node"/>
-		public new T? Node
+		/// <summary>
+		/// Constructor
+		/// </summary>
+		/// <param name="target">The referenced node</param>
+		public TreeNodeRef(T target) : base(target)
 		{
-			get => (T?)base.Node;
-			set => base.Node = value;
 		}
 
 		/// <summary>
 		/// Constructor
 		/// </summary>
-		/// <param name="owner">The node which owns the reference</param>
-		/// <param name="target">Hash of the referenced node</param>
-		public TreeNodeRef(TreeNode owner, ITreeBlobRef target) : base(owner, target)
+		/// <param name="store">Storage client containing the node data</param>
+		/// <param name="hash">Hash of the referenced node</param>
+		/// <param name="locator">Locator for the node</param>
+		internal TreeNodeRef(IStorageClient store, IoHash hash, NodeLocator locator) 
+			: base(store, hash, locator)
 		{
 		}
 
 		/// <summary>
 		/// Constructor
 		/// </summary>
-		/// <param name="owner">The node which owns the reference</param>
-		/// <param name="node">The referenced node</param>
-		public TreeNodeRef(TreeNode owner, T node) : base(owner, node)
+		/// <param name="reader"></param>
+		public TreeNodeRef(ITreeNodeReader reader)
+			: base(reader)
 		{
 		}
 
 		/// <summary>
-		/// Create an instance of the serializer for this node type
+		/// Resolve this reference to a concrete node
 		/// </summary>
-		/// <returns>New serializer instance</returns>
-		static TreeNodeSerializer<T> CreateSerializer()
+		/// <param name="cancellationToken">Cancellation token for the operation</param>
+		/// <returns></returns>
+		public override async ValueTask<TreeNode> ExpandBaseAsync(CancellationToken cancellationToken = default)
 		{
-			TreeSerializerAttribute? attribute = typeof(T).GetCustomAttribute<TreeSerializerAttribute>();
-			if (attribute == null)
-			{
-				throw new InvalidOperationException($"No serializer is defined for {typeof(T).Name}");
-			}
-			return (TreeNodeSerializer<T>)Activator.CreateInstance(attribute.Type)!;
+			return await ExpandAsync(cancellationToken);
 		}
 
 		/// <summary>
@@ -250,19 +138,13 @@ namespace EpicGames.Horde.Storage
 		/// <returns></returns>
 		public async ValueTask<T> ExpandAsync(CancellationToken cancellationToken = default)
 		{
-			T? result = (T?)ConvertToStrongRef();
+			T? result = Target;
 			if (result == null)
 			{
-				Debug.Assert(_target != null);
-				ITreeBlob node = await _target.GetTargetAsync(cancellationToken);
-
-				_strongRef = Serializer.Deserialize(node);
-				_strongRef.IncomingRef = this;
-
-				result = (T)_strongRef;
+				Target = await _store!.ReadNodeAsync<T>(_locator, cancellationToken);
+				result = Target;
 			}
 			return result;
 		}
 	}
 }
-
diff --git a/Engine/Source/Programs/Shared/EpicGames.Horde/Storage/TreeStore.cs b/Engine/Source/Programs/Shared/EpicGames.Horde/Storage/TreeStore.cs
deleted file mode 100644
diff --git a/Engine/Source/Programs/Shared/EpicGames.Horde/Storage/TreeWriter.cs b/Engine/Source/Programs/Shared/EpicGames.Horde/Storage/TreeWriter.cs
new file mode 100644
--- /dev/null
+++ b/Engine/Source/Programs/Shared/EpicGames.Horde/Storage/TreeWriter.cs
@@ -0,0 +1,359 @@
+// Copyright Epic Games, Inc. All Rights Reserved.
+
+using System;
+using System.Buffers;
+using System.Collections.Generic;
+using System.Diagnostics;
+using System.Linq;
+using System.Threading;
+using System.Threading.Tasks;
+using EpicGames.Core;
+
+namespace EpicGames.Horde.Storage
+{
+	/// <summary>
+	/// Options for configuring a bundle serializer
+	/// </summary>
+	public class TreeOptions
+	{
+		/// <summary>
+		/// Maximum payload size fo a blob
+		/// </summary>
+		public int MaxBlobSize { get; set; } = 10 * 1024 * 1024;
+
+		/// <summary>
+		/// Compression format to use
+		/// </summary>
+		public BundleCompressionFormat CompressionFormat { get; set; } = BundleCompressionFormat.LZ4;
+
+		/// <summary>
+		/// Minimum size of a block to be compressed
+		/// </summary>
+		public int MinCompressionPacketSize { get; set; } = 16 * 1024;
+	}
+
+	/// <summary>
+	/// Writes nodes of a tree to an <see cref="IStorageClient"/>, packed into bundles.
+	/// </summary>
+	public class TreeWriter
+	{
+		class NodeInfo
+		{
+			public readonly IoHash Hash;
+			public readonly int Length;
+			public readonly IReadOnlyList<IoHash> RefHashes;
+
+			public NodeInfo(IoHash hash, int length, IReadOnlyList<IoHash> refs)
+			{
+				Hash = hash;
+				Length = length;
+				RefHashes = refs;
+			}
+		}
+
+		class NodeWriter : ITreeNodeWriter
+		{
+			public ByteArrayBuilder Data = new ByteArrayBuilder();
+			public readonly List<TreeNodeRef> Refs = new List<TreeNodeRef>();
+
+			public int Length => Data.Length;
+
+			public void Clear()
+			{
+				Data.Clear();
+				Refs.Clear();
+			}
+
+			public void WriteRef(TreeNodeRef target)
+			{
+				Refs.Add(target);
+				Data.WriteUnsignedVarInt(Refs.Count - 1);
+			}
+
+			public Memory<byte> GetMemory(int minSize = 1) => Data.GetMemory(minSize);
+
+			public void Advance(int length) => Data.Advance(length);
+		}
+
+		static readonly TreeOptions s_defaultOptions = new TreeOptions();
+
+		readonly IStorageClient _store;
+		readonly TreeOptions _options;
+		readonly Utf8String _prefix;
+
+		// Writer used to create packets
+		readonly ArrayMemoryWriter _packetWriter;
+
+		// Map of hashes to existing nodes. Empty/invalid locators are used for items that are queued for writing, but not available yet.
+		readonly Dictionary<IoHash, NodeLocator> _hashToNode = new Dictionary<IoHash, NodeLocator>(); // TODO: this needs to include some additional state from external sources.
+
+		// Queue of nodes for the current bundle
+		readonly List<NodeInfo> _queue = new List<NodeInfo>();
+
+		// List of fences for particular nodes to be written
+		readonly Dictionary<IoHash, TaskCompletionSource<NodeLocator>> _fences = new Dictionary<IoHash, TaskCompletionSource<NodeLocator>>();
+
+		// List of packets in the current bundle
+		readonly List<BundlePacket> _packets = new List<BundlePacket>();
+
+		// List of compressed packets
+		readonly List<ReadOnlyMemory<byte>> _encodedPackets = new List<ReadOnlyMemory<byte>>();
+
+		// Total size of compressed packets in the current bundle
+		int _currentBundleLength;
+
+		/// <summary>
+		/// Constructor
+		/// </summary>
+		/// <param name="store">Store to write data to</param>
+		/// <param name="options">Options for the writer</param>
+		/// <param name="prefix">Prefix for blobs written to the store</param>
+		public TreeWriter(IStorageClient store, TreeOptions? options = null, Utf8String prefix = default)
+		{
+			_store = store;
+			_options = options ?? s_defaultOptions;
+			_prefix = prefix;
+			_packetWriter = new ArrayMemoryWriter(_options.MinCompressionPacketSize);
+		}
+
+		/// <summary>
+		/// Returns a task that can be used to wait for a particular node to be flushed to storage.
+		/// </summary>
+		/// <param name="hash">Hash of the node to wait for</param>
+		/// <returns></returns>
+		public Task<NodeLocator> WaitForFlushAsync(IoHash hash)
+		{
+			TaskCompletionSource<NodeLocator>? tcs;
+			if (!_fences.TryGetValue(hash, out tcs))
+			{
+				_ = _hashToNode[hash]; // Make sure it exists!
+				tcs = new TaskCompletionSource<NodeLocator>();
+				_fences.Add(hash, tcs);
+			}
+			return tcs.Task;
+		}
+
+		/// <summary>
+		/// Attempt to get the locator for the node with a given hash
+		/// </summary>
+		/// <param name="hash">Hash of the node</param>
+		/// <returns>True if the locator was found</returns>
+		public NodeLocator GetLocator(IoHash hash) => _hashToNode[hash];
+
+		/// <summary>
+		/// Attempt to get the locator for the node with a given hash
+		/// </summary>
+		/// <param name="hash">Hash of the node</param>
+		/// <param name="locator">Receives the locator for the node</param>
+		/// <returns>True if the locator was found</returns>
+		public bool TryGetLocator(IoHash hash, out NodeLocator locator) => _hashToNode.TryGetValue(hash, out locator);
+
+		/// <inheritdoc/>
+		public async Task<IoHash> WriteAsync(TreeNode node, CancellationToken cancellationToken = default)
+		{
+			// Serialize the node
+			NodeWriter nodeWriter = new NodeWriter();
+			node.Serialize(nodeWriter);
+
+			// Write all the references first
+			IoHash[] nextRefs = new IoHash[nodeWriter.Refs.Count];
+			for (int idx = 0; idx < nodeWriter.Refs.Count; idx++)
+			{
+				TreeNodeRef reference = nodeWriter.Refs[idx];
+				if (reference.Target == null || !reference.Target.IsDirty)
+				{
+					nextRefs[idx] = reference._hash;
+				}
+				else
+				{
+					nextRefs[idx] = await WriteAsync(reference.Target, cancellationToken);
+				}
+			}
+
+			// Get the hash for the new blob
+			ReadOnlySequence<byte> data = nodeWriter.Data.AsSequence();
+			IoHash hash = ComputeHash(data, nextRefs);
+
+			// Write the node if we don't already have it
+			if (!_hashToNode.ContainsKey(hash))
+			{
+				_packetWriter.WriteFixedLengthBytes(data);
+
+				NodeInfo nodeInfo = new NodeInfo(hash, (int)data.Length, nextRefs);
+				_queue.Add(nodeInfo);
+				_hashToNode.Add(hash, default);
+
+				if (_packetWriter.Length > Math.Min(_options.MinCompressionPacketSize, _options.MaxBlobSize))
+				{
+					FlushPacket();
+				}
+
+				if (_currentBundleLength > _options.MaxBlobSize)
+				{
+					await FlushAsync(cancellationToken);
+				}
+			}
+			return hash;
+		}
+
+		/// <summary>
+		/// Writes a node to the given ref
+		/// </summary>
+		/// <param name="name">Name of the ref to write</param>
+		/// <param name="node"></param>
+		/// <param name="cancellationToken"></param>
+		/// <returns></returns>
+		public async Task<NodeLocator> WriteRefAsync(RefName name, TreeNode node, CancellationToken cancellationToken = default)
+		{
+			IoHash hash = await WriteAsync(node, cancellationToken);
+			await FlushAsync(cancellationToken);
+
+			NodeLocator locator = GetLocator(hash);
+			await _store.WriteRefTargetAsync(name, locator, cancellationToken);
+			return locator;
+		}
+
+		/// <summary>
+		/// Writes a ref to the node with the given hash
+		/// </summary>
+		/// <param name="name">Name of the ref to write</param>
+		/// <param name="root">Hash of the root node</param>
+		/// <param name="cancellationToken">Cancellation token for the operation</param>
+		public async Task<NodeLocator> WriteAsync(RefName name, IoHash root, CancellationToken cancellationToken)
+		{
+			await FlushAsync(cancellationToken);
+			NodeLocator locator = GetLocator(root);
+			await _store.WriteRefTargetAsync(name, locator, cancellationToken);
+			return locator;
+		}
+
+		void FlushPacket()
+		{
+			if (_packetWriter.Length > 0)
+			{
+				ReadOnlyMemory<byte> encodedData = BundleData.Compress(_options.CompressionFormat, _packetWriter.WrittenMemory);
+				_encodedPackets.Add(encodedData);
+
+				BundlePacket packet = new BundlePacket(encodedData.Length, _packetWriter.Length);
+				_packets.Add(packet);
+
+				_packetWriter.Clear();
+				_currentBundleLength += encodedData.Length;
+			}
+		}
+
+		/// <summary>
+		/// Flushes all the current nodes to storage
+		/// </summary>
+		/// <param name="cancellationToken">Cancellation token for the operation</param>
+		/// <returns></returns>
+		public async Task FlushAsync(CancellationToken cancellationToken)
+		{
+			FlushPacket();
+
+			if (_queue.Count > 0)
+			{
+				// Create a set from the nodes to be written. We use this to determine references that are imported.
+				HashSet<IoHash> nodeSet = new HashSet<IoHash>(_queue.Select(x => x.Hash));
+
+				// Find all the imported nodes by bundle
+				Dictionary<BlobLocator, List<(int, IoHash)>> bundleToImportedNodes = new Dictionary<BlobLocator, List<(int, IoHash)>>();
+				foreach (NodeInfo nodeInfo in _queue)
+				{
+					foreach (IoHash refHash in nodeInfo.RefHashes)
+					{
+						if (nodeSet.Add(refHash))
+						{
+							NodeLocator refLocator = _hashToNode[refHash];
+
+							List<(int, IoHash)>? importedNodes;
+							if (!bundleToImportedNodes.TryGetValue(refLocator.Blob, out importedNodes))
+							{
+								importedNodes = new List<(int, IoHash)>();
+								bundleToImportedNodes.Add(refLocator.Blob, importedNodes);
+							}
+
+							importedNodes.Add((refLocator.ExportIdx, refHash));
+						}
+					}
+				}
+
+				// Map from node hash to index, with imported nodes first, ordered by blob, and exported nodes second.
+				Dictionary<IoHash, int> nodeToIndex = new Dictionary<IoHash, int>();
+
+				// Add all the imports and assign them identifiers
+				List<BundleImport> imports = new List<BundleImport>();
+				foreach ((BlobLocator blobLocator, List<(int, IoHash)> importedNodes) in bundleToImportedNodes)
+				{
+					importedNodes.SortBy(x => x.Item1);
+
+					foreach ((_, IoHash hash) in importedNodes)
+					{
+						nodeToIndex[hash] = nodeToIndex.Count;
+					}
+
+					imports.Add(new BundleImport(blobLocator, importedNodes));
+				}
+
+				// Create the export list
+				List<BundleExport> exports = new List<BundleExport>(_queue.Count);
+				foreach (NodeInfo nodeInfo in _queue)
+				{
+					int[] references = nodeInfo.RefHashes.Select(x => nodeToIndex[x]).ToArray();
+					BundleExport export = new BundleExport(nodeInfo.Hash, nodeInfo.Length, references);
+					nodeToIndex.Add(nodeInfo.Hash, nodeToIndex.Count);
+					exports.Add(export);
+				}
+
+				// Create the bundle
+				BundleHeader header = new BundleHeader(_options.CompressionFormat, imports, exports, _packets.ToArray());
+				Bundle bundle = new Bundle(header, _encodedPackets.ToArray());
+
+				// Write the bundle
+				BlobLocator locator = await _store.WriteBundleAsync(bundle, _prefix, cancellationToken);
+				for (int idx = 0; idx < _queue.Count; idx++)
+				{
+					NodeInfo nodeInfo = _queue[idx];
+					_hashToNode[nodeInfo.Hash] = new NodeLocator(locator, idx);
+				}
+
+				// Fire all the alerts for any nodes that have been written
+				List<IoHash> removeFences = new List<IoHash>(_fences.Keys.Where(x => nodeToIndex.ContainsKey(x)));
+				foreach ((IoHash hash, TaskCompletionSource<NodeLocator> tcs) in _fences)
+				{
+					if (nodeToIndex.ContainsKey(hash))
+					{
+						tcs.SetResult(_hashToNode[hash]);
+						removeFences.Add(hash);
+					}
+				}
+				foreach (IoHash removeFence in removeFences)
+				{
+					_hashToNode.Remove(removeFence);
+				}
+
+				// Reset the output state
+				_packets.Clear();
+				_encodedPackets.Clear();
+				_queue.Clear();
+				_currentBundleLength = 0;
+			}
+		}
+
+		static IoHash ComputeHash(ReadOnlySequence<byte> data, IoHash[] references)
+		{
+			byte[] buffer = new byte[IoHash.NumBytes * (references.Length + 1)];
+
+			Span<byte> span = buffer;
+			for (int idx = 0; idx < references.Length; idx++)
+			{
+				Debug.Assert(references[idx] != IoHash.Zero);
+				references[idx].CopyTo(span);
+				span = span[IoHash.NumBytes..];
+			}
+			IoHash.Compute(data).CopyTo(span);
+
+			return IoHash.Compute(buffer);
+		}
+	}
+}
