Fix missing comments.

diff --git a/Engine/Source/Programs/Shared/EpicGames.Core/MemoryWriter.cs b/Engine/Source/Programs/Shared/EpicGames.Core/MemoryWriter.cs
--- a/Engine/Source/Programs/Shared/EpicGames.Core/MemoryWriter.cs
+++ b/Engine/Source/Programs/Shared/EpicGames.Core/MemoryWriter.cs
@@ -91,6 +91,9 @@ namespace EpicGames.Core
 		byte[] _data;
 		int _length;
 
+		/// <summary>
+		/// Returns the memory that has been written to
+		/// </summary>
 		public ReadOnlyMemory<byte> WrittenMemory => _data.AsMemory(0, _length);
 
 		/// <summary>
@@ -105,11 +108,18 @@ namespace EpicGames.Core
 		/// <inheritdoc/>
 		public int Length => _length;
 
+		/// <summary>
+		/// Resets the current contents of the buffer
+		/// </summary>
 		public void Clear()
 		{
 			_length = 0;
 		}
 
+		/// <summary>
+		/// Resize the underlying buffer, updating the current length.
+		/// </summary>
+		/// <param name="newLength">New length of the buffer</param>
 		public void Resize(int newLength)
 		{
 			if (newLength > _length)
