EpicGames.Core: Add utility method for generating SourceFile log values with custom display text.

diff --git a/Engine/Source/Programs/Shared/EpicGames.Core/LogEvent.cs b/Engine/Source/Programs/Shared/EpicGames.Core/LogEvent.cs
--- a/Engine/Source/Programs/Shared/EpicGames.Core/LogEvent.cs
+++ b/Engine/Source/Programs/Shared/EpicGames.Core/LogEvent.cs
@@ -32,6 +32,9 @@ namespace EpicGames.Core
 		public static readonly Utf8String Type = new Utf8String("$type");
 		public static readonly Utf8String Text = new Utf8String("$text");
 
+		public static readonly Utf8String File = new Utf8String("file"); // For source file / asset types
+		public static readonly Utf8String Identifier = new Utf8String("identifier"); // For symbols
+
 		public static readonly Utf8String Exception = new Utf8String("exception");
 		public static readonly Utf8String Trace = new Utf8String("trace");
 		public static readonly Utf8String InnerException = new Utf8String("innerException");
diff --git a/Engine/Source/Programs/Shared/EpicGames.Core/LogEventMarkup.cs b/Engine/Source/Programs/Shared/EpicGames.Core/LogEventMarkup.cs
--- a/Engine/Source/Programs/Shared/EpicGames.Core/LogEventMarkup.cs
+++ b/Engine/Source/Programs/Shared/EpicGames.Core/LogEventMarkup.cs
@@ -70,7 +70,7 @@ namespace EpicGames.Core
 					{
 						string combinedPath = RemoveRelativeDirs(Path.Combine(baseDir, file));
 						properties = new Dictionary<Utf8String, object>();
-						properties["file"] = combinedPath;
+						properties[LogEventPropertyName.File] = combinedPath;
 					}
 					catch
 					{
@@ -106,7 +106,7 @@ namespace EpicGames.Core
 
 			// Add it to the list
 			Dictionary<Utf8String, object> properties = new Dictionary<Utf8String, object>();
-			properties["identifier"] = identifier;
+			properties[LogEventPropertyName.Identifier] = identifier;
 			builder.Annotate(group, new LogValue(LogValueType.Symbol, "", properties));
 		}
 	}
diff --git a/Engine/Source/Programs/Shared/EpicGames.Core/LogValue.cs b/Engine/Source/Programs/Shared/EpicGames.Core/LogValue.cs
--- a/Engine/Source/Programs/Shared/EpicGames.Core/LogValue.cs
+++ b/Engine/Source/Programs/Shared/EpicGames.Core/LogValue.cs
@@ -62,6 +62,17 @@ namespace EpicGames.Core
 			Properties = properties;
 		}
 
+		/// <summary>
+		/// Construct a source file log value with an overridden display text
+		/// </summary>
+		/// <param name="file">Source file to reference</param>
+		/// <param name="text">Display text for the fiel</param>
+		/// <returns>New log value instance</returns>
+		public static LogValue SourceFile(FileReference file, string text)
+		{
+			return new LogValue(LogValueType.SourceFile, text, new Dictionary<Utf8String, object> { [LogEventPropertyName.File] = file.FullName });
+		}
+
 		/// <summary>
 		/// Creates a LogValue from an object, overriding the type and display text for it
 		/// </summary>
