Fixed issue where already built UBT plugins were not being listed properly.

#rnx

diff --git a/Engine/Source/Programs/Shared/EpicGames.Build/System/EnumeratePlugins.cs b/Engine/Source/Programs/Shared/EpicGames.Build/System/EnumeratePlugins.cs
--- a/Engine/Source/Programs/Shared/EpicGames.Build/System/EnumeratePlugins.cs
+++ b/Engine/Source/Programs/Shared/EpicGames.Build/System/EnumeratePlugins.cs
@@ -148,7 +148,7 @@ namespace UnrealBuildBase
 			Plugins = null;
 
 			// Collect the list of plugins already build
-			List<FileReference> Built = FoundPlugins.Where(P => !UbtPluginFileCache.ContainsKey(P)).ToList();
+			List<FileReference> Built = FoundPlugins.Where(P => UbtPluginFileCache.ContainsKey(P)).ToList();
 
 			// Collect the list of plugins not build
 			List<FileReference> NotBuilt = FoundPlugins.Where(P => !UbtPluginFileCache.ContainsKey(P)).ToList();
