Horde: Multi-threaded processing of files being added to bundles.

diff --git a/Engine/Source/Programs/Horde/Horde.Agent/Commands/Bundles/CreateCommand.cs b/Engine/Source/Programs/Horde/Horde.Agent/Commands/Bundles/CreateCommand.cs
--- a/Engine/Source/Programs/Horde/Horde.Agent/Commands/Bundles/CreateCommand.cs
+++ b/Engine/Source/Programs/Horde/Horde.Agent/Commands/Bundles/CreateCommand.cs
@@ -1,10 +1,12 @@
 // Copyright Epic Games, Inc. All Rights Reserved.
 
 using System;
+using System.Diagnostics;
 using System.Globalization;
 using System.IO;
 using System.Threading;
 using System.Threading.Tasks;
+using System.Timers;
 using EpicGames.Core;
 using EpicGames.Horde.Storage;
 using EpicGames.Horde.Storage.Nodes;
@@ -118,12 +120,19 @@ namespace Horde.Agent.Commands.Bundles
 		{
 			using IMemoryCache cache = new MemoryCache(new MemoryCacheOptions());
 			IStorageClient store = CreateStorageClient(cache, logger);
+
 			TreeWriter writer = new TreeWriter(store, prefix: RefName.Text);
 
 			DirectoryNode node = new DirectoryNode(DirectoryFlags.None);
-			await node.CopyFromDirectoryAsync(InputDir.ToDirectoryInfo(), new ChunkingOptions(), writer, logger, CancellationToken.None);
+
+			Stopwatch timer = Stopwatch.StartNew();
+
+			ChunkingOptions options = new ChunkingOptions();
+			await node.CopyFromDirectoryAsync(InputDir.ToDirectoryInfo(), options, writer, logger, CancellationToken.None);
 
 			await writer.WriteRefAsync(RefName, node);
+
+			logger.LogInformation("Time: {Time}", timer.Elapsed.TotalSeconds);
 			return 0;
 		}
 	}
diff --git a/Engine/Source/Programs/Shared/EpicGames.Horde/Storage/Nodes/DirectoryNode.cs b/Engine/Source/Programs/Shared/EpicGames.Horde/Storage/Nodes/DirectoryNode.cs
--- a/Engine/Source/Programs/Shared/EpicGames.Horde/Storage/Nodes/DirectoryNode.cs
+++ b/Engine/Source/Programs/Shared/EpicGames.Horde/Storage/Nodes/DirectoryNode.cs
@@ -1,8 +1,8 @@
 // Copyright Epic Games, Inc. All Rights Reserved.
 
 using EpicGames.Core;
-using EpicGames.Horde.Storage.Git;
 using Microsoft.Extensions.Logging;
+using Microsoft.Extensions.Options;
 using System;
 using System.Buffers;
 using System.Collections.Generic;
@@ -10,8 +10,10 @@ using System.Diagnostics;
 using System.Diagnostics.CodeAnalysis;
 using System.IO;
 using System.Linq;
+using System.Reflection;
 using System.Threading;
 using System.Threading.Tasks;
+using System.Xml.Linq;
 
 namespace EpicGames.Horde.Storage.Nodes
 {
@@ -75,7 +77,15 @@ namespace EpicGames.Horde.Storage.Nodes
 		/// <summary>
 		/// Cached length of this node
 		/// </summary>
-		long _cachedLength;
+		readonly long _cachedLength;
+
+		/// <summary>
+		/// Constructor
+		/// </summary>
+		public FileEntry(Utf8String name, FileEntryFlags flags)
+			: this(name, flags, new LeafFileNode())
+		{
+		}
 
 		/// <summary>
 		/// Constructor
@@ -135,36 +145,8 @@ namespace EpicGames.Horde.Storage.Nodes
 		/// <param name="cancellationToken">Cancellation token for the operation</param>
 		public async Task AppendAsync(Stream stream, ChunkingOptions options, TreeWriter writer, CancellationToken cancellationToken)
 		{
-			const int BufferLength = 32 * 1024;
-
-			using IMemoryOwner<byte> owner = MemoryPool<byte>.Shared.Rent(BufferLength * 2);
-			Memory<byte> buffer = owner.Memory;
-
-			int readBufferOffset = 0;
-			Memory<byte> appendBuffer = Memory<byte>.Empty;
-			for (; ; )
-			{
-				// Start a read into memory
-				Memory<byte> readBuffer = buffer.Slice(readBufferOffset, BufferLength);
-				Task<int> readTask = Task.Run(async () => await stream.ReadAsync(readBuffer, cancellationToken), cancellationToken);
-
-				// In the meantime, append the last data that was read to the tree
-				if (appendBuffer.Length > 0)
-				{
-					await AppendAsync(appendBuffer, options, writer, cancellationToken);
-				}
-
-				// Wait for the read to finish
-				int numBytes = await readTask;
-				if (numBytes == 0)
-				{
-					break;
-				}
-
-				// Switch the buffers around
-				appendBuffer = readBuffer.Slice(0, numBytes);
-				readBufferOffset ^= BufferLength;
-			}
+			FileNode node = await ExpandAsync(cancellationToken);
+			Target = await node.AppendAsync(stream, options, writer, cancellationToken);
 		}
 
 		/// <inheritdoc/>
@@ -346,27 +328,33 @@ namespace EpicGames.Horde.Storage.Nodes
 		/// <returns>True if the entry exists</returns>
 		public bool Contains(Utf8String name) => TryGetFileEntry(name, out _) || TryGetDirectoryEntry(name, out _);
 
-#region File operations
+		#region File operations
 
 		/// <summary>
-		/// Adds a new directory with the given name
+		/// Adds a new file entry to this directory
 		/// </summary>
-		/// <param name="name">Name of the new directory</param>
-		/// <param name="flags">Flags for the new file</param>
-		/// <returns>The new directory object</returns>
-		public FileEntry AddFile(Utf8String name, FileEntryFlags flags)
+		/// <param name="entry">The entry to add</param>
+		public void AddFile(FileEntry entry)
 		{
-			if (TryGetDirectoryEntry(name, out _))
+			if (TryGetDirectoryEntry(entry.Name, out _))
 			{
-				throw new ArgumentException($"A directory with the name {name} already exists");
+				throw new ArgumentException($"A directory with the name {entry.Name} already exists");
 			}
 
-			FileNode newNode = new LeafFileNode();
-
-			FileEntry entry = new FileEntry(name, flags, newNode);
-			_nameToFileEntry[name] = entry;
+			_nameToFileEntry.Add(entry.Name, entry);
 			MarkAsDirty();
+		}
 
+		/// <summary>
+		/// Adds a new directory with the given name
+		/// </summary>
+		/// <param name="name">Name of the new directory</param>
+		/// <param name="flags">Flags for the new file</param>
+		/// <returns>The new directory object</returns>
+		public FileEntry AddFile(Utf8String name, FileEntryFlags flags)
+		{
+			FileEntry entry = new FileEntry(name, flags);
+			AddFile(entry);
 			return entry;
 		}
 
@@ -573,7 +561,7 @@ namespace EpicGames.Horde.Storage.Nodes
 			return false;
 		}
 
-#endregion
+		#endregion
 
 		/// <summary>
 		/// Adds files from a directory on disk
@@ -585,19 +573,79 @@ namespace EpicGames.Horde.Storage.Nodes
 		/// <param name="cancellationToken">Cancellation token for the operation</param>
 		/// <returns></returns>
 		public async Task CopyFromDirectoryAsync(DirectoryInfo directoryInfo, ChunkingOptions options, TreeWriter writer, ILogger logger, CancellationToken cancellationToken)
+		{
+			const int MaxWriters = 32;
+			const long MinSizePerWriter = 1024 * 1024;
+
+			// Enumerate all the files below this directory
+			List<(DirectoryNode DirectoryNode, FileInfo FileInfo)> files = new List<(DirectoryNode, FileInfo)>();
+			FindFilesToCopy(directoryInfo, files);
+
+			// Compute the total size
+			long totalSize = files.Sum(x => x.Item2.Length);
+			long chunkSize = Math.Max(MinSizePerWriter, totalSize / MaxWriters);
+
+			List<Task> tasks = new List<Task>();
+			long currentSize = 0;
+			long targetSize = chunkSize;
+			FileEntry[] fileEntries = new FileEntry[files.Count];
+
+			// Split it into separate writers
+			for (int minIdx = 0; minIdx < files.Count; )
+			{
+				currentSize += files[minIdx].FileInfo.Length;
+
+				int maxIdx = minIdx + 1;
+				while (maxIdx < files.Count && currentSize + files[maxIdx].FileInfo.Length <= targetSize)
+				{
+					currentSize += files[maxIdx].FileInfo.Length;
+					maxIdx++;
+				}
+
+				int minIdxCopy = minIdx;
+				tasks.Add(Task.Run(() => CopyFilesAsync(files, minIdxCopy, maxIdx, fileEntries, options, writer, logger, cancellationToken), cancellationToken));
+
+				targetSize += chunkSize;
+				minIdx = maxIdx;
+			}
+
+			// Wait for them all to finish
+			await Task.WhenAll(tasks);
+
+			// Update the directory with all the output entries
+			for (int idx = 0; idx < files.Count; idx++)
+			{
+				files[idx].DirectoryNode.AddFile(fileEntries[idx]);
+			}
+		}
+
+		void FindFilesToCopy(DirectoryInfo directoryInfo, List<(DirectoryNode, FileInfo)> files)
 		{
 			foreach (DirectoryInfo subDirectoryInfo in directoryInfo.EnumerateDirectories())
 			{
-				logger.LogInformation("Adding {Directory}", subDirectoryInfo.FullName);
-				DirectoryNode subDirectoryNode = AddDirectory(subDirectoryInfo.Name);
-				await subDirectoryNode.CopyFromDirectoryAsync(subDirectoryInfo, options, writer, logger, cancellationToken);
+				AddDirectory(subDirectoryInfo.Name).FindFilesToCopy(subDirectoryInfo, files);
 			}
 			foreach (FileInfo fileInfo in directoryInfo.EnumerateFiles())
 			{
-				logger.LogInformation("Adding {File}", fileInfo.FullName);
-				using Stream stream = fileInfo.OpenRead();
-				await AddFileAsync(fileInfo.Name, 0, stream, options, writer, cancellationToken);
+				files.Add((this, fileInfo));
+			}
+		}
+
+		static async Task CopyFilesAsync(List<(DirectoryNode DirectoryNode, FileInfo FileInfo)> files, int minIdx, int maxIdx, FileEntry[] entries, ChunkingOptions options, TreeWriter baseWriter, ILogger logger, CancellationToken cancellationToken)
+		{
+			TreeWriter writer = new TreeWriter(baseWriter);
+			for(int idx = minIdx; idx < maxIdx; idx++)
+			{
+				FileInfo fileInfo = files[idx].FileInfo;
+
+				FileEntry fileEntry = new FileEntry(fileInfo.Name, FileEntryFlags.None);
+				using (Stream stream = fileInfo.OpenRead())
+				{
+					await fileEntry.AppendAsync(stream, options, writer, cancellationToken);
+				}
+				entries[idx] = fileEntry;
 			}
+			await writer.FlushAsync(cancellationToken);
 		}
 
 		/// <summary>
diff --git a/Engine/Source/Programs/Shared/EpicGames.Horde/Storage/Nodes/FileNode.cs b/Engine/Source/Programs/Shared/EpicGames.Horde/Storage/Nodes/FileNode.cs
--- a/Engine/Source/Programs/Shared/EpicGames.Horde/Storage/Nodes/FileNode.cs
+++ b/Engine/Source/Programs/Shared/EpicGames.Horde/Storage/Nodes/FileNode.cs
@@ -132,6 +132,63 @@ namespace EpicGames.Horde.Storage.Nodes
 			return AppendAsync(this, input, options, writer, cancellationToken);
 		}
 
+		/// <summary>
+		/// Appends data to this file
+		/// </summary>
+		/// <param name="stream">Data to append to the file</param>
+		/// <param name="options">Options for chunking the data</param>
+		/// <param name="writer">Writer for new node data</param>
+		/// <param name="cancellationToken">Cancellation token for the operation</param>
+		/// <returns>New node at the root of this file</returns>
+		public ValueTask<FileNode> AppendAsync(Stream stream, ChunkingOptions options, TreeWriter writer, CancellationToken cancellationToken)
+		{
+			return AppendAsync(this, stream, options, writer, cancellationToken);
+		}
+
+		/// <summary>
+		/// Appends data to this file
+		/// </summary>
+		/// <param name="node">Node to append to</param>
+		/// <param name="stream">Data to append to the file</param>
+		/// <param name="options">Options for chunking the data</param>
+		/// <param name="writer">Writer for new node data</param>
+		/// <param name="cancellationToken">Cancellation token for the operation</param>
+		/// <returns>New node at the root of this file</returns>
+		static async ValueTask<FileNode> AppendAsync(FileNode node, Stream stream, ChunkingOptions options, TreeWriter writer, CancellationToken cancellationToken)
+		{
+			const int BufferLength = 32 * 1024;
+
+			using IMemoryOwner<byte> owner = MemoryPool<byte>.Shared.Rent(BufferLength * 2);
+			Memory<byte> buffer = owner.Memory;
+
+			int readBufferOffset = 0;
+			Memory<byte> appendBuffer = Memory<byte>.Empty;
+			for (; ; )
+			{
+				// Start a read into memory
+				Memory<byte> readBuffer = buffer.Slice(readBufferOffset, BufferLength);
+				Task<int> readTask = Task.Run(async () => await stream.ReadAsync(readBuffer, cancellationToken), cancellationToken);
+
+				// In the meantime, append the last data that was read to the tree
+				if (appendBuffer.Length > 0)
+				{
+					node = await node.AppendAsync(appendBuffer, options, writer, cancellationToken);
+				}
+
+				// Wait for the read to finish
+				int numBytes = await readTask;
+				if (numBytes == 0)
+				{
+					break;
+				}
+
+				// Switch the buffers around
+				appendBuffer = readBuffer.Slice(0, numBytes);
+				readBufferOffset ^= BufferLength;
+			}
+			return node;
+		}
+
 		static async ValueTask<FileNode> AppendAsync(FileNode root, ReadOnlyMemory<byte> input, ChunkingOptions options, TreeWriter writer, CancellationToken cancellationToken)
 		{
 			for (; ; )
@@ -303,22 +360,32 @@ namespace EpicGames.Horde.Storage.Nodes
 				return newData;
 			}
 
-			// Fast path for appending data to the buffer up to the chunk window size
-			int windowSize = options.LeafOptions.MinSize;
-			if (Length < windowSize)
+			// If the target option sizes are fixed, just chunk the data along fixed boundaries
+			if (options.LeafOptions.MinSize == options.LeafOptions.TargetSize && options.LeafOptions.MaxSize == options.LeafOptions.TargetSize)
 			{
-				int appendLength = Math.Min(windowSize - (int)Length, newData.Length);
-				AppendLeafData(newData.Span.Slice(0, appendLength));
-				newData = newData.Slice(appendLength);
+				int length = Math.Min(newData.Length, options.LeafOptions.MaxSize - (int)Length);
+				AppendLeafData(newData.Span.Slice(0, length), 0);
+				newData = newData.Slice(length);
 			}
-
-			// Cap the maximum amount of data to append to this node
-			int maxLength = Math.Min(newData.Length, options.LeafOptions.MaxSize - (int)Length);
-			if (maxLength > 0)
+			else
 			{
-				ReadOnlySpan<byte> inputSpan = newData.Span.Slice(0, maxLength);
-				int length = AppendLeafDataToChunkBoundary(inputSpan, options);
-				newData = newData.Slice(length);
+				// Fast path for appending data to the buffer up to the chunk window size
+				int windowSize = options.LeafOptions.MinSize;
+				if (Length < windowSize)
+				{
+					int appendLength = Math.Min(windowSize - (int)Length, newData.Length);
+					AppendLeafData(newData.Span.Slice(0, appendLength));
+					newData = newData.Slice(appendLength);
+				}
+
+				// Cap the maximum amount of data to append to this node
+				int maxLength = Math.Min(newData.Length, options.LeafOptions.MaxSize - (int)Length);
+				if (maxLength > 0)
+				{
+					ReadOnlySpan<byte> inputSpan = newData.Span.Slice(0, maxLength);
+					int length = AppendLeafDataToChunkBoundary(inputSpan, options);
+					newData = newData.Slice(length);
+				}
 			}
 
 			// Mark this node as complete if we've reached the max size
@@ -518,7 +585,7 @@ namespace EpicGames.Horde.Storage.Nodes
 						}
 
 						// Collapse the final node
-						await writer.WriteAsync(Children[^1].Target!, cancellationToken);
+						await writer.WriteAsync(Children[^1], cancellationToken);
 					}
 				}
 
diff --git a/Engine/Source/Programs/Shared/EpicGames.Horde/Storage/StorageClientBase.cs b/Engine/Source/Programs/Shared/EpicGames.Horde/Storage/StorageClientBase.cs
--- a/Engine/Source/Programs/Shared/EpicGames.Horde/Storage/StorageClientBase.cs
+++ b/Engine/Source/Programs/Shared/EpicGames.Horde/Storage/StorageClientBase.cs
@@ -343,8 +343,8 @@ namespace EpicGames.Horde.Storage
 
 			public void ReadRef(TreeNodeRef treeNodeRef)
 			{
-				treeNodeRef._store = _store;
-				(treeNodeRef._hash, treeNodeRef._locator) = _refs[(int)this.ReadUnsignedVarInt()];
+				(IoHash hash, NodeLocator locator) = _refs[(int)this.ReadUnsignedVarInt()];
+				treeNodeRef.MarkAsClean(_store, hash, locator, 0);
 			}
 		}
 
diff --git a/Engine/Source/Programs/Shared/EpicGames.Horde/Storage/TreeNode.cs b/Engine/Source/Programs/Shared/EpicGames.Horde/Storage/TreeNode.cs
--- a/Engine/Source/Programs/Shared/EpicGames.Horde/Storage/TreeNode.cs
+++ b/Engine/Source/Programs/Shared/EpicGames.Horde/Storage/TreeNode.cs
@@ -1,18 +1,9 @@
 // Copyright Epic Games, Inc. All Rights Reserved.
 
 using System;
-using System.Buffers;
-using System.Collections.Generic;
-using System.Diagnostics;
 using System.Reflection;
 using System.Reflection.Emit;
-using System.Reflection.Metadata;
-using System.Threading;
-using System.Threading.Tasks;
-using System.Xml.Linq;
 using EpicGames.Core;
-using EpicGames.Serialization;
-using JetBrains.Annotations;
 
 namespace EpicGames.Horde.Storage
 {
@@ -22,26 +13,30 @@ namespace EpicGames.Horde.Storage
 	public abstract class TreeNode
 	{
 		/// <summary>
-		/// Whether the node is dirty or not
+		/// Revision number for the node. Incremented whenever the node is modified; used to detect changes between a serialized ref and live instance.
 		/// </summary>
-		public bool IsDirty { get; private set; }
+		int _revision;
+
+		/// <summary>
+		/// Revision number for this node. Incremented Whether the node is dirty or not
+		/// </summary>
+		public int Revision => _revision;
 
 		/// <summary>
 		/// Default constructor
 		/// </summary>
 		protected TreeNode()
 		{
-			IsDirty = true;
+			_revision = 1;
 		}
 
 		/// <summary>
-		/// Serialization constructor. Clears the dirty flag by default.
+		/// Serialization constructor. Leaves the revision number zeroed by default.
 		/// </summary>
 		/// <param name="reader"></param>
 		protected TreeNode(IMemoryReader reader)
 		{
 			_ = reader;
-			IsDirty = false;
 		}
 
 		/// <summary>
@@ -49,7 +44,7 @@ namespace EpicGames.Horde.Storage
 		/// </summary>
 		protected void MarkAsDirty()
 		{
-			IsDirty = true;
+			_revision++;
 		}
 
 		/// <summary>
diff --git a/Engine/Source/Programs/Shared/EpicGames.Horde/Storage/TreeNodeRef.cs b/Engine/Source/Programs/Shared/EpicGames.Horde/Storage/TreeNodeRef.cs
--- a/Engine/Source/Programs/Shared/EpicGames.Horde/Storage/TreeNodeRef.cs
+++ b/Engine/Source/Programs/Shared/EpicGames.Horde/Storage/TreeNodeRef.cs
@@ -1,7 +1,6 @@
 // Copyright Epic Games, Inc. All Rights Reserved.
 
 using System;
-using System.Reflection;
 using System.Threading;
 using System.Threading.Tasks;
 using EpicGames.Core;
@@ -9,35 +8,49 @@ using EpicGames.Core;
 namespace EpicGames.Horde.Storage
 {
 	/// <summary>
-	/// Stores a reference from a parent to child node, which can be resurrected after the child node is flushed to storage if subsequently modified.
+	/// Stores a reference from a parent to child node. The reference may be to a node in memory, or to a node in the storage system.
 	/// </summary>
-	public abstract class TreeNodeRef
+	public class TreeNodeRef
 	{
 		/// <summary>
 		/// Store containing the node data. May be null for nodes in memory.
 		/// </summary>
-		internal IStorageClient? _store;
+		public IStorageClient? Store { get; private set; }
 
 		/// <summary>
 		/// Hash of the referenced node. Invalid for nodes in memory.
 		/// </summary>
-		internal IoHash _hash;
+		public IoHash Hash { get; private set; }
 
 		/// <summary>
 		/// Locator for the blob containing this node. Invalid for nodes in memory.
 		/// </summary>
-		internal NodeLocator _locator;
+		public NodeLocator Locator { get; private set; }
+
+		/// <summary>
+		/// Revision number of the target node
+		/// </summary>
+		private int _revision;
+
+		/// <summary>
+		/// The target node in memory
+		/// </summary>
+		private TreeNode? _target;
 
 		/// <summary>
 		/// The target node, or null if the node is not resident in memory.
 		/// </summary>
-		public TreeNode? Target { get; set; }
+		public TreeNode? Target
+		{
+			get => _target;
+			set => MarkAsDirty(value);
+		}
 
 		/// <summary>
 		/// Creates a reference to a node in memory.
 		/// </summary>
 		/// <param name="target">Node to reference</param>
-		protected TreeNodeRef(TreeNode target)
+		protected internal TreeNodeRef(TreeNode target)
 		{
 			Target = target;
 		}
@@ -50,13 +63,13 @@ namespace EpicGames.Horde.Storage
 		/// <param name="locator">Locator for the node</param>
 		internal TreeNodeRef(IStorageClient store, IoHash hash, NodeLocator locator)
 		{
-			_store = store;
-			_hash = hash;
-			_locator = locator;
+			Store = store;
+			Hash = hash;
+			Locator = locator;
 		}
 
 		/// <summary>
-		/// 
+		/// Deserialization constructor
 		/// </summary>
 		/// <param name="reader"></param>
 		public TreeNodeRef(ITreeNodeReader reader)
@@ -65,17 +78,61 @@ namespace EpicGames.Horde.Storage
 		}
 
 		/// <summary>
-		/// Serialize the node to the given writer
+		/// Determines whether the the referenced node has modified from the last version written to storage
 		/// </summary>
-		/// <param name="writer"></param>
-		public void Serialize(ITreeNodeWriter writer) => writer.WriteRef(this);
+		/// <returns></returns>
+		public bool IsDirty() => _target != null && _revision != _target.Revision;
 
 		/// <summary>
-		/// Resolve this reference to a concrete node
+		/// Update the reference to refer to a node in memory.
 		/// </summary>
-		/// <param name="cancellationToken">Cancellation token for the operation</param>
-		/// <returns></returns>
-		public abstract ValueTask<TreeNode> ExpandBaseAsync(CancellationToken cancellationToken = default);
+		/// <param name="target">The target node</param>
+		public void MarkAsDirty(TreeNode? target)
+		{
+			if (target == null)
+			{
+				if (!Locator.IsValid())
+				{
+					throw new InvalidOperationException("Node has not been serialized to disk; cannot clear target reference.");
+				}
+			}
+			else
+			{
+				Store = null;
+				Hash = default;
+				Locator = default;
+				_revision = 0;
+			}
+
+			_target = target;
+		}
+
+		/// <summary>
+		/// Update the reference to refer to a location in storage.
+		/// </summary>
+		/// <param name="store">The storage client</param>
+		/// <param name="hash">Hash of the node</param>
+		/// <param name="locator">Location of the node</param>
+		/// <param name="revision">Revision number for the node</param>
+		public bool MarkAsClean(IStorageClient store, IoHash hash, NodeLocator locator, int revision)
+		{
+			bool result = false;
+			if (_target == null || _target.Revision == revision)
+			{
+				Store = store;
+				Hash = hash;
+				Locator = locator;
+				_revision = revision;
+				result = true;
+			}
+			return result;
+		}
+
+		/// <summary>
+		/// Serialize the node to the given writer
+		/// </summary>
+		/// <param name="writer"></param>
+		public void Serialize(ITreeNodeWriter writer) => writer.WriteRef(this);
 	}
 
 	/// <summary>
@@ -121,16 +178,6 @@ namespace EpicGames.Horde.Storage
 		{
 		}
 
-		/// <summary>
-		/// Resolve this reference to a concrete node
-		/// </summary>
-		/// <param name="cancellationToken">Cancellation token for the operation</param>
-		/// <returns></returns>
-		public override async ValueTask<TreeNode> ExpandBaseAsync(CancellationToken cancellationToken = default)
-		{
-			return await ExpandAsync(cancellationToken);
-		}
-
 		/// <summary>
 		/// Resolve this reference to a concrete node
 		/// </summary>
@@ -141,7 +188,7 @@ namespace EpicGames.Horde.Storage
 			T? result = Target;
 			if (result == null)
 			{
-				Target = await _store!.ReadNodeAsync<T>(_locator, cancellationToken);
+				Target = await Store!.ReadNodeAsync<T>(Locator, cancellationToken);
 				result = Target;
 			}
 			return result;
diff --git a/Engine/Source/Programs/Shared/EpicGames.Horde/Storage/TreeWriter.cs b/Engine/Source/Programs/Shared/EpicGames.Horde/Storage/TreeWriter.cs
--- a/Engine/Source/Programs/Shared/EpicGames.Horde/Storage/TreeWriter.cs
+++ b/Engine/Source/Programs/Shared/EpicGames.Horde/Storage/TreeWriter.cs
@@ -41,13 +41,17 @@ namespace EpicGames.Horde.Storage
 		{
 			public readonly IoHash Hash;
 			public readonly int Length;
+			public readonly int Revision;
 			public readonly IReadOnlyList<IoHash> RefHashes;
+			public readonly TreeNodeRef NodeRef;
 
-			public NodeInfo(IoHash hash, int length, IReadOnlyList<IoHash> refs)
+			public NodeInfo(IoHash hash, int length, IReadOnlyList<IoHash> refs, TreeNodeRef nodeRef, int revision)
 			{
 				Hash = hash;
 				Length = length;
 				RefHashes = refs;
+				NodeRef = nodeRef;
+				Revision = revision;
 			}
 		}
 
@@ -90,8 +94,8 @@ namespace EpicGames.Horde.Storage
 		// Queue of nodes for the current bundle
 		readonly List<NodeInfo> _queue = new List<NodeInfo>();
 
-		// List of fences for particular nodes to be written
-		readonly Dictionary<IoHash, TaskCompletionSource<NodeLocator>> _fences = new Dictionary<IoHash, TaskCompletionSource<NodeLocator>>();
+		// Queue of nodes that are already scheduled to be written elsewhere
+		readonly List<NodeInfo> _secondaryQueue = new List<NodeInfo>();
 
 		// List of packets in the current bundle
 		readonly List<BundlePacket> _packets = new List<BundlePacket>();
@@ -102,6 +106,11 @@ namespace EpicGames.Horde.Storage
 		// Total size of compressed packets in the current bundle
 		int _currentBundleLength;
 
+		/// <summary>
+		/// Accessor for the store backing this writer
+		/// </summary>
+		public IStorageClient Store => _store;
+
 		/// <summary>
 		/// Constructor
 		/// </summary>
@@ -117,20 +126,12 @@ namespace EpicGames.Horde.Storage
 		}
 
 		/// <summary>
-		/// Returns a task that can be used to wait for a particular node to be flushed to storage.
+		/// Copies settings from another tree writer
 		/// </summary>
-		/// <param name="hash">Hash of the node to wait for</param>
-		/// <returns></returns>
-		public Task<NodeLocator> WaitForFlushAsync(IoHash hash)
+		/// <param name="other">Other instance to copy from</param>
+		public TreeWriter(TreeWriter other)
+			: this(other._store, other._options, other._prefix)
 		{
-			TaskCompletionSource<NodeLocator>? tcs;
-			if (!_fences.TryGetValue(hash, out tcs))
-			{
-				_ = _hashToNode[hash]; // Make sure it exists!
-				tcs = new TaskCompletionSource<NodeLocator>();
-				_fences.Add(hash, tcs);
-			}
-			return tcs.Task;
 		}
 
 		/// <summary>
@@ -149,37 +150,53 @@ namespace EpicGames.Horde.Storage
 		public bool TryGetLocator(IoHash hash, out NodeLocator locator) => _hashToNode.TryGetValue(hash, out locator);
 
 		/// <inheritdoc/>
-		public async Task<IoHash> WriteAsync(TreeNode node, CancellationToken cancellationToken = default)
+		public async Task<IoHash> WriteAsync(TreeNodeRef nodeRef, CancellationToken cancellationToken = default)
 		{
+			// If the target node hasn't been modified, use the existing state.
+			if (!nodeRef.IsDirty())
+			{
+				Debug.Assert(nodeRef.Locator.IsValid());
+				_hashToNode.TryAdd(nodeRef.Hash, nodeRef.Locator);
+				nodeRef.Target = null;
+				return nodeRef.Hash;
+			}
+
 			// Serialize the node
 			NodeWriter nodeWriter = new NodeWriter();
-			node.Serialize(nodeWriter);
+			nodeRef.Target!.Serialize(nodeWriter);
 
 			// Write all the references first
 			IoHash[] nextRefs = new IoHash[nodeWriter.Refs.Count];
 			for (int idx = 0; idx < nodeWriter.Refs.Count; idx++)
 			{
 				TreeNodeRef reference = nodeWriter.Refs[idx];
-				if (reference.Target == null || !reference.Target.IsDirty)
-				{
-					nextRefs[idx] = reference._hash;
-				}
-				else
-				{
-					nextRefs[idx] = await WriteAsync(reference.Target, cancellationToken);
-				}
+				nextRefs[idx] = await WriteAsync(reference, cancellationToken);
 			}
 
 			// Get the hash for the new blob
 			ReadOnlySequence<byte> data = nodeWriter.Data.AsSequence();
 			IoHash hash = ComputeHash(data, nextRefs);
 
-			// Write the node if we don't already have it
-			if (!_hashToNode.ContainsKey(hash))
+			// Check if we're already tracking a node with the same hash
+			NodeInfo nodeInfo = new NodeInfo(hash, (int)data.Length, nextRefs, nodeRef, nodeRef.Target.Revision);
+			if (_hashToNode.TryGetValue(hash, out NodeLocator locator))
 			{
+				// If the node is in the lookup but not currently valid, it's already queued for writing. Add this ref to the list of refs that needs fixing up,
+				// so we can update it after flushing.
+				if (locator.IsValid())
+				{
+					nodeRef.MarkAsClean(_store, hash, locator, nodeRef.Target.Revision);
+				}
+				else
+				{
+					_secondaryQueue.Add(nodeInfo);
+				}
+			}
+			else
+			{
+				// Write the node if we don't already have it
 				_packetWriter.WriteFixedLengthBytes(data);
 
-				NodeInfo nodeInfo = new NodeInfo(hash, (int)data.Length, nextRefs);
 				_queue.Add(nodeInfo);
 				_hashToNode.Add(hash, default);
 
@@ -205,12 +222,13 @@ namespace EpicGames.Horde.Storage
 		/// <returns></returns>
 		public async Task<NodeLocator> WriteRefAsync(RefName name, TreeNode node, CancellationToken cancellationToken = default)
 		{
-			IoHash hash = await WriteAsync(node, cancellationToken);
+			TreeNodeRef nodeRef = new TreeNodeRef(node);
+			await WriteAsync(nodeRef, cancellationToken);
 			await FlushAsync(cancellationToken);
 
-			NodeLocator locator = GetLocator(hash);
-			await _store.WriteRefTargetAsync(name, locator, cancellationToken);
-			return locator;
+			Debug.Assert(nodeRef.Locator.IsValid());
+			await _store.WriteRefTargetAsync(name, nodeRef.Locator, cancellationToken);
+			return nodeRef.Locator;
 		}
 
 		/// <summary>
@@ -313,24 +331,25 @@ namespace EpicGames.Horde.Storage
 				BlobLocator locator = await _store.WriteBundleAsync(bundle, _prefix, cancellationToken);
 				for (int idx = 0; idx < _queue.Count; idx++)
 				{
+					NodeLocator nodeLocator = new NodeLocator(locator, idx);
+
 					NodeInfo nodeInfo = _queue[idx];
-					_hashToNode[nodeInfo.Hash] = new NodeLocator(locator, idx);
+					nodeInfo.NodeRef.MarkAsClean(_store, nodeInfo.Hash, nodeLocator, nodeInfo.Revision);
+
+					_hashToNode[nodeInfo.Hash] = nodeLocator;
 				}
 
-				// Fire all the alerts for any nodes that have been written
-				List<IoHash> removeFences = new List<IoHash>(_fences.Keys.Where(x => nodeToIndex.ContainsKey(x)));
-				foreach ((IoHash hash, TaskCompletionSource<NodeLocator> tcs) in _fences)
+				// Update any refs with their target locator
+				int refIdx = 0;
+				for (; refIdx < _secondaryQueue.Count; refIdx++)
 				{
-					if (nodeToIndex.ContainsKey(hash))
+					NodeInfo nodeInfo = _secondaryQueue[refIdx];
+					if (_hashToNode.TryGetValue(nodeInfo.Hash, out NodeLocator refLocator))
 					{
-						tcs.SetResult(_hashToNode[hash]);
-						removeFences.Add(hash);
+						nodeInfo.NodeRef.MarkAsClean(_store, nodeInfo.Hash, refLocator, nodeInfo.Revision);
 					}
 				}
-				foreach (IoHash removeFence in removeFences)
-				{
-					_hashToNode.Remove(removeFence);
-				}
+				_secondaryQueue.RemoveRange(0, refIdx);
 
 				// Reset the output state
 				_packets.Clear();
