UnrealBuildTool: File IoHasher with caching based on last write time and file size, with simple saving and loading of the hash database
Add tests for Epicgames.IoHash

#rnx

diff --git a/Engine/Source/Programs/Shared/EpicGames.Build/EpicGames.Build.csproj b/Engine/Source/Programs/Shared/EpicGames.Build/EpicGames.Build.csproj
--- a/Engine/Source/Programs/Shared/EpicGames.Build/EpicGames.Build.csproj
+++ b/Engine/Source/Programs/Shared/EpicGames.Build/EpicGames.Build.csproj
@@ -39,6 +39,7 @@
   </ItemGroup>
   <ItemGroup>
     <ProjectReference Include="..\..\Shared\EpicGames.Core\EpicGames.Core.csproj" />
+    <ProjectReference Include="..\EpicGames.IoHash\EpicGames.IoHash.csproj" />
     <ProjectReference Include="..\EpicGames.MsBuild\EpicGames.MsBuild.csproj" />
   </ItemGroup>
   <ItemGroup>
diff --git a/Engine/Source/Programs/Shared/EpicGames.Build/System/FileHasher.cs b/Engine/Source/Programs/Shared/EpicGames.Build/System/FileHasher.cs
new file mode 100644
--- /dev/null
+++ b/Engine/Source/Programs/Shared/EpicGames.Build/System/FileHasher.cs
@@ -0,0 +1,145 @@
+﻿// Copyright Epic Games, Inc. All Rights Reserved.
+
+using System.Collections.Concurrent;
+using System.Collections.Generic;
+using System.IO;
+using System.Linq;
+using System.Text.Json;
+using System.Text.Json.Serialization;
+using System.Threading;
+using System.Threading.Tasks;
+using EpicGames.Core;
+using Microsoft.Extensions.Logging;
+
+namespace UnrealBuildBase
+{
+	public class FileHasher
+	{
+		private struct CachedDigest
+		{
+			[JsonPropertyName("l")]
+			public long Length { get; init; }
+			[JsonPropertyName("w")]
+			public long LastWriteTimeUtc { get; init; }
+			[JsonPropertyName("d")]
+			public IoHash Digest { get; init; }
+
+			public CachedDigest(FileItem item, IoHash digest)
+			{
+				Length = item.Length;
+				LastWriteTimeUtc = item.LastWriteTimeUtc.Ticks;
+				Digest = digest;
+			}
+
+			/// <summary>
+			/// Check if a digest is up to date for a FileItem
+			/// </summary>
+			/// <param name="item">The FileItem to check</param>
+			/// <returns>If the digest is up to date</returns>
+			public bool UpToDate(FileItem item) => item.Exists && Length == item.Length && LastWriteTimeUtc == item.LastWriteTimeUtc.Ticks;
+		}
+
+		// Dictionary containg cached digests, to prevent rehashing data if the file is unchanged
+		ConcurrentDictionary<string, CachedDigest> CachedDigests = new();
+
+		ILogger? Logger;
+
+		/// <summary>
+		/// FileHasher constructor
+		/// </summary>
+		/// <param name="logger">Optional logger</param>
+		public FileHasher(ILogger? logger = null)
+		{
+			Logger = logger;
+		}
+
+		/// <summary>
+		/// Saves cached digests to disk.
+		/// </summary>
+		/// <param name="location">The location to save the digest data (json)</param>
+		public async Task Save(FileReference location)
+		{
+			using FileStream stream = FileReference.Open(location, FileMode.Create);
+			await JsonSerializer.SerializeAsync(stream, new SortedDictionary<string, CachedDigest>(CachedDigests));
+			Logger?.LogDebug("Saved {Entries} cached digest entries", CachedDigests.Count);
+		}
+
+		/// <summary>
+		/// Loads cached digests from disk, replacing existing cache.
+		/// </summary>
+		/// <param name="location">The location to load the digest data from (json)</param>
+		public async Task Load(FileReference location)
+		{
+			if (!FileReference.Exists(location))
+			{
+				return;
+			}
+
+			using FileStream stream = FileReference.Open(location, FileMode.Open, FileAccess.Read, FileShare.Read);
+			ConcurrentDictionary<string, CachedDigest>? loadedDigests = await JsonSerializer.DeserializeAsync<ConcurrentDictionary<string, CachedDigest>>(stream);
+			if (loadedDigests != null)
+			{
+				CachedDigests = loadedDigests;
+				Logger?.LogDebug("Loaded {Entries} cached digest entries", CachedDigests.Count);
+			}
+		}
+
+		/// <summary>
+		/// Purges stale cached digests.
+		/// </summary>
+		public void PurgeStale()
+		{
+			int count = CachedDigests.Count;
+			IEnumerable<KeyValuePair<string, CachedDigest>> valid = CachedDigests.Where((x) => x.Value.UpToDate(FileItem.GetItemByFileReference(FileReference.FromString(x.Key))));
+			CachedDigests = new(valid);
+			Logger?.LogDebug("Purged {Entries} stale cached digest entries", count - CachedDigests.Count);
+		}
+
+		/// <summary>
+		/// Get the IoHash digest of a file's contents.
+		/// </summary>
+		/// <param name="item">The FileItem to digest</param>
+		/// <param name="cancellationToken">Cancellation token for the operation</param>
+		/// <remarks>A cache is maintained of digests and a file will not be rehashed if unchanged.</remarks>
+		/// <returns>The IoHash digest, or IoHash.Zero if the file doesn't exist.
+		public async Task<IoHash> GetDigestAsync(FileItem item, CancellationToken cancellationToken = default)
+		{
+			if (CachedDigests.TryGetValue(item.FullName, out CachedDigest CachedHash))
+			{
+				if (CachedHash.UpToDate(item))
+				{
+					// Hash already calculated
+					return CachedHash.Digest;
+				}
+			}
+
+			if (!item.Exists)
+			{
+				return IoHash.Zero;
+			}
+
+			CachedDigest digest = await ComputeDigest(item, Logger, cancellationToken);
+			return CachedDigests.AddOrUpdate(item.FullName, digest, (key, oldValue) => digest).Digest;
+		}
+
+		/// <summary>
+		/// Get the IoHash digest of a file's contents.
+		/// </summary>
+		/// <param name="location">The FileReference to digest</param>
+		/// <param name="cancellationToken">Cancellation token for the operation</param>
+		/// <remarks>A cache is maintained of digests and a file will not be rehashed if unchanged.</remarks>
+		/// <returns>The IoHash digest, or IoHash.Zero if the file doesn't exist.
+		public Task<IoHash> GetDigestAsync(FileReference location, CancellationToken cancellationToken = default) => GetDigestAsync(FileItem.GetItemByFileReference(location), cancellationToken);
+
+		public IoHash GetDigest(FileItem item) => GetDigestAsync(item).Result;
+		public IoHash GetDigest(FileReference location) => GetDigestAsync(location).Result;
+
+		static async Task<CachedDigest> ComputeDigest(FileItem item, ILogger? logger, CancellationToken CancellationToken = default)
+		{
+			using FileStream stream = FileReference.Open(item.Location, FileMode.Open, FileAccess.Read, FileShare.Read);
+			CachedDigest digest = new CachedDigest(item, await IoHash.ComputeAsync(stream));
+			logger?.LogDebug("Computed IoHash {Digest} File {Location} Size {Size} LastWrite {LastWrite}", digest.Digest, item.FullName, item.Length, item.LastWriteTimeUtc);
+			return digest;
+		}
+	}
+}
diff --git a/Engine/Source/Programs/Shared/EpicGames.IoHash.Tests/.editorconfig b/Engine/Source/Programs/Shared/EpicGames.IoHash.Tests/.editorconfig
new file mode 100644
--- /dev/null
+++ b/Engine/Source/Programs/Shared/EpicGames.IoHash.Tests/.editorconfig
@@ -0,0 +1,258 @@
+# NOTE: Requires **VS2019 16.3** or later
+
+# EditorConfig is awesome:http://EditorConfig.org
+
+# Dotnet code style settings:
+[*.cs]
+
+# Sort using and Import directives with System.* appearing first
+dotnet_sort_system_directives_first = true
+dotnet_separate_import_directive_groups = false
+
+# Avoid "this." and "Me." if not necessary
+dotnet_style_qualification_for_field = false:error
+dotnet_style_qualification_for_property = false:error
+dotnet_style_qualification_for_method = false:error
+dotnet_style_qualification_for_event = false:error
+
+# Use language keywords instead of framework type names for type references
+dotnet_style_predefined_type_for_locals_parameters_members = true:error
+dotnet_style_predefined_type_for_member_access = false:error
+
+# Suggest more modern language features when available
+dotnet_style_object_initializer = false:suggestion
+dotnet_style_collection_initializer = false:suggestion
+dotnet_style_coalesce_expression = true:suggestion
+dotnet_style_null_propagation = true:suggestion
+dotnet_style_explicit_tuple_names = true:suggestion
+
+# Whitespace options
+dotnet_style_allow_multiple_blank_lines_experimental = false
+dotnet_style_allow_statement_immediately_after_block_experimental = true
+
+# Non-private static fields are PascalCase
+dotnet_naming_rule.non_private_static_fields_should_be_pascal_case.severity = suggestion
+dotnet_naming_rule.non_private_static_fields_should_be_pascal_case.symbols = non_private_static_fields
+dotnet_naming_rule.non_private_static_fields_should_be_pascal_case.style = non_private_static_field_style
+
+dotnet_naming_symbols.non_private_static_fields.applicable_kinds = field
+dotnet_naming_symbols.non_private_static_fields.applicable_accessibilities = public, protected, internal, protected_internal, private_protected
+dotnet_naming_symbols.non_private_static_fields.required_modifiers = static
+
+dotnet_naming_style.non_private_static_field_style.capitalization = pascal_case
+
+# Non-private readonly fields are PascalCase
+dotnet_naming_rule.non_private_readonly_fields_should_be_pascal_case.severity = suggestion
+dotnet_naming_rule.non_private_readonly_fields_should_be_pascal_case.symbols = non_private_readonly_fields
+dotnet_naming_rule.non_private_readonly_fields_should_be_pascal_case.style = non_private_readonly_field_style
+
+dotnet_naming_symbols.non_private_readonly_fields.applicable_kinds = field
+dotnet_naming_symbols.non_private_readonly_fields.applicable_accessibilities = public, protected, internal, protected_internal, private_protected
+dotnet_naming_symbols.non_private_readonly_fields.required_modifiers = readonly
+
+dotnet_naming_style.non_private_readonly_field_style.capitalization = pascal_case
+
+# Constants are PascalCase
+dotnet_naming_rule.constants_should_be_pascal_case.severity = suggestion
+dotnet_naming_rule.constants_should_be_pascal_case.symbols = constants
+dotnet_naming_rule.constants_should_be_pascal_case.style = constant_style
+
+dotnet_naming_symbols.constants.applicable_kinds = field, local
+dotnet_naming_symbols.constants.required_modifiers = const
+
+dotnet_naming_style.constant_style.capitalization = pascal_case
+
+# Static fields are camelCase and start with s_
+dotnet_naming_rule.static_fields_should_be_camel_case.severity = suggestion
+dotnet_naming_rule.static_fields_should_be_camel_case.symbols = static_fields
+dotnet_naming_rule.static_fields_should_be_camel_case.style = static_field_style
+
+dotnet_naming_symbols.static_fields.applicable_kinds = field
+dotnet_naming_symbols.static_fields.required_modifiers = static
+
+dotnet_naming_style.static_field_style.capitalization = camel_case
+dotnet_naming_style.static_field_style.required_prefix = s_
+
+# Instance fields are camelCase and start with _
+dotnet_naming_rule.instance_fields_should_be_camel_case.severity = suggestion
+dotnet_naming_rule.instance_fields_should_be_camel_case.symbols = instance_fields
+dotnet_naming_rule.instance_fields_should_be_camel_case.style = instance_field_style
+
+dotnet_naming_symbols.instance_fields.applicable_kinds = field
+
+dotnet_naming_style.instance_field_style.capitalization = camel_case
+dotnet_naming_style.instance_field_style.required_prefix = _
+
+# Locals and parameters are camelCase
+dotnet_naming_rule.locals_should_be_camel_case.severity = suggestion
+dotnet_naming_rule.locals_should_be_camel_case.symbols = locals_and_parameters
+dotnet_naming_rule.locals_should_be_camel_case.style = camel_case_style
+
+dotnet_naming_symbols.locals_and_parameters.applicable_kinds = parameter, local
+
+dotnet_naming_style.camel_case_style.capitalization = camel_case
+
+# Local functions are PascalCase
+dotnet_naming_rule.local_functions_should_be_pascal_case.severity = suggestion
+dotnet_naming_rule.local_functions_should_be_pascal_case.symbols = local_functions
+dotnet_naming_rule.local_functions_should_be_pascal_case.style = local_function_style
+
+dotnet_naming_symbols.local_functions.applicable_kinds = local_function
+
+dotnet_naming_style.local_function_style.capitalization = pascal_case
+
+# By default, name items with PascalCase
+dotnet_naming_rule.members_should_be_pascal_case.severity = suggestion
+dotnet_naming_rule.members_should_be_pascal_case.symbols = all_members
+dotnet_naming_rule.members_should_be_pascal_case.style = pascal_case_style
+
+dotnet_naming_symbols.all_members.applicable_kinds = *
+
+dotnet_naming_style.pascal_case_style.capitalization = pascal_case
+
+# IDE0005: Remove unnecessary import
+dotnet_diagnostic.IDE0005.severity = none
+
+# IDE0073: File header
+dotnet_diagnostic.IDE0073.severity = warning
+file_header_template = Copyright Epic Games, Inc. All Rights Reserved.
+
+# IDE0035: Remove unreachable code
+dotnet_diagnostic.IDE0035.severity = warning
+
+# IDE0036: Order modifiers
+dotnet_diagnostic.IDE0036.severity = warning
+
+# IDE0043: Format string contains invalid placeholder
+dotnet_diagnostic.IDE0043.severity = warning
+
+# IDE0044: Make field readonly
+dotnet_diagnostic.IDE0044.severity = warning
+
+# IDE0011: Add braces
+csharp_prefer_braces = when_multiline:warning
+# NOTE: We need the below severity entry for Add Braces due to [REDACTED_PROFILE_URL]/roslyn/issues/44201
+dotnet_diagnostic.IDE0011.severity = warning
+
+# IDE0040: Add accessibility modifiers
+dotnet_diagnostic.IDE0040.severity = warning
+
+# CONSIDER: Are IDE0051 and IDE0052 too noisy to be warnings for IDE editing scenarios? Should they be made build-only warnings?
+# IDE0051: Remove unused private member
+dotnet_diagnostic.IDE0051.severity = warning
+
+# IDE0052: Remove unread private member
+dotnet_diagnostic.IDE0052.severity = warning
+
+# IDE0059: Unnecessary assignment to a value
+dotnet_diagnostic.IDE0059.severity = warning
+
+# IDE0060: Remove unused parameter
+dotnet_diagnostic.IDE0060.severity = warning
+
+# CA1012: Abstract types should not have public constructors
+dotnet_diagnostic.CA1012.severity = warning
+
+# CA1822: Make member static
+dotnet_diagnostic.CA1822.severity = warning
+
+# IDE0005: Using directive is unnecessary
+dotnet_diagnostic.IDE0005.severity = warning
+
+# dotnet_style_allow_multiple_blank_lines_experimental
+dotnet_diagnostic.IDE2000.severity = warning
+
+# csharp_style_allow_embedded_statements_on_same_line_experimental
+dotnet_diagnostic.IDE2001.severity = warning
+
+# csharp_style_allow_blank_lines_between_consecutive_braces_experimental
+dotnet_diagnostic.IDE2002.severity = warning
+
+# dotnet_style_allow_statement_immediately_after_block_experimental
+dotnet_diagnostic.IDE2003.severity = warning
+
+# csharp_style_allow_blank_line_after_colon_in_constructor_initializer_experimental
+dotnet_diagnostic.IDE2004.severity = warning
+
+dotnet_style_require_accessibility_modifiers=false
+
+# CSharp code style settings:
+[*.cs]
+csharp_prefer_simple_using_statement=false
+csharp_style_prefer_range_operator=false
+
+# Newline settings
+csharp_new_line_before_open_brace = all
+csharp_new_line_before_else = true
+csharp_new_line_before_catch = true
+csharp_new_line_before_finally = true
+csharp_new_line_before_members_in_object_initializers = true
+csharp_new_line_before_members_in_anonymous_types = true
+csharp_new_line_between_query_expression_clauses = true
+
+# Indentation preferences
+csharp_indent_block_contents = true
+csharp_indent_braces = false
+csharp_indent_case_contents = true
+csharp_indent_case_contents_when_block = true
+csharp_indent_switch_labels = true
+csharp_indent_labels = flush_left
+
+# Whitespace options
+csharp_style_allow_embedded_statements_on_same_line_experimental = false
+csharp_style_allow_blank_lines_between_consecutive_braces_experimental = false
+csharp_style_allow_blank_line_after_colon_in_constructor_initializer_experimental = false
+
+# Prefer "var" everywhere
+dotnet_diagnostic.IDE0007.severity = error
+csharp_style_var_for_built_in_types = false:error
+csharp_style_var_when_type_is_apparent = false:error
+csharp_style_var_elsewhere = false:error
+
+# Prefer method-like constructs to have a block body
+csharp_style_expression_bodied_methods = false:none
+csharp_style_expression_bodied_constructors = false:none
+csharp_style_expression_bodied_operators = false:none
+
+# Prefer property-like constructs to have an expression-body
+csharp_style_expression_bodied_properties = true:error
+csharp_style_expression_bodied_indexers = true:error
+csharp_style_expression_bodied_accessors = true:error
+
+# Suggest more modern language features when available
+csharp_style_pattern_matching_over_is_with_cast_check = true:error
+csharp_style_pattern_matching_over_as_with_null_check = false:error
+csharp_style_inlined_variable_declaration = false:suggestion
+csharp_style_throw_expression = true:error
+csharp_style_conditional_delegate_call = true:suggestion
+
+# Spacing
+csharp_space_after_cast = false
+csharp_space_after_colon_in_inheritance_clause = true
+csharp_space_after_comma = true
+csharp_space_after_dot = false
+csharp_space_after_keywords_in_control_flow_statements = true
+csharp_space_after_semicolon_in_for_statement = true
+csharp_space_around_binary_operators = before_and_after
+csharp_space_around_declaration_statements = do_not_ignore
+csharp_space_before_colon_in_inheritance_clause = true
+csharp_space_before_comma = false
+csharp_space_before_dot = false
+csharp_space_before_open_square_brackets = false
+csharp_space_before_semicolon_in_for_statement = false
+csharp_space_between_empty_square_brackets = false
+csharp_space_between_method_call_empty_parameter_list_parentheses = false
+csharp_space_between_method_call_name_and_opening_parenthesis = false
+csharp_space_between_method_call_parameter_list_parentheses = false
+csharp_space_between_method_declaration_empty_parameter_list_parentheses = false
+csharp_space_between_method_declaration_name_and_open_parenthesis = false
+csharp_space_between_method_declaration_parameter_list_parentheses = false
+csharp_space_between_parentheses = false
+csharp_space_between_square_brackets = false
+
+# Blocks are allowed
+csharp_prefer_braces = true:silent
+csharp_preserve_single_line_blocks = true
+csharp_preserve_single_line_statements = true
+
+
diff --git a/Engine/Source/Programs/Shared/EpicGames.IoHash.Tests/Blake3HashTests.cs b/Engine/Source/Programs/Shared/EpicGames.IoHash.Tests/Blake3HashTests.cs
new file mode 100644
--- /dev/null
+++ b/Engine/Source/Programs/Shared/EpicGames.IoHash.Tests/Blake3HashTests.cs
@@ -0,0 +1,54 @@
+﻿// Copyright Epic Games, Inc. All Rights Reserved.
+
+using System;
+using System.Buffers;
+using System.IO;
+using System.Text;
+using System.Threading.Tasks;
+using Microsoft.VisualStudio.TestTools.UnitTesting;
+
+namespace EpicGames.IoHash.Tests
+{
+	[TestClass]
+	public class Blake3HashTests
+	{
+		static readonly ReadOnlyMemory<byte> s_testData = Encoding.ASCII.GetBytes("The quick brown fox jumped over the lazy dog.");
+		static readonly Core.Blake3Hash s_testHash = Core.Blake3Hash.Parse("7b0142eafb4d0bed8265978281a1398099b1d727acf4c4b8158b407b1e8385b3");
+
+		[TestMethod]
+		public void ComputeSpan()
+		{
+			Assert.AreEqual(Core.Blake3Hash.Compute(s_testData.Span), s_testHash);
+		}
+
+		[TestMethod]
+		public void ComputeSequence()
+		{
+			Assert.AreEqual(Core.Blake3Hash.Compute(new ReadOnlySequence<byte>(s_testData)), s_testHash);
+		}
+
+		[TestMethod]
+		public void ComputeStream()
+		{
+			using MemoryStream stream = new MemoryStream(s_testData.Length);
+			stream.Write(s_testData.Span);
+			stream.Position = 0;
+			Assert.AreEqual(Core.Blake3Hash.Compute(stream), s_testHash);
+		}
+
+		[TestMethod]
+		public async Task ComputeStreamAsync()
+		{
+			using MemoryStream stream = new MemoryStream(s_testData.Length);
+			stream.Write(s_testData.Span);
+			stream.Position = 0;
+			Assert.AreEqual(await Core.Blake3Hash.ComputeAsync(stream), s_testHash);
+		}
+
+		[TestMethod]
+		public void Parse()
+		{
+			Assert.AreEqual(Core.Blake3Hash.Parse(s_testHash.ToString()), s_testHash);
+		}
+	}
+}
diff --git a/Engine/Source/Programs/Shared/EpicGames.IoHash.Tests/EpicGames.IoHash.Tests.csproj b/Engine/Source/Programs/Shared/EpicGames.IoHash.Tests/EpicGames.IoHash.Tests.csproj
new file mode 100644
--- /dev/null
+++ b/Engine/Source/Programs/Shared/EpicGames.IoHash.Tests/EpicGames.IoHash.Tests.csproj
@@ -0,0 +1,24 @@
+<Project Sdk="Microsoft.NET.Sdk">
+
+  <PropertyGroup>
+    <TargetFramework>net6.0</TargetFramework>
+    <Configurations>Debug;Release;Development</Configurations>
+    <Nullable>enable</Nullable>
+    <IsPackable>false</IsPackable>
+  </PropertyGroup>
+
+  <ItemGroup>
+    <PackageReference Include="Microsoft.NET.Test.Sdk" Version="17.3.0" />
+    <PackageReference Include="MSTest.TestAdapter" Version="2.2.10" />
+    <PackageReference Include="MSTest.TestFramework" Version="2.2.10" />
+    <PackageReference Include="coverlet.collector" Version="3.1.2">
+      <PrivateAssets>all</PrivateAssets>
+      <IncludeAssets>runtime; build; native; contentfiles; analyzers; buildtransitive</IncludeAssets>
+    </PackageReference>
+  </ItemGroup>
+
+  <ItemGroup>
+    <ProjectReference Include="..\EpicGames.IoHash\EpicGames.IoHash.csproj" />
+  </ItemGroup>
+
+</Project>
diff --git a/Engine/Source/Programs/Shared/EpicGames.IoHash.Tests/IoHashTests.cs b/Engine/Source/Programs/Shared/EpicGames.IoHash.Tests/IoHashTests.cs
new file mode 100644
--- /dev/null
+++ b/Engine/Source/Programs/Shared/EpicGames.IoHash.Tests/IoHashTests.cs
@@ -0,0 +1,75 @@
+﻿// Copyright Epic Games, Inc. All Rights Reserved.
+
+using System;
+using System.Buffers;
+using System.IO;
+using System.Text;
+using System.Threading.Tasks;
+using Microsoft.VisualStudio.TestTools.UnitTesting;
+
+namespace EpicGames.IoHash.Tests
+{
+	[TestClass]
+	public class IoHashTests
+	{
+		static readonly ReadOnlyMemory<byte> s_testData = Encoding.ASCII.GetBytes("The quick brown fox jumped over the lazy dog.");
+		static readonly Core.IoHash s_testHash = Core.IoHash.Parse("7b0142eafb4d0bed8265978281a1398099b1d727");
+
+		[TestMethod]
+		public void ComputeSpan()
+		{
+			Assert.AreEqual(Core.IoHash.Compute(s_testData.Span), s_testHash);
+		}
+
+		[TestMethod]
+		public void ComputeSequence()
+		{
+			Assert.AreEqual(Core.IoHash.Compute(new ReadOnlySequence<byte>(s_testData)), s_testHash);
+		}
+
+		[TestMethod]
+		public void ComputeStream()
+		{
+			using MemoryStream stream = new MemoryStream(s_testData.Length);
+			stream.Write(s_testData.Span);
+			stream.Position = 0;
+			Assert.AreEqual(Core.IoHash.Compute(stream), s_testHash);
+		}
+
+		[TestMethod]
+		public async Task ComputeStreamAsync()
+		{
+			using MemoryStream stream = new MemoryStream(s_testData.Length);
+			stream.Write(s_testData.Span);
+			stream.Position = 0;
+			Assert.AreEqual(await Core.IoHash.ComputeAsync(stream), s_testHash);
+		}
+
+		[TestMethod]
+		public void FromBlake3Hash()
+		{
+			Assert.AreEqual(new Core.IoHash(Core.Blake3Hash.Compute(s_testData.Span).Span), s_testHash);
+		}
+
+		[TestMethod]
+		public void Parse()
+		{
+			Assert.AreEqual(Core.IoHash.Parse(s_testHash.ToString()), s_testHash);
+			Assert.AreEqual(Core.IoHash.Parse(Encoding.ASCII.GetBytes(s_testHash.ToString())), s_testHash);
+		}
+
+		[TestMethod]
+		public void TryParse()
+		{
+			Core.IoHash result;
+			Assert.IsTrue(Core.IoHash.TryParse(s_testHash.ToString(), out result));
+			Assert.AreEqual(result, s_testHash);
+
+			Assert.IsTrue(Core.IoHash.TryParse(Encoding.ASCII.GetBytes(s_testHash.ToString()), out result));
+			Assert.AreEqual(result, s_testHash);
+
+			Assert.IsFalse(Core.IoHash.TryParse(String.Empty, out _));
+			Assert.IsFalse(Core.IoHash.TryParse(Encoding.ASCII.GetBytes(String.Empty), out _));
+		}
+	}
+}
diff --git a/Engine/Source/Programs/Shared/EpicGames.IoHash/Blake3Hash.cs b/Engine/Source/Programs/Shared/EpicGames.IoHash/Blake3Hash.cs
--- a/Engine/Source/Programs/Shared/EpicGames.IoHash/Blake3Hash.cs
+++ b/Engine/Source/Programs/Shared/EpicGames.IoHash/Blake3Hash.cs
@@ -1,7 +1,11 @@
 // Copyright Epic Games, Inc. All Rights Reserved.
 
 using System;
+using System.Buffers;
 using System.Buffers.Binary;
+using System.IO;
+using System.Threading.Tasks;
+using System.Threading;
 
 namespace EpicGames.Core
 {
@@ -50,7 +54,7 @@ namespace EpicGames.Core
 		}
 
 		/// <summary>
-		/// Creates a content hash for a block of data, using a given algorithm.
+		/// Creates a content hash for a block of data.
 		/// </summary>
 		/// <param name="data">Data to compute the hash for</param>
 		/// <returns>New content hash instance containing the hash of the data</returns>
@@ -61,6 +65,75 @@ namespace EpicGames.Core
 			return new Blake3Hash(output);
 		}
 
+		/// <summary>
+		/// Creates a content hash for a block of data.
+		/// </summary>
+		/// <param name="sequence">Data to compute the hash for</param>
+		/// <returns>New content hash instance containing the hash of the data</returns>
+		public static IoHash Compute(ReadOnlySequence<byte> sequence)
+		{
+			if (sequence.IsSingleSegment)
+			{
+				return Compute(sequence.FirstSpan);
+			}
+
+			using (Blake3.Hasher hasher = Blake3.Hasher.New())
+			{
+				foreach (ReadOnlyMemory<byte> segment in sequence)
+				{
+					hasher.Update(segment.Span);
+				}
+
+				byte[] output = new byte[32];
+				hasher.Finalize(output);
+				return new Blake3Hash(output);
+			}
+		}
+
+		/// <summary>
+		/// Creates a content hash for a stream.
+		/// </summary>
+		/// <param name="stream">Data to compute the hash for</param>
+		/// <returns>New content hash instance containing the hash of the data</returns>
+		public static IoHash Compute(Stream stream)
+		{
+			using (Blake3.Hasher hasher = Blake3.Hasher.New())
+			{
+				Span<byte> buffer = new byte[16384];
+				int length;
+				while ((length = stream.Read(buffer)) > 0)
+				{
+					hasher.Update(buffer.Slice(0, length));
+				}
+
+				byte[] output = new byte[32];
+				hasher.Finalize(output);
+				return new Blake3Hash(output);
+			}
+		}
+
+		/// <summary>
+		/// Creates a content hash for a stream asynchronously. 
+		/// </summary>
+		/// <param name="stream">Data to compute the hash for</param>
+		/// <returns>New content hash instance containing the hash of the data</returns>
+		public static async Task<IoHash> ComputeAsync(Stream stream, CancellationToken cancellationToken = default)
+		{
+			using (Blake3.Hasher hasher = Blake3.Hasher.New())
+			{
+				Memory<byte> buffer = new byte[16384];
+				int length;
+				while ((length = await stream.ReadAsync(buffer, cancellationToken)) > 0)
+				{
+					hasher.Update(buffer.Slice(0, length).Span);
+				}
+
+				byte[] output = new byte[32];
+				hasher.Finalize(output);
+				return new Blake3Hash(output);
+			}
+		}
+
 		/// <summary>
 		/// Parses a digest from the given hex string
 		/// </summary>
diff --git a/Engine/Source/Programs/Shared/EpicGames.IoHash/IoHash.cs b/Engine/Source/Programs/Shared/EpicGames.IoHash/IoHash.cs
--- a/Engine/Source/Programs/Shared/EpicGames.IoHash/IoHash.cs
+++ b/Engine/Source/Programs/Shared/EpicGames.IoHash/IoHash.cs
@@ -7,6 +7,9 @@ using System.ComponentModel;
 using System.Globalization;
 using System.Text.Json.Serialization;
 using System.Text.Json;
+using System.IO;
+using System.Threading.Tasks;
+using System.Threading;
 
 namespace EpicGames.Core
 {
@@ -100,6 +103,44 @@ namespace EpicGames.Core
 			}
 		}
 
+		/// <summary>
+		/// Creates the IoHash for a stream.
+		/// </summary>
+		/// <param name="stream">Data to compute the hash for</param>
+		/// <returns>New content hash instance containing the hash of the data</returns>
+		public static IoHash Compute(Stream stream)
+		{
+			using (Blake3.Hasher hasher = Blake3.Hasher.New())
+			{
+				Span<byte> buffer = stackalloc byte[16384];
+				int length;
+				while ((length = stream.Read(buffer)) > 0)
+				{
+					hasher.Update(buffer.Slice(0, length));
+				}
+				return FromBlake3(hasher);
+			}
+		}
+
+		/// <summary>
+		/// Creates the IoHash for a stream asynchronously. 
+		/// </summary>
+		/// <param name="stream">Data to compute the hash for</param>
+		/// <returns>New content hash instance containing the hash of the data</returns>
+		public static async Task<IoHash> ComputeAsync(Stream stream, CancellationToken cancellationToken = default)
+		{
+			using (Blake3.Hasher hasher = Blake3.Hasher.New())
+			{
+				Memory<byte> buffer = new byte[16384];
+				int length;
+				while ((length = await stream.ReadAsync(buffer, cancellationToken)) > 0)
+				{
+					hasher.Update(buffer.Slice(0, length).Span);
+				}
+				return FromBlake3(hasher);
+			}
+		}
+
 		/// <summary>
 		/// Parses a digest from the given hex string
 		/// </summary>
diff --git a/Engine/Source/Programs/UnrealBuildTool/Executors/Experimental/HordeExecutor.cs b/Engine/Source/Programs/UnrealBuildTool/Executors/Experimental/HordeExecutor.cs
--- a/Engine/Source/Programs/UnrealBuildTool/Executors/Experimental/HordeExecutor.cs
+++ b/Engine/Source/Programs/UnrealBuildTool/Executors/Experimental/HordeExecutor.cs
@@ -163,23 +163,26 @@ namespace UnrealBuildTool
 
 		internal class ZenUploader
 		{
-			Dictionary<string, IoHash> FileHashes = new Dictionary<string, IoHash>();
-
+			FileHasher FileHasher;
 			Dictionary<IoHash, byte[]> Datas = new Dictionary<IoHash, byte[]>();
 			HashSet<IoHash> Uploaded = new HashSet<IoHash>();
 
+			public ZenUploader(ILogger? Logger = null)
+			{
+				FileHasher = new(Logger);
+			}
+
 			public IoHash HashFile(FileReference File)
 			{
-				if (!FileHashes.ContainsKey(File.FullName))
+				IoHash Hash = FileHasher.GetDigest(File);
+				if (!Datas.ContainsKey(Hash))
 				{
-					byte[] Data = FileReference.ReadAllBytes(File);
-					IoHash Hash = HashAndStoreData(Data);
-					lock (FileHashes)
+					lock (Datas)
 					{
-						FileHashes.TryAdd(File.FullName, Hash);
+						Datas.TryAdd(Hash, FileReference.ReadAllBytes(File));
 					}
 				}
-				return FileHashes[File.FullName];
+				return Hash;
 			}
 
 			// HashString is not cached
