UnrealBuildTool: Fix analysis warnings

#rnx

diff --git a/Engine/Source/Programs/Shared/EpicGames.Build/System/FileHasher.cs b/Engine/Source/Programs/Shared/EpicGames.Build/System/FileHasher.cs
--- a/Engine/Source/Programs/Shared/EpicGames.Build/System/FileHasher.cs
+++ b/Engine/Source/Programs/Shared/EpicGames.Build/System/FileHasher.cs
@@ -101,7 +101,7 @@ namespace UnrealBuildBase
 		/// <param name="item">The FileItem to digest</param>
 		/// <param name="cancellationToken">Cancellation token for the operation</param>
 		/// <remarks>A cache is maintained of digests and a file will not be rehashed if unchanged.</remarks>
-		/// <returns>The IoHash digest, or IoHash.Zero if the file doesn't exist.
+		/// <returns>The IoHash digest, or IoHash.Zero if the file doesn't exist.</returns>
 		public async Task<IoHash> GetDigestAsync(FileItem item, CancellationToken cancellationToken = default)
 		{
 			if (CachedDigests.TryGetValue(item.FullName, out CachedDigest CachedHash))
@@ -128,7 +128,7 @@ namespace UnrealBuildBase
 		/// <param name="location">The FileReference to digest</param>
 		/// <param name="cancellationToken">Cancellation token for the operation</param>
 		/// <remarks>A cache is maintained of digests and a file will not be rehashed if unchanged.</remarks>
-		/// <returns>The IoHash digest, or IoHash.Zero if the file doesn't exist.
+		/// <returns>The IoHash digest, or IoHash.Zero if the file doesn't exist.</returns>
 		public Task<IoHash> GetDigestAsync(FileReference location, CancellationToken cancellationToken = default) => GetDigestAsync(FileItem.GetItemByFileReference(location), cancellationToken);
 
 		public IoHash GetDigest(FileItem item) => GetDigestAsync(item).Result;
@@ -137,7 +137,7 @@ namespace UnrealBuildBase
 		static async Task<CachedDigest> ComputeDigest(FileItem item, ILogger? logger, CancellationToken CancellationToken = default)
 		{
 			using FileStream stream = FileReference.Open(item.Location, FileMode.Open, FileAccess.Read, FileShare.Read);
-			CachedDigest digest = new CachedDigest(item, await IoHash.ComputeAsync(stream));
+			CachedDigest digest = new CachedDigest(item, await IoHash.ComputeAsync(stream, CancellationToken));
 			logger?.LogDebug("Computed IoHash {Digest} File {Location} Size {Size} LastWrite {LastWrite}", digest.Digest, item.FullName, item.Length, item.LastWriteTimeUtc);
 			return digest;
 		}
