EpicGames.Perforce: Add support for "p4 sizes" command. #9584

diff --git a/Engine/Source/Programs/Shared/EpicGames.Perforce/PerforceConnection.cs b/Engine/Source/Programs/Shared/EpicGames.Perforce/PerforceConnection.cs
--- a/Engine/Source/Programs/Shared/EpicGames.Perforce/PerforceConnection.cs
+++ b/Engine/Source/Programs/Shared/EpicGames.Perforce/PerforceConnection.cs
@@ -3233,6 +3233,92 @@ namespace EpicGames.Perforce
 
 		#endregion
 
+		#region p4 sizes
+
+		/// <summary>
+		/// Execute the 'sizes' command
+		/// </summary>
+		/// <param name="connection">Connection to the Perforce server</param>
+		/// <param name="options">Options for the command</param>
+		/// <param name="fileSpecs">List of file specifications to query</param>
+		/// <param name="cancellationToken">Token used to cancel the operation</param>
+		/// <returns>List of response objects</returns>
+		public static async Task<List<SizesRecord>> SizesAsync(this IPerforceConnection connection, SizesOptions options, FileSpecList fileSpecs, CancellationToken cancellationToken = default)
+		{
+			return (await TrySizesAsync(connection, options, -1, fileSpecs, cancellationToken)).Data;
+		}
+
+		/// <summary>
+		/// Execute the 'sizes' command
+		/// </summary>
+		/// <param name="connection">Connection to the Perforce server</param>
+		/// <param name="options">Options for the command</param>
+		/// <param name="fileSpecs">List of file specifications to query</param>
+		/// <param name="cancellationToken">Token used to cancel the operation</param>
+		/// <returns>List of response objects</returns>
+		public static Task<PerforceResponseList<SizesRecord>> TrySizesAsync(this IPerforceConnection connection, SizesOptions options, FileSpecList fileSpecs, CancellationToken cancellationToken = default)
+		{
+			return TrySizesAsync(connection, options, -1, fileSpecs, cancellationToken);
+		}
+
+		/// <summary>
+		/// Execute the 'sizes' command
+		/// </summary>
+		/// <param name="connection">Connection to the Perforce server</param>
+		/// <param name="options">Options for the command</param>
+		/// <param name="maxLines">Maximum number of results to return. Ignored if less than or equal to zero.</param>
+		/// <param name="fileSpecs">List of file specifications to query</param>
+		/// <param name="cancellationToken">Token used to cancel the operation</param>
+		/// <returns>List of response objects</returns>
+		public static async Task<List<SizesRecord>> SizesAsync(this IPerforceConnection connection, SizesOptions options, int maxLines, FileSpecList fileSpecs, CancellationToken cancellationToken = default)
+		{
+			return (await TrySizesAsync(connection, options, maxLines, fileSpecs, cancellationToken)).Data;
+		}
+
+		/// <summary>
+		/// Execute the 'sizes' command
+		/// </summary>
+		/// <param name="connection">Connection to the Perforce server</param>
+		/// <param name="options">Options for the command</param>
+		/// <param name="maxLines">Maximum number of results to return. Ignored if less than or equal to zero.</param>
+		/// <param name="fileSpecs">List of file specifications to query</param>
+		/// <param name="cancellationToken">Token used to cancel the operation</param>
+		/// <returns>List of response objects</returns>
+		public static async Task<PerforceResponseList<SizesRecord>> TrySizesAsync(this IPerforceConnection connection, SizesOptions options, int maxLines, FileSpecList fileSpecs, CancellationToken cancellationToken = default)
+		{
+			List<string> arguments = new List<string>();
+			if ((options & SizesOptions.AllRevisions) != 0)
+			{
+				arguments.Add("-a");
+			}
+			if ((options & SizesOptions.LimitToArchiveDepots) != 0)
+			{
+				arguments.Add("-A");
+			}
+			if ((options & SizesOptions.CalculateSum) != 0)
+			{
+				arguments.Add("-s");
+			}
+			if ((options & SizesOptions.DisplayForShelvedFilesOnly) != 0)
+			{
+				arguments.Add("-S");
+			}
+            if ((options & SizesOptions.ExcludeLazyCopies) != 0)
+            {
+                arguments.Add("-z");
+            }
+            if (maxLines > 0)
+			{
+				arguments.Add($"-m{maxLines}");
+			}
+
+			PerforceResponseList<SizesRecord> records = await BatchedCommandAsync<SizesRecord>(connection, "sizes", arguments, fileSpecs.List, cancellationToken);
+			records.RemoveAll(x => x.Error != null && x.Error.Generic == PerforceGenericCode.Empty);
+			return records;
+		}
+
+		#endregion
+
 		#region p4 stream
 
 		/// <summary>
diff --git a/Engine/Source/Programs/Shared/EpicGames.Perforce/Records/SizeOptions.cs b/Engine/Source/Programs/Shared/EpicGames.Perforce/Records/SizeOptions.cs
new file mode 100644
--- /dev/null
+++ b/Engine/Source/Programs/Shared/EpicGames.Perforce/Records/SizeOptions.cs
@@ -0,0 +1,43 @@
+// Copyright Epic Games, Inc. All Rights Reserved.
+
+using System;
+
+namespace EpicGames.Perforce
+{
+	/// <summary>
+	/// Options for the sizes command
+	/// </summary>
+	[Flags]
+	public enum SizesOptions
+	{
+		/// <summary>
+		/// No options are set
+		/// </summary>
+		None = 0,
+
+        /// <summary>
+        /// For each file, list all revisions within a specified revision range, rather than only the highest revision in the range.
+        /// </summary>
+        AllRevisions = 1,
+
+        /// <summary>
+        /// Display files in archive depots.
+        /// </summary>
+        LimitToArchiveDepots = 2,
+
+        /// <summary>
+        /// Calculate the sum of the file sizes for the specified file argument.
+        /// </summary>
+        CalculateSum = 4,
+
+        /// <summary>
+        /// Display size information for shelved files only. If you use this option, revision specifications are not permitted.
+        /// </summary>
+        DisplayForShelvedFilesOnly = 8,
+
+        /// <summary>
+        /// When calculating size information, exclude lazy copies.
+        /// </summary>
+        ExcludeLazyCopies = 16,
+    }
+}
diff --git a/Engine/Source/Programs/Shared/EpicGames.Perforce/Records/SizesRecord.cs b/Engine/Source/Programs/Shared/EpicGames.Perforce/Records/SizesRecord.cs
new file mode 100644
--- /dev/null
+++ b/Engine/Source/Programs/Shared/EpicGames.Perforce/Records/SizesRecord.cs
@@ -0,0 +1,37 @@
+// Copyright Epic Games, Inc. All Rights Reserved.
+
+using System;
+using System.Collections.Generic;
+
+namespace EpicGames.Perforce
+{
+    /// <summary>
+    /// Record returned by the sizes command
+    /// </summary>
+    public class SizesRecord
+	{
+        /// <summary>
+        /// Depot path to file
+        /// </summary>
+        [PerforceTag("depotFile", Optional = true)]
+        public string? DepotFile { get; set; }
+
+        /// <summary>
+        /// Revision for opened file(s)
+        /// </summary>
+        [PerforceTag("rev", Optional = true)]
+        public int Revision { get; set; }
+
+        /// <summary>
+        /// Size for file(s)
+        /// </summary>
+        [PerforceTag("fileSize")]
+        public long Size { get; set; }
+
+        /// <summary>
+        /// Count of file(s)
+        /// </summary>
+        [PerforceTag("fileCount", Optional = true)]
+        public int Count { get; set; }
+    }
+}
