Horde: Add a read queue to StorageClientBase, and add support for concurrent reads.

* Background task creates up to 4 concurrent read threads.
* When reading a bundle header, we read a fixed 1mb chunk of data from the start of the bundle and add any encoded packets within that range to the cache.
* When reading a bundle packet, we also read more data that specifically requested, and add that to the cache.
* Packet decode requests are also performed asynchronously and shared between requests.

diff --git a/Engine/Source/Programs/Horde/Horde.Agent/Commands/Bundles/CreateCommand.cs b/Engine/Source/Programs/Horde/Horde.Agent/Commands/Bundles/CreateCommand.cs
--- a/Engine/Source/Programs/Horde/Horde.Agent/Commands/Bundles/CreateCommand.cs
+++ b/Engine/Source/Programs/Horde/Horde.Agent/Commands/Bundles/CreateCommand.cs
@@ -23,7 +23,7 @@ namespace Horde.Agent.Commands.Bundles
 			readonly ILogger _logger;
 
 			public FileBlobStore(DirectoryReference rootDir, IMemoryCache cache, ILogger logger)
-				: base(cache)
+				: base(cache, logger)
 			{
 				_rootDir = rootDir;
 				_logger = logger;
@@ -40,7 +40,7 @@ namespace Horde.Agent.Commands.Bundles
 			{
 				FileReference file = GetBlobFile(id);
 				_logger.LogInformation("Reading {File}", file);
-				return Task.FromResult<Stream>(FileReference.Open(file, FileMode.Open, FileAccess.Read));
+				return Task.FromResult<Stream>(FileReference.Open(file, FileMode.Open, FileAccess.Read, FileShare.Read));
 			}
 
 			public override async Task<BlobLocator> WriteBlobAsync(Stream stream, Utf8String prefix = default, CancellationToken cancellationToken = default)
diff --git a/Engine/Source/Programs/Horde/Horde.Agent/Commands/Bundles/ExtractCommand.cs b/Engine/Source/Programs/Horde/Horde.Agent/Commands/Bundles/ExtractCommand.cs
--- a/Engine/Source/Programs/Horde/Horde.Agent/Commands/Bundles/ExtractCommand.cs
+++ b/Engine/Source/Programs/Horde/Horde.Agent/Commands/Bundles/ExtractCommand.cs
@@ -1,5 +1,6 @@
 // Copyright Epic Games, Inc. All Rights Reserved.
 
+using System.Diagnostics;
 using System.Threading;
 using System.Threading.Tasks;
 using EpicGames.Core;
@@ -24,9 +25,12 @@ namespace Horde.Agent.Commands.Bundles
 			using IMemoryCache cache = new MemoryCache(new MemoryCacheOptions());
 			IStorageClient store = CreateStorageClient(cache, logger);
 
+			Stopwatch timer = Stopwatch.StartNew();
+
 			DirectoryNode node = await store.ReadNodeAsync<DirectoryNode>(RefName);
 			await node.CopyToDirectoryAsync(OutputDir.ToDirectoryInfo(), logger, CancellationToken.None);
 
+			logger.LogInformation("Elapsed: {Time}s", timer.Elapsed.Seconds);
 			return 0;
 		}
 	}
diff --git a/Engine/Source/Programs/Horde/Horde.Build.Tests/BlobStoreTests.cs b/Engine/Source/Programs/Horde/Horde.Build.Tests/BlobStoreTests.cs
--- a/Engine/Source/Programs/Horde/Horde.Build.Tests/BlobStoreTests.cs
+++ b/Engine/Source/Programs/Horde/Horde.Build.Tests/BlobStoreTests.cs
@@ -11,6 +11,7 @@ using EpicGames.Horde.Storage;
 using System.Buffers;
 using EpicGames.Horde.Storage.Backends;
 using Microsoft.Extensions.Caching.Memory;
+using Microsoft.Extensions.Logging.Abstractions;
 
 namespace Horde.Build.Tests
 {
@@ -23,7 +24,7 @@ namespace Horde.Build.Tests
 		IStorageClient CreateBlobStore()
 		{
 			_cache ??= new MemoryCache(new MemoryCacheOptions());
-			_store ??= new InMemoryBlobStore(_cache);
+			_store ??= new InMemoryBlobStore(_cache, NullLogger.Instance);
 			return _store;
 		}
 
diff --git a/Engine/Source/Programs/Horde/Horde.Build/Storage/StorageService.cs b/Engine/Source/Programs/Horde/Horde.Build/Storage/StorageService.cs
--- a/Engine/Source/Programs/Horde/Horde.Build/Storage/StorageService.cs
+++ b/Engine/Source/Programs/Horde/Horde.Build/Storage/StorageService.cs
@@ -53,8 +53,8 @@ namespace Horde.Build.Storage
 			readonly IStorageBackend _backend;
 			readonly IRefCollection _refs;
 
-			public StorageClient(NamespaceConfig config, IStorageBackend backend, IRefCollection refs, IMemoryCache cache)
-				: base(cache)
+			public StorageClient(NamespaceConfig config, IStorageBackend backend, IRefCollection refs, IMemoryCache cache, ILogger logger)
+				: base(cache, logger)
 			{
 				Config = config;
 				_backend = backend;
@@ -153,6 +153,7 @@ namespace Horde.Build.Storage
 		readonly IRefCollection _refCollection;
 		readonly IMemoryCache _cache;
 		readonly IServiceProvider _serviceProvider;
+		readonly ILogger _logger;
 
 		State? _lastState;
 		string? _lastConfigRevision;
@@ -164,14 +165,15 @@ namespace Horde.Build.Storage
 		/// <summary>
 		/// Constructor
 		/// </summary>
-		public StorageService(GlobalsService globalsService, AclService aclService, IRefCollection refCollection, IMemoryCache cache, IServiceProvider serviceProvider)
+		public StorageService(GlobalsService globalsService, AclService aclService, IRefCollection refCollection, IMemoryCache cache, IServiceProvider serviceProvider, ILogger<StorageService> logger)
 		{
 			_globalsService = globalsService;
 			_aclService = aclService;
 			_refCollection = refCollection;
 			_cache = cache;
 			_serviceProvider = serviceProvider;
-			_cachedState = new AsyncCachedValue<State>(() => GetNextState(), TimeSpan.FromMinutes(1.0)); 
+			_cachedState = new AsyncCachedValue<State>(() => GetNextState(), TimeSpan.FromMinutes(1.0));
+			_logger = logger;
 		}
 
 		/// <inheritdoc/>
@@ -262,7 +264,7 @@ namespace Horde.Build.Storage
 						IStorageBackend backend = CreateStorageBackend(backendConfig);
 
 #pragma warning disable CA2000 // Dispose objects before losing scope
-						StorageClient client = new StorageClient(namespaceConfig, backend, _refCollection, _cache);
+						StorageClient client = new StorageClient(namespaceConfig, backend, _refCollection, _cache, _logger);
 						nextState.Namespaces.Add(namespaceConfig.Id, new NamespaceInfo(namespaceConfig, client));
 #pragma warning restore CA2000 // Dispose objects before losing scope
 					}
diff --git a/Engine/Source/Programs/Shared/EpicGames.Core/AsyncUtils.cs b/Engine/Source/Programs/Shared/EpicGames.Core/AsyncUtils.cs
--- a/Engine/Source/Programs/Shared/EpicGames.Core/AsyncUtils.cs
+++ b/Engine/Source/Programs/Shared/EpicGames.Core/AsyncUtils.cs
@@ -32,6 +32,21 @@ namespace EpicGames.Core
 			return Task.Delay(-1, token).ContinueWith(_ => Task.FromCanceled<T>(token)).Unwrap();
 		}
 
+		/// <summary>
+		/// Returns a task that will be abandoned if a cancellation token is activated. This differs from the normal cancellation pattern in that the task will run to completion, but waiting for it can be cancelled.
+		/// </summary>
+		/// <param name="task">Task to wait for</param>
+		/// <param name="cancellationToken">Cancellation token for the operation</param>
+		/// <returns>Wrapped task</returns>
+		public static async Task<T> AbandonOnCancel<T>(this Task<T> task, CancellationToken cancellationToken)
+		{
+			if (cancellationToken.CanBeCanceled)
+			{
+				await await Task.WhenAny(task, Task.Delay(-1, cancellationToken)); // Double await to ensure cancellation exception is rethrown if returned
+			}
+			return await task;
+		}
+
 		/// <summary>
 		/// Attempts to get the result of a task, if it has finished
 		/// </summary>
diff --git a/Engine/Source/Programs/Shared/EpicGames.Horde.Tests/BundleStoreTests.cs b/Engine/Source/Programs/Shared/EpicGames.Horde.Tests/BundleStoreTests.cs
--- a/Engine/Source/Programs/Shared/EpicGames.Horde.Tests/BundleStoreTests.cs
+++ b/Engine/Source/Programs/Shared/EpicGames.Horde.Tests/BundleStoreTests.cs
@@ -15,6 +15,7 @@ using System.IO;
 using System.Security.Cryptography;
 using EpicGames.Core;
 using Microsoft.Extensions.Caching.Memory;
+using Microsoft.Extensions.Logging.Abstractions;
 
 namespace EpicGames.Horde.Tests
 {
@@ -26,7 +27,7 @@ namespace EpicGames.Horde.Tests
 		{
 			using IMemoryCache cache = new MemoryCache(new MemoryCacheOptions());
 
-			InMemoryBlobStore blobStore = new InMemoryBlobStore(cache);
+			InMemoryBlobStore blobStore = new InMemoryBlobStore(cache, NullLogger.Instance);
 			await TestTreeAsync(blobStore, new TreeOptions { MaxBlobSize = 1024 * 1024 });
 
 			Assert.AreEqual(1, blobStore.Blobs.Count);
@@ -38,7 +39,7 @@ namespace EpicGames.Horde.Tests
 		{
 			using IMemoryCache cache = new MemoryCache(new MemoryCacheOptions());
 
-			InMemoryBlobStore blobStore = new InMemoryBlobStore(cache);
+			InMemoryBlobStore blobStore = new InMemoryBlobStore(cache, NullLogger.Instance);
 			await TestTreeAsync(blobStore, new TreeOptions { MaxBlobSize = 1 });
 
 			Assert.AreEqual(5, blobStore.Blobs.Count);
@@ -131,7 +132,7 @@ namespace EpicGames.Horde.Tests
 		public async Task DirectoryNodesAsync()
 		{
 			using IMemoryCache cache = new MemoryCache(new MemoryCacheOptions());
-			InMemoryBlobStore store = new InMemoryBlobStore(cache);
+			InMemoryBlobStore store = new InMemoryBlobStore(cache, NullLogger.Instance);
 
 			// Generate a tree
 			{
@@ -167,7 +168,7 @@ namespace EpicGames.Horde.Tests
 		public async Task FileNodesAsync()
 		{
 			using IMemoryCache cache = new MemoryCache(new MemoryCacheOptions());
-			InMemoryBlobStore store = new InMemoryBlobStore(cache);
+			InMemoryBlobStore store = new InMemoryBlobStore(cache, NullLogger.Instance);
 
 			// Generate a tree
 			{
@@ -212,7 +213,7 @@ namespace EpicGames.Horde.Tests
 		{
 			using IMemoryCache cache = new MemoryCache(new MemoryCacheOptions());
 
-			InMemoryBlobStore store = new InMemoryBlobStore(cache);
+			InMemoryBlobStore store = new InMemoryBlobStore(cache, NullLogger.Instance);
 			TreeWriter writer = new TreeWriter(store, new TreeOptions());
 
 			byte[] chunk = RandomNumberGenerator.GetBytes(256);
@@ -288,7 +289,7 @@ namespace EpicGames.Horde.Tests
 		public async Task LargeFileTestAsync()
 		{
 			using IMemoryCache cache = new MemoryCache(new MemoryCacheOptions());
-			InMemoryBlobStore store = new InMemoryBlobStore(cache);
+			InMemoryBlobStore store = new InMemoryBlobStore(cache, NullLogger.Instance);
 
 			const int length = 1024;
 			const int copies = 4096;
diff --git a/Engine/Source/Programs/Shared/EpicGames.Horde.Tests/BundleTests.cs b/Engine/Source/Programs/Shared/EpicGames.Horde.Tests/BundleTests.cs
--- a/Engine/Source/Programs/Shared/EpicGames.Horde.Tests/BundleTests.cs
+++ b/Engine/Source/Programs/Shared/EpicGames.Horde.Tests/BundleTests.cs
@@ -5,6 +5,7 @@ using EpicGames.Horde.Storage;
 using EpicGames.Horde.Storage.Backends;
 using EpicGames.Horde.Storage.Nodes;
 using Microsoft.Extensions.Caching.Memory;
+using Microsoft.Extensions.Logging.Abstractions;
 using Microsoft.VisualStudio.TestTools.UnitTesting;
 using System;
 using System.Collections.Generic;
@@ -24,7 +25,7 @@ namespace EpicGames.Horde.Tests
 		public BundleTests()
 		{
 			_cache = new MemoryCache(new MemoryCacheOptions());
-			_storage = new InMemoryBlobStore(_cache);
+			_storage = new InMemoryBlobStore(_cache, NullLogger.Instance);
 		}
 
 		public void Dispose()
diff --git a/Engine/Source/Programs/Shared/EpicGames.Horde/Storage/Backends/InMemoryBlobStore.cs b/Engine/Source/Programs/Shared/EpicGames.Horde/Storage/Backends/InMemoryBlobStore.cs
--- a/Engine/Source/Programs/Shared/EpicGames.Horde/Storage/Backends/InMemoryBlobStore.cs
+++ b/Engine/Source/Programs/Shared/EpicGames.Horde/Storage/Backends/InMemoryBlobStore.cs
@@ -9,6 +9,7 @@ using System.Threading;
 using System.Threading.Tasks;
 using EpicGames.Core;
 using Microsoft.Extensions.Caching.Memory;
+using Microsoft.Extensions.Logging;
 
 namespace EpicGames.Horde.Storage.Backends
 {
@@ -36,8 +37,8 @@ namespace EpicGames.Horde.Storage.Backends
 		/// <summary>
 		/// Constructor
 		/// </summary>
-		public InMemoryBlobStore(IMemoryCache cache) 
-			: base(cache)
+		public InMemoryBlobStore(IMemoryCache cache, ILogger logger) 
+			: base(cache, logger)
 		{
 		}
 
diff --git a/Engine/Source/Programs/Shared/EpicGames.Horde/Storage/Nodes/DirectoryNode.cs b/Engine/Source/Programs/Shared/EpicGames.Horde/Storage/Nodes/DirectoryNode.cs
--- a/Engine/Source/Programs/Shared/EpicGames.Horde/Storage/Nodes/DirectoryNode.cs
+++ b/Engine/Source/Programs/Shared/EpicGames.Horde/Storage/Nodes/DirectoryNode.cs
@@ -664,14 +664,14 @@ namespace EpicGames.Horde.Storage.Nodes
 			{
 				FileInfo fileInfo = new FileInfo(Path.Combine(directoryInfo.FullName, fileEntry.Name.ToString()));
 				FileNode fileNode = await fileEntry.ExpandAsync(cancellationToken);
-				logger.LogInformation("Writing {File}", fileInfo.FullName);
+//				logger.LogInformation("Writing {File}", fileInfo.FullName);
 				tasks.Add(Task.Run(() => fileNode.CopyToFileAsync(fileInfo, cancellationToken), cancellationToken));
 			}
 			foreach (DirectoryEntry directoryEntry in Directories)
 			{
 				DirectoryInfo subDirectoryInfo = directoryInfo.CreateSubdirectory(directoryEntry.Name.ToString());
 				DirectoryNode subDirectoryNode = await directoryEntry.ExpandAsync(cancellationToken);
-				logger.LogInformation("Writing {Dir}", subDirectoryInfo.FullName);
+//				logger.LogInformation("Writing {Dir}", subDirectoryInfo.FullName);
 				tasks.Add(Task.Run(() => subDirectoryNode.CopyToDirectoryAsync(subDirectoryInfo, logger, cancellationToken), cancellationToken));
 			}
 
diff --git a/Engine/Source/Programs/Shared/EpicGames.Horde/Storage/StorageClientBase.cs b/Engine/Source/Programs/Shared/EpicGames.Horde/Storage/StorageClientBase.cs
--- a/Engine/Source/Programs/Shared/EpicGames.Horde/Storage/StorageClientBase.cs
+++ b/Engine/Source/Programs/Shared/EpicGames.Horde/Storage/StorageClientBase.cs
@@ -8,7 +8,10 @@ using System.Linq;
 using System.Threading;
 using System.Threading.Tasks;
 using EpicGames.Core;
+using Microsoft.CodeAnalysis;
+using Microsoft.CodeAnalysis.CSharp.Syntax;
 using Microsoft.Extensions.Caching.Memory;
+using Microsoft.Extensions.Logging;
 
 namespace EpicGames.Horde.Storage
 {
@@ -17,6 +20,9 @@ namespace EpicGames.Horde.Storage
 	/// </summary>
 	public abstract class StorageClientBase : IStorageClient
 	{
+		/// <summary>
+		/// Wraps a stream and returns a portion of it
+		/// </summary>
 		class RangedReadStream : Stream, IDisposable
 		{
 			readonly Stream _inner;
@@ -154,15 +160,60 @@ namespace EpicGames.Horde.Storage
 			}
 		}
 
+		/// <summary>
+		/// Bundle header queued to be read
+		/// </summary>
+		class QueuedHeader
+		{
+			public readonly BlobLocator Blob;
+			public readonly TaskCompletionSource<BundleInfo> CompletionSource = new TaskCompletionSource<BundleInfo>();
+
+			public BlobId BlobId => Blob.BlobId;
+
+			public QueuedHeader(BlobLocator blob)
+			{
+				Blob = blob;
+			}
+		}
+
+		/// <summary>
+		/// Encoded bundle packet queued to be read
+		/// </summary>
+		class QueuedPacket
+		{
+			public readonly BundleInfo Bundle;
+			public readonly int PacketIdx;
+			public readonly TaskCompletionSource<ReadOnlyMemory<byte>> CompletionSource = new TaskCompletionSource<ReadOnlyMemory<byte>>();
+
+			public BlobId BlobId => Bundle.Locator.BlobId;
+
+			public QueuedPacket(BundleInfo bundle, int packetIdx)
+			{
+				Bundle = bundle;
+				PacketIdx = packetIdx;
+			}
+		}
+
+		// Size of data to fetch by default. This is larger than the minimum request size to reduce number of reads.
+		const int DefaultFetchSize = 15 * 1024 * 1024;
+
+		readonly object _queueLock = new object();
+		readonly List<QueuedHeader> _queuedHeaders = new List<QueuedHeader>();
+		readonly List<QueuedPacket> _queuedPackets = new List<QueuedPacket>();
+		readonly Dictionary<string, Task<ReadOnlyMemory<byte>>> _decodeTasks = new Dictionary<string, Task<ReadOnlyMemory<byte>>>(StringComparer.Ordinal);
+		Task? _readTask;
 		readonly IMemoryCache _cache;
+		readonly ILogger _logger;
 
 		/// <summary>
 		/// Constructor
 		/// </summary>
-		/// <param name="cache"></param>
-		protected StorageClientBase(IMemoryCache cache)
+		/// <param name="cache">Cache for read data</param>
+		/// <param name="logger">Logger for output</param>
+		protected StorageClientBase(IMemoryCache cache, ILogger logger)
 		{
 			_cache = cache;
+			_logger = logger;
 		}
 
 		#region Blobs
@@ -173,6 +224,7 @@ namespace EpicGames.Horde.Storage
 		/// <inheritdoc/>
 		public virtual async Task<Stream> ReadBlobRangeAsync(BlobLocator locator, int offset, int length, CancellationToken cancellationToken = default)
 		{
+			_logger.LogInformation("Reading {Locator} [{Offset},{Length}]", locator, offset, length);
 			Stream stream = await ReadBlobAsync(locator, cancellationToken);
 			stream.Seek(offset, SeekOrigin.Begin);
 			return new RangedReadStream(stream, Math.Min((int)(stream.Length - offset), length));
@@ -207,6 +259,12 @@ namespace EpicGames.Horde.Storage
 		static string GetEncodedPacketCacheKey(BlobId blobId, int packetIdx) => $"encoded-packet:{blobId}#{packetIdx}";
 		static string GetDecodedPacketCacheKey(BlobId blobId, int packetIdx) => $"decoded-packet:{blobId}#{packetIdx}";
 
+		/// <summary>
+		/// Adds an object to the storage cache
+		/// </summary>
+		/// <param name="cacheKey">Key for the item</param>
+		/// <param name="value">Value to add</param>
+		/// <param name="size">Size of the value</param>
 		void AddToCache(string cacheKey, object value, int size)
 		{
 			using (ICacheEntry entry = _cache.CreateEntry(cacheKey))
@@ -216,6 +274,208 @@ namespace EpicGames.Horde.Storage
 			}
 		}
 
+		/// <summary>
+		/// Starts the background task for reading data from the store
+		/// </summary>
+		void StartReadTask()
+		{
+			if (_readTask == null)
+			{
+				_readTask = Task.Run(() => ServiceReadQueueAsync(CancellationToken.None));
+			}
+		}
+
+		/// <summary>
+		/// Dispatches requests in the read queue
+		/// </summary>
+		/// <param name="cancellationToken">Cancellation token for the background task</param>
+		async Task ServiceReadQueueAsync(CancellationToken cancellationToken)
+		{
+			const int MaxConcurrentReads = 4;
+
+			List<(BlobId BlobId, Task Task)> currentTasks = new List<(BlobId, Task)>();
+			for (; ; )
+			{
+				// Start any new reads
+				lock (_queueLock)
+				{
+					while (currentTasks.Count < MaxConcurrentReads)
+					{
+						HashSet<BlobId> currentBlobIds = new HashSet<BlobId>(currentTasks.Select(x => x.BlobId));
+
+						// Try to start another header read
+						QueuedHeader? queuedHeader = _queuedHeaders.FirstOrDefault(x => !currentBlobIds.Contains(x.BlobId));
+						if (queuedHeader != null)
+						{
+							Task task = Task.Run(() => PerformHeaderReadAsync(queuedHeader, cancellationToken), cancellationToken);
+							currentTasks.Add((queuedHeader.BlobId, task));
+							continue;
+						}
+
+						// Try to start another packet read
+						QueuedPacket? queuedPacket = _queuedPackets.FirstOrDefault(x => !currentBlobIds.Contains(x.BlobId));
+						if (queuedPacket != null)
+						{
+							Task task = Task.Run(() => PerformPacketReadAsync(queuedPacket, cancellationToken), cancellationToken);
+							currentTasks.Add((queuedPacket.BlobId, task));
+							continue;
+						}
+
+						// If we're not waiting for anything else and there are no more requests, end the task thread.
+						if (currentTasks.Count == 0)
+						{
+							_readTask = null;
+							return;
+						}
+
+						// Break out of the loop
+						break;
+					}
+				}
+
+				// Wait for any read task to complete
+				await Task.WhenAny(currentTasks.Select(x => x.Task));
+
+				// Remove any tasks which are complete
+				for (int idx = 0; idx < currentTasks.Count; idx++)
+				{
+					Task task = currentTasks[idx].Task;
+					if (task.IsCompleted)
+					{
+						if (task.Exception != null)
+						{
+							_logger.LogError(task.Exception, "Exception while reading from blob {BlobId}.", currentTasks[idx].BlobId);
+						}
+						currentTasks.RemoveAt(idx--);
+					}
+				}
+			}
+		}
+
+		/// <summary>
+		/// Reads a bundle header from the queue
+		/// </summary>
+		/// <param name="queuedHeader">The header to read</param>
+		/// <param name="cancellationToken">Cancellation token for the operation</param>
+		async Task PerformHeaderReadAsync(QueuedHeader queuedHeader, CancellationToken cancellationToken)
+		{
+			int prefetchSize = DefaultFetchSize;
+			for (; ; )
+			{
+				using (IMemoryOwner<byte> owner = MemoryPool<byte>.Shared.Rent(prefetchSize))
+				{
+					// Read the prefetch size from the blob
+					Memory<byte> memory = owner.Memory.Slice(0, prefetchSize);
+					memory = await ReadBlobRangeAsync(queuedHeader.Blob, 0, memory, cancellationToken);
+
+					// Make sure it's large enough to hold the header
+					int headerSize = BundleHeader.ReadPrelude(memory);
+					if (headerSize > prefetchSize)
+					{
+						prefetchSize = headerSize;
+						continue;
+					}
+
+					// Parse the header and construct the bundle info from it
+					BundleHeader header = new BundleHeader(new MemoryReader(memory));
+					BundleInfo bundleInfo = new BundleInfo(queuedHeader.Blob, header, headerSize);
+
+					// Also add any encoded packets we prefetched
+					List<ReadOnlyMemory<byte>> packets = new List<ReadOnlyMemory<byte>>();
+					for (int packetOffset = headerSize; packets.Count < header.Packets.Count;)
+					{
+						int packetIdx = packets.Count;
+
+						int packetLength = header.Packets[packetIdx].EncodedLength;
+						if (packetOffset + packetLength > memory.Length)
+						{
+							break;
+						}
+
+						ReadOnlyMemory<byte> packetData = memory.Slice(packetOffset, packetLength).ToArray();
+						AddToCache(GetEncodedPacketCacheKey(queuedHeader.Blob.BlobId, packetIdx), packetData, packetData.Length);
+						packets.Add(packetData);
+						packetOffset += packetLength;
+					}
+
+					// Add the info to the cache
+					string cacheKey = GetBundleInfoCacheKey(queuedHeader.Blob.BlobId);
+					AddToCache(cacheKey, bundleInfo, headerSize);
+
+					// Update the task
+					queuedHeader.CompletionSource.SetResult(bundleInfo);
+
+					// Remove it from the queue
+					lock (_queueLock)
+					{
+						_queuedHeaders.Remove(queuedHeader);
+					}
+					break;
+				}
+			}
+		}
+
+		/// <summary>
+		/// Reads a packet from storage
+		/// </summary>
+		/// <param name="queuedPacket">The packet to read</param>
+		/// <param name="cancellationToken">Cancellation token for the operation</param>
+		async Task PerformPacketReadAsync(QueuedPacket queuedPacket, CancellationToken cancellationToken)
+		{
+			BundleInfo bundleInfo = queuedPacket.Bundle;
+			int minPacketIdx = queuedPacket.PacketIdx;
+
+			BundlePacket packet = bundleInfo.Header.Packets[minPacketIdx];
+			int readLength = packet.EncodedLength;
+
+			int maxPacketIdx = minPacketIdx + 1;
+			for (; maxPacketIdx < bundleInfo.Header.Packets.Count; maxPacketIdx++)
+			{
+				int nextReadLength = readLength + bundleInfo.Header.Packets[maxPacketIdx].EncodedLength;
+				if (nextReadLength > DefaultFetchSize)
+				{
+					break;
+				}
+				readLength = nextReadLength;
+			}
+
+			using (IMemoryOwner<byte> owner = MemoryPool<byte>.Shared.Rent(readLength))
+			{
+				Memory<byte> buffer = owner.Memory.Slice(0, readLength);
+				buffer = await ReadBlobRangeAsync(bundleInfo.Locator, bundleInfo.PacketOffsets[minPacketIdx], buffer, cancellationToken);
+
+				// Copy all the packets that have been read into separate buffers, so we can cache them individually.
+				List<ReadOnlyMemory<byte>> packets = new List<ReadOnlyMemory<byte>>();
+				for (int idx = minPacketIdx; idx < maxPacketIdx; idx++)
+				{
+					string cacheKey = GetEncodedPacketCacheKey(bundleInfo.Locator.BlobId, idx);
+					ReadOnlyMemory<byte> data = buffer.Slice(bundleInfo.PacketOffsets[idx] - bundleInfo.PacketOffsets[minPacketIdx], bundleInfo.Header.Packets[idx].EncodedLength).ToArray();
+					packets.Add(data);
+					AddToCache(cacheKey, data, data.Length);
+				}
+
+				// Find all the packets we can mark as complete
+				List<QueuedPacket> updatePackets = new List<QueuedPacket>();
+				lock (_queueLock)
+				{
+					updatePackets.AddRange(_queuedPackets.Where(x => x.BlobId == bundleInfo.Locator.BlobId && (x.PacketIdx >= minPacketIdx && x.PacketIdx < maxPacketIdx)));
+				}
+
+				// Mark them all as complete
+				foreach (QueuedPacket updatePacket in updatePackets)
+				{
+					ReadOnlyMemory<byte> data = packets[updatePacket.PacketIdx - minPacketIdx];
+					updatePacket.CompletionSource.SetResult(data);
+				}
+
+				// Remove all the completed packets from the queue
+				lock (_queueLock)
+				{
+					_queuedPackets.RemoveAll(x => updatePackets.Contains(x));
+				}
+			}
+		}
+
 		/// <summary>
 		/// Reads a bundle from the given blob id, or retrieves it from the cache
 		/// </summary>
@@ -239,47 +499,30 @@ namespace EpicGames.Horde.Storage
 		async ValueTask<BundleInfo> GetBundleInfoAsync(BlobLocator locator, CancellationToken cancellationToken = default)
 		{
 			string cacheKey = GetBundleInfoCacheKey(locator.BlobId);
-			if (!_cache.TryGetValue(cacheKey, out BundleInfo bundleInfo))
+			if (_cache.TryGetValue(cacheKey, out BundleInfo bundleInfo))
 			{
-				int prefetchSize = 1024 * 1024;
-				for (; ; )
-				{
-					using (IMemoryOwner<byte> owner = MemoryPool<byte>.Shared.Rent(prefetchSize))
-					{
-						// Read the prefetch size from the blob
-						Memory<byte> memory = owner.Memory.Slice(0, prefetchSize);
-						memory = await ReadBlobRangeAsync(locator, 0, memory, cancellationToken);
-
-						// Make sure it's large enough to hold the header
-						int headerSize = BundleHeader.ReadPrelude(memory);
-						if (headerSize > prefetchSize)
-						{
-							prefetchSize = headerSize;
-							continue;
-						}
-
-						// Parse the header and construct the bundle info from it
-						BundleHeader header = new BundleHeader(new MemoryReader(memory));
-						bundleInfo = new BundleInfo(locator, header, headerSize);
-
-						// Add the info to the cache
-						AddToCache(cacheKey, bundleInfo, headerSize);
+				return bundleInfo;
+			}
 
-						// Also add any encoded packets we prefetched
-						int packetOffset = headerSize;
-						for (int packetIdx = 0; packetIdx < header.Packets.Count && packetOffset + header.Packets[packetIdx].EncodedLength < memory.Length; packetIdx++)
-						{
-							int packetLength = header.Packets[packetIdx].EncodedLength;
-							ReadOnlyMemory<byte> packetData = memory.Slice(packetOffset, packetLength).ToArray();
-							AddToCache(GetEncodedPacketCacheKey(locator.BlobId, packetIdx), packetData, packetData.Length);
-							packetOffset += packetLength;
-						}
+			QueuedHeader? queuedHeader;
+			lock (_queueLock)
+			{
+				// Check the cache again inside lock scope to avoid races
+				if (_cache.TryGetValue(cacheKey, out bundleInfo))
+				{
+					return bundleInfo;
+				}
 
-						break;
-					}
+				// Find or start the read
+				queuedHeader = _queuedHeaders.FirstOrDefault(x => x.Blob.BlobId == locator.BlobId);
+				if (queuedHeader == null)
+				{
+					queuedHeader = new QueuedHeader(locator);
+					_queuedHeaders.Add(queuedHeader);
+					StartReadTask();
 				}
 			}
-			return bundleInfo;
+			return await queuedHeader.CompletionSource.Task.AbandonOnCancel(cancellationToken);
 		}
 
 		/// <summary>
@@ -291,33 +534,92 @@ namespace EpicGames.Horde.Storage
 		/// <returns>The decoded data</returns>
 		async ValueTask<ReadOnlyMemory<byte>> ReadBundlePacketAsync(BundleInfo bundleInfo, int packetIdx, CancellationToken cancellationToken)
 		{
-			string cacheKey = GetDecodedPacketCacheKey(bundleInfo.Locator.BlobId, packetIdx);
-			if (!_cache.TryGetValue(cacheKey, out ReadOnlyMemory<byte> decodedPacket))
+			string decodedCacheKey = GetDecodedPacketCacheKey(bundleInfo.Locator.BlobId, packetIdx);
+			if (_cache.TryGetValue(decodedCacheKey, out ReadOnlyMemory<byte> decodedPacket))
 			{
-				BundlePacket packet = bundleInfo.Header.Packets[packetIdx];
-				byte[] decodedPacketBuffer = new byte[packet.DecodedLength];
+				return decodedPacket;
+			}
 
-				ReadOnlyMemory<byte> encodedPacket;
-				if (_cache.TryGetValue(GetEncodedPacketCacheKey(bundleInfo.Locator.BlobId, packetIdx), out encodedPacket))
+			Task<ReadOnlyMemory<byte>>? decodeTask;
+			lock (_queueLock)
+			{
+				// Query the cache again, to eliminate races between cache checks and decode tasks finishing.
+				if (_cache.TryGetValue(decodedCacheKey, out decodedPacket))
 				{
-					BundleData.Decompress(bundleInfo.Header.CompressionFormat, encodedPacket, decodedPacketBuffer);
+					return decodedPacket;
 				}
-				else
+
+				// Create an async task to read the data
+				if (!_decodeTasks.TryGetValue(decodedCacheKey, out decodeTask))
 				{
-					using IMemoryOwner<byte> owner = MemoryPool<byte>.Shared.Rent(packet.EncodedLength);
+					decodeTask = Task.Run(() => ReadAndDecodePacketAsync(bundleInfo, packetIdx), CancellationToken.None);
+					_decodeTasks.Add(decodedCacheKey, decodeTask);
+				}
+			}
+			return await decodeTask.AbandonOnCancel(cancellationToken);
+		}
+
+		/// <summary>
+		/// Reads and decodes a packet from a bundle
+		/// </summary>
+		/// <param name="bundleInfo">Bundle to read from</param>
+		/// <param name="packetIdx">Index of the packet to return</param>
+		/// <returns>The decoded packet data</returns>
+		async Task<ReadOnlyMemory<byte>> ReadAndDecodePacketAsync(BundleInfo bundleInfo, int packetIdx)
+		{
+			ReadOnlyMemory<byte> encodedPacket = await ReadEncodedPacketAsync(bundleInfo, packetIdx);
 
-					Memory<byte> buffer = owner.Memory.Slice(0, packet.EncodedLength);
-					buffer = await ReadBlobRangeAsync(bundleInfo.Locator, bundleInfo.PacketOffsets[packetIdx], buffer, cancellationToken);
+			BundlePacket packet = bundleInfo.Header.Packets[packetIdx];
+			byte[] decodedPacket = new byte[packet.DecodedLength];
 
-					BundleData.Decompress(bundleInfo.Header.CompressionFormat, buffer, decodedPacketBuffer);
-				}
+			BundleData.Decompress(bundleInfo.Header.CompressionFormat, encodedPacket, decodedPacket);
+
+			string decodedCacheKey = GetDecodedPacketCacheKey(bundleInfo.Locator.BlobId, packetIdx);
+			AddToCache(decodedCacheKey, (ReadOnlyMemory<byte>)decodedPacket, decodedPacket.Length);
 
-				decodedPacket = decodedPacketBuffer;
-				AddToCache(cacheKey, decodedPacket, decodedPacket.Length);
+			lock (_queueLock)
+			{
+				_decodeTasks.Remove(decodedCacheKey);
 			}
+
 			return decodedPacket;
 		}
 
+		/// <summary>
+		/// Reads an encoded packet from a bundle
+		/// </summary>
+		/// <param name="bundleInfo">Bundle to read from</param>
+		/// <param name="packetIdx">Index of the packet to return</param>
+		/// <returns>The encoded packet data</returns>
+		async ValueTask<ReadOnlyMemory<byte>> ReadEncodedPacketAsync(BundleInfo bundleInfo, int packetIdx)
+		{
+			string encodedCacheKey = GetEncodedPacketCacheKey(bundleInfo.Locator.BlobId, packetIdx);
+			if (_cache.TryGetValue(encodedCacheKey, out ReadOnlyMemory<byte> encodedPacket))
+			{
+				return encodedPacket;
+			}
+
+			QueuedPacket? queuedPacket;
+			lock (_queueLock)
+			{
+				// Query the cache again, to eliminate races between cache checks and decode tasks finishing.
+				if (_cache.TryGetValue(encodedCacheKey, out encodedPacket))
+				{
+					return encodedPacket;
+				}
+
+				// Add a read to the queue
+				queuedPacket = _queuedPackets.FirstOrDefault(x => x.Bundle.Locator.BlobId == bundleInfo.Locator.BlobId && x.PacketIdx == packetIdx);
+				if (queuedPacket == null)
+				{
+					queuedPacket = new QueuedPacket(bundleInfo, packetIdx);
+					_queuedPackets.Add(queuedPacket);
+					StartReadTask();
+				}
+			}
+			return await queuedPacket.CompletionSource.Task;
+		}
+
 		/// <inheritdoc/>
 		public async Task<BlobLocator> WriteBundleAsync(Bundle bundle, Utf8String prefix = default, CancellationToken cancellationToken = default)
 		{
