Horde: Flush nodes to the backing store as they are generated.

diff --git a/Engine/Source/Programs/Horde/Horde.Agent/Commands/Bundles/CreateCommand.cs b/Engine/Source/Programs/Horde/Horde.Agent/Commands/Bundles/CreateCommand.cs
--- a/Engine/Source/Programs/Horde/Horde.Agent/Commands/Bundles/CreateCommand.cs
+++ b/Engine/Source/Programs/Horde/Horde.Agent/Commands/Bundles/CreateCommand.cs
@@ -134,7 +134,7 @@ namespace Horde.Agent.Commands.Bundles
 			Stopwatch timer = Stopwatch.StartNew();
 
 			ChunkingOptions options = new ChunkingOptions();
-			await node.CopyFromDirectoryAsync(InputDir.ToDirectoryInfo(), options, writer, logger, CancellationToken.None);
+			await node.CopyFromDirectoryAsync(InputDir.ToDirectoryInfo(), options, writer, CancellationToken.None);
 
 			await writer.WriteRefAsync(RefName, node);
 
diff --git a/Engine/Source/Programs/Horde/Horde.Build/Commands/Bundles/CreateCommand.cs b/Engine/Source/Programs/Horde/Horde.Build/Commands/Bundles/CreateCommand.cs
--- a/Engine/Source/Programs/Horde/Horde.Build/Commands/Bundles/CreateCommand.cs
+++ b/Engine/Source/Programs/Horde/Horde.Build/Commands/Bundles/CreateCommand.cs
@@ -38,7 +38,7 @@ namespace Horde.Build.Commands.Bundles
 			TreeWriter writer = new TreeWriter(store, prefix: RefName.Text);
 
 			DirectoryNode node = new DirectoryNode(DirectoryFlags.None);
-			await node.CopyFromDirectoryAsync(InputDir.ToDirectoryInfo(), new ChunkingOptions(), writer, logger, CancellationToken.None);
+			await node.CopyFromDirectoryAsync(InputDir.ToDirectoryInfo(), new ChunkingOptions(), writer, CancellationToken.None);
 
 			await writer.WriteRefAsync(RefName, node, CancellationToken.None);
 			return 0;
diff --git a/Engine/Source/Programs/Shared/EpicGames.Horde/Storage/Nodes/DirectoryNode.cs b/Engine/Source/Programs/Shared/EpicGames.Horde/Storage/Nodes/DirectoryNode.cs
--- a/Engine/Source/Programs/Shared/EpicGames.Horde/Storage/Nodes/DirectoryNode.cs
+++ b/Engine/Source/Programs/Shared/EpicGames.Horde/Storage/Nodes/DirectoryNode.cs
@@ -569,10 +569,9 @@ namespace EpicGames.Horde.Storage.Nodes
 		/// <param name="directoryInfo"></param>
 		/// <param name="options">Options for chunking file content</param>
 		/// <param name="writer">Writer for new node data</param>
-		/// <param name="logger"></param>
 		/// <param name="cancellationToken">Cancellation token for the operation</param>
 		/// <returns></returns>
-		public async Task CopyFromDirectoryAsync(DirectoryInfo directoryInfo, ChunkingOptions options, TreeWriter writer, ILogger logger, CancellationToken cancellationToken)
+		public async Task CopyFromDirectoryAsync(DirectoryInfo directoryInfo, ChunkingOptions options, TreeWriter writer, CancellationToken cancellationToken)
 		{
 			const int MaxWriters = 32;
 			const long MinSizePerWriter = 1024 * 1024;
@@ -603,7 +602,7 @@ namespace EpicGames.Horde.Storage.Nodes
 				}
 
 				int minIdxCopy = minIdx;
-				tasks.Add(Task.Run(() => CopyFilesAsync(files, minIdxCopy, maxIdx, fileEntries, options, writer, logger, cancellationToken), cancellationToken));
+				tasks.Add(Task.Run(() => CopyFilesAsync(files, minIdxCopy, maxIdx, fileEntries, options, writer, cancellationToken), cancellationToken));
 
 				targetSize += chunkSize;
 				minIdx = maxIdx;
@@ -631,7 +630,7 @@ namespace EpicGames.Horde.Storage.Nodes
 			}
 		}
 
-		static async Task CopyFilesAsync(List<(DirectoryNode DirectoryNode, FileInfo FileInfo)> files, int minIdx, int maxIdx, FileEntry[] entries, ChunkingOptions options, TreeWriter baseWriter, ILogger logger, CancellationToken cancellationToken)
+		static async Task CopyFilesAsync(List<(DirectoryNode DirectoryNode, FileInfo FileInfo)> files, int minIdx, int maxIdx, FileEntry[] entries, ChunkingOptions options, TreeWriter baseWriter, CancellationToken cancellationToken)
 		{
 			TreeWriter writer = new TreeWriter(baseWriter);
 			for(int idx = minIdx; idx < maxIdx; idx++)
@@ -664,14 +663,12 @@ namespace EpicGames.Horde.Storage.Nodes
 			{
 				FileInfo fileInfo = new FileInfo(Path.Combine(directoryInfo.FullName, fileEntry.Name.ToString()));
 				FileNode fileNode = await fileEntry.ExpandAsync(cancellationToken);
-//				logger.LogInformation("Writing {File}", fileInfo.FullName);
 				tasks.Add(Task.Run(() => fileNode.CopyToFileAsync(fileInfo, cancellationToken), cancellationToken));
 			}
 			foreach (DirectoryEntry directoryEntry in Directories)
 			{
 				DirectoryInfo subDirectoryInfo = directoryInfo.CreateSubdirectory(directoryEntry.Name.ToString());
 				DirectoryNode subDirectoryNode = await directoryEntry.ExpandAsync(cancellationToken);
-//				logger.LogInformation("Writing {Dir}", subDirectoryInfo.FullName);
 				tasks.Add(Task.Run(() => subDirectoryNode.CopyToDirectoryAsync(subDirectoryInfo, logger, cancellationToken), cancellationToken));
 			}
 
diff --git a/Engine/Source/Programs/Shared/EpicGames.Horde/Storage/Nodes/FileNode.cs b/Engine/Source/Programs/Shared/EpicGames.Horde/Storage/Nodes/FileNode.cs
--- a/Engine/Source/Programs/Shared/EpicGames.Horde/Storage/Nodes/FileNode.cs
+++ b/Engine/Source/Programs/Shared/EpicGames.Horde/Storage/Nodes/FileNode.cs
@@ -555,7 +555,9 @@ namespace EpicGames.Horde.Storage.Nodes
 				// Try to write to the last node
 				if (_children.Count > 0)
 				{
-					FileNode? lastNode = _children[^1].Target;
+					TreeNodeRef<FileNode> lastNodeRef = _children[^1];
+
+					FileNode? lastNode = lastNodeRef.Target;
 					if (lastNode != null)
 					{
 						// Update the length to match the new node
@@ -566,6 +568,9 @@ namespace EpicGames.Horde.Storage.Nodes
 						// If the last node is complete, write it to the buffer
 						if (lastNode.IsReadOnly())
 						{
+							// Write the last node to allow it to be flushed
+							await writer.WriteAsync(lastNodeRef, cancellationToken);
+
 							// Update the hash
 							AppendChildHash(lastNode.RollingHash);
 
diff --git a/Engine/Source/Programs/Shared/EpicGames.Horde/Storage/TreeNodeRef.cs b/Engine/Source/Programs/Shared/EpicGames.Horde/Storage/TreeNodeRef.cs
--- a/Engine/Source/Programs/Shared/EpicGames.Horde/Storage/TreeNodeRef.cs
+++ b/Engine/Source/Programs/Shared/EpicGames.Horde/Storage/TreeNodeRef.cs
@@ -128,6 +128,17 @@ namespace EpicGames.Horde.Storage
 			return result;
 		}
 
+		/// <summary>
+		/// Updates the hash and revision number for the ref.
+		/// </summary>
+		/// <param name="hash">Hash of the node</param>
+		/// <returns></returns>
+		public void MarkAsPendingWrite(IoHash hash)
+		{
+			Hash = hash;
+			_revision = Target!.Revision;
+		}
+
 		/// <summary>
 		/// Serialize the node to the given writer
 		/// </summary>
diff --git a/Engine/Source/Programs/Shared/EpicGames.Horde/Storage/TreeWriter.cs b/Engine/Source/Programs/Shared/EpicGames.Horde/Storage/TreeWriter.cs
--- a/Engine/Source/Programs/Shared/EpicGames.Horde/Storage/TreeWriter.cs
+++ b/Engine/Source/Programs/Shared/EpicGames.Horde/Storage/TreeWriter.cs
@@ -155,9 +155,12 @@ namespace EpicGames.Horde.Storage
 			// If the target node hasn't been modified, use the existing state.
 			if (!nodeRef.IsDirty())
 			{
-				Debug.Assert(nodeRef.Locator.IsValid());
-				_hashToNode.TryAdd(nodeRef.Hash, nodeRef.Locator);
-				nodeRef.Target = null;
+				// Make sure the locator is valid. The node may be queued for writing but not flushed to disk yet.
+				if (nodeRef.Locator.IsValid())
+				{
+					_hashToNode.TryAdd(nodeRef.Hash, nodeRef.Locator);
+					nodeRef.Target = null;
+				}
 				return nodeRef.Hash;
 			}
 
@@ -210,6 +213,9 @@ namespace EpicGames.Horde.Storage
 					await FlushAsync(cancellationToken);
 				}
 			}
+
+			// Mark the node as queued for writing
+			nodeRef.MarkAsPendingWrite(hash);
 			return hash;
 		}
 
