Fixed C# UHT issue where struct properties were not performing all required validations
Fixed C# UHT issue where struct serializers were appearing in the wrong place.

diff --git a/Engine/Source/Programs/Shared/EpicGames.UHT/Types/Properties/UhtStructProperty.cs b/Engine/Source/Programs/Shared/EpicGames.UHT/Types/Properties/UhtStructProperty.cs
--- a/Engine/Source/Programs/Shared/EpicGames.UHT/Types/Properties/UhtStructProperty.cs
+++ b/Engine/Source/Programs/Shared/EpicGames.UHT/Types/Properties/UhtStructProperty.cs
@@ -207,6 +207,7 @@ namespace EpicGames.UHT.Types
 		/// <inheritdoc/>
 		protected override void ValidateMember(UhtStruct structObj, UhtValidationOptions options)
 		{
+			base.ValidateMember(structObj, options);
 			if (PropertyFlags.HasAnyFlags(EPropertyFlags.Net))
 			{
 				Session.ValidateScriptStructOkForNet(this, ScriptStruct);
