Horde: Fix EpicGames.Horde.Tests compile errors.

diff --git a/Engine/Source/Programs/Shared/EpicGames.Horde.Tests/BundleStoreTests.cs b/Engine/Source/Programs/Shared/EpicGames.Horde.Tests/BundleStoreTests.cs
--- a/Engine/Source/Programs/Shared/EpicGames.Horde.Tests/BundleStoreTests.cs
+++ b/Engine/Source/Programs/Shared/EpicGames.Horde.Tests/BundleStoreTests.cs
@@ -27,7 +27,7 @@ namespace EpicGames.Horde.Tests
 		{
 			using IMemoryCache cache = new MemoryCache(new MemoryCacheOptions());
 
-			MemoryBlobStore blobStore = new MemoryBlobStore(cache, NullLogger.Instance);
+			MemoryStorageClient blobStore = new MemoryStorageClient(cache, NullLogger.Instance);
 			await TestTreeAsync(blobStore, new TreeOptions { MaxBlobSize = 1024 * 1024 });
 
 			Assert.AreEqual(1, blobStore.Blobs.Count);
@@ -39,7 +39,7 @@ namespace EpicGames.Horde.Tests
 		{
 			using IMemoryCache cache = new MemoryCache(new MemoryCacheOptions());
 
-			MemoryBlobStore blobStore = new MemoryBlobStore(cache, NullLogger.Instance);
+			MemoryStorageClient blobStore = new MemoryStorageClient(cache, NullLogger.Instance);
 			await TestTreeAsync(blobStore, new TreeOptions { MaxBlobSize = 1 });
 
 			Assert.AreEqual(5, blobStore.Blobs.Count);
@@ -70,7 +70,7 @@ namespace EpicGames.Horde.Tests
 			}
 		}
 
-		static async Task TestTreeAsync(MemoryBlobStore store, TreeOptions options)
+		static async Task TestTreeAsync(MemoryStorageClient store, TreeOptions options)
 		{
 			// Generate a tree
 			{
@@ -132,7 +132,7 @@ namespace EpicGames.Horde.Tests
 		public async Task DirectoryNodesAsync()
 		{
 			using IMemoryCache cache = new MemoryCache(new MemoryCacheOptions());
-			MemoryBlobStore store = new MemoryBlobStore(cache, NullLogger.Instance);
+			MemoryStorageClient store = new MemoryStorageClient(cache, NullLogger.Instance);
 
 			// Generate a tree
 			{
@@ -168,7 +168,7 @@ namespace EpicGames.Horde.Tests
 		public async Task FileNodesAsync()
 		{
 			using IMemoryCache cache = new MemoryCache(new MemoryCacheOptions());
-			MemoryBlobStore store = new MemoryBlobStore(cache, NullLogger.Instance);
+			MemoryStorageClient store = new MemoryStorageClient(cache, NullLogger.Instance);
 
 			// Generate a tree
 			{
@@ -213,7 +213,7 @@ namespace EpicGames.Horde.Tests
 		{
 			using IMemoryCache cache = new MemoryCache(new MemoryCacheOptions());
 
-			MemoryBlobStore store = new MemoryBlobStore(cache, NullLogger.Instance);
+			MemoryStorageClient store = new MemoryStorageClient(cache, NullLogger.Instance);
 			TreeWriter writer = new TreeWriter(store, new TreeOptions());
 
 			byte[] chunk = RandomNumberGenerator.GetBytes(256);
@@ -289,7 +289,7 @@ namespace EpicGames.Horde.Tests
 		public async Task LargeFileTestAsync()
 		{
 			using IMemoryCache cache = new MemoryCache(new MemoryCacheOptions());
-			MemoryBlobStore store = new MemoryBlobStore(cache, NullLogger.Instance);
+			MemoryStorageClient store = new MemoryStorageClient(cache, NullLogger.Instance);
 
 			const int length = 1024;
 			const int copies = 4096;
diff --git a/Engine/Source/Programs/Shared/EpicGames.Horde.Tests/BundleTests.cs b/Engine/Source/Programs/Shared/EpicGames.Horde.Tests/BundleTests.cs
--- a/Engine/Source/Programs/Shared/EpicGames.Horde.Tests/BundleTests.cs
+++ b/Engine/Source/Programs/Shared/EpicGames.Horde.Tests/BundleTests.cs
@@ -20,12 +20,12 @@ namespace EpicGames.Horde.Tests
 	public sealed class BundleTests : IDisposable
 	{
 		readonly IMemoryCache _cache;
-		readonly MemoryBlobStore _storage;
+		readonly MemoryStorageClient _storage;
 
 		public BundleTests()
 		{
 			_cache = new MemoryCache(new MemoryCacheOptions());
-			_storage = new MemoryBlobStore(_cache, NullLogger.Instance);
+			_storage = new MemoryStorageClient(_cache, NullLogger.Instance);
 		}
 
 		public void Dispose()
@@ -166,7 +166,7 @@ namespace EpicGames.Horde.Tests
 		[TestMethod]
 		public async Task BasicTestDirectory()
 		{
-			MemoryBlobStore store = _storage;
+			MemoryStorageClient store = _storage;
 
 			DirectoryNode root = new DirectoryNode(DirectoryFlags.None);
 			DirectoryNode node = root.AddDirectory("hello");
diff --git a/Engine/Source/Programs/Shared/EpicGames.Horde.Tests/ChunkingTests.cs b/Engine/Source/Programs/Shared/EpicGames.Horde.Tests/ChunkingTests.cs
--- a/Engine/Source/Programs/Shared/EpicGames.Horde.Tests/ChunkingTests.cs
+++ b/Engine/Source/Programs/Shared/EpicGames.Horde.Tests/ChunkingTests.cs
@@ -11,6 +11,7 @@ using System.Threading;
 using EpicGames.Horde.Storage;
 using EpicGames.Horde.Storage.Backends;
 using Microsoft.Extensions.Caching.Memory;
+using Microsoft.Extensions.Logging.Abstractions;
 
 namespace EpicGames.Horde.Tests
 {
@@ -51,7 +52,7 @@ namespace EpicGames.Horde.Tests
 		{
 			using MemoryCache cache = new MemoryCache(new MemoryCacheOptions());
 
-			InMemoryBlobStore store = new InMemoryBlobStore(cache);
+			MemoryStorageClient store = new MemoryStorageClient(cache, NullLogger.Instance);
 			TreeWriter writer = new TreeWriter(store, new TreeOptions(), "test");
 
 			ChunkingOptions options = new ChunkingOptions();
@@ -110,7 +111,7 @@ namespace EpicGames.Horde.Tests
 		{
 			using MemoryCache cache = new MemoryCache(new MemoryCacheOptions());
 
-			InMemoryBlobStore store = new InMemoryBlobStore(cache);
+			MemoryStorageClient store = new MemoryStorageClient(cache, NullLogger.Instance);
 			TreeWriter writer = new TreeWriter(store, new TreeOptions());
 
 			byte[] data = new byte[4096];
diff --git a/Engine/Source/Programs/Shared/EpicGames.Horde/Storage/Backends/MemoryStorageClient.cs b/Engine/Source/Programs/Shared/EpicGames.Horde/Storage/Backends/MemoryStorageClient.cs
--- a/Engine/Source/Programs/Shared/EpicGames.Horde/Storage/Backends/MemoryStorageClient.cs
+++ b/Engine/Source/Programs/Shared/EpicGames.Horde/Storage/Backends/MemoryStorageClient.cs
@@ -16,7 +16,7 @@ namespace EpicGames.Horde.Storage.Backends
 	/// <summary>
 	/// Implementation of <see cref="IStorageClient"/> which stores data in memory. Not intended for production use.
 	/// </summary>
-	public class MemoryBlobStore : StorageClientBase
+	public class MemoryStorageClient : StorageClientBase
 	{
 		/// <summary>
 		/// Map of blob id to blob data
@@ -37,7 +37,7 @@ namespace EpicGames.Horde.Storage.Backends
 		/// <summary>
 		/// Constructor
 		/// </summary>
-		public MemoryBlobStore(IMemoryCache cache, ILogger logger) 
+		public MemoryStorageClient(IMemoryCache cache, ILogger logger) 
 			: base(cache, logger)
 		{
 		}
