EpicGames.Serialization: Fix warnings compiling against NET 6 SDK.

diff --git a/Engine/Source/Programs/Shared/EpicGames.Serialization/CbField.cs b/Engine/Source/Programs/Shared/EpicGames.Serialization/CbField.cs
--- a/Engine/Source/Programs/Shared/EpicGames.Serialization/CbField.cs
+++ b/Engine/Source/Programs/Shared/EpicGames.Serialization/CbField.cs
@@ -263,10 +263,10 @@ namespace EpicGames.Serialization
 	sealed class CbBinaryAttachmentTypeConverter : TypeConverter
 	{
 		/// <inheritdoc/>
-		public override bool CanConvertFrom(ITypeDescriptorContext context, Type sourceType) => sourceType == typeof(string);
+		public override bool CanConvertFrom(ITypeDescriptorContext? context, Type sourceType) => sourceType == typeof(string);
 
 		/// <inheritdoc/>
-		public override object ConvertFrom(ITypeDescriptorContext context, CultureInfo culture, object value) => new CbBinaryAttachment(IoHash.Parse((string)value));
+		public override object ConvertFrom(ITypeDescriptorContext? context, CultureInfo? culture, object value) => new CbBinaryAttachment(IoHash.Parse((string)value));
 	}
 
 	/// <summary>
@@ -342,10 +342,10 @@ namespace EpicGames.Serialization
 	sealed class CbObjectAttachmentTypeConverter : TypeConverter
 	{
 		/// <inheritdoc/>
-		public override bool CanConvertFrom(ITypeDescriptorContext context, Type sourceType) => sourceType == typeof(string);
+		public override bool CanConvertFrom(ITypeDescriptorContext? context, Type sourceType) => sourceType == typeof(string);
 
 		/// <inheritdoc/>
-		public override object ConvertFrom(ITypeDescriptorContext context, CultureInfo culture, object value) => new CbObjectAttachment(IoHash.Parse((string)value));
+		public override object ConvertFrom(ITypeDescriptorContext? context, CultureInfo? culture, object value) => new CbObjectAttachment(IoHash.Parse((string)value));
 	}
 
 	/// <summary>
diff --git a/Engine/Source/Programs/Shared/EpicGames.Serialization/Converters/CbStringConverter.cs b/Engine/Source/Programs/Shared/EpicGames.Serialization/Converters/CbStringConverter.cs
--- a/Engine/Source/Programs/Shared/EpicGames.Serialization/Converters/CbStringConverter.cs
+++ b/Engine/Source/Programs/Shared/EpicGames.Serialization/Converters/CbStringConverter.cs
@@ -17,7 +17,7 @@ namespace EpicGames.Serialization.Converters
 
 		public override T Read(CbField field)
 		{
-			return (T)_typeConverter.ConvertFromInvariantString(field.AsString());
+			return (T)_typeConverter.ConvertFromInvariantString(field.AsString())!;
 		}
 
 		public override void WriteNamed(CbWriter writer, Utf8String name, T value)
@@ -27,7 +27,7 @@ namespace EpicGames.Serialization.Converters
 
 		public override void Write(CbWriter writer, T value)
 		{ 
-			writer.WriteStringValue(_typeConverter.ConvertToInvariantString(value));
+			writer.WriteStringValue(_typeConverter.ConvertToInvariantString(value)!);
 		}
 	}
 
