EpicGames.Core: Fix warnings compiling against NET 6 SDK.

diff --git a/Engine/Source/Programs/Shared/EpicGames.Core/BinaryFormatterUtils.cs b/Engine/Source/Programs/Shared/EpicGames.Core/BinaryFormatterUtils.cs
--- a/Engine/Source/Programs/Shared/EpicGames.Core/BinaryFormatterUtils.cs
+++ b/Engine/Source/Programs/Shared/EpicGames.Core/BinaryFormatterUtils.cs
@@ -3,6 +3,8 @@
 using System.IO;
 using System.Runtime.Serialization.Formatters.Binary;
 
+#pragma warning disable SYSLIB0011
+
 namespace EpicGames.Core
 {
 	/// <summary>
diff --git a/Engine/Source/Programs/Shared/EpicGames.Core/CommandHost.cs b/Engine/Source/Programs/Shared/EpicGames.Core/CommandHost.cs
--- a/Engine/Source/Programs/Shared/EpicGames.Core/CommandHost.cs
+++ b/Engine/Source/Programs/Shared/EpicGames.Core/CommandHost.cs
@@ -201,7 +201,7 @@ namespace EpicGames.Core
 				}
 				else 
 				{
-					command = (ICommand)serviceProvider.GetService(defaultCommandType);
+					command = (ICommand)serviceProvider.GetService(defaultCommandType)!;
 				}
 			}
 			else
diff --git a/Engine/Source/Programs/Shared/EpicGames.Core/CommandLineArguments.cs b/Engine/Source/Programs/Shared/EpicGames.Core/CommandLineArguments.cs
--- a/Engine/Source/Programs/Shared/EpicGames.Core/CommandLineArguments.cs
+++ b/Engine/Source/Programs/Shared/EpicGames.Core/CommandLineArguments.cs
@@ -1295,7 +1295,7 @@ namespace EpicGames.Core
 				TypeConverter typeConverter = TypeDescriptor.GetConverter(fieldType);
 				if (typeConverter.CanConvertFrom(typeof(string)))
 				{
-					value = typeConverter.ConvertFrom(text);
+					value = typeConverter.ConvertFrom(text)!;
 					return true;
 				}
 
diff --git a/Engine/Source/Programs/Shared/EpicGames.Core/ContentHash.cs b/Engine/Source/Programs/Shared/EpicGames.Core/ContentHash.cs
--- a/Engine/Source/Programs/Shared/EpicGames.Core/ContentHash.cs
+++ b/Engine/Source/Programs/Shared/EpicGames.Core/ContentHash.cs
@@ -288,7 +288,7 @@ namespace EpicGames.Core
 		/// <inheritdoc/>
 		public override ContentHash Read(ref Utf8JsonReader reader, Type typeToConvert, JsonSerializerOptions options)
 		{
-			return ContentHash.Parse(reader.GetString());
+			return ContentHash.Parse(reader.GetString()!);
 		}
 
 		/// <inheritdoc/>
diff --git a/Engine/Source/Programs/Shared/EpicGames.Core/CsProjectInfo.cs b/Engine/Source/Programs/Shared/EpicGames.Core/CsProjectInfo.cs
--- a/Engine/Source/Programs/Shared/EpicGames.Core/CsProjectInfo.cs
+++ b/Engine/Source/Programs/Shared/EpicGames.Core/CsProjectInfo.cs
@@ -372,7 +372,7 @@ namespace EpicGames.Core
 
 			// Check the root element is the right type
 			//			HashSet<FileReference> ProjectBuildProducts = new HashSet<FileReference>();
-			if (document.DocumentElement.Name != "Project")
+			if (document.DocumentElement!.Name != "Project")
 			{
 				outProjectInfo = null;
 				return false;
@@ -643,7 +643,7 @@ namespace EpicGames.Core
 		/// <returns>The contents of the child element, or default value if it's not present</returns>
 		static string? GetChildElementString(XmlElement parentElement, string name, string? defaultValue)
 		{
-			XmlElement childElement = parentElement.ChildNodes.OfType<XmlElement>().FirstOrDefault(x => x.Name == name);
+			XmlElement? childElement = parentElement.ChildNodes.OfType<XmlElement>().FirstOrDefault(x => x.Name == name);
 			if (childElement == null)
 			{
 				return defaultValue;
diff --git a/Engine/Source/Programs/Shared/EpicGames.Core/DirectoryUtils.cs b/Engine/Source/Programs/Shared/EpicGames.Core/DirectoryUtils.cs
--- a/Engine/Source/Programs/Shared/EpicGames.Core/DirectoryUtils.cs
+++ b/Engine/Source/Programs/Shared/EpicGames.Core/DirectoryUtils.cs
@@ -8,7 +8,7 @@ namespace EpicGames.Core
 	/// <summary>
 	/// Utility functions for manipulating directories
 	/// </summary>
-	public class DirectoryUtils
+	public static class DirectoryUtils
 	{
 		/// <summary>
 		/// Finds the on-disk case of a a directory
@@ -17,7 +17,7 @@ namespace EpicGames.Core
 		/// <returns>New DirectoryInfo instance that represents the directory with the correct case</returns>
 		public static DirectoryInfo FindCorrectCase(DirectoryInfo info)
 		{
-			DirectoryInfo parentInfo = info.Parent;
+			DirectoryInfo? parentInfo = info.Parent;
 			if (parentInfo == null)
 			{
 				string fullName = info.FullName;
diff --git a/Engine/Source/Programs/Shared/EpicGames.Core/FileFilter.cs b/Engine/Source/Programs/Shared/EpicGames.Core/FileFilter.cs
--- a/Engine/Source/Programs/Shared/EpicGames.Core/FileFilter.cs
+++ b/Engine/Source/Programs/Shared/EpicGames.Core/FileFilter.cs
@@ -330,7 +330,7 @@ namespace EpicGames.Core
 			FileFilterNode lastNode = _rootNode;
 			foreach (string branchPattern in branchPatterns)
 			{
-				FileFilterNode nextNode = lastNode.Branches.FirstOrDefault(x => x.Pattern == branchPattern);
+				FileFilterNode? nextNode = lastNode.Branches.FirstOrDefault(x => x.Pattern == branchPattern);
 				if (nextNode == null)
 				{
 					nextNode = new FileFilterNode(lastNode, branchPattern);
diff --git a/Engine/Source/Programs/Shared/EpicGames.Core/FileUtils.cs b/Engine/Source/Programs/Shared/EpicGames.Core/FileUtils.cs
--- a/Engine/Source/Programs/Shared/EpicGames.Core/FileUtils.cs
+++ b/Engine/Source/Programs/Shared/EpicGames.Core/FileUtils.cs
@@ -82,7 +82,7 @@ namespace EpicGames.Core
 		/// <returns>New FileInfo instance that represents the file with the correct case</returns>
 		public static FileInfo FindCorrectCase(FileInfo info)
 		{
-			DirectoryInfo parentInfo = DirectoryUtils.FindCorrectCase(info.Directory);
+			DirectoryInfo parentInfo = DirectoryUtils.FindCorrectCase(info.Directory!);
 			if (info.Exists)
 			{
 				foreach (FileInfo childInfo in parentInfo.EnumerateFiles())
diff --git a/Engine/Source/Programs/Shared/EpicGames.Core/HttpClientExtensions.cs b/Engine/Source/Programs/Shared/EpicGames.Core/HttpClientExtensions.cs
--- a/Engine/Source/Programs/Shared/EpicGames.Core/HttpClientExtensions.cs
+++ b/Engine/Source/Programs/Shared/EpicGames.Core/HttpClientExtensions.cs
@@ -116,7 +116,7 @@ namespace EpicGames.Core
 		private static async Task<T> ParseJsonContent<T>(HttpResponseMessage message)
 		{
 			byte[] bytes = await message.Content.ReadAsByteArrayAsync();
-			return JsonSerializer.Deserialize<T>(bytes, new JsonSerializerOptions { PropertyNameCaseInsensitive = true });
+			return JsonSerializer.Deserialize<T>(bytes, new JsonSerializerOptions { PropertyNameCaseInsensitive = true })!;
 		}
 	}
 }
diff --git a/Engine/Source/Programs/Shared/EpicGames.Core/JsonKnownTypes.cs b/Engine/Source/Programs/Shared/EpicGames.Core/JsonKnownTypes.cs
--- a/Engine/Source/Programs/Shared/EpicGames.Core/JsonKnownTypes.cs
+++ b/Engine/Source/Programs/Shared/EpicGames.Core/JsonKnownTypes.cs
@@ -120,13 +120,13 @@ namespace EpicGames.Core
 				}
 
 				Type? targetType;
-				if (!_discriminatorToType.TryGetValue(element.GetString(), out targetType))
+				if (!_discriminatorToType.TryGetValue(element.GetString()!, out targetType))
 				{
 					throw new JsonException($"Invalid '{TypePropertyName}' field ('{element.GetString()}') for serializing {typeToConvert.Name}");
 				}
 
 				string text = document.RootElement.GetRawText();
-				T result = (T)JsonSerializer.Deserialize(text, targetType, options);
+				T result = (T)JsonSerializer.Deserialize(text, targetType, options)!;
 				return result;
 			}
 		}
diff --git a/Engine/Source/Programs/Shared/EpicGames.Core/JsonSchema.cs b/Engine/Source/Programs/Shared/EpicGames.Core/JsonSchema.cs
--- a/Engine/Source/Programs/Shared/EpicGames.Core/JsonSchema.cs
+++ b/Engine/Source/Programs/Shared/EpicGames.Core/JsonSchema.cs
@@ -303,7 +303,7 @@ namespace EpicGames.Core
 
 			string selector = $"//member[@name='P:{type.FullName}.{name}']/summary";
 
-			XmlNode node = xmlDoc.SelectSingleNode(selector);
+			XmlNode? node = xmlDoc.SelectSingleNode(selector);
 			if (node == null)
 			{
 				return null;
diff --git a/Engine/Source/Programs/Shared/EpicGames.Core/LogEvent.cs b/Engine/Source/Programs/Shared/EpicGames.Core/LogEvent.cs
--- a/Engine/Source/Programs/Shared/EpicGames.Core/LogEvent.cs
+++ b/Engine/Source/Programs/Shared/EpicGames.Core/LogEvent.cs
@@ -224,7 +224,7 @@ namespace EpicGames.Core
 				}
 				else if (Utf8StringComparer.OrdinalIgnoreCase.Equals(propertyName, LogEventPropertyName.Level.Span))
 				{
-					level = Enum.Parse<LogLevel>(reader.GetString());
+					level = Enum.Parse<LogLevel>(reader.GetString()!);
 				}
 				else if (Utf8StringComparer.OrdinalIgnoreCase.Equals(propertyName, LogEventPropertyName.Id.Span))
 				{
@@ -240,11 +240,11 @@ namespace EpicGames.Core
 				}
 				else if (Utf8StringComparer.OrdinalIgnoreCase.Equals(propertyName, LogEventPropertyName.Message.Span))
 				{
-					message = reader.GetString();
+					message = reader.GetString()!;
 				}
 				else if (Utf8StringComparer.OrdinalIgnoreCase.Equals(propertyName, LogEventPropertyName.Format.Span))
 				{
-					format = reader.GetString();
+					format = reader.GetString()!;
 				}
 				else if (Utf8StringComparer.OrdinalIgnoreCase.Equals(propertyName, LogEventPropertyName.Properties.Span))
 				{
@@ -281,7 +281,7 @@ namespace EpicGames.Core
 				case JsonTokenType.StartObject:
 					return ReadStructuredPropertyValue(ref reader);
 				case JsonTokenType.String:
-					return reader.GetString();
+					return reader.GetString()!;
 				case JsonTokenType.Number:
 					if (reader.TryGetInt32(out int intValue))
 					{
@@ -311,11 +311,11 @@ namespace EpicGames.Core
 			{
 				if (Utf8StringComparer.OrdinalIgnoreCase.Equals(propertyName, LogEventPropertyName.Type.Span))
 				{
-					type = reader.GetString();
+					type = reader.GetString() ?? String.Empty;
 				}
 				else if (Utf8StringComparer.OrdinalIgnoreCase.Equals(propertyName, LogEventPropertyName.Text.Span))
 				{
-					text = reader.GetString();
+					text = reader.GetString() ?? String.Empty;
 				}
 				else
 				{
diff --git a/Engine/Source/Programs/Shared/EpicGames.Core/LogEventParser.cs b/Engine/Source/Programs/Shared/EpicGames.Core/LogEventParser.cs
--- a/Engine/Source/Programs/Shared/EpicGames.Core/LogEventParser.cs
+++ b/Engine/Source/Programs/Shared/EpicGames.Core/LogEventParser.cs
@@ -404,11 +404,11 @@ namespace EpicGames.Core
 			{
 				// Try to match an event
 				List<LogEvent>? events = null;
-				if (Regex.IsMatch(_buffer[0], "<-- Suspend Log Parsing -->", RegexOptions.IgnoreCase))
+				if (Regex.IsMatch(_buffer[0]!, "<-- Suspend Log Parsing -->", RegexOptions.IgnoreCase))
 				{
 					_matchingEnabled--;
 				}
-				else if (Regex.IsMatch(_buffer[0], "<-- Resume Log Parsing -->", RegexOptions.IgnoreCase))
+				else if (Regex.IsMatch(_buffer[0]!, "<-- Resume Log Parsing -->", RegexOptions.IgnoreCase))
 				{
 					_matchingEnabled++;
 				}
@@ -429,7 +429,7 @@ namespace EpicGames.Core
 				{
 					foreach (Regex ignorePattern in IgnorePatterns)
 					{
-						if (ignorePattern.IsMatch(_buffer[0]))
+						if (ignorePattern.IsMatch(_buffer[0]!))
 						{
 							events = null;
 							break;
diff --git a/Engine/Source/Programs/Shared/EpicGames.Core/Md5Hash.cs b/Engine/Source/Programs/Shared/EpicGames.Core/Md5Hash.cs
--- a/Engine/Source/Programs/Shared/EpicGames.Core/Md5Hash.cs
+++ b/Engine/Source/Programs/Shared/EpicGames.Core/Md5Hash.cs
@@ -194,13 +194,13 @@ namespace EpicGames.Core
 	sealed class Md5HashTypeConverter : TypeConverter
 	{
 		/// <inheritdoc/>
-		public override bool CanConvertFrom(ITypeDescriptorContext context, Type sourceType)
+		public override bool CanConvertFrom(ITypeDescriptorContext? context, Type sourceType)
 		{
 			return sourceType == typeof(string);
 		}
 
 		/// <inheritdoc/>
-		public override object ConvertFrom(ITypeDescriptorContext context, CultureInfo culture, object value)
+		public override object ConvertFrom(ITypeDescriptorContext? context, CultureInfo? culture, object value)
 		{
 			return Md5Hash.Parse((string)value);
 		}
diff --git a/Engine/Source/Programs/Shared/EpicGames.Core/ProcessUtils.cs b/Engine/Source/Programs/Shared/EpicGames.Core/ProcessUtils.cs
--- a/Engine/Source/Programs/Shared/EpicGames.Core/ProcessUtils.cs
+++ b/Engine/Source/Programs/Shared/EpicGames.Core/ProcessUtils.cs
@@ -171,7 +171,7 @@ namespace EpicGames.Core
 					FileReference? imageFile;
 					try
 					{
-						imageFile = new FileReference(process.MainModule.FileName);
+						imageFile = new FileReference(process.MainModule!.FileName!);
 					}
 					catch
 					{
diff --git a/Engine/Source/Programs/Shared/EpicGames.Core/Sha1Hash.cs b/Engine/Source/Programs/Shared/EpicGames.Core/Sha1Hash.cs
--- a/Engine/Source/Programs/Shared/EpicGames.Core/Sha1Hash.cs
+++ b/Engine/Source/Programs/Shared/EpicGames.Core/Sha1Hash.cs
@@ -229,13 +229,13 @@ namespace EpicGames.Core
 	sealed class Sha1HashTypeConverter : TypeConverter
 	{
 		/// <inheritdoc/>
-		public override bool CanConvertFrom(ITypeDescriptorContext context, Type sourceType)
+		public override bool CanConvertFrom(ITypeDescriptorContext? context, Type sourceType)
 		{
 			return sourceType == typeof(string);
 		}
 
 		/// <inheritdoc/>
-		public override object ConvertFrom(ITypeDescriptorContext context, CultureInfo culture, object value)
+		public override object ConvertFrom(ITypeDescriptorContext? context, CultureInfo? culture, object value)
 		{
 			return Sha1Hash.Parse((string)value);
 		}
diff --git a/Engine/Source/Programs/Shared/EpicGames.Core/StartupTraceListener.cs b/Engine/Source/Programs/Shared/EpicGames.Core/StartupTraceListener.cs
--- a/Engine/Source/Programs/Shared/EpicGames.Core/StartupTraceListener.cs
+++ b/Engine/Source/Programs/Shared/EpicGames.Core/StartupTraceListener.cs
@@ -28,7 +28,7 @@ namespace EpicGames.Core
 		/// Write a message to the buffer
 		/// </summary>
 		/// <param name="message">The message to write</param>
-		public override void Write(string message)
+		public override void Write(string? message)
 		{
 			if(NeedIndent)
 			{
@@ -41,10 +41,10 @@ namespace EpicGames.Core
 		/// Write a message to the buffer, followed by a newline
 		/// </summary>
 		/// <param name="message">The message to write</param>
-		public override void WriteLine(string message)
+		public override void WriteLine(string? message)
 		{
 			Write(message);
-			_buffer.Append("\n");
+			_buffer.Append('\n');
 		}
 	}
 }
