EpicGames.BuildGraph: Fix warnings when compiling against NET 6 SDK.

diff --git a/Engine/Source/Programs/Shared/EpicGames.BuildGraph.Tests/GraphTests.cs b/Engine/Source/Programs/Shared/EpicGames.BuildGraph.Tests/GraphTests.cs
--- a/Engine/Source/Programs/Shared/EpicGames.BuildGraph.Tests/GraphTests.cs
+++ b/Engine/Source/Programs/Shared/EpicGames.BuildGraph.Tests/GraphTests.cs
@@ -1,15 +1,9 @@
 // Copyright Epic Games, Inc. All Rights Reserved.
 
 using Microsoft.VisualStudio.TestTools.UnitTesting;
-using System;
 using System.Collections.Generic;
-using System.IO;
 using System.Linq;
-using System.Text;
-using System.Text.Json;
-using OpenTracing;
 using EpicGames.BuildGraph.Expressions;
-using Microsoft.Extensions.Logging;
 using System.Threading.Tasks;
 
 namespace EpicGames.BuildGraph.Tests
@@ -21,7 +15,7 @@ namespace EpicGames.BuildGraph.Tests
 		static Task CompileShooterGameWin64() => Task.CompletedTask;
 		static Task CookShooterGameWin64() => Task.CompletedTask;
 
-		object Evaluate(BgExpr expr)
+		static object Evaluate(BgExpr expr)
 		{
 			(byte[] data, BgThunkDef[] methods) = BgCompiler.Compile(expr);
 			BgInterpreter interpreter = new BgInterpreter(data, methods, new Dictionary<string, string>());
diff --git a/Engine/Source/Programs/Shared/EpicGames.BuildGraph.Tests/OpcodeTests.cs b/Engine/Source/Programs/Shared/EpicGames.BuildGraph.Tests/OpcodeTests.cs
--- a/Engine/Source/Programs/Shared/EpicGames.BuildGraph.Tests/OpcodeTests.cs
+++ b/Engine/Source/Programs/Shared/EpicGames.BuildGraph.Tests/OpcodeTests.cs
@@ -5,11 +5,7 @@ using System;
 using System.Collections.Generic;
 using System.IO;
 using System.Linq;
-using System.Text;
-using System.Text.Json;
-using OpenTracing;
 using EpicGames.BuildGraph.Expressions;
-using System.Collections.Immutable;
 using System.Threading.Tasks;
 
 namespace EpicGames.BuildGraph.Tests
@@ -17,14 +13,14 @@ namespace EpicGames.BuildGraph.Tests
 	[TestClass]
 	public class OpcodeTests
 	{
-		object Evaluate(BgExpr expr)
+		static object Evaluate(BgExpr expr)
 		{
 			(byte[] data, BgThunkDef[] methods) = BgCompiler.Compile(expr);
 			BgInterpreter interpreter = new BgInterpreter(data, methods, new Dictionary<string, string>());
 			return interpreter.Evaluate();
 		}
 
-		BgOptionDef EvaluateOption(BgExpr expr)
+		static BgOptionDef EvaluateOption(BgExpr expr)
 		{
 			(byte[] data, BgThunkDef[] methods) = BgCompiler.Compile(expr);
 			BgInterpreter interpreter = new BgInterpreter(data, methods, new Dictionary<string, string>());
@@ -353,8 +349,8 @@ namespace EpicGames.BuildGraph.Tests
 			Assert.AreEqual("value", result.DefaultValue);
 			Assert.AreEqual("abc", result.Pattern);
 			Assert.AreEqual("def", result.PatternFailed);
-			Assert.IsTrue(result.Values.SequenceEqual(new[] { "value1", "value2" }));
-			Assert.IsTrue(result.ValueDescriptions.SequenceEqual(new[] { "desc1", "desc2" }));
+			Assert.IsTrue(result.Values!.SequenceEqual(new[] { "value1", "value2" }));
+			Assert.IsTrue(result.ValueDescriptions!.SequenceEqual(new[] { "desc1", "desc2" }));
 		}
 
 		#endregion
@@ -489,8 +485,8 @@ namespace EpicGames.BuildGraph.Tests
 			Assert.AreEqual("stringOption", result.Name);
 			Assert.AreEqual("description", result.Description);
 			Assert.AreEqual("value", result.DefaultValue);
-			Assert.IsTrue(result.Values.SequenceEqual(new[] { "value1", "value2" }));
-			Assert.IsTrue(result.ValueDescriptions.SequenceEqual(new[] { "desc1", "desc2" }));
+			Assert.IsTrue(result.Values!.SequenceEqual(new[] { "value1", "value2" }));
+			Assert.IsTrue(result.ValueDescriptions!.SequenceEqual(new[] { "desc1", "desc2" }));
 		}
 
 		#endregion
diff --git a/Engine/Source/Programs/Shared/EpicGames.BuildGraph/.editorconfig b/Engine/Source/Programs/Shared/EpicGames.BuildGraph/.editorconfig
new file mode 100644
--- /dev/null
+++ b/Engine/Source/Programs/Shared/EpicGames.BuildGraph/.editorconfig
@@ -0,0 +1,17 @@
+# Dotnet code style settings:
+[*.cs]
+
+# CA2227: Collection properties should be read only
+dotnet_diagnostic.CA2227.severity = none
+
+# CA1000: Do not declare static members on generic types
+dotnet_diagnostic.CA1000.severity = none
+
+# CA1822: Mark members as static
+dotnet_diagnostic.CA1822.severity = none
+
+# CA1005: Avoid excessive parameters on generic types
+dotnet_diagnostic.CA1005.severity = none
+
+# CA1043: Use Integral Or String Argument For Indexers
+dotnet_diagnostic.CA1043.severity = none
diff --git a/Engine/Source/Programs/Shared/EpicGames.BuildGraph/BgContext.cs b/Engine/Source/Programs/Shared/EpicGames.BuildGraph/BgContext.cs
--- a/Engine/Source/Programs/Shared/EpicGames.BuildGraph/BgContext.cs
+++ b/Engine/Source/Programs/Shared/EpicGames.BuildGraph/BgContext.cs
@@ -42,12 +42,12 @@ namespace EpicGames.BuildGraph
 		/// </summary>
 		public HashSet<FileReference> BuildProducts { get; } = new HashSet<FileReference>();
 
-		Dictionary<string, FileSet> _tagNameToFileSet;
+		readonly Dictionary<string, FileSet> _tagNameToFileSet;
 
 		/// <summary>
 		/// Constructor
 		/// </summary>
-		public BgContext(Dictionary<string, FileSet> tagNameToFileSet)
+		protected BgContext(Dictionary<string, FileSet> tagNameToFileSet)
 		{
 			_tagNameToFileSet = tagNameToFileSet;
 		}
diff --git a/Engine/Source/Programs/Shared/EpicGames.BuildGraph/BgGraphDef.cs b/Engine/Source/Programs/Shared/EpicGames.BuildGraph/BgGraphDef.cs
--- a/Engine/Source/Programs/Shared/EpicGames.BuildGraph/BgGraphDef.cs
+++ b/Engine/Source/Programs/Shared/EpicGames.BuildGraph/BgGraphDef.cs
@@ -6,7 +6,6 @@ using System.Diagnostics.CodeAnalysis;
 using System.Linq;
 using System.Runtime.CompilerServices;
 using System.Text;
-using System.Xml;
 using EpicGames.Core;
 using Microsoft.Extensions.Logging;
 
@@ -17,8 +16,14 @@ namespace EpicGames.BuildGraph
 	/// <summary>
 	/// Options for how the graph should be printed
 	/// </summary>
+	[Flags]
 	public enum GraphPrintOptions
 	{
+		/// <summary>
+		/// No options specified
+		/// </summary>
+		None = 0,
+
 		/// <summary>
 		/// Includes a list of the graph options
 		/// </summary>
@@ -134,7 +139,7 @@ namespace EpicGames.BuildGraph
 		public bool TryResolveReference(string name, [NotNullWhen(true)] out BgNodeDef[]? outNodes)
 		{
 			// Check if it's a tag reference or node reference
-			if (name.StartsWith("#"))
+			if (name.StartsWith("#", StringComparison.Ordinal))
 			{
 				// Check if it's a regular node or output name
 				BgNodeOutput? output;
@@ -185,7 +190,7 @@ namespace EpicGames.BuildGraph
 		public bool TryResolveInputReference(string name, [NotNullWhen(true)] out BgNodeOutput[]? outOutputs)
 		{
 			// Check if it's a tag reference or node reference
-			if (name.StartsWith("#"))
+			if (name.StartsWith("#", StringComparison.Ordinal))
 			{
 				// Check if it's a regular node or output name
 				BgNodeOutput? output;
@@ -605,7 +610,7 @@ namespace EpicGames.BuildGraph
 
 					foreach (string line in lines)
 					{
-						logger.Log(LogLevel.Information, line);
+						logger.Log(LogLevel.Information, "{Line}", line);
 					}
 				}
 			}
@@ -615,21 +620,21 @@ namespace EpicGames.BuildGraph
 			logger.LogInformation("Graph:");
 			foreach (BgAgentDef agent in Agents)
 			{
-				logger.LogInformation("        Agent: {0} ({1})", agent.Name, String.Join(";", agent.PossibleTypes));
+				logger.LogInformation("        Agent: {Name} ({Types})", agent.Name, String.Join(";", agent.PossibleTypes));
 				foreach (BgNodeDef node in agent.Nodes)
 				{
-					logger.LogInformation("            Node: {0}{1}", node.Name, completedNodes.Contains(node) ? " (completed)" : node.RunEarly ? " (early)" : "");
+					logger.LogInformation("            Node: {Name}{Type}", node.Name, completedNodes.Contains(node) ? " (completed)" : node.RunEarly ? " (early)" : "");
 					if (printOptions.HasFlag(GraphPrintOptions.ShowDependencies))
 					{
 						HashSet<BgNodeDef> inputDependencies = new HashSet<BgNodeDef>(node.GetDirectInputDependencies());
 						foreach (BgNodeDef inputDependency in inputDependencies)
 						{
-							logger.LogInformation("                input> {0}", inputDependency.Name);
+							logger.LogInformation("                input> {Name}", inputDependency.Name);
 						}
 						HashSet<BgNodeDef> orderDependencies = new HashSet<BgNodeDef>(node.GetDirectOrderDependencies());
 						foreach (BgNodeDef orderDependency in orderDependencies.Except(inputDependencies))
 						{
-							logger.LogInformation("                after> {0}", orderDependency.Name);
+							logger.LogInformation("                after> {Name}", orderDependency.Name);
 						}
 					}
 					if (printOptions.HasFlag(GraphPrintOptions.ShowNotifications))
@@ -637,11 +642,11 @@ namespace EpicGames.BuildGraph
 						string label = node.NotifyOnWarnings ? "warnings" : "errors";
 						foreach (string user in node.NotifyUsers)
 						{
-							logger.LogInformation("                {0}> {1}", label, user);
+							logger.LogInformation("                {Name}> {User}", label, user);
 						}
 						foreach (string submitter in node.NotifySubmitters)
 						{
-							logger.LogInformation("                {0}> submitters to {1}", label, submitter);
+							logger.LogInformation("                {Name}> submitters to {Submitter}", label, submitter);
 						}
 					}
 				}
@@ -655,7 +660,7 @@ namespace EpicGames.BuildGraph
 				logger.LogInformation("Aggregates:");
 				foreach (string aggregateName in aggregates.Select(x => x.Name))
 				{
-					logger.LogInformation("    {0}", aggregateName);
+					logger.LogInformation("    {Aggregate}", aggregateName);
 				}
 				logger.LogInformation("");
 			}
diff --git a/Engine/Source/Programs/Shared/EpicGames.BuildGraph/BgInterpreter.cs b/Engine/Source/Programs/Shared/EpicGames.BuildGraph/BgInterpreter.cs
--- a/Engine/Source/Programs/Shared/EpicGames.BuildGraph/BgInterpreter.cs
+++ b/Engine/Source/Programs/Shared/EpicGames.BuildGraph/BgInterpreter.cs
@@ -637,7 +637,7 @@ namespace EpicGames.BuildGraph
 		/// </summary>
 		string ReadString(Frame frame)
 		{
-			ReadOnlySpan<byte> buffer = _data[frame.Offset..];
+			ReadOnlySpan<byte> buffer = _data.AsSpan(frame.Offset);
 
 			int length = (int)VarInt.ReadUnsigned(buffer, out int bytesRead);
 			frame.Offset += bytesRead;
@@ -664,7 +664,7 @@ namespace EpicGames.BuildGraph
 		/// <returns>The value that was read</returns>
 		ulong ReadUnsignedInteger(Frame frame)
 		{
-			ulong value = VarInt.ReadUnsigned(_data[frame.Offset..], out int bytesRead);
+			ulong value = VarInt.ReadUnsigned(_data.AsSpan(frame.Offset), out int bytesRead);
 			frame.Offset += bytesRead;
 			return value;
 		}
diff --git a/Engine/Source/Programs/Shared/EpicGames.BuildGraph/BgNodeDef.cs b/Engine/Source/Programs/Shared/EpicGames.BuildGraph/BgNodeDef.cs
--- a/Engine/Source/Programs/Shared/EpicGames.BuildGraph/BgNodeDef.cs
+++ b/Engine/Source/Programs/Shared/EpicGames.BuildGraph/BgNodeDef.cs
@@ -181,7 +181,7 @@ namespace EpicGames.BuildGraph
 
 			List<BgNodeOutput> allOutputs = new List<BgNodeOutput>();
 			allOutputs.Add(new BgNodeOutput(this, "#" + Name));
-			allOutputs.AddRange(outputNames.Where(x => String.Compare(x, Name, StringComparison.InvariantCultureIgnoreCase) != 0).Select(x => new BgNodeOutput(this, x)));
+			allOutputs.AddRange(outputNames.Where(x => !String.Equals(x, Name, StringComparison.OrdinalIgnoreCase)).Select(x => new BgNodeOutput(this, x)));
 			Outputs = allOutputs.ToArray();
 		}
 
diff --git a/Engine/Source/Programs/Shared/EpicGames.BuildGraph/BgObjectDef.cs b/Engine/Source/Programs/Shared/EpicGames.BuildGraph/BgObjectDef.cs
--- a/Engine/Source/Programs/Shared/EpicGames.BuildGraph/BgObjectDef.cs
+++ b/Engine/Source/Programs/Shared/EpicGames.BuildGraph/BgObjectDef.cs
@@ -295,8 +295,9 @@ namespace EpicGames.BuildGraph
 	}
 
 	/// <summary>
-	/// Attribute marking that a property should be serialized to BuildGraph
+	/// Attribute marking the type of serializer for an object
 	/// </summary>
+	[AttributeUsage(AttributeTargets.Class)]
 	public sealed class BgObjectAttribute : Attribute
 	{
 		/// <summary>
diff --git a/Engine/Source/Programs/Shared/EpicGames.BuildGraph/BgOptionDef.cs b/Engine/Source/Programs/Shared/EpicGames.BuildGraph/BgOptionDef.cs
--- a/Engine/Source/Programs/Shared/EpicGames.BuildGraph/BgOptionDef.cs
+++ b/Engine/Source/Programs/Shared/EpicGames.BuildGraph/BgOptionDef.cs
@@ -2,7 +2,6 @@
 
 using System;
 using System.Collections.Generic;
-using System.Text;
 
 namespace EpicGames.BuildGraph
 {
@@ -30,7 +29,7 @@ namespace EpicGames.BuildGraph
 		/// Constructor
 		/// </summary>
 		/// <param name="name">The name of this option</param>
-		public BgOptionDef(string name)
+		protected BgOptionDef(string name)
 		{
 			Name = name;
 		}
diff --git a/Engine/Source/Programs/Shared/EpicGames.BuildGraph/Expressions/BgBool.cs b/Engine/Source/Programs/Shared/EpicGames.BuildGraph/Expressions/BgBool.cs
--- a/Engine/Source/Programs/Shared/EpicGames.BuildGraph/Expressions/BgBool.cs
+++ b/Engine/Source/Programs/Shared/EpicGames.BuildGraph/Expressions/BgBool.cs
@@ -26,7 +26,7 @@ namespace EpicGames.BuildGraph.Expressions
 		/// Constructor
 		/// </summary>
 		/// <param name="flags">Flags for this expression</param>
-		public BgBool(BgExprFlags flags)
+		protected BgBool(BgExprFlags flags)
 			: base(flags)
 		{
 		}
diff --git a/Engine/Source/Programs/Shared/EpicGames.BuildGraph/Expressions/BgEnum.cs b/Engine/Source/Programs/Shared/EpicGames.BuildGraph/Expressions/BgEnum.cs
--- a/Engine/Source/Programs/Shared/EpicGames.BuildGraph/Expressions/BgEnum.cs
+++ b/Engine/Source/Programs/Shared/EpicGames.BuildGraph/Expressions/BgEnum.cs
@@ -27,7 +27,7 @@ namespace EpicGames.BuildGraph.Expressions
 		/// <summary>
 		/// Constructor
 		/// </summary>
-		public BgEnum(BgExprFlags flags)
+		protected BgEnum(BgExprFlags flags)
 			: base(flags)
 		{
 		}
diff --git a/Engine/Source/Programs/Shared/EpicGames.BuildGraph/Expressions/BgExpr.cs b/Engine/Source/Programs/Shared/EpicGames.BuildGraph/Expressions/BgExpr.cs
--- a/Engine/Source/Programs/Shared/EpicGames.BuildGraph/Expressions/BgExpr.cs
+++ b/Engine/Source/Programs/Shared/EpicGames.BuildGraph/Expressions/BgExpr.cs
@@ -56,7 +56,7 @@ namespace EpicGames.BuildGraph
 		/// Constructor
 		/// </summary>
 		/// <param name="flags"></param>
-		public BgExpr(BgExprFlags flags)
+		protected BgExpr(BgExprFlags flags)
 		{
 			Flags = flags;
 		}
diff --git a/Engine/Source/Programs/Shared/EpicGames.BuildGraph/Expressions/BgFileSet.cs b/Engine/Source/Programs/Shared/EpicGames.BuildGraph/Expressions/BgFileSet.cs
--- a/Engine/Source/Programs/Shared/EpicGames.BuildGraph/Expressions/BgFileSet.cs
+++ b/Engine/Source/Programs/Shared/EpicGames.BuildGraph/Expressions/BgFileSet.cs
@@ -22,7 +22,7 @@ namespace EpicGames.BuildGraph.Expressions
 		/// Constructor
 		/// </summary>
 		/// <param name="flags"></param>
-		public BgFileSet(BgExprFlags flags)
+		protected BgFileSet(BgExprFlags flags)
 			: base(flags)
 		{
 		}
diff --git a/Engine/Source/Programs/Shared/EpicGames.BuildGraph/Expressions/BgInt.cs b/Engine/Source/Programs/Shared/EpicGames.BuildGraph/Expressions/BgInt.cs
--- a/Engine/Source/Programs/Shared/EpicGames.BuildGraph/Expressions/BgInt.cs
+++ b/Engine/Source/Programs/Shared/EpicGames.BuildGraph/Expressions/BgInt.cs
@@ -24,7 +24,7 @@ namespace EpicGames.BuildGraph.Expressions
 		/// Constructor
 		/// </summary>
 		/// <param name="flags">Flags for this expression</param>
-		public BgInt(BgExprFlags flags)
+		protected BgInt(BgExprFlags flags)
 			: base(flags)
 		{
 		}
diff --git a/Engine/Source/Programs/Shared/EpicGames.BuildGraph/Expressions/BgList.cs b/Engine/Source/Programs/Shared/EpicGames.BuildGraph/Expressions/BgList.cs
--- a/Engine/Source/Programs/Shared/EpicGames.BuildGraph/Expressions/BgList.cs
+++ b/Engine/Source/Programs/Shared/EpicGames.BuildGraph/Expressions/BgList.cs
@@ -83,7 +83,7 @@ namespace EpicGames.BuildGraph.Expressions
 		/// Constructor
 		/// </summary>
 		/// <param name="flags">Flags for this expression</param>
-		public BgList(BgExprFlags flags)
+		protected BgList(BgExprFlags flags)
 			: base(flags)
 		{
 		}
diff --git a/Engine/Source/Programs/Shared/EpicGames.BuildGraph/Expressions/BgNode.cs b/Engine/Source/Programs/Shared/EpicGames.BuildGraph/Expressions/BgNode.cs
--- a/Engine/Source/Programs/Shared/EpicGames.BuildGraph/Expressions/BgNode.cs
+++ b/Engine/Source/Programs/Shared/EpicGames.BuildGraph/Expressions/BgNode.cs
@@ -29,7 +29,7 @@ namespace EpicGames.BuildGraph.Expressions
 	/// Speecifies the node name for a method. Parameters from the method may be embedded in the name using the {ParamName} syntax.
 	/// </summary>
 	[AttributeUsage(AttributeTargets.Method)]
-	public class BgNodeNameAttribute : Attribute
+	public sealed class BgNodeNameAttribute : Attribute
 	{
 		/// <summary>
 		/// The format string
diff --git a/Engine/Source/Programs/Shared/EpicGames.BuildGraph/Expressions/BgObject.cs b/Engine/Source/Programs/Shared/EpicGames.BuildGraph/Expressions/BgObject.cs
--- a/Engine/Source/Programs/Shared/EpicGames.BuildGraph/Expressions/BgObject.cs
+++ b/Engine/Source/Programs/Shared/EpicGames.BuildGraph/Expressions/BgObject.cs
@@ -25,7 +25,7 @@ namespace EpicGames.BuildGraph.Expressions
 		/// Constructor
 		/// </summary>
 		/// <param name="flags">Flags for this expression</param>
-		public BgObject(BgExprFlags flags)
+		protected BgObject(BgExprFlags flags)
 			: base(flags)
 		{
 		}
diff --git a/Engine/Source/Programs/Shared/EpicGames.BuildGraph/Expressions/BgString.cs b/Engine/Source/Programs/Shared/EpicGames.BuildGraph/Expressions/BgString.cs
--- a/Engine/Source/Programs/Shared/EpicGames.BuildGraph/Expressions/BgString.cs
+++ b/Engine/Source/Programs/Shared/EpicGames.BuildGraph/Expressions/BgString.cs
@@ -2,8 +2,6 @@
 
 using System;
 using System.Collections.Generic;
-using System.Linq;
-using System.Text.RegularExpressions;
 
 namespace EpicGames.BuildGraph.Expressions
 {
@@ -22,7 +20,7 @@ namespace EpicGames.BuildGraph.Expressions
 		/// Constructor
 		/// </summary>
 		/// <param name="flags">Flags for this expression</param>
-		public BgString(BgExprFlags flags)
+		protected BgString(BgExprFlags flags)
 			: base(flags)
 		{
 		}
diff --git a/Engine/Source/Programs/Shared/EpicGames.BuildGraph/Expressions/BgType.cs b/Engine/Source/Programs/Shared/EpicGames.BuildGraph/Expressions/BgType.cs
--- a/Engine/Source/Programs/Shared/EpicGames.BuildGraph/Expressions/BgType.cs
+++ b/Engine/Source/Programs/Shared/EpicGames.BuildGraph/Expressions/BgType.cs
@@ -41,7 +41,8 @@ namespace EpicGames.BuildGraph.Expressions
 	/// <summary>
 	/// Attribute used to specify the converter class to use for a type
 	/// </summary>
-	class BgTypeAttribute : Attribute
+	[AttributeUsage(AttributeTargets.Class)]
+	sealed class BgTypeAttribute : Attribute
 	{
 		/// <summary>
 		/// The converter type
