Suppress NET 6 warnings in EpicGames.Jupiter.

diff --git a/Engine/Source/Programs/Shared/EpicGames.Jupiter/.editorconfig b/Engine/Source/Programs/Shared/EpicGames.Jupiter/.editorconfig
new file mode 100644
--- /dev/null
+++ b/Engine/Source/Programs/Shared/EpicGames.Jupiter/.editorconfig
@@ -0,0 +1,42 @@
+# Dotnet code style settings:
+[*.cs]
+
+
+# CA1034: Nested types should not be visible
+dotnet_diagnostic.CA1034.severity = none
+
+# CA1819: Properties should not return arrays
+dotnet_diagnostic.CA1819.severity = none
+
+# CA1051: Do not declare visible instance fields
+dotnet_diagnostic.CA1051.severity = none
+
+# IDE0025: Use expression body for properties
+dotnet_diagnostic.IDE0025.severity = none
+
+# CA5350: Do Not Use Weak Cryptographic Algorithms
+dotnet_diagnostic.CA5350.severity = none
+
+# CA2234: Pass system uri objects instead of strings
+dotnet_diagnostic.CA2234.severity = none
+
+# IDE0049: Use framework type
+dotnet_diagnostic.IDE0049.severity = none
+
+# IDE1006: Naming Styles
+dotnet_diagnostic.IDE1006.severity = none
+
+# CA2227: Collection properties should be read only
+dotnet_diagnostic.CA2227.severity = none
+
+# CA1054: URI-like parameters should not be strings
+dotnet_diagnostic.CA1054.severity = none
+
+# CA1307: Specify StringComparison for clarity
+dotnet_diagnostic.CA1307.severity = none
+
+# CA1056: URI-like properties should not be strings
+dotnet_diagnostic.CA1056.severity = none
+
+# CA1849: Call async methods when in an async method
+dotnet_diagnostic.CA1849.severity = none
diff --git a/Engine/Source/Programs/Shared/EpicGames.Jupiter/JupiterTree.cs b/Engine/Source/Programs/Shared/EpicGames.Jupiter/JupiterTree.cs
--- a/Engine/Source/Programs/Shared/EpicGames.Jupiter/JupiterTree.cs
+++ b/Engine/Source/Programs/Shared/EpicGames.Jupiter/JupiterTree.cs
@@ -117,7 +117,7 @@ namespace EpicGames.Jupiter
 						Metadata = Metadata,
 					};
 					string PutTreeRootString = JsonSerializer.Serialize(TreeRoot);
-					StringContent PutTreeRootContent = new StringContent(PutTreeRootString, Encoding.UTF8, "application/json");
+					using StringContent PutTreeRootContent = new StringContent(PutTreeRootString, Encoding.UTF8, "application/json");
 					// upload tree information
 					HttpResponseMessage PutTreeResult = await JupiterClient.PutAsync(string.Format("api/v1/c/tree-root/{0}", JupiterNamespace), PutTreeRootContent);
 					if (!PutTreeResult.IsSuccessStatusCode)
@@ -140,7 +140,7 @@ namespace EpicGames.Jupiter
 							Blobs = Tree.ContentHashes,
 						};
 					string PutTreeString = JsonSerializer.Serialize(PutTreesRequest);
-					StringContent PutTreeContent = new StringContent(PutTreeString, Encoding.UTF8, "application/json");
+					using StringContent PutTreeContent = new StringContent(PutTreeString, Encoding.UTF8, "application/json");
 					HttpResponseMessage PutTreeContentResult = await JupiterClient.PutAsync(string.Format("api/v1/c/tree/{0}", JupiterNamespace), PutTreeContent);
 					if (!PutTreeContentResult.IsSuccessStatusCode)
 					{
@@ -166,7 +166,7 @@ namespace EpicGames.Jupiter
 					};
 
 					string FilterBlobString = JsonSerializer.Serialize(FilterBlobRequest);
-					StringContent FilterBlobContent = new StringContent(FilterBlobString, Encoding.UTF8, "application/json");
+					using StringContent FilterBlobContent = new StringContent(FilterBlobString, Encoding.UTF8, "application/json");
 					HttpResponseMessage FilterBlobResponse = await JupiterClient.PostAsync("api/v1/s", FilterBlobContent);
 					if (!FilterBlobResponse.IsSuccessStatusCode)
 					{
