Upgrade shared C# projects to NET 6.

diff --git a/Engine/Source/Programs/Shared/EpicGames.BuildGraph.Tests/EpicGames.BuildGraph.Tests.csproj b/Engine/Source/Programs/Shared/EpicGames.BuildGraph.Tests/EpicGames.BuildGraph.Tests.csproj
--- a/Engine/Source/Programs/Shared/EpicGames.BuildGraph.Tests/EpicGames.BuildGraph.Tests.csproj
+++ b/Engine/Source/Programs/Shared/EpicGames.BuildGraph.Tests/EpicGames.BuildGraph.Tests.csproj
@@ -1,7 +1,7 @@
 <Project Sdk="Microsoft.NET.Sdk">
 
   <PropertyGroup>
-    <TargetFramework>netcoreapp3.1</TargetFramework>
+    <TargetFramework>net6.0</TargetFramework>
     <Configuration Condition=" '$(Configuration)' == '' ">Development</Configuration>
     <OutputType>Library</OutputType>
     <Deterministic>true</Deterministic>
diff --git a/Engine/Source/Programs/Shared/EpicGames.BuildGraph/EpicGames.BuildGraph.csproj b/Engine/Source/Programs/Shared/EpicGames.BuildGraph/EpicGames.BuildGraph.csproj
--- a/Engine/Source/Programs/Shared/EpicGames.BuildGraph/EpicGames.BuildGraph.csproj
+++ b/Engine/Source/Programs/Shared/EpicGames.BuildGraph/EpicGames.BuildGraph.csproj
@@ -1,6 +1,6 @@
 ﻿<Project Sdk="Microsoft.NET.Sdk">
   <PropertyGroup>
-    <TargetFramework>netcoreapp3.1</TargetFramework>
+    <TargetFramework>net6.0</TargetFramework>
     <Configuration Condition=" '$(Configuration)' == '' ">Development</Configuration>
     <OutputType>Library</OutputType>
     <RootNamespace>EpicGames.BuildGraph</RootNamespace>
diff --git a/Engine/Source/Programs/Shared/EpicGames.Core.Tests/EpicGames.Core.Tests.csproj b/Engine/Source/Programs/Shared/EpicGames.Core.Tests/EpicGames.Core.Tests.csproj
--- a/Engine/Source/Programs/Shared/EpicGames.Core.Tests/EpicGames.Core.Tests.csproj
+++ b/Engine/Source/Programs/Shared/EpicGames.Core.Tests/EpicGames.Core.Tests.csproj
@@ -1,7 +1,7 @@
 <Project Sdk="Microsoft.NET.Sdk">
 
   <PropertyGroup>
-    <TargetFramework>netcoreapp3.1</TargetFramework>
+    <TargetFramework>net6.0</TargetFramework>
     <Configuration Condition=" '$(Configuration)' == '' ">Development</Configuration>
     <OutputType>Library</OutputType>
     <Deterministic>true</Deterministic>
diff --git a/Engine/Source/Programs/Shared/EpicGames.Core/EpicGames.Core.csproj b/Engine/Source/Programs/Shared/EpicGames.Core/EpicGames.Core.csproj
--- a/Engine/Source/Programs/Shared/EpicGames.Core/EpicGames.Core.csproj
+++ b/Engine/Source/Programs/Shared/EpicGames.Core/EpicGames.Core.csproj
@@ -1,6 +1,6 @@
 ﻿<Project Sdk="Microsoft.NET.Sdk">
   <PropertyGroup>
-    <TargetFramework>netcoreapp3.1</TargetFramework>
+    <TargetFramework>net6.0</TargetFramework>
     <Configuration Condition=" '$(Configuration)' == '' ">Development</Configuration>
     <OutputType>Library</OutputType>
     <Deterministic>true</Deterministic>
diff --git a/Engine/Source/Programs/Shared/EpicGames.IoHash/EpicGames.IoHash.csproj b/Engine/Source/Programs/Shared/EpicGames.IoHash/EpicGames.IoHash.csproj
--- a/Engine/Source/Programs/Shared/EpicGames.IoHash/EpicGames.IoHash.csproj
+++ b/Engine/Source/Programs/Shared/EpicGames.IoHash/EpicGames.IoHash.csproj
@@ -1,6 +1,6 @@
 ﻿<Project Sdk="Microsoft.NET.Sdk">
   <PropertyGroup>
-    <TargetFramework>netcoreapp3.1</TargetFramework>
+    <TargetFramework>net6.0</TargetFramework>
     <Configuration Condition=" '$(Configuration)' == '' ">Development</Configuration>
     <OutputType>Library</OutputType>
     <Deterministic>true</Deterministic>
diff --git a/Engine/Source/Programs/Shared/EpicGames.Jupiter/EpicGames.Jupiter.csproj b/Engine/Source/Programs/Shared/EpicGames.Jupiter/EpicGames.Jupiter.csproj
--- a/Engine/Source/Programs/Shared/EpicGames.Jupiter/EpicGames.Jupiter.csproj
+++ b/Engine/Source/Programs/Shared/EpicGames.Jupiter/EpicGames.Jupiter.csproj
@@ -1,6 +1,6 @@
 <Project Sdk="Microsoft.NET.Sdk">
   <PropertyGroup>
-    <TargetFramework>netcoreapp3.1</TargetFramework>
+    <TargetFramework>net6.0</TargetFramework>
     <Configuration Condition=" '$(Configuration)' == '' ">Development</Configuration>
     <OutputType>Library</OutputType>
     <Deterministic>true</Deterministic>
diff --git a/Engine/Source/Programs/Shared/EpicGames.Perforce.Managed/EpicGames.Perforce.Managed.csproj b/Engine/Source/Programs/Shared/EpicGames.Perforce.Managed/EpicGames.Perforce.Managed.csproj
--- a/Engine/Source/Programs/Shared/EpicGames.Perforce.Managed/EpicGames.Perforce.Managed.csproj
+++ b/Engine/Source/Programs/Shared/EpicGames.Perforce.Managed/EpicGames.Perforce.Managed.csproj
@@ -1,7 +1,7 @@
 <Project Sdk="Microsoft.NET.Sdk">
 
   <PropertyGroup>
-    <TargetFramework>netcoreapp3.1</TargetFramework>
+    <TargetFramework>net6.0</TargetFramework>
     <Nullable>enable</Nullable>
     <GenerateAssemblyInfo>false</GenerateAssemblyInfo>
     <GenerateTargetFrameworkAttribute>false</GenerateTargetFrameworkAttribute>
diff --git a/Engine/Source/Programs/Shared/EpicGames.Perforce/EpicGames.Perforce.csproj b/Engine/Source/Programs/Shared/EpicGames.Perforce/EpicGames.Perforce.csproj
--- a/Engine/Source/Programs/Shared/EpicGames.Perforce/EpicGames.Perforce.csproj
+++ b/Engine/Source/Programs/Shared/EpicGames.Perforce/EpicGames.Perforce.csproj
@@ -1,7 +1,7 @@
 <Project Sdk="Microsoft.NET.Sdk">
 
   <PropertyGroup>
-    <TargetFramework>netcoreapp3.1</TargetFramework>
+    <TargetFramework>net6.0</TargetFramework>
     <Configuration Condition=" '$(Configuration)' == '' ">Development</Configuration>
     <OutputType>Library</OutputType>
     <RootNamespace>EpicGames.Perforce</RootNamespace>
diff --git a/Engine/Source/Programs/Shared/EpicGames.Serialization/EpicGames.Serialization.csproj b/Engine/Source/Programs/Shared/EpicGames.Serialization/EpicGames.Serialization.csproj
--- a/Engine/Source/Programs/Shared/EpicGames.Serialization/EpicGames.Serialization.csproj
+++ b/Engine/Source/Programs/Shared/EpicGames.Serialization/EpicGames.Serialization.csproj
@@ -1,6 +1,6 @@
 <Project Sdk="Microsoft.NET.Sdk">
   <PropertyGroup>
-    <TargetFramework>netcoreapp3.1</TargetFramework>
+    <TargetFramework>net6.0</TargetFramework>
     <Configuration Condition=" '$(Configuration)' == '' ">Development</Configuration>
     <OutputType>Library</OutputType>
     <Deterministic>true</Deterministic>
