EpicGames.Core: Fix static analysis warnings.

diff --git a/Engine/Source/Programs/Horde/Horde.Build/Jobs/JobMessages.cs b/Engine/Source/Programs/Horde/Horde.Build/Jobs/JobMessages.cs
--- a/Engine/Source/Programs/Horde/Horde.Build/Jobs/JobMessages.cs
+++ b/Engine/Source/Programs/Horde/Horde.Build/Jobs/JobMessages.cs
@@ -68,7 +68,7 @@ namespace Horde.Build.Jobs
 		/// <summary>
 		/// Parameters to use when selecting the change to execute at.
 		/// </summary>
-		[Obsolete]
+		[Obsolete("Use ChangeQueries instead")]
 		public ChangeQueryConfig? ChangeQuery
 		{
 			get => (ChangeQueries != null && ChangeQueries.Count > 0) ? ChangeQueries[0] : null;
diff --git a/Engine/Source/Programs/Horde/Horde.Build/Logs/Builder/RedisLogBuilder.cs b/Engine/Source/Programs/Horde/Horde.Build/Logs/Builder/RedisLogBuilder.cs
--- a/Engine/Source/Programs/Horde/Horde.Build/Logs/Builder/RedisLogBuilder.cs
+++ b/Engine/Source/Programs/Horde/Horde.Build/Logs/Builder/RedisLogBuilder.cs
@@ -253,7 +253,7 @@ namespace Horde.Build.Logs.Builder
 					int subChunkLineIndex = lineIndex;
 
 					List<LogSubChunkData> subChunks = new List<LogSubChunkData>();
-					while (reader.Memory.Length > 0)
+					while (reader.RemainingMemory.Length > 0)
 					{
 						LogSubChunkData subChunkData = reader.ReadLogSubChunkData(subChunkOffset, subChunkLineIndex);
 						subChunkOffset += subChunkData.Length;
diff --git a/Engine/Source/Programs/Horde/Horde.Build/Logs/Data/LogChunkData.cs b/Engine/Source/Programs/Horde/Horde.Build/Logs/Data/LogChunkData.cs
--- a/Engine/Source/Programs/Horde/Horde.Build/Logs/Data/LogChunkData.cs
+++ b/Engine/Source/Programs/Horde/Horde.Build/Logs/Data/LogChunkData.cs
@@ -153,8 +153,8 @@ namespace Horde.Build.Logs.Data
 			if ((signature & 0xffffff) != CurrentSignature)
 			{
 				List<LogSubChunkData> subChunks = new List<LogSubChunkData>();
-				subChunks.Add(new LogSubChunkData(LogType.Json, offset, lineIndex, new ReadOnlyLogText(reader.Memory)));
-				reader.Advance(reader.Memory.Length);
+				subChunks.Add(new LogSubChunkData(LogType.Json, offset, lineIndex, new ReadOnlyLogText(reader.RemainingMemory)));
+				reader.Advance(reader.RemainingMemory.Length);
 				return new LogChunkData(offset, lineIndex, subChunks);
 			}
 
diff --git a/Engine/Source/Programs/Horde/Horde.Build/Logs/LogFileService.cs b/Engine/Source/Programs/Horde/Horde.Build/Logs/LogFileService.cs
--- a/Engine/Source/Programs/Horde/Horde.Build/Logs/LogFileService.cs
+++ b/Engine/Source/Programs/Horde/Horde.Build/Logs/LogFileService.cs
@@ -1226,7 +1226,7 @@ namespace Horde.Build.Logs
 				// Wait for all tasks to be complete OR (any task has completed AND 30 seconds has elapsed)
 				Task allCompleteTask = Task.WhenAll(chunkWriteTasks);
 				Task anyCompleteTask = Task.WhenAny(chunkWriteTasks);
-				await Task.WhenAny(allCompleteTask, Task.WhenAll(anyCompleteTask, Task.Delay(TimeSpan.FromSeconds(30.0))));
+				await Task.WhenAny(allCompleteTask, Task.WhenAll(anyCompleteTask, Task.Delay(TimeSpan.FromSeconds(30.0), cancellationToken)));
 
 				// Update the log file with the written chunks
 				List<LogChunkData?> writtenChunks = chunkWriteTasks.RemoveCompleteTasks();
diff --git a/Engine/Source/Programs/Horde/Horde.Build/Logs/Storage/PersistentLogStorage.cs b/Engine/Source/Programs/Horde/Horde.Build/Logs/Storage/PersistentLogStorage.cs
--- a/Engine/Source/Programs/Horde/Horde.Build/Logs/Storage/PersistentLogStorage.cs
+++ b/Engine/Source/Programs/Horde/Horde.Build/Logs/Storage/PersistentLogStorage.cs
@@ -82,9 +82,9 @@ namespace Horde.Build.Logs.Storage
 			MemoryReader reader = new MemoryReader(data.Value);
 			LogChunkData chunkData = reader.ReadLogChunkData(offset, lineIndex);
 
-			if (reader.Memory.Length > 0)
+			if (reader.RemainingMemory.Length > 0)
 			{
-				throw new Exception($"Serialization of persistent chunk {path} is not at expected offset ({reader.Memory.Length} bytes remaining)");
+				throw new Exception($"Serialization of persistent chunk {path} is not at expected offset ({reader.RemainingMemory.Length} bytes remaining)");
 			}
 
 			return chunkData;
diff --git a/Engine/Source/Programs/Horde/Horde.Build/Streams/StreamConfig.cs b/Engine/Source/Programs/Horde/Horde.Build/Streams/StreamConfig.cs
--- a/Engine/Source/Programs/Horde/Horde.Build/Streams/StreamConfig.cs
+++ b/Engine/Source/Programs/Horde/Horde.Build/Streams/StreamConfig.cs
@@ -29,7 +29,7 @@ namespace Horde.Build.Streams
 	/// Flags identifying content of a changelist
 	/// </summary>
 	[Flags]
-	[Obsolete]
+	[Obsolete("Use tags instead")]
 	public enum ChangeContentFlags
 	{
 		/// <summary>
diff --git a/Engine/Source/Programs/Shared/EpicGames.BuildGraph/BgInterpreter.cs b/Engine/Source/Programs/Shared/EpicGames.BuildGraph/BgInterpreter.cs
--- a/Engine/Source/Programs/Shared/EpicGames.BuildGraph/BgInterpreter.cs
+++ b/Engine/Source/Programs/Shared/EpicGames.BuildGraph/BgInterpreter.cs
@@ -112,7 +112,7 @@ namespace EpicGames.BuildGraph
 			_names = reader.ReadVariableLengthArray(() => reader.ReadString());
 			int[] fragmentLengths = reader.ReadVariableLengthArray(() => (int)reader.ReadUnsignedVarInt());
 
-			int offset = data.Length - reader.Memory.Length;
+			int offset = data.Length - reader.RemainingMemory.Length;
 
 			_fragments = new int[fragmentLengths.Length];
 			for (int idx = 0; idx < _fragments.Length; idx++)
diff --git a/Engine/Source/Programs/Shared/EpicGames.Core/AssemblyUtils.cs b/Engine/Source/Programs/Shared/EpicGames.Core/AssemblyUtils.cs
--- a/Engine/Source/Programs/Shared/EpicGames.Core/AssemblyUtils.cs
+++ b/Engine/Source/Programs/Shared/EpicGames.Core/AssemblyUtils.cs
@@ -43,7 +43,7 @@ namespace EpicGames.Core
 
 				return (
 					from knownAssemblyName in new[] { "SwarmAgent.exe", "../ThirdParty/Ionic/Ionic.Zip.Reduced.dll", "../ThirdParty/Newtonsoft/NewtonSoft.Json.dll" }
-					where assemblyName.Equals(Path.GetFileNameWithoutExtension(knownAssemblyName), StringComparison.InvariantCultureIgnoreCase)
+					where assemblyName.Equals(Path.GetFileNameWithoutExtension(knownAssemblyName), StringComparison.OrdinalIgnoreCase)
 					let resolvedAssemblyFilename = Path.Combine(pathToBinariesDotNET, knownAssemblyName)
 					// check if the file exists first. If we just try to load it, we correctly throw an exception, but it's a generic
 					// FileNotFoundException, which is not informative. Better to return null.
diff --git a/Engine/Source/Programs/Shared/EpicGames.Core/AsyncUtils.cs b/Engine/Source/Programs/Shared/EpicGames.Core/AsyncUtils.cs
--- a/Engine/Source/Programs/Shared/EpicGames.Core/AsyncUtils.cs
+++ b/Engine/Source/Programs/Shared/EpicGames.Core/AsyncUtils.cs
@@ -19,7 +19,7 @@ namespace EpicGames.Core
 		/// <returns></returns>
 		public static Task AsTask(this CancellationToken token)
 		{
-			return Task.Delay(-1, token).ContinueWith(x => { });
+			return Task.Delay(-1, token).ContinueWith(x => { }, TaskScheduler.Default);
 		}
 
 		/// <summary>
@@ -29,7 +29,7 @@ namespace EpicGames.Core
 		/// <returns></returns>
 		public static Task<T> AsTask<T>(this CancellationToken token)
 		{
-			return Task.Delay(-1, token).ContinueWith(_ => Task.FromCanceled<T>(token)).Unwrap();
+			return Task.Delay(-1, token).ContinueWith(_ => Task.FromCanceled<T>(token), TaskScheduler.Default).Unwrap();
 		}
 
 		/// <summary>
@@ -76,7 +76,7 @@ namespace EpicGames.Core
 		/// <returns></returns>
 		public static Task DelayNoThrow(TimeSpan time, CancellationToken token)
 		{
-			return Task.Delay(time, token).ContinueWith(x => { });
+			return Task.Delay(time, token).ContinueWith(x => { }, CancellationToken.None, TaskContinuationOptions.None, TaskScheduler.Default);
 		}
 
 		/// <summary>
diff --git a/Engine/Source/Programs/Shared/EpicGames.Core/BinaryArchive.cs b/Engine/Source/Programs/Shared/EpicGames.Core/BinaryArchive.cs
--- a/Engine/Source/Programs/Shared/EpicGames.Core/BinaryArchive.cs
+++ b/Engine/Source/Programs/Shared/EpicGames.Core/BinaryArchive.cs
@@ -106,11 +106,9 @@ namespace EpicGames.Core
 		}
 
 		static readonly MethodInfo s_readBoolMethodInfo = GetMethodInfo<BinaryArchiveReader>(x => x.ReadBool());
-		static readonly MethodInfo s_readIntMethodInfo = GetMethodInfo<BinaryArchiveReader>(x => x.ReadInt());
 		static readonly MethodInfo s_readArrayMethodInfo = GetGenericMethodInfo<BinaryArchiveReader, Func<int>>((x, y) => x.ReadArray(y));
 		static readonly MethodInfo s_readPolymorphicObjectMethodInfo = GetMethodInfo(() => ReadObject(null!));
 		static readonly MethodInfo s_writeBoolMethodInfo = GetMethodInfo<BinaryArchiveWriter>(x => x.WriteBool(false));
-		static readonly MethodInfo s_writeIntMethodInfo = GetMethodInfo<BinaryArchiveWriter>(x => x.WriteInt(0));
 		static readonly MethodInfo s_writeArrayMethodInfo = GetGenericMethodInfo<BinaryArchiveWriter, Action<int>>((x, y) => x.WriteArray(null!, y));
 		static readonly MethodInfo s_writePolymorphicObjectMethodInfo = GetMethodInfo(() => WritePolymorphicObject(null!, null!));
 
diff --git a/Engine/Source/Programs/Shared/EpicGames.Core/BinaryArchiveAttributes.cs b/Engine/Source/Programs/Shared/EpicGames.Core/BinaryArchiveAttributes.cs
--- a/Engine/Source/Programs/Shared/EpicGames.Core/BinaryArchiveAttributes.cs
+++ b/Engine/Source/Programs/Shared/EpicGames.Core/BinaryArchiveAttributes.cs
@@ -8,7 +8,7 @@ namespace EpicGames.Core
 	/// Instructs the serializer to ignore this property during serialization
 	/// </summary>
 	[AttributeUsage(AttributeTargets.Property)]
-	public class BinaryIgnoreAttribute : Attribute
+	public sealed class BinaryIgnoreAttribute : Attribute
 	{
 		/// <summary>
 		/// Constructor
@@ -22,7 +22,7 @@ namespace EpicGames.Core
 	/// Marks this class as supporting binary serialization. Used to automatically register types via BinaryArchive.RegisterTypes().
 	/// </summary>
 	[AttributeUsage(AttributeTargets.Class)]
-	public class BinarySerializableAttribute : Attribute
+	public sealed class BinarySerializableAttribute : Attribute
 	{
 		/// <summary>
 		/// Constructor
@@ -36,7 +36,7 @@ namespace EpicGames.Core
 	/// Instructs the serializer to use a specific converter for this property or class
 	/// </summary>
 	[AttributeUsage(AttributeTargets.Property | AttributeTargets.Class)]
-	public class BinaryConverterAttribute : Attribute
+	public sealed class BinaryConverterAttribute : Attribute
 	{
 		/// <summary>
 		/// The serializer type
diff --git a/Engine/Source/Programs/Shared/EpicGames.Core/BinaryArchiveReader.cs b/Engine/Source/Programs/Shared/EpicGames.Core/BinaryArchiveReader.cs
--- a/Engine/Source/Programs/Shared/EpicGames.Core/BinaryArchiveReader.cs
+++ b/Engine/Source/Programs/Shared/EpicGames.Core/BinaryArchiveReader.cs
@@ -12,7 +12,7 @@ namespace EpicGames.Core
 	/// Reads data from a binary output stream. Similar to the NET Framework BinaryReader class, but supports fast serialization of object graphs and container types, and supports nullable objects.
 	/// Significantly faster than BinaryReader due to the expectation that the whole stream is in memory before deserialization.
 	/// </summary>
-	public class BinaryArchiveReader : IDisposable
+	public sealed class BinaryArchiveReader : IDisposable
 	{
 		/// <summary>
 		/// The input stream.
diff --git a/Engine/Source/Programs/Shared/EpicGames.Core/BinaryArchiveWriter.cs b/Engine/Source/Programs/Shared/EpicGames.Core/BinaryArchiveWriter.cs
--- a/Engine/Source/Programs/Shared/EpicGames.Core/BinaryArchiveWriter.cs
+++ b/Engine/Source/Programs/Shared/EpicGames.Core/BinaryArchiveWriter.cs
@@ -11,7 +11,7 @@ namespace EpicGames.Core
     /// <summary>
     /// Writes data to a binary output stream. Similar to the NET Framework BinaryWriter class, but supports fast serialization of object graphs and container types, and supports nullable objects.
     /// </summary>
-    public class BinaryArchiveWriter : IDisposable
+    public sealed class BinaryArchiveWriter : IDisposable
     {
         /// <summary>
         /// Comparer which tests for reference equality between two objects
diff --git a/Engine/Source/Programs/Shared/EpicGames.Core/BinaryFormatterUtils.cs b/Engine/Source/Programs/Shared/EpicGames.Core/BinaryFormatterUtils.cs
--- a/Engine/Source/Programs/Shared/EpicGames.Core/BinaryFormatterUtils.cs
+++ b/Engine/Source/Programs/Shared/EpicGames.Core/BinaryFormatterUtils.cs
@@ -4,6 +4,8 @@ using System.IO;
 using System.Runtime.Serialization.Formatters.Binary;
 
 #pragma warning disable SYSLIB0011
+#pragma warning disable CA2300 // Do not use insecure deserializer BinaryFormatter
+#pragma warning disable CA2301 // Do not use insecure deserializer BinaryFormatter
 
 namespace EpicGames.Core
 {
diff --git a/Engine/Source/Programs/Shared/EpicGames.Core/BuzHash.cs b/Engine/Source/Programs/Shared/EpicGames.Core/BuzHash.cs
--- a/Engine/Source/Programs/Shared/EpicGames.Core/BuzHash.cs
+++ b/Engine/Source/Programs/Shared/EpicGames.Core/BuzHash.cs
@@ -156,7 +156,9 @@ namespace EpicGames.Core
 		/// <param name="threshold">Threshold for hash below which to terminate the scan</param>
 		/// <param name="state">The current hash value</param>
 		/// <returns>Number of bytes consumed from next</returns>
+#pragma warning disable CA1045 // Do not pass types by reference
 		public static unsafe int Update(ReadOnlySpan<byte> prev, ReadOnlySpan<byte> next, uint threshold, ref uint state)
+#pragma warning restore CA1045 // Do not pass types by reference
 		{
 			Debug.Assert(prev.Length == next.Length);
 
diff --git a/Engine/Source/Programs/Shared/EpicGames.Core/CancellationTask.cs b/Engine/Source/Programs/Shared/EpicGames.Core/CancellationTask.cs
--- a/Engine/Source/Programs/Shared/EpicGames.Core/CancellationTask.cs
+++ b/Engine/Source/Programs/Shared/EpicGames.Core/CancellationTask.cs
@@ -9,7 +9,7 @@ namespace EpicGames.Core
 	/// <summary>
 	/// Wrapper class to create a waitable task out of a cancellation token
 	/// </summary>
-	public class CancellationTask : IDisposable
+	public sealed class CancellationTask : IDisposable
 	{
 		/// <summary>
 		/// Completion source for the task
diff --git a/Engine/Source/Programs/Shared/EpicGames.Core/CertificateUtils.cs b/Engine/Source/Programs/Shared/EpicGames.Core/CertificateUtils.cs
--- a/Engine/Source/Programs/Shared/EpicGames.Core/CertificateUtils.cs
+++ b/Engine/Source/Programs/Shared/EpicGames.Core/CertificateUtils.cs
@@ -1,6 +1,7 @@
 // Copyright Epic Games, Inc. All Rights Reserved.
 
 using System;
+using System.Runtime.InteropServices;
 using System.Security.Cryptography;
 using System.Security.Cryptography.X509Certificates;
 using System.Text.RegularExpressions;
@@ -39,7 +40,10 @@ namespace EpicGames.Core
 				// NB: MacOS requires 825 days or fewer (https://support.apple.com/en-us/HT210176)
 				using (X509Certificate2 certificate = request.CreateSelfSigned(new DateTimeOffset(DateTime.UtcNow.AddDays(-1)), new DateTimeOffset(DateTime.UtcNow.AddDays(800))))
 				{
-					certificate.FriendlyName = friendlyName;
+					if (RuntimeInformation.IsOSPlatform(OSPlatform.Windows))
+					{
+						certificate.FriendlyName = friendlyName;
+					}
 					byte[] privateCertData = certificate.Export(X509ContentType.Pkcs12); // Note: Need to reimport this to use immediately, otherwise key is ephemeral
 					return privateCertData;
 				}
diff --git a/Engine/Source/Programs/Shared/EpicGames.Core/CityHash.cs b/Engine/Source/Programs/Shared/EpicGames.Core/CityHash.cs
--- a/Engine/Source/Programs/Shared/EpicGames.Core/CityHash.cs
+++ b/Engine/Source/Programs/Shared/EpicGames.Core/CityHash.cs
@@ -37,6 +37,10 @@ namespace EpicGames.Core
 #pragma warning disable IDE0180
 #pragma warning disable IDE1006
 #pragma warning disable CS1591 // Missing documentation
+#pragma warning disable CA1034 // Nested types should not be visible
+#pragma warning disable CA1707 // Identifiers should not contain underscores
+#pragma warning disable CA1051 // Do not declare visible instance fields
+#pragma warning disable CA1822 // Mark members as static
 
 	public class CityHash
 	{
diff --git a/Engine/Source/Programs/Shared/EpicGames.Core/CommandHost.cs b/Engine/Source/Programs/Shared/EpicGames.Core/CommandHost.cs
--- a/Engine/Source/Programs/Shared/EpicGames.Core/CommandHost.cs
+++ b/Engine/Source/Programs/Shared/EpicGames.Core/CommandHost.cs
@@ -45,7 +45,7 @@ namespace EpicGames.Core
 		/// <summary>
 		/// Names for this command
 		/// </summary>
-		public string[] Names { get; }
+		public IReadOnlyList<string> Names { get; }
 
 		/// <summary>
 		/// Short description for the mode. Will be displayed in the help text.
@@ -58,10 +58,12 @@ namespace EpicGames.Core
 		public ICommand CreateInstance(IServiceProvider serviceProvider);
 	}
 
+#pragma warning disable CA1019 // Define accessors for attribute arguments
 	/// <summary>
 	/// Attribute used to specify names of program modes, and help text
 	/// </summary>
-	public class CommandAttribute : Attribute
+	[AttributeUsage(AttributeTargets.Class)]
+	public sealed class CommandAttribute : Attribute
 	{
 		/// <summary>
 		/// Names for this command
@@ -96,6 +98,7 @@ namespace EpicGames.Core
 			Description = description;
 		}
 	}
+#pragma warning restore CA1019 // Define accessors for attribute arguments
 
 	/// <summary>
 	/// Base class for all commands that can be executed by HordeAgent
@@ -123,7 +126,7 @@ namespace EpicGames.Core
 	/// </summary>
 	class CommandFactory : ICommandFactory
 	{
-		public string[] Names { get; }
+		public IReadOnlyList<string> Names { get; }
 		public string Description { get; }
 
 		public Type Type { get; }
diff --git a/Engine/Source/Programs/Shared/EpicGames.Core/CommandLineArguments.cs b/Engine/Source/Programs/Shared/EpicGames.Core/CommandLineArguments.cs
--- a/Engine/Source/Programs/Shared/EpicGames.Core/CommandLineArguments.cs
+++ b/Engine/Source/Programs/Shared/EpicGames.Core/CommandLineArguments.cs
@@ -11,6 +11,8 @@ using System.Reflection;
 using System.Text;
 using Microsoft.Extensions.Logging;
 
+#pragma warning disable CA1710 // Identifiers should have correct suffix
+
 namespace EpicGames.Core
 {
 	/// <summary>
@@ -701,7 +703,7 @@ namespace EpicGames.Core
 			}
 			else
 			{
-				if (valueType != typeof(bool) && attribute.Value == null && !prefix.EndsWith("=") && !prefix.EndsWith(":"))
+				if (valueType != typeof(bool) && attribute.Value == null && !prefix.EndsWith("=", StringComparison.Ordinal) && !prefix.EndsWith(":", StringComparison.Ordinal))
 				{
 					prefix += "=";
 				}
@@ -810,11 +812,11 @@ namespace EpicGames.Core
 			{
 				if(missingArguments.Count == 1)
 				{
-					throw new CommandLineArgumentException(String.Format("Missing {0} argument", missingArguments[0].Replace("=", "=...")));
+					throw new CommandLineArgumentException(String.Format("Missing {0} argument", missingArguments[0].Replace("=", "=...", StringComparison.Ordinal)));
 				}
 				else
 				{
-					throw new CommandLineArgumentException(String.Format("Missing {0} arguments", StringUtils.FormatList(missingArguments.Select(x => x.Replace("=", "=...")))));
+					throw new CommandLineArgumentException(String.Format("Missing {0} arguments", StringUtils.FormatList(missingArguments.Select(x => x.Replace("=", "=...", StringComparison.Ordinal)))));
 				}
 			}
 		}
@@ -873,7 +875,7 @@ namespace EpicGames.Core
 				{
 					if(descriptionBuilder.Length > 0)
 					{
-						descriptionBuilder.Append("\n");
+						descriptionBuilder.Append('\n');
 					}
 					descriptionBuilder.Append(attribute.Description);
 				}
@@ -929,13 +931,13 @@ namespace EpicGames.Core
 				{
 					if (argument[idx] == '=')
 					{
-						return String.Format("{0}=\"{1}\"", argument.Substring(0, idx), argument.Substring(idx + 1).Replace("\"", "\\\""));
+						return String.Format("{0}=\"{1}\"", argument.Substring(0, idx), argument.Substring(idx + 1).Replace("\"", "\\\"", StringComparison.Ordinal));
 					}
 				}
 			}
 
 			// Quote the whole thing
-			return "\"" + argument.Replace("\"", "\\\"") + "\"";
+			return "\"" + argument.Replace("\"", "\\\"", StringComparison.Ordinal) + "\"";
 		}
 
 		/// <summary>
@@ -1079,11 +1081,11 @@ namespace EpicGames.Core
 			{
 				if(remainingArguments.Count == 1)
 				{
-					logger.LogWarning(String.Format("Invalid argument: {0}", remainingArguments[0]));
+					logger.LogWarning("Invalid argument: {Argument}", remainingArguments[0]);
 				}
 				else
 				{
-					logger.LogWarning(String.Format("Invalid arguments:\n{0}", String.Join("\n", remainingArguments)));
+					logger.LogWarning("Invalid arguments:\n{Arguments}", String.Join("\n", remainingArguments));
 				}
 			}
 		}
@@ -1094,7 +1096,7 @@ namespace EpicGames.Core
 		/// <returns></returns>
 		public bool AreAnyArgumentsUsed()
 		{
-			return _usedArguments.Cast<bool>().Count(b => b) != 0;
+			return _usedArguments.Cast<bool>().Any(b => b);
 		}
 
 		/// <summary>
@@ -1171,7 +1173,7 @@ namespace EpicGames.Core
 			object? value;
 			if(!TryParseValue(valueType, valueText, out value))
 			{
-				logger.LogWarning("Unable to parse value for argument '{0}'.", argumentText);
+				logger.LogWarning("Unable to parse value for argument '{Argument}'.", argumentText);
 				return false;
 			}
 
@@ -1184,7 +1186,7 @@ namespace EpicGames.Core
 					object? previousValue = target.GetValue(targetObject);
 					if(!Object.Equals(previousValue, value))
 					{
-						logger.LogWarning("Argument '{0}' conflicts with '{1}'; ignoring.", argumentText, previousArgumentText);
+						logger.LogWarning("Argument '{Argument}' conflicts with '{PrevArgument}'; ignoring.", argumentText, previousArgumentText);
 					}
 					return false;
 				}
@@ -1353,15 +1355,15 @@ namespace EpicGames.Core
 				commandLine.Append(' ');
 			}
 
-			int spaceIdx = argument.IndexOf(' ');
+			int spaceIdx = argument.IndexOf(' ', StringComparison.Ordinal);
 			if(spaceIdx == -1)
 			{
 				commandLine.Append(argument);
 			}
 			else
 			{
-				int equalsIdx = argument.IndexOf('=');
-				if(equalsIdx == -1)
+				int equalsIdx = argument.IndexOf('=', StringComparison.Ordinal);
+				if(equalsIdx == -1 || equalsIdx > spaceIdx)
 				{
 					commandLine.AppendFormat("\"{0}\"", argument);
 				}
diff --git a/Engine/Source/Programs/Shared/EpicGames.Core/CommandLineAttribute.cs b/Engine/Source/Programs/Shared/EpicGames.Core/CommandLineAttribute.cs
--- a/Engine/Source/Programs/Shared/EpicGames.Core/CommandLineAttribute.cs
+++ b/Engine/Source/Programs/Shared/EpicGames.Core/CommandLineAttribute.cs
@@ -8,12 +8,14 @@ namespace EpicGames.Core
 	/// Attribute to indicate the name of a command line argument
 	/// </summary>
 	[AttributeUsage(AttributeTargets.Field | AttributeTargets.Property, AllowMultiple = true)]
-	public class CommandLineAttribute : Attribute
+	public sealed class CommandLineAttribute : Attribute
 	{
 		/// <summary>
 		/// Prefix for the option, with a leading '-' and trailing '=' character if a value is expected.
 		/// </summary>
+#pragma warning disable CA1019 // Define accessors for attribute arguments
 		public string? Prefix { get; set; }
+#pragma warning restore CA1019 // Define accessors for attribute arguments
 
 		/// <summary>
 		/// Specifies a fixed value for this argument. Specifying an alternate value is not permitted.
@@ -45,7 +47,7 @@ namespace EpicGames.Core
 
 			if(prefix != null)
 			{
-				if(!prefix.StartsWith("-"))
+				if(!prefix.StartsWith("-", StringComparison.Ordinal))
 				{
 					throw new Exception("Command-line arguments must begin with a '-' character");
 				}
diff --git a/Engine/Source/Programs/Shared/EpicGames.Core/CommandLineExtensions.cs b/Engine/Source/Programs/Shared/EpicGames.Core/CommandLineExtensions.cs
--- a/Engine/Source/Programs/Shared/EpicGames.Core/CommandLineExtensions.cs
+++ b/Engine/Source/Programs/Shared/EpicGames.Core/CommandLineExtensions.cs
@@ -1,5 +1,6 @@
 // Copyright Epic Games, Inc. All Rights Reserved.
 
+using System;
 using System.Text;
 
 namespace EpicGames.Core
@@ -16,7 +17,7 @@ namespace EpicGames.Core
 		/// <returns>True if the argument needs to be escaped</returns>
 		static bool NeedsEscaping(string argument)
 		{
-			return argument.Contains(' ') || argument.Contains('\"');
+			return argument.Contains(' ', StringComparison.Ordinal) || argument.Contains('\"', StringComparison.Ordinal);
 		}
 
 		/// <summary>
@@ -31,7 +32,7 @@ namespace EpicGames.Core
 				builder.Append(' ');
 			}
 
-			int equalsIdx = argument.IndexOf('=');
+			int equalsIdx = argument.IndexOf('=', StringComparison.Ordinal);
 			if (equalsIdx != -1)
 			{
 				string name = argument.Substring(0, equalsIdx + 1);
diff --git a/Engine/Source/Programs/Shared/EpicGames.Core/ContentHash.cs b/Engine/Source/Programs/Shared/EpicGames.Core/ContentHash.cs
--- a/Engine/Source/Programs/Shared/EpicGames.Core/ContentHash.cs
+++ b/Engine/Source/Programs/Shared/EpicGames.Core/ContentHash.cs
@@ -9,6 +9,10 @@ using System.Text;
 using System.Text.Json;
 using System.Text.Json.Serialization;
 
+#pragma warning disable CA5351 // Do Not Use Broken Cryptographic Algorithms
+#pragma warning disable CA5350 // Do Not Use Weak Cryptographic Algorithms
+#pragma warning disable CA1819 // Properties should not return arrays
+
 namespace EpicGames.Core
 {
 	/// <summary>
diff --git a/Engine/Source/Programs/Shared/EpicGames.Core/CsProjBuildHook.cs b/Engine/Source/Programs/Shared/EpicGames.Core/CsProjBuildHook.cs
--- a/Engine/Source/Programs/Shared/EpicGames.Core/CsProjBuildHook.cs
+++ b/Engine/Source/Programs/Shared/EpicGames.Core/CsProjBuildHook.cs
@@ -3,6 +3,8 @@
 using System;
 using System.Collections.Generic;
 
+#pragma warning disable CA1715 // Identifiers should have correct prefix
+
 namespace EpicGames.Core
 {
 
diff --git a/Engine/Source/Programs/Shared/EpicGames.Core/CsProjBuildRecord.cs b/Engine/Source/Programs/Shared/EpicGames.Core/CsProjBuildRecord.cs
--- a/Engine/Source/Programs/Shared/EpicGames.Core/CsProjBuildRecord.cs
+++ b/Engine/Source/Programs/Shared/EpicGames.Core/CsProjBuildRecord.cs
@@ -2,7 +2,9 @@
 
 using System;
 using System.Collections.Generic;
-using System.Text;
+
+#pragma warning disable CA2227 // Collection properties should be read only
+#pragma warning disable CA1034 // Nested types should not be visible
 
 namespace EpicGames.Core
 {
diff --git a/Engine/Source/Programs/Shared/EpicGames.Core/CsProjectInfo.cs b/Engine/Source/Programs/Shared/EpicGames.Core/CsProjectInfo.cs
--- a/Engine/Source/Programs/Shared/EpicGames.Core/CsProjectInfo.cs
+++ b/Engine/Source/Programs/Shared/EpicGames.Core/CsProjectInfo.cs
@@ -88,7 +88,7 @@ namespace EpicGames.Core
 		private bool IsNetCoreProject()
 		{
 			string? framework;
-			return Properties.TryGetValue("TargetFramework", out framework) && framework.StartsWith("netcoreapp");
+			return Properties.TryGetValue("TargetFramework", out framework) && framework.StartsWith("netcoreapp", StringComparison.Ordinal);
 		}
 
 		/// <summary>
@@ -287,7 +287,7 @@ namespace EpicGames.Core
 			string? targetFramework;
 			if (Properties.TryGetValue("TargetFramework", out targetFramework))
 			{
-				bIsDotNetCoreProject = targetFramework.ToLower().Contains("netstandard") || targetFramework.ToLower().Contains("netcoreapp");
+				bIsDotNetCoreProject = targetFramework.ToLower().Contains("netstandard", StringComparison.Ordinal) || targetFramework.ToLower().Contains("netcoreapp", StringComparison.Ordinal);
 			}
 
 			return bIsDotNetCoreProject;
@@ -509,7 +509,7 @@ namespace EpicGames.Core
 			if (!remainingComponents.Any())
 			{
 				// ** means include all files under this tree, so enumerate them all
-				if (currentComponent.Contains("**"))
+				if (currentComponent.Contains("**", StringComparison.Ordinal))
 				{
 					outFoundFiles.AddRange(DirectoryReference.EnumerateFiles(existingPath, "*", SearchOption.AllDirectories));
 				}
@@ -523,7 +523,7 @@ namespace EpicGames.Core
 			{
 				// new component contains a wildcard, and based on the above we know there are more entries so find 
 				// matching directories
-				if (currentComponent.Contains("*"))
+				if (currentComponent.Contains('*', StringComparison.Ordinal))
 				{
 					// ** means all directories, no matter how deep
 					SearchOption option = currentComponent == "**" ? SearchOption.AllDirectories : SearchOption.TopDirectoryOnly;
@@ -551,7 +551,7 @@ namespace EpicGames.Core
 
 					// but... we can just take all the next components that don't have wildcards in them instead of recursing
 					// into each one!
-					IEnumerable<string> nonWildCardComponents = remainingComponents.TakeWhile(c => !c.Contains("*"));
+					IEnumerable<string> nonWildCardComponents = remainingComponents.TakeWhile(c => !c.Contains('*', StringComparison.Ordinal));
 					remainingComponents = remainingComponents.Skip(nonWildCardComponents.Count());
 
 					existingPath = DirectoryReference.Combine(existingPath, nonWildCardComponents.ToArray());
@@ -605,7 +605,7 @@ namespace EpicGames.Core
 			{
 				FileReference sourceFile = FileReference.Combine(baseDirectory, includePath);
 
-				if (sourceFile.FullName.Contains("*"))
+				if (sourceFile.FullName.Contains('*', StringComparison.Ordinal))
 				{
 					compileReferences.AddRange(FindMatchingFiles(sourceFile));
 				}
@@ -628,7 +628,7 @@ namespace EpicGames.Core
 			if (!String.IsNullOrEmpty(includePath))
 			{
 				string? copyTo = GetChildElementString(parentElement, "CopyToOutputDirectory", null);
-				bool shouldCopy = !String.IsNullOrEmpty(copyTo) && (copyTo.Equals("Always", StringComparison.InvariantCultureIgnoreCase) || copyTo.Equals("PreserveNewest", StringComparison.InvariantCultureIgnoreCase));
+				bool shouldCopy = !String.IsNullOrEmpty(copyTo) && (copyTo.Equals("Always", StringComparison.OrdinalIgnoreCase) || copyTo.Equals("PreserveNewest", StringComparison.OrdinalIgnoreCase));
 				FileReference contentFile = FileReference.Combine(baseDirectory, includePath);
 				contents.Add(contentFile, shouldCopy);
 			}
@@ -668,11 +668,11 @@ namespace EpicGames.Core
 			{
 				return defaultValue;
 			}
-			else if (value.Equals("True", StringComparison.InvariantCultureIgnoreCase))
+			else if (value.Equals("True", StringComparison.OrdinalIgnoreCase))
 			{
 				return true;
 			}
-			else if (value.Equals("False", StringComparison.InvariantCultureIgnoreCase))
+			else if (value.Equals("False", StringComparison.OrdinalIgnoreCase))
 			{
 				return false;
 			}
@@ -734,13 +734,13 @@ namespace EpicGames.Core
 				}
 			}
 
-			if (tokens.Length == 3 && tokens[0].StartsWith("'") && tokens[1] == "==" && tokens[2].StartsWith("'"))
+			if (tokens.Length == 3 && tokens[0].StartsWith("'", StringComparison.Ordinal) && tokens[1] == "==" && tokens[2].StartsWith("'", StringComparison.Ordinal))
 			{
-				bResult = String.Compare(tokens[0], tokens[2], StringComparison.InvariantCultureIgnoreCase) == 0;
+				bResult = String.Equals(tokens[0], tokens[2], StringComparison.OrdinalIgnoreCase);
 			}
-			else if (tokens.Length == 3 && tokens[0].StartsWith("'") && tokens[1] == "!=" && tokens[2].StartsWith("'"))
+			else if (tokens.Length == 3 && tokens[0].StartsWith("'", StringComparison.Ordinal) && tokens[1] == "!=" && tokens[2].StartsWith("'", StringComparison.Ordinal))
 			{
-				bResult = String.Compare(tokens[0], tokens[2], StringComparison.InvariantCultureIgnoreCase) != 0;
+				bResult = !String.Equals(tokens[0], tokens[2], StringComparison.OrdinalIgnoreCase);
 			}
 			else
 			{
@@ -759,7 +759,7 @@ namespace EpicGames.Core
 		static string ExpandProperties(string text, Dictionary<string, string> properties)
 		{
 			string newText = text;
-			for (int idx = newText.IndexOf("$("); idx != -1; idx = newText.IndexOf("$(", idx))
+			for (int idx = newText.IndexOf("$(", StringComparison.Ordinal); idx != -1; idx = newText.IndexOf("$(", idx, StringComparison.Ordinal))
 			{
 				// Find the end of the variable name, accounting for changes in scope
 				int endIdx = idx + 2;
@@ -1038,10 +1038,10 @@ namespace EpicGames.Core
 				{
 					if(newText[idx] == '%')
 					{
-						int upperDigitIdx = HexChars.IndexOf(Char.ToLowerInvariant(newText[idx + 1]));
+						int upperDigitIdx = HexChars.IndexOf(Char.ToLowerInvariant(newText[idx + 1]), StringComparison.Ordinal);
 						if(upperDigitIdx != -1)
 						{
-							int lowerDigitIdx = HexChars.IndexOf(Char.ToLowerInvariant(newText[idx + 2]));
+							int lowerDigitIdx = HexChars.IndexOf(Char.ToLowerInvariant(newText[idx + 2]), StringComparison.Ordinal);
 							if(lowerDigitIdx != -1)
 							{
 								char newChar = (char)((upperDigitIdx << 4) | lowerDigitIdx);
diff --git a/Engine/Source/Programs/Shared/EpicGames.Core/Digest.cs b/Engine/Source/Programs/Shared/EpicGames.Core/Digest.cs
--- a/Engine/Source/Programs/Shared/EpicGames.Core/Digest.cs
+++ b/Engine/Source/Programs/Shared/EpicGames.Core/Digest.cs
@@ -5,12 +5,16 @@ using System.Buffers.Binary;
 using System.Security.Cryptography;
 using System.Text;
 
+#pragma warning disable CA1000 // Do not declare static members on generic types
+#pragma warning disable CA5351 // Do Not Use Broken Cryptographic Algorithms
+#pragma warning disable CA5350 // Do Not Use Weak Cryptographic Algorithms
+
 namespace EpicGames.Core
 {
 	/// <summary>
 	/// Struct representing a weakly typed hash value. Counterpart to <see cref="Digest{T}"/> - a strongly typed digest.
 	/// </summary>
-	public struct Digest
+	public struct Digest : IEquatable<Digest>
 	{
 		/// <summary>
 		/// Memory storing the digest data
@@ -91,7 +95,10 @@ namespace EpicGames.Core
 		}
 
 		/// <inheritdoc/>
-		public override bool Equals(object? obj) => (obj is Digest digest) && digest.Span.SequenceEqual(Span);
+		public override bool Equals(object? obj) => (obj is Digest digest) && Equals(digest);
+
+		/// <inheritdoc/>
+		public bool Equals(Digest other) => other.Span.SequenceEqual(Span);
 
 		/// <inheritdoc/>
 		public override int GetHashCode() => BinaryPrimitives.ReadInt32LittleEndian(Span);
@@ -154,7 +161,7 @@ namespace EpicGames.Core
 		/// Constructor
 		/// </summary>
 		/// <param name="length"></param>
-		public DigestTraits(int length)
+		protected DigestTraits(int length)
 		{
 			Length = length;
 		}
@@ -235,7 +242,7 @@ namespace EpicGames.Core
 	/// <summary>
 	/// Generic HashValue implementation
 	/// </summary>
-	public struct Digest<T> where T : DigestTraits, new()
+	public struct Digest<T> : IEquatable<Digest<T>> where T : DigestTraits, new()
 	{
 		/// <summary>
 		/// Traits instance
@@ -272,7 +279,7 @@ namespace EpicGames.Core
 		}
 
 		/// <inheritdoc/>
-		public override bool Equals(object? obj) => (obj is Digest<T> hash) && hash.Span.SequenceEqual(Span);
+		public override bool Equals(object? obj) => (obj is Digest<T> hash) && Equals(hash);
 
 		/// <inheritdoc/>
 		public override int GetHashCode() => BinaryPrimitives.ReadInt32LittleEndian(Span);
@@ -280,6 +287,9 @@ namespace EpicGames.Core
 		/// <inheritdoc/>
 		public override string ToString() => StringUtils.FormatHexString(Memory.Span);
 
+		/// <inheritdoc/>
+		public bool Equals(Digest<T> other) => other.Span.SequenceEqual(Span);
+
 		/// <summary>
 		/// Test two hash values for equality
 		/// </summary>
diff --git a/Engine/Source/Programs/Shared/EpicGames.Core/DirectoryReference.cs b/Engine/Source/Programs/Shared/EpicGames.Core/DirectoryReference.cs
--- a/Engine/Source/Programs/Shared/EpicGames.Core/DirectoryReference.cs
+++ b/Engine/Source/Programs/Shared/EpicGames.Core/DirectoryReference.cs
@@ -213,6 +213,18 @@ namespace EpicGames.Core
 		/// <inheritdoc/>
 		public int CompareTo(DirectoryReference? other) => Comparer.Compare(FullName, other?.FullName);
 
+		/// <inheritdoc/>
+		public static bool operator <(DirectoryReference left, DirectoryReference right) => ReferenceEquals(left, null) ? !ReferenceEquals(right, null) : left.CompareTo(right) < 0;
+
+		/// <inheritdoc/>
+		public static bool operator <=(DirectoryReference left, DirectoryReference right) => ReferenceEquals(left, null) || left.CompareTo(right) <= 0;
+
+		/// <inheritdoc/>
+		public static bool operator >(DirectoryReference left, DirectoryReference right) => !ReferenceEquals(left, null) && left.CompareTo(right) > 0;
+
+		/// <inheritdoc/>
+		public static bool operator >=(DirectoryReference left, DirectoryReference right) => ReferenceEquals(left, null) ? ReferenceEquals(right, null) : left.CompareTo(right) >= 0;
+
 		/// <summary>
 		/// Helper function to create a remote directory reference. Unlike normal DirectoryReference objects, these aren't converted to a full path in the local filesystem.
 		/// </summary>
diff --git a/Engine/Source/Programs/Shared/EpicGames.Core/FileFilter.cs b/Engine/Source/Programs/Shared/EpicGames.Core/FileFilter.cs
--- a/Engine/Source/Programs/Shared/EpicGames.Core/FileFilter.cs
+++ b/Engine/Source/Programs/Shared/EpicGames.Core/FileFilter.cs
@@ -89,7 +89,7 @@ namespace EpicGames.Core
 		/// <param name="rule">Pattern to match. See CreateRegex() for details.</param>
 		public void AddRule(string rule)
 		{
-			if (rule.StartsWith("-"))
+			if (rule.StartsWith("-", StringComparison.Ordinal))
 			{
 				Exclude(rule.Substring(1).TrimStart());
 			}
@@ -107,10 +107,10 @@ namespace EpicGames.Core
 		public void AddRule(string rule, params string[] allowTags)
 		{
 			string cleanRule = rule.Trim();
-			if(cleanRule.StartsWith("{"))
+			if(cleanRule.StartsWith("{", StringComparison.Ordinal))
 			{
 				// Find the end of the condition
-				int conditionEnd = cleanRule.IndexOf('}');
+				int conditionEnd = cleanRule.IndexOf('}', StringComparison.Ordinal);
 				if(conditionEnd == -1)
 				{
 					throw new Exception(String.Format("Missing closing parenthesis in rule: {0}", cleanRule));
@@ -205,9 +205,9 @@ namespace EpicGames.Core
 			foreach(string line in File.ReadAllLines(fileName.FullName))
 			{
 				string trimLine = line.Trim();
-				if(!trimLine.StartsWith(";") && trimLine.Length > 0)
+				if(!trimLine.StartsWith(";", StringComparison.Ordinal) && trimLine.Length > 0)
 				{
-					if(trimLine.StartsWith("["))
+					if(trimLine.StartsWith("[", StringComparison.Ordinal))
 					{
 						bInSection = (trimLine == "[" + sectionName + "]");
 					}
@@ -290,23 +290,23 @@ namespace EpicGames.Core
 
 			// Remove the slash from the start of the pattern. Any exclude pattern that doesn't contain a directory separator is assumed to apply to any directory (eg. *.cpp), otherwise it's 
 			// taken relative to the root.
-			if (normalizedPattern.StartsWith("/"))
+			if (normalizedPattern.StartsWith("/", StringComparison.Ordinal))
 			{
 				normalizedPattern = normalizedPattern.Substring(1);
 			}
-			else if(!normalizedPattern.Contains("/") && !normalizedPattern.StartsWith("..."))
+			else if(!normalizedPattern.Contains('/', StringComparison.Ordinal) && !normalizedPattern.StartsWith("...", StringComparison.Ordinal))
 			{
 				normalizedPattern = ".../" + normalizedPattern;
 			}
 
 			// All directories indicate a wildcard match
-			if (normalizedPattern.EndsWith("/"))
+			if (normalizedPattern.EndsWith("/", StringComparison.Ordinal))
 			{
 				normalizedPattern += "...";
 			}
 
 			// Replace any directory wildcards mid-string
-			for (int idx = normalizedPattern.IndexOf("..."); idx != -1; idx = normalizedPattern.IndexOf("...", idx))
+			for (int idx = normalizedPattern.IndexOf("...", StringComparison.Ordinal); idx != -1; idx = normalizedPattern.IndexOf("...", idx, StringComparison.Ordinal))
 			{
 				if (idx > 0 && normalizedPattern[idx - 1] != '/')
 				{
diff --git a/Engine/Source/Programs/Shared/EpicGames.Core/FileFilterNode.cs b/Engine/Source/Programs/Shared/EpicGames.Core/FileFilterNode.cs
--- a/Engine/Source/Programs/Shared/EpicGames.Core/FileFilterNode.cs
+++ b/Engine/Source/Programs/Shared/EpicGames.Core/FileFilterNode.cs
@@ -61,9 +61,9 @@ namespace EpicGames.Core
 		/// </summary>
 		public bool IsMatch(string token)
 		{
-			if(Pattern.EndsWith("."))
+			if(Pattern.EndsWith(".", StringComparison.Ordinal))
 			{
-				return !token.Contains('.') && IsMatch(token, 0, Pattern.Substring(0, Pattern.Length - 1), 0);
+				return !token.Contains('.', StringComparison.Ordinal) && IsMatch(token, 0, Pattern.Substring(0, Pattern.Length - 1), 0);
 			}
 			else
 			{
diff --git a/Engine/Source/Programs/Shared/EpicGames.Core/FilePattern.cs b/Engine/Source/Programs/Shared/EpicGames.Core/FilePattern.cs
--- a/Engine/Source/Programs/Shared/EpicGames.Core/FilePattern.cs
+++ b/Engine/Source/Programs/Shared/EpicGames.Core/FilePattern.cs
@@ -6,6 +6,8 @@ using System.IO;
 using System.Text;
 using System.Text.RegularExpressions;
 
+#pragma warning disable CA1045 // Do not pass types by reference
+
 namespace EpicGames.Core
 {
 	/// <summary>
@@ -44,12 +46,12 @@ namespace EpicGames.Core
 		/// <summary>
 		/// Base directory for all matched files
 		/// </summary>
-		public readonly DirectoryReference BaseDirectory;
+		public DirectoryReference BaseDirectory { get; }
 		
 		/// <summary>
 		/// List of tokens in the pattern. Every second token is a wildcard, other tokens are string fragments. Always has an odd number of elements. Path separators are normalized to the host platform format.
 		/// </summary>
-		public readonly List<string> Tokens = new List<string>();
+		public IReadOnlyList<string> Tokens { get; }
 
 		/// <summary>
 		/// Constructs a file pattern which matches a single file
@@ -58,7 +60,7 @@ namespace EpicGames.Core
 		public FilePattern(FileReference file)
 		{
 			BaseDirectory = file.Directory;
-			Tokens.Add(file.GetFileName());
+			Tokens = new[] { file.GetFileName() };
 		}
 
 		/// <summary>
@@ -123,23 +125,25 @@ namespace EpicGames.Core
 			}
 
 			// Parse the tokens
+			List<string> tokens = new List<string>();
 			int lastIdx = baseDirectoryLen;
 			for(int idx = baseDirectoryLen; idx < text.Length; idx++)
 			{
 				if(text[idx] == '?' || text[idx] == '*')
 				{
-					Tokens.Add(text.ToString(lastIdx, idx - lastIdx));
-					Tokens.Add(text.ToString(idx, 1));
+					tokens.Add(text.ToString(lastIdx, idx - lastIdx));
+					tokens.Add(text.ToString(idx, 1));
 					lastIdx = idx + 1;
 				}
 				else if(idx - 3 >= baseDirectoryLen && text[idx] == Path.DirectorySeparatorChar && text[idx - 1] == '.' && text[idx - 2] == '.' && text[idx - 3] == '.')
 				{
-					Tokens.Add(text.ToString(lastIdx, idx - 3 - lastIdx));
-					Tokens.Add(text.ToString(idx - 3, 4));
+					tokens.Add(text.ToString(lastIdx, idx - 3 - lastIdx));
+					tokens.Add(text.ToString(idx - 3, 4));
 					lastIdx = idx + 1;
 				}
 			}
-			Tokens.Add(text.ToString(lastIdx, text.Length - lastIdx));
+			tokens.Add(text.ToString(lastIdx, text.Length - lastIdx));
+			Tokens = tokens;
 		}
 
 		/// <summary>
@@ -255,7 +259,7 @@ namespace EpicGames.Core
 				// Append the next sequence of characters to match
 				pattern.Append(Regex.Escape(Tokens[idx + 1]));
 			}
-			pattern.Append("$");
+			pattern.Append('$');
 			return pattern.ToString();
 		}
 
@@ -269,7 +273,7 @@ namespace EpicGames.Core
 			for(int idx = 0;;idx += 2)
 			{
 				// Append the escaped replacement character
-				pattern.Append(Tokens[idx].Replace("$", "$$"));
+				pattern.Append(Tokens[idx].Replace("$", "$$", StringComparison.Ordinal));
 
 				// Check if we've reached the end of the string
 				if(idx == Tokens.Count - 1)
diff --git a/Engine/Source/Programs/Shared/EpicGames.Core/FileReference.cs b/Engine/Source/Programs/Shared/EpicGames.Core/FileReference.cs
--- a/Engine/Source/Programs/Shared/EpicGames.Core/FileReference.cs
+++ b/Engine/Source/Programs/Shared/EpicGames.Core/FileReference.cs
@@ -58,6 +58,7 @@ namespace EpicGames.Core
 		public FileReference(string fullName, Sanitize sanitize)
 			: base(fullName)
 		{
+			_ = sanitize;
 		}
 
 		/// <summary>
@@ -246,6 +247,18 @@ namespace EpicGames.Core
 		/// <inheritdoc/>
 		public int CompareTo(FileReference? other) => Comparer.Compare(FullName, other?.FullName);
 
+		/// <inheritdoc/>
+		public static bool operator <(FileReference left, FileReference right) => ReferenceEquals(left, null) ? !ReferenceEquals(right, null) : left.CompareTo(right) < 0;
+
+		/// <inheritdoc/>
+		public static bool operator <=(FileReference left, FileReference right) => ReferenceEquals(left, null) || left.CompareTo(right) <= 0;
+
+		/// <inheritdoc/>
+		public static bool operator >(FileReference left, FileReference right) => !ReferenceEquals(left, null) && left.CompareTo(right) > 0;
+
+		/// <inheritdoc/>
+		public static bool operator >=(FileReference left, FileReference right) => ReferenceEquals(left, null) ? ReferenceEquals(right, null) : left.CompareTo(right) >= 0;
+
 		/// <summary>
 		/// Helper function to create a remote file reference. Unlike normal FileReference objects, these aren't converted to a full path in the local filesystem, but are
 		/// left as they are passed in.
@@ -650,7 +663,7 @@ namespace EpicGames.Core
 			if(FileReference.Exists(location))
 			{
 				string currentContents = FileReference.ReadAllText(location);
-				if (String.Equals(contents, currentContents))
+				if (String.Equals(contents, currentContents, StringComparison.Ordinal))
 				{
 					return false;
 				}
diff --git a/Engine/Source/Programs/Shared/EpicGames.Core/FileSet.cs b/Engine/Source/Programs/Shared/EpicGames.Core/FileSet.cs
--- a/Engine/Source/Programs/Shared/EpicGames.Core/FileSet.cs
+++ b/Engine/Source/Programs/Shared/EpicGames.Core/FileSet.cs
@@ -28,7 +28,7 @@ namespace EpicGames.Core
 		/// Constructor
 		/// </summary>
 		/// <param name="path">Relative path within the tree</param>
-		public FileSet(string path)
+		protected FileSet(string path)
 		{
 			Path = path;
 		}
diff --git a/Engine/Source/Programs/Shared/EpicGames.Core/FileSystemReference.cs b/Engine/Source/Programs/Shared/EpicGames.Core/FileSystemReference.cs
--- a/Engine/Source/Programs/Shared/EpicGames.Core/FileSystemReference.cs
+++ b/Engine/Source/Programs/Shared/EpicGames.Core/FileSystemReference.cs
@@ -18,17 +18,17 @@ namespace EpicGames.Core
 		/// <summary>
 		/// The path to this object. Stored as an absolute path, with O/S preferred separator characters, and no trailing slash for directories.
 		/// </summary>
-		public readonly string FullName;
+		public string FullName { get; }
 
 		/// <summary>
 		/// The comparer to use for file system references
 		/// </summary>
-		public static readonly StringComparer Comparer = StringComparer.OrdinalIgnoreCase;
+		public static StringComparer Comparer { get; } = StringComparer.OrdinalIgnoreCase;
 
 		/// <summary>
 		/// The comparison to use for file system references
 		/// </summary>
-		public static readonly StringComparison Comparison = StringComparison.OrdinalIgnoreCase;
+		public static StringComparison Comparison { get; } = StringComparison.OrdinalIgnoreCase;
 
 		/// <summary>
 		/// Direct constructor for a path
@@ -241,7 +241,7 @@ namespace EpicGames.Core
 		public string MakeRelativeTo(DirectoryReference directory)
 		{
 			StringBuilder result = s_makeRelativeToStringBuilder.Value!.Clear();
-			WriteRelativeTo(directory, ref result);
+			WriteRelativeTo(directory, result);
 			return result.ToString();
 		}
 
@@ -250,7 +250,7 @@ namespace EpicGames.Core
 		/// </summary>
 		/// <param name="directory"></param>
 		/// <param name="result"></param>
-		public void WriteRelativeTo(DirectoryReference directory, ref StringBuilder result)
+		public void WriteRelativeTo(DirectoryReference directory, StringBuilder result)
 		{
 			// Find how much of the path is common between the two paths. This length does not include a trailing directory separator character.
 			int commonDirectoryLength = -1;
@@ -335,7 +335,7 @@ namespace EpicGames.Core
 		/// <returns></returns>
 		public string ToNormalizedPath()
 		{
-			return FullName.Replace("\\", "/");
+			return FullName.Replace("\\", "/", StringComparison.Ordinal);
 		}
 
 		/// <summary>
diff --git a/Engine/Source/Programs/Shared/EpicGames.Core/FileTransaction.cs b/Engine/Source/Programs/Shared/EpicGames.Core/FileTransaction.cs
--- a/Engine/Source/Programs/Shared/EpicGames.Core/FileTransaction.cs
+++ b/Engine/Source/Programs/Shared/EpicGames.Core/FileTransaction.cs
@@ -169,6 +169,7 @@ namespace EpicGames.Core
 		/// <inheritdoc/>
 		public override async ValueTask DisposeAsync()
 		{
+			GC.SuppressFinalize(this);
 			await base.DisposeAsync();
 			await _inner.DisposeAsync();
 		}
diff --git a/Engine/Source/Programs/Shared/EpicGames.Core/FileUtils.cs b/Engine/Source/Programs/Shared/EpicGames.Core/FileUtils.cs
--- a/Engine/Source/Programs/Shared/EpicGames.Core/FileUtils.cs
+++ b/Engine/Source/Programs/Shared/EpicGames.Core/FileUtils.cs
@@ -9,12 +9,15 @@ using System.Runtime.InteropServices.ComTypes;
 using System.Text;
 using Microsoft.Win32.SafeHandles;
 
+#pragma warning disable CA1707 // Identifiers should not contain underscores
+#pragma warning disable CA1806 // Do not ignore method results
+
 namespace EpicGames.Core
 {
 	/// <summary>
 	/// Exception used to represent caught file/directory exceptions.
 	/// </summary>
-	class WrappedFileOrDirectoryException : Exception
+	public class WrappedFileOrDirectoryException : Exception
 	{
 		/// <summary>
 		/// Constructor
@@ -35,6 +38,53 @@ namespace EpicGames.Core
 		}
 	}
 
+	/// <summary>
+	/// Information about a locked file
+	/// </summary>
+	public class FileLockInfoWin32
+	{
+		/// <summary>
+		/// Process id
+		/// </summary>
+		public int ProcessId { get; }
+
+		/// <summary>
+		/// Path to the process holding the lock
+		/// </summary>
+		public string? FileName { get; }
+
+		/// <summary>
+		/// Name of the application
+		/// </summary>
+		public string AppName { get; }
+
+		/// <summary>
+		/// Time at which the process started
+		/// </summary>
+		public DateTime StartTime { get; }
+
+		/// <summary>
+		/// Constructor
+		/// </summary>
+		/// <param name="processId">Process id</param>
+		/// <param name="fileName">Path to the process holding the lock</param>
+		/// <param name="appName">Name of the application</param>
+		/// <param name="startTime">Time at which the process started</param>
+		public FileLockInfoWin32(int processId, string? fileName, string appName, DateTime startTime)
+		{
+			ProcessId = processId;
+			FileName = fileName;
+			AppName = appName;
+			StartTime = startTime;
+		}
+
+		/// <inheritdoc/>
+		public override string ToString()
+		{
+			return $"{ProcessId}: {FileName ?? AppName} (started {StartTime})";
+		}
+	}
+
 	/// <summary>
 	/// Utility functions for manipulating files. Where these methods have similar functionality to those in the NET Framework, they generally go the extra mile to produce concise, specific error messages where possible.
 	/// </summary>
@@ -43,12 +93,12 @@ namespace EpicGames.Core
 		/// <summary>
 		/// Comparer that should be used for native path comparisons
 		/// </summary>
-		public static IEqualityComparer<string> PlatformPathComparer = RuntimePlatform.IsLinux ? StringComparer.Ordinal : StringComparer.OrdinalIgnoreCase;
+		public static IEqualityComparer<string> PlatformPathComparer { get; } = RuntimePlatform.IsLinux ? StringComparer.Ordinal : StringComparer.OrdinalIgnoreCase;
 
 		/// <summary>
 		/// Utf8 string comparer that should be used for native path comparisons
 		/// </summary>
-		public static IEqualityComparer<Utf8String> PlatformPathComparerUtf8 = RuntimePlatform.IsLinux ? Utf8StringComparer.Ordinal : Utf8StringComparer.OrdinalIgnoreCase;
+		public static IEqualityComparer<Utf8String> PlatformPathComparerUtf8 { get; } = RuntimePlatform.IsLinux ? Utf8StringComparer.Ordinal : Utf8StringComparer.OrdinalIgnoreCase;
 
 		/// <summary>
 		/// Read all text for a file
@@ -542,7 +592,7 @@ namespace EpicGames.Core
 				List<FileLockInfoWin32>? lockInfoList;
 				try
 				{
-					lockInfoList = GetFileLockInfo_Win32(fileName);
+					lockInfoList = GetFileLockInfoWin32(fileName);
 				}
 				catch
 				{
@@ -611,7 +661,9 @@ namespace EpicGames.Core
 		static extern int RmRegisterResources(uint pSessionHandle, uint nFiles, string[] rgsFilenames, uint nApplications, [In] RM_UNIQUE_PROCESS[]? rgApplications, uint nServices, string[]? rgsServiceNames);
 
 		[DllImport("rstrtmgr.dll", CharSet = CharSet.Auto)]
+#pragma warning disable CA1838 // Avoid 'StringBuilder' parameters for P/Invokes
 		static extern int RmStartSession(out uint pSessionHandle, int dwSessionFlags, StringBuilder strSessionKey);
+#pragma warning restore CA1838 // Avoid 'StringBuilder' parameters for P/Invokes
 
 		[DllImport("rstrtmgr.dll")]
 		static extern int RmEndSession(uint pSessionHandle);
@@ -639,62 +691,17 @@ namespace EpicGames.Core
 		   out FILETIME lpUserTime);
 
 		[DllImport("kernel32.dll", SetLastError = true)]
+#pragma warning disable CA1838 // Avoid 'StringBuilder' parameters for P/Invokes
 		static extern bool QueryFullProcessImageName([In]SafeProcessHandle hProcess, [In]int dwFlags, [Out]StringBuilder lpExeName, ref int lpdwSize);
+#pragma warning restore CA1838 // Avoid 'StringBuilder' parameters for P/Invokes
 #pragma warning restore IDE1006
 
-		/// <summary>
-		/// Information about a locked file
-		/// </summary>
-		public class FileLockInfoWin32
-		{
-			/// <summary>
-			/// Process id
-			/// </summary>
-			public int ProcessId { get; }
-
-			/// <summary>
-			/// Path to the process holding the lock
-			/// </summary>
-			public string? FileName { get; }
-
-			/// <summary>
-			/// Name of the application
-			/// </summary>
-			public string AppName { get; }
-
-			/// <summary>
-			/// Time at which the process started
-			/// </summary>
-			public DateTime StartTime { get; }
-
-			/// <summary>
-			/// Constructor
-			/// </summary>
-			/// <param name="processId">Process id</param>
-			/// <param name="fileName">Path to the process holding the lock</param>
-			/// <param name="appName">Name of the application</param>
-			/// <param name="startTime">Time at which the process started</param>
-			public FileLockInfoWin32(int processId, string? fileName, string appName, DateTime startTime)
-			{
-				ProcessId = processId;
-				FileName = fileName;
-				AppName = appName;
-				StartTime = startTime;
-			}
-
-			/// <inheritdoc/>
-			public override string ToString()
-			{
-				return $"{ProcessId}: {FileName ?? AppName} (started {StartTime})";
-			}
-		}
-
 		/// <summary>
 		/// Gets a list of processes that have a handle to the given file open
 		/// </summary>
 		/// <param name="fileName">File to check</param>
 		/// <returns>List of processes with a lock open</returns>
-		public static List<FileLockInfoWin32> GetFileLockInfo_Win32(string fileName)
+		public static List<FileLockInfoWin32> GetFileLockInfoWin32(string fileName)
 		{
 			uint sessionHandle = 0;
 			try
diff --git a/Engine/Source/Programs/Shared/EpicGames.Core/ForwardingLogger.cs b/Engine/Source/Programs/Shared/EpicGames.Core/ForwardingLogger.cs
--- a/Engine/Source/Programs/Shared/EpicGames.Core/ForwardingLogger.cs
+++ b/Engine/Source/Programs/Shared/EpicGames.Core/ForwardingLogger.cs
@@ -29,7 +29,7 @@ namespace EpicGames.Core
 		/// <summary>
 		/// The loggers that are being written to
 		/// </summary>
-		public ILogger[] Loggers
+		public IReadOnlyList<ILogger> Loggers
 		{
 			get;
 		}
diff --git a/Engine/Source/Programs/Shared/EpicGames.Core/HttpClientExtensions.cs b/Engine/Source/Programs/Shared/EpicGames.Core/HttpClientExtensions.cs
--- a/Engine/Source/Programs/Shared/EpicGames.Core/HttpClientExtensions.cs
+++ b/Engine/Source/Programs/Shared/EpicGames.Core/HttpClientExtensions.cs
@@ -6,6 +6,9 @@ using System.Text.Json;
 using System.Threading;
 using System.Threading.Tasks;
 
+#pragma warning disable CA2234 // Pass system uri objects instead of strings
+#pragma warning disable CA1054 // URI-like parameters should not be strings
+
 namespace EpicGames.Core
 {
 	/// <summary>
@@ -41,7 +44,8 @@ namespace EpicGames.Core
 		/// <returns>Response message</returns>
 		public static async Task<HttpResponseMessage> PostAsync<TRequest>(this HttpClient client, string url, TRequest request, CancellationToken cancellationToken)
 		{
-			return await client.PostAsync(url, ToJsonContent(request), cancellationToken);
+			using HttpContent content = ToJsonContent(request);
+			return await client.PostAsync(url, content, cancellationToken);
 		}
 
 		/// <summary>
@@ -74,7 +78,8 @@ namespace EpicGames.Core
 		/// <returns>Response message</returns>
 		public static async Task<HttpResponseMessage> PutAsync<TRequest>(this HttpClient client, string url, TRequest obj, CancellationToken cancellationToken)
 		{
-			return await client.PutAsync(url, ToJsonContent(obj), cancellationToken);
+			using HttpContent content = ToJsonContent(obj);
+			return await client.PutAsync(url, content, cancellationToken);
 		}
 
 		/// <summary>
diff --git a/Engine/Source/Programs/Shared/EpicGames.Core/ILoggerIndent.cs b/Engine/Source/Programs/Shared/EpicGames.Core/ILoggerIndent.cs
--- a/Engine/Source/Programs/Shared/EpicGames.Core/ILoggerIndent.cs
+++ b/Engine/Source/Programs/Shared/EpicGames.Core/ILoggerIndent.cs
@@ -69,7 +69,7 @@ namespace EpicGames.Core
 		/// Struct to wrap a formatted set of log values with applied indent
 		/// </summary>
 		/// <typeparam name="TState">Arbitrary type parameter</typeparam>
-		public struct FormattedLogValues<TState> : IEnumerable<KeyValuePair<string, object>>
+		struct FormattedLogValues<TState> : IEnumerable<KeyValuePair<string, object>>
 		{
 			/// <summary>
 			/// The indent to apply
diff --git a/Engine/Source/Programs/Shared/EpicGames.Core/ILoggerProgress.cs b/Engine/Source/Programs/Shared/EpicGames.Core/ILoggerProgress.cs
--- a/Engine/Source/Programs/Shared/EpicGames.Core/ILoggerProgress.cs
+++ b/Engine/Source/Programs/Shared/EpicGames.Core/ILoggerProgress.cs
@@ -96,7 +96,7 @@ namespace EpicGames.Core
 				Message = message;
 				_progressInternal = String.Empty;
 
-				logger.LogInformation(message);
+				logger.LogInformation("{Progress}", message);
 			}
 
 			/// <summary>
@@ -119,7 +119,7 @@ namespace EpicGames.Core
 				}
 				if (!String.Equals(output, _lastOutput, StringComparison.Ordinal))
 				{
-					_logger.LogInformation(output);
+					_logger.LogInformation("{Progress}", output);
 					_lastOutput = output;
 				}
 			}
diff --git a/Engine/Source/Programs/Shared/EpicGames.Core/JsonExtensions.cs b/Engine/Source/Programs/Shared/EpicGames.Core/JsonExtensions.cs
--- a/Engine/Source/Programs/Shared/EpicGames.Core/JsonExtensions.cs
+++ b/Engine/Source/Programs/Shared/EpicGames.Core/JsonExtensions.cs
@@ -4,6 +4,8 @@ using System;
 using System.Diagnostics;
 using System.Text.Json;
 
+#pragma warning disable CA1045 // Do not pass types by reference
+
 namespace EpicGames.Core
 {
 	/// <summary>
diff --git a/Engine/Source/Programs/Shared/EpicGames.Core/JsonKnownTypes.cs b/Engine/Source/Programs/Shared/EpicGames.Core/JsonKnownTypes.cs
--- a/Engine/Source/Programs/Shared/EpicGames.Core/JsonKnownTypes.cs
+++ b/Engine/Source/Programs/Shared/EpicGames.Core/JsonKnownTypes.cs
@@ -12,7 +12,7 @@ namespace EpicGames.Core
 	/// Sets a name used to discriminate between classes derived from a base class
 	/// </summary>
 	[AttributeUsage(AttributeTargets.Class)]
-	public class JsonDiscriminatorAttribute : Attribute
+	public sealed class JsonDiscriminatorAttribute : Attribute
 	{
 		/// <summary>
 		/// Name to use to discriminate between different classes
@@ -33,7 +33,7 @@ namespace EpicGames.Core
 	/// Marks a class as the base class of a hierarchy, allowing any known subclasses to be serialized for it.
 	/// </summary>
 	[AttributeUsage(AttributeTargets.Class)]
-	public class JsonKnownTypesAttribute : Attribute
+	public sealed class JsonKnownTypesAttribute : Attribute
 	{
 		/// <summary>
 		/// Array of derived classes
diff --git a/Engine/Source/Programs/Shared/EpicGames.Core/JsonLogEvent.cs b/Engine/Source/Programs/Shared/EpicGames.Core/JsonLogEvent.cs
--- a/Engine/Source/Programs/Shared/EpicGames.Core/JsonLogEvent.cs
+++ b/Engine/Source/Programs/Shared/EpicGames.Core/JsonLogEvent.cs
@@ -147,7 +147,9 @@ namespace EpicGames.Core
 		static readonly sbyte[] s_firstCharToLogLevel;
 		static readonly byte[][] s_logLevelNames;
 
+#pragma warning disable CA2207 // Initialize value type static fields inline
 		static JsonLogEvent()
+#pragma warning restore CA2207 // Initialize value type static fields inline
 		{
 			const int LogLevelCount = (int)LogLevel.None;
 
diff --git a/Engine/Source/Programs/Shared/EpicGames.Core/JsonSchemaAttributes.cs b/Engine/Source/Programs/Shared/EpicGames.Core/JsonSchemaAttributes.cs
--- a/Engine/Source/Programs/Shared/EpicGames.Core/JsonSchemaAttributes.cs
+++ b/Engine/Source/Programs/Shared/EpicGames.Core/JsonSchemaAttributes.cs
@@ -8,7 +8,7 @@ namespace EpicGames.Core
 	/// Attribute setting the identifier for a schema
 	/// </summary>
 	[AttributeUsage(AttributeTargets.Class)]
-	public class JsonSchemaAttribute : Attribute
+	public sealed class JsonSchemaAttribute : Attribute
 	{
 		/// <summary>
 		/// The schema identifier
@@ -29,22 +29,22 @@ namespace EpicGames.Core
 	/// Attribute setting catalog entries for a schema
 	/// </summary>
 	[AttributeUsage(AttributeTargets.Class)]
-	public class JsonSchemaCatalogAttribute : Attribute
+	public sealed class JsonSchemaCatalogAttribute : Attribute
 	{
 		/// <summary>
 		/// The schema name
 		/// </summary>
-		public string Name { get; set; }
+		public string Name { get; }
 
 		/// <summary>
 		/// Description of the schema
 		/// </summary>
-		public string Description { get; set; }
+		public string Description { get; }
 
 		/// <summary>
 		/// File patterns to match
 		/// </summary>
-		public string[]? FileMatch { get; set; }
+		public string[]? FileMatch { get; }
 
 		/// <summary>
 		/// Constructor
@@ -52,7 +52,9 @@ namespace EpicGames.Core
 		/// <param name="name"></param>
 		/// <param name="description"></param>
 		/// <param name="fileMatch">File patterns to match</param>
+#pragma warning disable CA1019 // Define accessors for attribute arguments
 		public JsonSchemaCatalogAttribute(string name, string description, string? fileMatch)
+#pragma warning restore CA1019 // Define accessors for attribute arguments
 			: this(name, description, (fileMatch == null) ? (string[]?)null : new[] { fileMatch })
 		{
 		}
@@ -82,7 +84,7 @@ namespace EpicGames.Core
 	/// Attribute setting properties for a type to be serialized as a string
 	/// </summary>
 	[AttributeUsage(AttributeTargets.Class | AttributeTargets.Struct | AttributeTargets.Property)]
-	public class JsonSchemaStringAttribute : JsonSchemaTypeAttribute
+	public sealed class JsonSchemaStringAttribute : JsonSchemaTypeAttribute
 	{
 		/// <summary>
 		/// Format of the string
diff --git a/Engine/Source/Programs/Shared/EpicGames.Core/JsonSchemaType.cs b/Engine/Source/Programs/Shared/EpicGames.Core/JsonSchemaType.cs
--- a/Engine/Source/Programs/Shared/EpicGames.Core/JsonSchemaType.cs
+++ b/Engine/Source/Programs/Shared/EpicGames.Core/JsonSchemaType.cs
@@ -133,7 +133,9 @@ namespace EpicGames.Core
 			writer.WriteString("type", "string");
 			if (Format != null)
 			{
+#pragma warning disable CA1308 // Normalize strings to uppercase
 				writer.WriteString("format", Format.Value.ToString().ToLowerInvariant());
+#pragma warning restore CA1308 // Normalize strings to uppercase
 			}
 			if (Pattern != null)
 			{
@@ -255,7 +257,7 @@ namespace EpicGames.Core
 		/// <summary>
 		/// Properties to allow in this object
 		/// </summary>
-		public List<JsonSchemaProperty> Properties { get; set; } = new List<JsonSchemaProperty>();
+		public List<JsonSchemaProperty> Properties { get; } = new List<JsonSchemaProperty>();
 
 		/// <summary>
 		/// Type for additional properties
diff --git a/Engine/Source/Programs/Shared/EpicGames.Core/JsonTracer.cs b/Engine/Source/Programs/Shared/EpicGames.Core/JsonTracer.cs
--- a/Engine/Source/Programs/Shared/EpicGames.Core/JsonTracer.cs
+++ b/Engine/Source/Programs/Shared/EpicGames.Core/JsonTracer.cs
@@ -202,7 +202,7 @@ namespace EpicGames.Core
 			// return the context of the parent, if applicable
 			foreach (Reference reference in references)
 			{
-				if (OpenTracing.References.ChildOf.Equals(reference.ReferenceType))
+				if (OpenTracing.References.ChildOf.Equals(reference.ReferenceType, StringComparison.Ordinal))
 				{
 					return reference.Context;
 				}
@@ -219,8 +219,10 @@ namespace EpicGames.Core
 				throw new InvalidOperationException(String.Format(format, args));
 			}
 		}
-		
+
+#pragma warning disable CA1034 // Nested types should not be visible
 		public sealed class Reference : IEquatable<Reference>
+#pragma warning restore CA1034 // Nested types should not be visible
 		{
 			public JsonTracerSpanContext Context { get; }
 
diff --git a/Engine/Source/Programs/Shared/EpicGames.Core/JsonWriter.cs b/Engine/Source/Programs/Shared/EpicGames.Core/JsonWriter.cs
--- a/Engine/Source/Programs/Shared/EpicGames.Core/JsonWriter.cs
+++ b/Engine/Source/Programs/Shared/EpicGames.Core/JsonWriter.cs
@@ -27,7 +27,7 @@ namespace EpicGames.Core
 	/// <summary>
 	/// Writer for JSON data, which indents the output text appropriately, and adds commas and newlines between fields
 	/// </summary>
-	public class JsonWriter : IDisposable
+	public sealed class JsonWriter : IDisposable
 	{
 		TextWriter _writer;
 		readonly bool _leaveOpen;
diff --git a/Engine/Source/Programs/Shared/EpicGames.Core/ListSegment.cs b/Engine/Source/Programs/Shared/EpicGames.Core/ListSegment.cs
--- a/Engine/Source/Programs/Shared/EpicGames.Core/ListSegment.cs
+++ b/Engine/Source/Programs/Shared/EpicGames.Core/ListSegment.cs
@@ -4,8 +4,21 @@ using System;
 using System.Collections;
 using System.Collections.Generic;
 
+#pragma warning disable CA1710 // Identifiers should have correct suffix
+
 namespace EpicGames.Core
 {
+	/// <summary>
+	/// Static methods for <see cref="ListSegment{T}"/>
+	/// </summary>
+	public static class ListSegment
+	{
+		/// <summary>
+		/// An empty list segment of the given element type
+		/// </summary>
+		public static ListSegment<T> Empty<T>() => new ListSegment<T>(null!, 0, 0);
+	}
+
 	/// <summary>
 	/// Delimits a section of a list
 	/// </summary>
@@ -29,11 +42,6 @@ namespace EpicGames.Core
 		/// </summary>
 		public int Count { get; }
 
-		/// <summary>
-		/// An empty list segment of the given element type
-		/// </summary>
-		public static ListSegment<T> Empty { get; } = new ListSegment<T>(null!, 0, 0);
-
 		/// <inheritdoc/>
 		public bool IsReadOnly => List.IsReadOnly;
 
diff --git a/Engine/Source/Programs/Shared/EpicGames.Core/Log.cs b/Engine/Source/Programs/Shared/EpicGames.Core/Log.cs
--- a/Engine/Source/Programs/Shared/EpicGames.Core/Log.cs
+++ b/Engine/Source/Programs/Shared/EpicGames.Core/Log.cs
@@ -8,7 +8,6 @@ using System.IO;
 using System.Linq;
 using System.Reflection;
 using System.Text;
-using System.Text.Json;
 using System.Text.RegularExpressions;
 using System.Threading.Channels;
 using System.Threading.Tasks;
@@ -21,7 +20,9 @@ namespace EpicGames.Core
 	/// <summary>
 	/// Log Event Type
 	/// </summary>
+#pragma warning disable CA1027 // Mark enums with FlagsAttribute
 	public enum LogEventType
+#pragma warning restore CA1027 // Mark enums with FlagsAttribute
 	{
 		/// <summary>
 		/// The log event is a fatal error
@@ -56,7 +57,9 @@ namespace EpicGames.Core
 		/// <summary>
 		/// The log event should only be displayed if very verbose logging is enabled
 		/// </summary>
+#pragma warning disable CA1069 // Enums values should not be duplicated
 		VeryVerbose = LogLevel.Trace
+#pragma warning restore CA1069 // Enums values should not be duplicated
 	}
 
 	/// <summary>
@@ -166,13 +169,13 @@ namespace EpicGames.Core
 		/// Log.txt will be backed up with its UTC creation time in the name e.g.
 		/// Log-backup-2021.10.29-19.53.17.txt
 		/// </summary>
-		public static bool BackupLogFiles = true;
+		public static bool BackupLogFiles { get; set; } = true;
 		
 		/// <summary>
 		/// The number of backups to be preserved - when there are more than this, the oldest backups will be deleted.
 		/// Backups will not be deleted if BackupLogFiles is false.
 		/// </summary>
-		public static int LogFileBackupCount = 10;
+		public static int LogFileBackupCount { get; set; } = 10;
 		
 		/// <summary>
 		/// Path to the log file being written to. May be null.
@@ -375,7 +378,7 @@ namespace EpicGames.Core
 
 				// Remove any trailing whitespace
 				int trimLen = message.Length;
-				while (trimLen > 0 && " \t\r\n".Contains(message[trimLen - 1]))
+				while (trimLen > 0 && " \t\r\n".Contains(message[trimLen - 1], StringComparison.Ordinal))
 				{
 					trimLen--;
 				}
@@ -406,7 +409,7 @@ namespace EpicGames.Core
 				// Forward it on to the internal logger
 				if (verbosity < LogEventType.Console)
 				{
-					Logger.Log((LogLevel)verbosity, message.ToString());
+					Logger.Log((LogLevel)verbosity, "{Message}", message.ToString());
 				}
 				else
 				{
@@ -813,7 +816,7 @@ namespace EpicGames.Core
 	/// Wrapper around a custom logger interface which flushes the event parser when switching between legacy
 	/// and native structured logging
 	/// </summary>
-	class LegacyEventLogger : ILogger
+	sealed class LegacyEventLogger : ILogger, IDisposable
 	{
 		private ILogger _inner;
 		private readonly LogEventParser _parser;
@@ -826,6 +829,11 @@ namespace EpicGames.Core
 			_parser = new LogEventParser(inner);
 		}
 
+		public void Dispose()
+		{
+			_parser.Dispose();
+		}
+
 		public void SetInnerLogger(ILogger inner)
 		{
 			lock (_parser)
@@ -943,7 +951,7 @@ namespace EpicGames.Core
 		/// <summary>
 		/// Whether console output is redirected. This prevents writing status updates that rely on moving the cursor.
 		/// </summary>
-		private bool AllowStatusUpdates => !Console.IsOutputRedirected;
+		private static bool AllowStatusUpdates { get; set; } = !Console.IsOutputRedirected;
 
 		/// <summary>
 		/// When configured, this tracks time since initialization to prepend a timestamp to each log.
@@ -1375,7 +1383,7 @@ namespace EpicGames.Core
 	/// <summary>
 	/// Provider for default logger instances
 	/// </summary>
-	public class DefaultLoggerProvider : ILoggerProvider
+	public sealed class DefaultLoggerProvider : ILoggerProvider
 	{
 		/// <inheritdoc/>
 		public ILogger CreateLogger(string categoryName)
diff --git a/Engine/Source/Programs/Shared/EpicGames.Core/LogEvent.cs b/Engine/Source/Programs/Shared/EpicGames.Core/LogEvent.cs
--- a/Engine/Source/Programs/Shared/EpicGames.Core/LogEvent.cs
+++ b/Engine/Source/Programs/Shared/EpicGames.Core/LogEvent.cs
@@ -203,7 +203,9 @@ namespace EpicGames.Core
 		/// </summary>
 		/// <param name="reader">The Json reader</param>
 		/// <returns>New log event</returns>
+#pragma warning disable CA1045 // Do not pass types by reference
 		public static LogEvent Read(ref Utf8JsonReader reader)
+#pragma warning restore CA1045 // Do not pass types by reference
 		{
 			DateTime time = new DateTime(0);
 			LogLevel level = LogLevel.None;
@@ -542,7 +544,7 @@ namespace EpicGames.Core
 		/// <summary>
 		/// Multiple inner exceptions, in the case of an <see cref="AggregateException"/>
 		/// </summary>
-		public List<LogException> InnerExceptions { get; set; } = new List<LogException>();
+		public List<LogException> InnerExceptions { get; } = new List<LogException>();
 
 		/// <summary>
 		/// Constructor
@@ -575,7 +577,6 @@ namespace EpicGames.Core
 				AggregateException? aggregateException = exception as AggregateException;
 				if (aggregateException != null && aggregateException.InnerExceptions.Count > 0)
 				{
-					result.InnerExceptions = new List<LogException>();
 					for (int idx = 0; idx < 16 && idx < aggregateException.InnerExceptions.Count; idx++) // Cap number of exceptions returned to avoid huge messages
 					{
 						LogException innerException = FromException(aggregateException.InnerExceptions[idx]);
diff --git a/Engine/Source/Programs/Shared/EpicGames.Core/LogEventParser.cs b/Engine/Source/Programs/Shared/EpicGames.Core/LogEventParser.cs
--- a/Engine/Source/Programs/Shared/EpicGames.Core/LogEventParser.cs
+++ b/Engine/Source/Programs/Shared/EpicGames.Core/LogEventParser.cs
@@ -129,7 +129,7 @@ namespace EpicGames.Core
 		/// <summary>
 		/// Buffer for holding partial line data
 		/// </summary>
-		readonly MemoryStream _partialLine = new MemoryStream();
+		readonly ByteArrayBuilder _partialLine = new ByteArrayBuilder();
 
 		/// <summary>
 		/// Whether matching is currently enabled
@@ -176,7 +176,23 @@ namespace EpicGames.Core
 		}
 
 		/// <inheritdoc/>
-		public void Dispose() => Flush();
+		public void Dispose()
+		{
+			Dispose(true);
+			GC.SuppressFinalize(this);
+		}
+
+		/// <summary>
+		/// Standard Dispose pattern method
+		/// </summary>
+		/// <param name="disposing"></param>
+		protected virtual void Dispose(bool disposing)
+		{
+			if (disposing)
+			{
+				Flush();
+			}
+		}
 
 		/// <summary>
 		/// Enumerate all the types that implement <see cref="ILogEventMatcher"/> in the given assembly, and create instances of them
@@ -319,7 +335,7 @@ namespace EpicGames.Core
 				{
 					if (span[scanIdx] == '\n')
 					{
-						_partialLine.Write(span.Slice(baseIdx, scanIdx - baseIdx));
+						_partialLine.WriteFixedLengthBytes(span.Slice(baseIdx, scanIdx - baseIdx));
 						FlushPartialLine();
 						baseIdx = ++scanIdx;
 						break;
@@ -338,7 +354,7 @@ namespace EpicGames.Core
 			}
 
 			// Add the rest of the text to the partial line buffer
-			_partialLine.Write(span.Slice(baseIdx));
+			_partialLine.WriteFixedLengthBytes(span.Slice(baseIdx));
 
 			// Process the new data
 			ProcessData(false);
@@ -387,9 +403,8 @@ namespace EpicGames.Core
 		/// </summary>
 		private void FlushPartialLine()
 		{
-			AddLine(_partialLine.ToArray());
-			_partialLine.Position = 0;
-			_partialLine.SetLength(0);
+			AddLine(_partialLine.ToByteArray());
+			_partialLine.Clear();
 		}
 
 		/// <summary>
diff --git a/Engine/Source/Programs/Shared/EpicGames.Core/LogIndentScope.cs b/Engine/Source/Programs/Shared/EpicGames.Core/LogIndentScope.cs
--- a/Engine/Source/Programs/Shared/EpicGames.Core/LogIndentScope.cs
+++ b/Engine/Source/Programs/Shared/EpicGames.Core/LogIndentScope.cs
@@ -7,7 +7,7 @@ namespace EpicGames.Core
 	/// <summary>
 	/// Class to apply a log indent for the lifetime of an object 
 	/// </summary>
-	public class LogIndentScope : IDisposable
+	public sealed class LogIndentScope : IDisposable
 	{
 		/// <summary>
 		/// Whether the object has been disposed
diff --git a/Engine/Source/Programs/Shared/EpicGames.Core/LogStatusScope.cs b/Engine/Source/Programs/Shared/EpicGames.Core/LogStatusScope.cs
--- a/Engine/Source/Programs/Shared/EpicGames.Core/LogStatusScope.cs
+++ b/Engine/Source/Programs/Shared/EpicGames.Core/LogStatusScope.cs
@@ -8,7 +8,7 @@ namespace EpicGames.Core
 	/// <summary>
 	/// Writes a status message to the log, which can be updated with a progress indicator as a slow task is being performed.
 	/// </summary>
-	public class LogStatusScope : IDisposable
+	public sealed class LogStatusScope : IDisposable
 	{
 		/// <summary>
 		/// The base status message
diff --git a/Engine/Source/Programs/Shared/EpicGames.Core/LogValueFormatter.cs b/Engine/Source/Programs/Shared/EpicGames.Core/LogValueFormatter.cs
--- a/Engine/Source/Programs/Shared/EpicGames.Core/LogValueFormatter.cs
+++ b/Engine/Source/Programs/Shared/EpicGames.Core/LogValueFormatter.cs
@@ -25,7 +25,7 @@ namespace EpicGames.Core
 	/// Specifies a type to use for serializing objects to a log event
 	/// </summary>
 	[AttributeUsage(AttributeTargets.Struct | AttributeTargets.Class)]
-	public class LogValueFormatterAttribute : Attribute
+	public sealed class LogValueFormatterAttribute : Attribute
 	{
 		/// <summary>
 		/// Type of the formatter
diff --git a/Engine/Source/Programs/Shared/EpicGames.Core/ManagedProcess.cs b/Engine/Source/Programs/Shared/EpicGames.Core/ManagedProcess.cs
--- a/Engine/Source/Programs/Shared/EpicGames.Core/ManagedProcess.cs
+++ b/Engine/Source/Programs/Shared/EpicGames.Core/ManagedProcess.cs
@@ -36,7 +36,7 @@ namespace EpicGames.Core
 	/// <summary>
 	/// Tracks a set of processes, and destroys them when the object is disposed.
 	/// </summary>
-	public class ManagedProcessGroup : IDisposable
+	public sealed class ManagedProcessGroup : IDisposable
 	{
 #pragma warning disable IDE0049 // Naming Styles
 #pragma warning disable IDE1006 // Naming Styles
@@ -266,7 +266,9 @@ namespace EpicGames.Core
 		}
 
 		[DllImport("kernel32.dll", SetLastError = true)]
-        static extern int CreateProcess(/*[MarshalAs(UnmanagedType.LPTStr)]*/ string? lpApplicationName, StringBuilder lpCommandLine, IntPtr lpProcessAttributes, IntPtr lpThreadAttributes, bool bInheritHandles, ProcessCreationFlags dwCreationFlags, IntPtr lpEnvironment, /*[MarshalAs(UnmanagedType.LPTStr)]*/ string? lpCurrentDirectory, STARTUPINFO lpStartupInfo, PROCESS_INFORMATION lpProcessInformation);
+#pragma warning disable CA1838 // Avoid 'StringBuilder' parameters for P/Invokes
+		static extern int CreateProcess(/*[MarshalAs(UnmanagedType.LPTStr)]*/ string? lpApplicationName, StringBuilder lpCommandLine, IntPtr lpProcessAttributes, IntPtr lpThreadAttributes, bool bInheritHandles, ProcessCreationFlags dwCreationFlags, IntPtr lpEnvironment, /*[MarshalAs(UnmanagedType.LPTStr)]*/ string? lpCurrentDirectory, STARTUPINFO lpStartupInfo, PROCESS_INFORMATION lpProcessInformation);
+#pragma warning restore CA1838 // Avoid 'StringBuilder' parameters for P/Invokes
 
 		[DllImport("kernel32.dll", SetLastError = true)]
 		static extern int ResumeThread(IntPtr hThread);
@@ -297,6 +299,7 @@ namespace EpicGames.Core
 		static extern int TerminateProcess(SafeHandleZeroOrMinusOneIsInvalid hProcess, uint uExitCode);
 
 		const UInt32 INFINITE = 0xFFFFFFFF;
+		const UInt32 WAIT_FAILED = 0xFFFFFFFF;
 
 		[DllImport("kernel32.dll", SetLastError = true)]
 		static extern UInt32 WaitForSingleObject(SafeHandleZeroOrMinusOneIsInvalid hHandle, UInt32 dwMilliseconds);
@@ -425,7 +428,9 @@ namespace EpicGames.Core
 		/// <summary>
 		/// Cancellation token for background threads
 		/// </summary>
+#pragma warning disable CA2213
 		CancellationTokenSource? _frameworkCancellationSource;
+#pragma warning restore CA2213
 
 		/// <summary>
 		/// Static lock object. This is used to synchronize the creation of child processes - in particular, the inheritance of stdout/stderr write pipes. If processes
@@ -734,12 +739,9 @@ namespace EpicGames.Core
 
 					// Create a JobObject for each spawned process to do CPU usage accounting of spawned process and all its children
 					_accountingProcessGroup = new ManagedProcessGroup();
-					if (_accountingProcessGroup != null)
+					if (AssignProcessToJobObject(_accountingProcessGroup.JobHandle, processInfo.hProcess) == 0) 
 					{
-						if (AssignProcessToJobObject(_accountingProcessGroup.JobHandle, processInfo.hProcess) == 0) 
-						{
-							throw new Win32Exception();
-						}
+						throw new Win32Exception();
 					}
 
 					// On systems with more than one processor group (more than one CPU socket, or more than 64 cores), it is possible
@@ -764,7 +766,10 @@ namespace EpicGames.Core
 						GROUP_AFFINITY groupAffinity = new GROUP_AFFINITY();
 						groupAffinity.Mask = ~0ul >> (int)(64 - groupProcessorCount);
 						groupAffinity.Group = processorGroup;
-						SetThreadGroupAffinity(processInfo.hThread, groupAffinity, null);
+						if (SetThreadGroupAffinity(processInfo.hThread, groupAffinity, null) == 0)
+						{
+							throw new Win32Exception();
+						}
 					}
 					
 
@@ -800,11 +805,11 @@ namespace EpicGames.Core
 				{
 					if (processInfo.hProcess != IntPtr.Zero && _processHandle == null)
 					{
-						CloseHandle(processInfo.hProcess);
+						_ = CloseHandle(processInfo.hProcess);
 					}
 					if (processInfo.hThread != IntPtr.Zero)
 					{
-						CloseHandle(processInfo.hThread);
+						_ = CloseHandle(processInfo.hThread);
 					}
 				}
 			}
@@ -888,7 +893,7 @@ namespace EpicGames.Core
 				CancellationToken cancellationToken = _frameworkCancellationSource.Token;
 				Task stdOutTask = Task.Run(() => CopyPipeAsync(_frameworkProcess.StandardOutput.BaseStream, _frameworkChannel, cancellationToken));
 				Task stdErrTask = Task.Run(() => CopyPipeAsync(_frameworkProcess.StandardError.BaseStream, _frameworkChannel, cancellationToken));
-				_frameworkOutputTask = Task.WhenAll(stdOutTask, stdErrTask).ContinueWith(x => _frameworkChannel.Writer.Complete());
+				_frameworkOutputTask = Task.WhenAll(stdOutTask, stdErrTask).ContinueWith(x => _frameworkChannel.Writer.Complete(), TaskScheduler.Default);
 			}
 			else
 			{
@@ -929,8 +934,8 @@ namespace EpicGames.Core
 		{
 			if(_processHandle != null)
 			{
-				TerminateProcess(_processHandle, 0);
-				WaitForSingleObject(_processHandle, INFINITE);
+				_ = TerminateProcess(_processHandle, 0);
+				_ = WaitForSingleObject(_processHandle, INFINITE);
 
 				_processHandle.Dispose();
 				_processHandle = null;
@@ -976,7 +981,7 @@ namespace EpicGames.Core
 				_frameworkProcess.Dispose();
 				_frameworkProcess = null;
 
-				_frameworkOutputTask?.ContinueWith(x => _frameworkCancellationSource?.Dispose());
+				_frameworkOutputTask?.ContinueWith(x => _frameworkCancellationSource?.Dispose(), TaskScheduler.Default);
 			}
 		}
 
@@ -1253,7 +1258,10 @@ namespace EpicGames.Core
 		{
 			if (_frameworkProcess == null)
 			{
-				WaitForSingleObject(_processHandle!, INFINITE);
+				if (WaitForSingleObject(_processHandle!, INFINITE) == WAIT_FAILED)
+				{
+					throw new Win32Exception();
+				}
 			}
 			else
 			{
diff --git a/Engine/Source/Programs/Shared/EpicGames.Core/Md5Hash.cs b/Engine/Source/Programs/Shared/EpicGames.Core/Md5Hash.cs
--- a/Engine/Source/Programs/Shared/EpicGames.Core/Md5Hash.cs
+++ b/Engine/Source/Programs/Shared/EpicGames.Core/Md5Hash.cs
@@ -7,6 +7,8 @@ using System.Globalization;
 using System.IO;
 using System.Security.Cryptography;
 
+#pragma warning disable CA5351 // Do Not Use Broken Cryptographic Algorithms
+
 namespace EpicGames.Core
 {
 	/// <summary>
diff --git a/Engine/Source/Programs/Shared/EpicGames.Core/MemoryReader.cs b/Engine/Source/Programs/Shared/EpicGames.Core/MemoryReader.cs
--- a/Engine/Source/Programs/Shared/EpicGames.Core/MemoryReader.cs
+++ b/Engine/Source/Programs/Shared/EpicGames.Core/MemoryReader.cs
@@ -33,7 +33,7 @@ namespace EpicGames.Core
 		/// <summary>
 		/// The memory to read from
 		/// </summary>
-		public ReadOnlyMemory<byte> Memory
+		public ReadOnlyMemory<byte> RemainingMemory
 		{
 			get; private set;
 		}
@@ -41,7 +41,7 @@ namespace EpicGames.Core
 		/// <summary>
 		/// Returns the memory at the current offset
 		/// </summary>
-		public ReadOnlySpan<byte> Span => Memory.Span;
+		public ReadOnlySpan<byte> RemainingSpan => RemainingMemory.Span;
 
 		/// <summary>
 		/// Constructor
@@ -49,7 +49,7 @@ namespace EpicGames.Core
 		/// <param name="memory">The memory to read from</param>
 		public MemoryReader(ReadOnlyMemory<byte> memory)
 		{
-			Memory = memory;
+			RemainingMemory = memory;
 		}
 
 		/// <summary>
@@ -57,17 +57,17 @@ namespace EpicGames.Core
 		/// </summary>
 		public void CheckEmpty()
 		{
-			if (Memory.Length > 0)
+			if (RemainingMemory.Length > 0)
 			{
-				throw new Exception($"Serialization is not at expected offset within the input buffer ({Memory.Length} bytes unused)");
+				throw new Exception($"Serialization is not at expected offset within the input buffer ({RemainingMemory.Length} bytes unused)");
 			}
 		}
 
 		/// <inheritdoc/>
-		public ReadOnlyMemory<byte> GetMemory(int minSize) => Memory;
+		public ReadOnlyMemory<byte> GetMemory(int minSize) => RemainingMemory;
 
 		/// <inheritdoc/>
-		public void Advance(int length) => Memory = Memory.Slice(length);
+		public void Advance(int length) => RemainingMemory = RemainingMemory.Slice(length);
 	}
 
 	/// <summary>
diff --git a/Engine/Source/Programs/Shared/EpicGames.Core/MessageTemplate.cs b/Engine/Source/Programs/Shared/EpicGames.Core/MessageTemplate.cs
--- a/Engine/Source/Programs/Shared/EpicGames.Core/MessageTemplate.cs
+++ b/Engine/Source/Programs/Shared/EpicGames.Core/MessageTemplate.cs
@@ -115,7 +115,7 @@ namespace EpicGames.Core
 			char lastChar = '\0';
 			foreach (char character in text)
 			{
-				if ((character != '{' || character != '}') || character != lastChar)
+				if ((character != '{' && character != '}') || character != lastChar)
 				{
 					result.Append(character);
 				}
diff --git a/Engine/Source/Programs/Shared/EpicGames.Core/ProcessUtils.cs b/Engine/Source/Programs/Shared/EpicGames.Core/ProcessUtils.cs
--- a/Engine/Source/Programs/Shared/EpicGames.Core/ProcessUtils.cs
+++ b/Engine/Source/Programs/Shared/EpicGames.Core/ProcessUtils.cs
@@ -2,6 +2,7 @@
 
 using System;
 using System.Collections.Generic;
+using System.ComponentModel;
 using System.Diagnostics;
 using System.Linq;
 using System.Runtime.InteropServices;
@@ -28,11 +29,15 @@ namespace EpicGames.Core
 		static extern bool TerminateProcess(SafeProcessHandle hProcess, uint uExitCode);
 
 		[DllImport("kernel32.dll", SetLastError = true)]
+#pragma warning disable CA1838 // Avoid 'StringBuilder' parameters for P/Invokes
 		static extern int QueryFullProcessImageName([In]SafeProcessHandle hProcess, [In]int dwFlags, [Out]StringBuilder lpExeName, ref int lpdwSize);
+#pragma warning restore CA1838 // Avoid 'StringBuilder' parameters for P/Invokes
 
 		[DllImport("Psapi.dll", SetLastError = true)]
 		static extern bool EnumProcesses([MarshalAs(UnmanagedType.LPArray, ArraySubType = UnmanagedType.U4)] [In][Out] uint[] processIds, int arraySizeBytes, [MarshalAs(UnmanagedType.U4)] out int bytesCopied);
 
+		const uint WAIT_FAILED = 0xffffffff;
+
 		[DllImport("kernel32.dll")]
 		static extern uint WaitForMultipleObjects(int nCount, IntPtr[] lpHandles, bool bWaitAll, uint dwMilliseconds);
 
@@ -134,7 +139,10 @@ namespace EpicGames.Core
 					}
 
 					// Wait for them all to complete
-					WaitForMultipleObjects(waitHandles.Count, waitHandles.Select(x => x.DangerousGetHandle()).ToArray(), true, 10 * 1000);
+					if (WaitForMultipleObjects(waitHandles.Count, waitHandles.Select(x => x.DangerousGetHandle()).ToArray(), true, 10 * 1000) == WAIT_FAILED)
+					{
+						throw new Win32Exception();
+					}
 				}
 				finally
 				{
diff --git a/Engine/Source/Programs/Shared/EpicGames.Core/ReadOnlySequenceBuilder.cs b/Engine/Source/Programs/Shared/EpicGames.Core/ReadOnlySequenceBuilder.cs
--- a/Engine/Source/Programs/Shared/EpicGames.Core/ReadOnlySequenceBuilder.cs
+++ b/Engine/Source/Programs/Shared/EpicGames.Core/ReadOnlySequenceBuilder.cs
@@ -3,10 +3,22 @@
 using System;
 using System.Buffers;
 using System.Collections.Generic;
-using System.Text;
 
 namespace EpicGames.Core
 {
+	/// <summary>
+	/// Static utility methods
+	/// </summary>
+	public static class ReadOnlySequence
+	{
+		/// <summary>
+		/// Create a sequence from a list of segments
+		/// </summary>
+		/// <param name="segments"></param>
+		/// <returns>Sequence for the list of segments</returns>
+		public static ReadOnlySequence<T> Create<T>(IReadOnlyList<ReadOnlyMemory<T>> segments) => ReadOnlySequenceBuilder<T>.Create(segments);
+	}
+
 	/// <summary>
 	/// Utility class to combine buffers into a <see cref="ReadOnlySequence{T}"/>
 	/// </summary>
@@ -39,7 +51,7 @@ namespace EpicGames.Core
 		/// </summary>
 		/// <param name="segments"></param>
 		/// <returns>Sequence for the list of segments</returns>
-		public static ReadOnlySequence<T> Create(IReadOnlyList<ReadOnlyMemory<T>> segments)
+		internal static ReadOnlySequence<T> Create(IReadOnlyList<ReadOnlyMemory<T>> segments)
 		{
 			if (segments.Count == 0)
 			{
diff --git a/Engine/Source/Programs/Shared/EpicGames.Core/RuntimePlatform.cs b/Engine/Source/Programs/Shared/EpicGames.Core/RuntimePlatform.cs
--- a/Engine/Source/Programs/Shared/EpicGames.Core/RuntimePlatform.cs
+++ b/Engine/Source/Programs/Shared/EpicGames.Core/RuntimePlatform.cs
@@ -7,7 +7,7 @@ namespace EpicGames.Core
 	/// <summary>
 	/// Allows testing different platforsm
 	/// </summary>
-	public class RuntimePlatform
+	public static class RuntimePlatform
 	{
 		/// <summary>
 		/// Whether we are currently running on Linux.
diff --git a/Engine/Source/Programs/Shared/EpicGames.Core/ScopeProfiler.cs b/Engine/Source/Programs/Shared/EpicGames.Core/ScopeProfiler.cs
--- a/Engine/Source/Programs/Shared/EpicGames.Core/ScopeProfiler.cs
+++ b/Engine/Source/Programs/Shared/EpicGames.Core/ScopeProfiler.cs
@@ -1,5 +1,7 @@
 // Copyright Epic Games, Inc. All Rights Reserved.
 
+#pragma warning disable CA1822 // Mark members as static
+
 namespace EpicGames.Core
 {
 #if !__SCOPEPROFILER_AVAILABLE__
@@ -14,8 +16,8 @@ namespace EpicGames.Core
 		/// <summary>
 		/// Static instance of the profiler
 		/// </summary>
-		public static ScopeProfiler Instance = new ScopeProfiler();
-		
+		public static ScopeProfiler Instance { get; } = new ScopeProfiler();
+
 		/// <summary>
 		/// 
 		/// </summary>
diff --git a/Engine/Source/Programs/Shared/EpicGames.Core/Sha1Hash.cs b/Engine/Source/Programs/Shared/EpicGames.Core/Sha1Hash.cs
--- a/Engine/Source/Programs/Shared/EpicGames.Core/Sha1Hash.cs
+++ b/Engine/Source/Programs/Shared/EpicGames.Core/Sha1Hash.cs
@@ -8,6 +8,8 @@ using System.Security.Cryptography;
 using System.Text.Json;
 using System.Text.Json.Serialization;
 
+#pragma warning disable CA5350 // Do Not Use Weak Cryptographic Algorithms
+
 namespace EpicGames.Core
 {
 	/// <summary>
diff --git a/Engine/Source/Programs/Shared/EpicGames.Core/StringComparerUE.cs b/Engine/Source/Programs/Shared/EpicGames.Core/StringComparerUE.cs
--- a/Engine/Source/Programs/Shared/EpicGames.Core/StringComparerUE.cs
+++ b/Engine/Source/Programs/Shared/EpicGames.Core/StringComparerUE.cs
@@ -13,12 +13,12 @@ namespace EpicGames.Core
 		/// <summary>
 		/// Static instance of an ordinal String comparer
 		/// </summary>
-		public static StringComparerUE Ordinal = new StringComparerUE(StringComparison.Ordinal);
+		public static StringComparerUE Ordinal { get; } = new StringComparerUE(StringComparison.Ordinal);
 
 		/// <summary>
 		/// Static instance of an ordinal String comparer which ignores case
 		/// </summary>
-		public static StringComparerUE OrdinalIgnoreCase = new StringComparerUE(StringComparison.OrdinalIgnoreCase);
+		public static StringComparerUE OrdinalIgnoreCase { get; } = new StringComparerUE(StringComparison.OrdinalIgnoreCase);
 
 		/// <summary>
 		/// The comparison type
diff --git a/Engine/Source/Programs/Shared/EpicGames.Core/StringUtils.cs b/Engine/Source/Programs/Shared/EpicGames.Core/StringUtils.cs
--- a/Engine/Source/Programs/Shared/EpicGames.Core/StringUtils.cs
+++ b/Engine/Source/Programs/Shared/EpicGames.Core/StringUtils.cs
@@ -16,7 +16,7 @@ namespace EpicGames.Core
 		/// <summary>
 		/// Array mapping from ascii index to hexadecimal digits.
 		/// </summary>
-		static readonly sbyte[] s_hexDigits;
+		static readonly sbyte[] s_hexDigits = CreateHexDigitsTable();
 
 		/// <summary>
 		/// Hex digits to utf8 byte
@@ -31,25 +31,26 @@ namespace EpicGames.Core
 		/// <summary>
 		/// Static constructor. Initializes the HexDigits array.
 		/// </summary>
-		static StringUtils()
+		static sbyte[] CreateHexDigitsTable()
 		{
-			s_hexDigits = new sbyte[256];
+			sbyte[] hexDigits = new sbyte[256];
 			for (int idx = 0; idx < 256; idx++)
 			{
-				s_hexDigits[idx] = -1;
+				hexDigits[idx] = -1;
 			}
 			for (int idx = '0'; idx <= '9'; idx++)
 			{
-				s_hexDigits[idx] = (sbyte)(idx - '0');
+				hexDigits[idx] = (sbyte)(idx - '0');
 			}
 			for (int idx = 'a'; idx <= 'f'; idx++)
 			{
-				s_hexDigits[idx] = (sbyte)(10 + idx - 'a');
+				hexDigits[idx] = (sbyte)(10 + idx - 'a');
 			}
 			for (int idx = 'A'; idx <= 'F'; idx++)
 			{
-				s_hexDigits[idx] = (sbyte)(10 + idx - 'A');
+				hexDigits[idx] = (sbyte)(10 + idx - 'A');
 			}
+			return hexDigits;
 		}
 
 		/// <summary>
@@ -63,7 +64,7 @@ namespace EpicGames.Core
 			string result = "";
 			if(text.Length > 0)
 			{
-				result = indent + text.Replace("\n", "\n" + indent);
+				result = indent + text.Replace("\n", "\n" + indent, StringComparison.Ordinal);
 			}
 			return result;
 		}
@@ -92,7 +93,7 @@ namespace EpicGames.Core
 		public static string ExpandProperties(string text, Func<string, string?> getPropertyValue)
 		{
 			string result = text;
-			for (int idx = result.IndexOf("$("); idx != -1; idx = result.IndexOf("$(", idx))
+			for (int idx = result.IndexOf("$(", StringComparison.Ordinal); idx != -1; idx = result.IndexOf("$(", idx, StringComparison.Ordinal))
 			{
 				// Find the end of the variable name
 				int endIdx = result.IndexOf(')', idx + 2);
@@ -503,7 +504,7 @@ namespace EpicGames.Core
 		/// <returns>The quoted argument if it contains any spaces, otherwise the original string</returns>
 		public static string QuoteArgument(this string str)
 		{
-			if (str.Contains(' '))
+			if (str.Contains(' ', StringComparison.Ordinal))
 			{
 				return $"\"{str}\"";
 			}
@@ -558,7 +559,7 @@ namespace EpicGames.Core
 			bytesString = bytesString.Trim();
 
 			int power = s_byteSizes.FindIndex( s => (s != s_byteSizes[0]) && bytesString.EndsWith(s, StringComparison.InvariantCultureIgnoreCase ) ); // need to handle 'B' suffix separately
-			if (power == -1 && bytesString.EndsWith(s_byteSizes[0]))
+			if (power == -1 && bytesString.EndsWith(s_byteSizes[0], StringComparison.Ordinal))
 			{
 				power = 0;
 			}
@@ -605,7 +606,7 @@ namespace EpicGames.Core
 		{
 			char escapeChar = '\u001b';
 
-			int index = line.IndexOf(escapeChar);
+			int index = line.IndexOf(escapeChar, StringComparison.Ordinal);
 			if (index != -1)
 			{
 				int lastIndex = 0;
diff --git a/Engine/Source/Programs/Shared/EpicGames.Core/StringView.cs b/Engine/Source/Programs/Shared/EpicGames.Core/StringView.cs
--- a/Engine/Source/Programs/Shared/EpicGames.Core/StringView.cs
+++ b/Engine/Source/Programs/Shared/EpicGames.Core/StringView.cs
@@ -146,7 +146,7 @@ namespace EpicGames.Core
 		/// <inheritdoc/>
 		public override int GetHashCode()
 		{
-			return String.GetHashCode(Memory.Span);
+			return String.GetHashCode(Memory.Span, StringComparison.Ordinal);
 		}
 
 		/// <inheritdoc/>
@@ -188,12 +188,12 @@ namespace EpicGames.Core
 		/// <summary>
 		/// Static instance of an ordinal StringView comparer
 		/// </summary>
-		public static StringViewComparer Ordinal = new StringViewComparer(StringComparison.Ordinal);
+		public static StringViewComparer Ordinal { get; } = new StringViewComparer(StringComparison.Ordinal);
 
 		/// <summary>
 		/// Static instance of an ordinal StringView comparer which ignores case
 		/// </summary>
-		public static StringViewComparer OrdinalIgnoreCase = new StringViewComparer(StringComparison.OrdinalIgnoreCase);
+		public static StringViewComparer OrdinalIgnoreCase { get; } = new StringViewComparer(StringComparison.OrdinalIgnoreCase);
 
 		/// <summary>
 		/// The comparison type
@@ -236,12 +236,12 @@ namespace EpicGames.Core
 		/// <summary>
 		/// Static instance of an ordinal StringView comparer
 		/// </summary>
-		public static StringViewComparerUe Ordinal = new StringViewComparerUe(StringComparison.Ordinal);
+		public static StringViewComparerUe Ordinal { get; } = new StringViewComparerUe(StringComparison.Ordinal);
 
 		/// <summary>
 		/// Static instance of an ordinal StringView comparer which ignores case
 		/// </summary>
-		public static StringViewComparerUe OrdinalIgnoreCase = new StringViewComparerUe(StringComparison.OrdinalIgnoreCase);
+		public static StringViewComparerUe OrdinalIgnoreCase { get; } = new StringViewComparerUe(StringComparison.OrdinalIgnoreCase);
 
 		/// <summary>
 		/// The comparison type
diff --git a/Engine/Source/Programs/Shared/EpicGames.Core/ThreadPoolWorkQueue.cs b/Engine/Source/Programs/Shared/EpicGames.Core/ThreadPoolWorkQueue.cs
--- a/Engine/Source/Programs/Shared/EpicGames.Core/ThreadPoolWorkQueue.cs
+++ b/Engine/Source/Programs/Shared/EpicGames.Core/ThreadPoolWorkQueue.cs
@@ -9,7 +9,7 @@ namespace EpicGames.Core
 	/// <summary>
 	/// Allows queuing a large number of tasks to a thread pool and waiting for them all to complete.
 	/// </summary>
-	public class ThreadPoolWorkQueue : IDisposable
+	public sealed class ThreadPoolWorkQueue : IDisposable
 	{
 		/// <summary>
 		/// Object used for controlling access to NumOutstandingJobs and updating EmptyEvent
diff --git a/Engine/Source/Programs/Shared/EpicGames.Core/TraceSpan.cs b/Engine/Source/Programs/Shared/EpicGames.Core/TraceSpan.cs
--- a/Engine/Source/Programs/Shared/EpicGames.Core/TraceSpan.cs
+++ b/Engine/Source/Programs/Shared/EpicGames.Core/TraceSpan.cs
@@ -117,7 +117,7 @@ namespace EpicGames.Core
 			{
 				DirectoryReference.CreateDirectory(_telemetryDir);
 
-				string fileName = String.Format("{0}.{1}.json", Path.GetFileName(Assembly.GetEntryAssembly()!.Location), process.Id, process.StartTime.Ticks);
+				string fileName = String.Format("{0}.{1}.{2}.json", Path.GetFileName(Assembly.GetEntryAssembly()!.Location), process.Id, process.StartTime.Ticks);
 				file = FileReference.Combine(_telemetryDir, fileName);
 			}
 
diff --git a/Engine/Source/Programs/Shared/EpicGames.Core/UHTTypes.cs b/Engine/Source/Programs/Shared/EpicGames.Core/UHTTypes.cs
--- a/Engine/Source/Programs/Shared/EpicGames.Core/UHTTypes.cs
+++ b/Engine/Source/Programs/Shared/EpicGames.Core/UHTTypes.cs
@@ -4,6 +4,11 @@ using System.Collections.Generic;
 using System.Text.Json.Serialization;
 
 #pragma warning disable CS1591 // Missing documentation
+#pragma warning disable CA2227 // Collection properties should be read only
+#pragma warning disable CA1027 // Mark enums with FlagsAttribute
+#pragma warning disable CA1034 // Nested types should not be visible
+#pragma warning disable CA1716 // Identifiers should not match keywords
+#pragma warning disable IDE0049 // Use framework type
 
 namespace EpicGames.Core
 {
@@ -105,7 +110,7 @@ namespace EpicGames.Core
 				int Length = Name.Length;
 				foreach (string Define in PublicDefines)
 				{
-					if (!Define.StartsWith(Name))
+					if (!Define.StartsWith(Name, System.StringComparison.Ordinal))
 					{
 						continue;
 					}
diff --git a/Engine/Source/Programs/Shared/EpicGames.Core/UnrealEngineTypes.cs b/Engine/Source/Programs/Shared/EpicGames.Core/UnrealEngineTypes.cs
--- a/Engine/Source/Programs/Shared/EpicGames.Core/UnrealEngineTypes.cs
+++ b/Engine/Source/Programs/Shared/EpicGames.Core/UnrealEngineTypes.cs
@@ -6,6 +6,10 @@ using System.Reflection;
 
 #pragma warning disable IDE1006 // Naming Styles
 #pragma warning disable CS1591 // Missing documentation
+#pragma warning disable CA1028 // Enum Storage should be Int32
+#pragma warning disable CA2217 // Do not mark enums with FlagsAttribute
+#pragma warning disable CA1069 // Enums values should not be duplicated
+#pragma warning disable CA1008 // Enums should have zero value
 
 namespace EpicGames.Core
 {
diff --git a/Engine/Source/Programs/Shared/EpicGames.Core/Utf8String.cs b/Engine/Source/Programs/Shared/EpicGames.Core/Utf8String.cs
--- a/Engine/Source/Programs/Shared/EpicGames.Core/Utf8String.cs
+++ b/Engine/Source/Programs/Shared/EpicGames.Core/Utf8String.cs
@@ -111,6 +111,18 @@ namespace EpicGames.Core
 		/// <inheritdoc/>
 		public int CompareTo(Utf8String other) => Utf8StringComparer.Ordinal.Compare(Span, other.Span);
 
+		/// <inheritdoc/>
+		public static bool operator <(Utf8String left, Utf8String right) => left.CompareTo(right) < 0;
+
+		/// <inheritdoc/>
+		public static bool operator <=(Utf8String left, Utf8String right) => left.CompareTo(right) <= 0;
+
+		/// <inheritdoc/>
+		public static bool operator >(Utf8String left, Utf8String right) => left.CompareTo(right) > 0;
+
+		/// <inheritdoc/>
+		public static bool operator >=(Utf8String left, Utf8String right) =>  left.CompareTo(right) >= 0;
+
 		/// <inheritdoc cref="String.Contains(String)"/>
 		public bool Contains(Utf8String str) => IndexOf(str) != -1;
 
@@ -351,7 +363,7 @@ namespace EpicGames.Core
 		/// <summary>
 		/// Ordinal comparer for utf8 strings
 		/// </summary>
-		public sealed class OrdinalComparer : Utf8StringComparer
+		sealed class OrdinalComparer : Utf8StringComparer
 		{
 			/// <inheritdoc/>
 			public override bool Equals(ReadOnlySpan<byte> strA, ReadOnlySpan<byte> strB)
@@ -380,7 +392,7 @@ namespace EpicGames.Core
 		/// <summary>
 		/// Comparison between ReadOnlyUtf8String objects that ignores case for ASCII characters
 		/// </summary>
-		public sealed class OrdinalIgnoreCaseComparer : Utf8StringComparer 
+		sealed class OrdinalIgnoreCaseComparer : Utf8StringComparer 
 		{
 			/// <inheritdoc/>
 			public override bool Equals(ReadOnlySpan<byte> strA, ReadOnlySpan<byte> strB)
diff --git a/Engine/Source/Programs/Shared/EpicGames.Core/ZipArchiveExtensions.cs b/Engine/Source/Programs/Shared/EpicGames.Core/ZipArchiveExtensions.cs
--- a/Engine/Source/Programs/Shared/EpicGames.Core/ZipArchiveExtensions.cs
+++ b/Engine/Source/Programs/Shared/EpicGames.Core/ZipArchiveExtensions.cs
@@ -6,6 +6,8 @@ using System.IO.Compression;
 using System.Reflection;
 using System.Runtime.InteropServices;
 
+#pragma warning disable CA1707 // Identifiers should not contain underscores
+
 namespace EpicGames.Core
 {
 	/// <summary>
