EpicGames.Serialization: Fix static analysis warnings.

diff --git a/Engine/Source/Programs/Shared/EpicGames.Serialization.Tests/CbPackageTests.cs b/Engine/Source/Programs/Shared/EpicGames.Serialization.Tests/CbPackageTests.cs
--- a/Engine/Source/Programs/Shared/EpicGames.Serialization.Tests/CbPackageTests.cs
+++ b/Engine/Source/Programs/Shared/EpicGames.Serialization.Tests/CbPackageTests.cs
@@ -39,7 +39,7 @@ namespace EpicGames.Serialization.Tests
 			CbObject rootObject = CbSerializer.Serialize(root);
 			IoHash rootHash = IoHash.Compute(rootObject.GetView().Span);
 			
-			CbPackageBuilder builder = new();
+			using CbPackageBuilder builder = new();
 			await builder.AddAttachment(rootHash, CbPackageAttachmentFlags.IsObject, rootObject.GetView().ToArray());
 			byte[] bytes = builder.ToByteArray();
 
@@ -73,7 +73,7 @@ namespace EpicGames.Serialization.Tests
 			CbObject rootObject = CbSerializer.Serialize(root);
 			IoHash rootHash = IoHash.Compute(rootObject.GetView().Span);
 
-			CbPackageBuilder builder = new();
+			using CbPackageBuilder builder = new();
 			await builder.AddAttachment(rootHash, CbPackageAttachmentFlags.IsObject, rootObject.GetView().ToArray());
 			await builder.AddAttachment(simpleHash, CbPackageAttachmentFlags.IsObject, simpleObject.GetView().ToArray());
 			await builder.AddAttachment(blobHash, 0, blob);
@@ -114,7 +114,7 @@ namespace EpicGames.Serialization.Tests
 		[TestMethod]
 		public async Task ReadNoAttachments()
 		{
-			CbPackageBuilder packageBuilder = new CbPackageBuilder();
+			using CbPackageBuilder packageBuilder = new CbPackageBuilder();
 			byte[] buf = packageBuilder.ToByteArray();
 			await using MemoryStream ms = new MemoryStream(buf);
 
diff --git a/Engine/Source/Programs/Shared/EpicGames.Serialization/CbConverter.cs b/Engine/Source/Programs/Shared/EpicGames.Serialization/CbConverter.cs
--- a/Engine/Source/Programs/Shared/EpicGames.Serialization/CbConverter.cs
+++ b/Engine/Source/Programs/Shared/EpicGames.Serialization/CbConverter.cs
@@ -13,7 +13,7 @@ namespace EpicGames.Serialization
 	/// Attribute declaring the converter to use for a type
 	/// </summary>
 	[AttributeUsage(AttributeTargets.Class | AttributeTargets.Struct | AttributeTargets.Property)]
-	public class CbConverterAttribute : Attribute
+	public sealed class CbConverterAttribute : Attribute
 	{
 		/// <summary>
 		/// The converter type
@@ -235,7 +235,7 @@ namespace EpicGames.Serialization
 		/// <summary>
 		/// Cache of type to converter
 		/// </summary>
-		public static Dictionary<Type, ICbConverter> TypeToConverter = new Dictionary<Type, ICbConverter>()
+		public static Dictionary<Type, ICbConverter> TypeToConverter { get; } = new Dictionary<Type, ICbConverter>()
 		{
 			[typeof(bool)] = new CbPrimitiveConverter<bool>(x => x.AsBool(), (w, v) => w.WriteBoolValue(v), (w, n, v) => w.WriteBool(n, v)),
 			[typeof(int)] = new CbPrimitiveConverter<int>(x => x.AsInt32(), (w, v) => w.WriteIntegerValue(v), (w, n, v) => w.WriteInteger(n, v)),
@@ -258,7 +258,7 @@ namespace EpicGames.Serialization
 		/// <summary>
 		/// List of converter factories. Must be 
 		/// </summary>
-		public static List<CbConverterFactory> ConverterFactories = new List<CbConverterFactory>
+		public static List<CbConverterFactory> ConverterFactories { get; } = new List<CbConverterFactory>
 		{
 			new CbClassConverterFactory(),
 			new CbStringConverterFactory(),
diff --git a/Engine/Source/Programs/Shared/EpicGames.Serialization/CbField.cs b/Engine/Source/Programs/Shared/EpicGames.Serialization/CbField.cs
--- a/Engine/Source/Programs/Shared/EpicGames.Serialization/CbField.cs
+++ b/Engine/Source/Programs/Shared/EpicGames.Serialization/CbField.cs
@@ -14,6 +14,9 @@ using System.Text;
 using System.Text.Json;
 using System.Text.Json.Serialization;
 
+#pragma warning disable CA1721 // Property names should not match get methods
+#pragma warning disable CA1028 // Enum Storage should be Int32
+
 namespace EpicGames.Serialization
 {
 	/// <summary>
@@ -171,11 +174,6 @@ namespace EpicGames.Serialization
 		/// </summary>
 		CustomByName = 0x1f,
 
-		/// <summary>
-		/// Reserved for future use as a flag. Do not add types in this range. 
-		/// </summary>
-		Reserved = 0x20,
-
 		/// <summary>
 		/// A transient flag which indicates that the object or array containing this field has stored
 		/// the field type before the payload and name. Non-uniform objects and fields will set this.
@@ -196,7 +194,7 @@ namespace EpicGames.Serialization
 	[DebuggerDisplay("{_hash}")]
 	[JsonConverter(typeof(CbBinaryAttachmentJsonConverter))]
 	[TypeConverter(typeof(CbBinaryAttachmentTypeConverter))]
-	public struct CbBinaryAttachment
+	public struct CbBinaryAttachment : IEquatable<CbBinaryAttachment>
 	{
 		/// <summary>
 		/// Attachment with a hash of zero
@@ -221,7 +219,10 @@ namespace EpicGames.Serialization
 		public override string ToString() => Hash.ToString();
 
 		/// <inheritdoc/>
-		public override bool Equals(object? obj) => obj is CbBinaryAttachment other && other.Hash == Hash;
+		public override bool Equals(object? obj) => obj is CbBinaryAttachment other && Equals(other);
+
+		/// <inheritdoc/>
+		public bool Equals(CbBinaryAttachment other) => other.Hash == Hash;
 
 		/// <inheritdoc/>
 		public override int GetHashCode() => Hash.GetHashCode();
@@ -275,7 +276,7 @@ namespace EpicGames.Serialization
 	[DebuggerDisplay("{_hash}")]
 	[JsonConverter(typeof(CbObjectAttachmentJsonConverter))]
 	[TypeConverter(typeof(CbObjectAttachmentTypeConverter))]
-	public struct CbObjectAttachment
+	public struct CbObjectAttachment : IEquatable<CbObjectAttachment>
 	{
 		/// <summary>
 		/// Attachment with a hash of zero
@@ -300,7 +301,10 @@ namespace EpicGames.Serialization
 		public override string ToString() => Hash.ToString();
 
 		/// <inheritdoc/>
-		public override bool Equals(object? obj) => obj is CbObjectAttachment other && other.Hash == Hash;
+		public override bool Equals(object? obj) => obj is CbObjectAttachment other && Equals(other);
+
+		/// <inheritdoc/>
+		public bool Equals(CbObjectAttachment other) => other.Hash == Hash;
 
 		/// <inheritdoc/>
 		public override int GetHashCode() => Hash.GetHashCode();
@@ -716,12 +720,12 @@ namespace EpicGames.Serialization
 		/// <summary>
 		/// Offset of the name with the memory
 		/// </summary>
-		public int _nameLen;
+		internal int _nameLen;
 
 		/// <summary>
 		/// Offset of the payload within the memory
 		/// </summary>
-		public int _payloadOffset;
+		internal int _payloadOffset;
 
 		/// <summary>
 		/// Error for parsing the current field type
@@ -1524,6 +1528,12 @@ namespace EpicGames.Serialization
 			hasher.Update(GetViewNoType().Span);
 		}
 
+		/// <inheritdoc/>
+		public override int GetHashCode() => throw new NotImplementedException();
+
+		/// <inheritdoc/>
+		public override bool Equals(object? other) => Equals(other as CbField);
+
 		/// <summary>
 		/// Whether this field is identical to the other field.
 		/// 
@@ -1600,7 +1610,9 @@ namespace EpicGames.Serialization
 		/// </summary>
 		/// <param name="name"></param>
 		/// <returns></returns>
+#pragma warning disable CA1043 // Use Integral Or String Argument For Indexers
 		public CbField this[Utf8String name] => this.FirstOrDefault(field => field.Name == name) ?? CbField.Empty;
+#pragma warning restore CA1043 // Use Integral Or String Argument For Indexers
 
 		/// <summary>
 		/// Create an iterator for the fields of an array or object, otherwise an empty iterator.
@@ -1899,7 +1911,7 @@ namespace EpicGames.Serialization
 	/// <summary>
 	/// Enumerator for contents of a field
 	/// </summary>
-	public class CbFieldEnumerator : IEnumerator<CbField>
+	public sealed class CbFieldEnumerator : IEnumerator<CbField>
 	{
 		/// <summary>
 		/// The underlying buffer
@@ -2426,7 +2438,7 @@ namespace EpicGames.Serialization
 		/// <summary>
 		/// Empty array constant
 		/// </summary>
-		public static CbObject Empty = CbObject.FromFieldNoCheck(new CbField(new byte[] { (byte)CbFieldType.Object, 0 }));
+		public static CbObject Empty { get; } = CbObject.FromFieldNoCheck(new CbField(new byte[] { (byte)CbFieldType.Object, 0 }));
 
 		/// <summary>
 		/// The inner field object
@@ -2488,7 +2500,9 @@ namespace EpicGames.Serialization
 		/// </summary>
 		/// <param name="name">The name of the field.</param>
 		/// <returns>The matching field if found, otherwise a field with no value.</returns>
+#pragma warning disable CA1043 // Use Integral Or String Argument For Indexers
 		public CbField this[Utf8String name] => _innerField[name];
+#pragma warning restore CA1043 // Use Integral Or String Argument For Indexers
 
 		/// <summary>
 		/// Gets the underlying field for this object
diff --git a/Engine/Source/Programs/Shared/EpicGames.Serialization/CbPackage.cs b/Engine/Source/Programs/Shared/EpicGames.Serialization/CbPackage.cs
--- a/Engine/Source/Programs/Shared/EpicGames.Serialization/CbPackage.cs
+++ b/Engine/Source/Programs/Shared/EpicGames.Serialization/CbPackage.cs
@@ -7,6 +7,8 @@ using System.Text;
 using System.Threading.Tasks;
 using EpicGames.Core;
 
+#pragma warning disable CA1028 // Enum Storage should be Int32
+
 namespace EpicGames.Serialization
 {
 	static class CbPackageConstants
@@ -84,10 +86,12 @@ namespace EpicGames.Serialization
 		/// Is marshaled using compressed buffer storage format
 		/// </summary>
 		IsCompressed = (1u << 0),
+
 		/// <summary>
 		/// Is compact binary object
 		/// </summary>
 		IsObject = (1u << 1),
+
 		/// <summary>
 		/// Is error (compact binary formatted) object
 		/// </summary>
@@ -267,9 +271,11 @@ namespace EpicGames.Serialization
 	/// <summary>
 	/// Builds a in-memory representation of a CbPackage
 	/// </summary>
-	public class CbPackageBuilder
+	public sealed class CbPackageBuilder : IDisposable
 	{
+#pragma warning disable CA2213
 		private readonly MemoryStream _packageBuffer;
+#pragma warning restore CA2213
 		private uint _countOfWrittenAttachments;
 
 		/// <summary>
@@ -290,6 +296,12 @@ namespace EpicGames.Serialization
 			header.Write(_packageBuffer);
 		}
 
+		/// <inheritdoc/>
+		public void Dispose()
+		{
+			_packageBuffer.Dispose();
+		}
+
 		/// <summary>
 		/// Add a attachment to the package builder using in memory buffer
 		/// </summary>
diff --git a/Engine/Source/Programs/Shared/EpicGames.Serialization/CbSerializer.cs b/Engine/Source/Programs/Shared/EpicGames.Serialization/CbSerializer.cs
--- a/Engine/Source/Programs/Shared/EpicGames.Serialization/CbSerializer.cs
+++ b/Engine/Source/Programs/Shared/EpicGames.Serialization/CbSerializer.cs
@@ -10,7 +10,7 @@ namespace EpicGames.Serialization
 	/// when base class objects are empty.
 	/// </summary>
 	[AttributeUsage(AttributeTargets.Class)]
-	public class CbObjectAttribute : Attribute
+	public sealed class CbObjectAttribute : Attribute
 	{
 	}
 
@@ -18,12 +18,12 @@ namespace EpicGames.Serialization
 	/// Attribute used to mark a property that should be serialized to compact binary
 	/// </summary>
 	[AttributeUsage(AttributeTargets.Property)]
-	public class CbFieldAttribute : Attribute
+	public sealed class CbFieldAttribute : Attribute
 	{
 		/// <summary>
 		/// Name of the serialized field
 		/// </summary>
-		public string? Name { get; set; }
+		public string? Name { get; }
 
 		/// <summary>
 		/// Default constructor
@@ -45,14 +45,16 @@ namespace EpicGames.Serialization
 	/// <summary>
 	/// Attribute used to indicate that this object is the base for a class hierarchy. Each derived class must have a [CbDiscriminator] attribute.
 	/// </summary>
-	public class CbPolymorphicAttribute : Attribute
+	[AttributeUsage(AttributeTargets.Class)]
+	public sealed class CbPolymorphicAttribute : Attribute
 	{
 	}
 
 	/// <summary>
 	/// Sets the name used for discriminating between derived classes during serialization
 	/// </summary>
-	public class CbDiscriminatorAttribute : Attribute
+	[AttributeUsage(AttributeTargets.Class)]
+	public sealed class CbDiscriminatorAttribute : Attribute
 	{
 		/// <summary>
 		/// Name used to identify this class
@@ -90,15 +92,15 @@ namespace EpicGames.Serialization
 		/// <summary>
 		/// Type with missing field annotations
 		/// </summary>
-		public Type Type { get; }
+		public Type ClassType { get; }
 
 		/// <summary>
 		/// Constructor
 		/// </summary>
-		public CbEmptyClassException(Type type)
-			: base($"{type.Name} does not have any fields marked with a [CbField] attribute. If this is intended, explicitly mark the class with a [CbObject] attribute.")
+		public CbEmptyClassException(Type classType)
+			: base($"{classType.Name} does not have any fields marked with a [CbField] attribute. If this is intended, explicitly mark the class with a [CbObject] attribute.")
 		{
-			Type = type;
+			ClassType = classType;
 		}
 	}
 
diff --git a/Engine/Source/Programs/Shared/EpicGames.Serialization/Converters/CbClassConverter.cs b/Engine/Source/Programs/Shared/EpicGames.Serialization/Converters/CbClassConverter.cs
--- a/Engine/Source/Programs/Shared/EpicGames.Serialization/Converters/CbClassConverter.cs
+++ b/Engine/Source/Programs/Shared/EpicGames.Serialization/Converters/CbClassConverter.cs
@@ -38,7 +38,6 @@ namespace EpicGames.Serialization.Converters
 		public DynamicMethod ReadConcreteMethod { get; }
 		public DynamicMethod WriteConcreteContentsMethod { get; }
 
-		static readonly Dictionary<PropertyInfo, CbClassConverterMethods> s_propertyToMethods = new Dictionary<PropertyInfo, CbClassConverterMethods>();
 		static readonly Dictionary<Type, CbClassConverterMethods> s_typeToMethods = new Dictionary<Type, CbClassConverterMethods>();
 
 		public static CbClassConverterMethods Create(Type classType)
