BuildGraph: Use table-driven opcode parsing, to decouple instruction decode from execution.

diff --git a/Engine/Source/Programs/Shared/EpicGames.BuildGraph/BgInterpreter.cs b/Engine/Source/Programs/Shared/EpicGames.BuildGraph/BgInterpreter.cs
--- a/Engine/Source/Programs/Shared/EpicGames.BuildGraph/BgInterpreter.cs
+++ b/Engine/Source/Programs/Shared/EpicGames.BuildGraph/BgInterpreter.cs
@@ -60,6 +60,41 @@ namespace EpicGames.BuildGraph
 	/// </summary>
 	public class BgInterpreter
 	{
+		enum BgArg
+		{
+			None,
+			Value,
+			ArgList,
+			Fragment,
+			Name,
+			Thunk,
+			StringLiteral,
+			VarIntSigned,
+			VarIntUnsigned,
+		}
+
+		readonly struct BgOpcodeInfo
+		{
+			readonly uint _data;
+
+			public BgOpcodeInfo(BgArg arg1) : this(1, arg1, BgArg.None, BgArg.None) { }
+			public BgOpcodeInfo(BgArg arg1, BgArg arg2) : this(2, arg1, arg2, BgArg.None) { }
+			public BgOpcodeInfo(BgArg arg1, BgArg arg2, BgArg arg3) : this(3, arg1, arg2, arg3) { }
+
+			private BgOpcodeInfo(int count, BgArg arg1, BgArg arg2, BgArg arg3)
+			{
+				_data = (uint)arg1 | ((uint)arg2 << 4) | ((uint)arg3 << 8) | ((uint)count << 12);
+			}
+
+			public int ArgCount => (int)((_data >> 12) & 0x7);
+
+			public BgArg Arg0 => GetArg(0);
+			public BgArg Arg1 => GetArg(1);
+			public BgArg Arg2 => GetArg(2);
+
+			public BgArg GetArg(int idx) => (BgArg)((_data >> (idx * 4)) & 0xf);
+		}
+
 		class Frame
 		{
 			public int Offset { get; set; }
@@ -94,6 +129,78 @@ namespace EpicGames.BuildGraph
 		/// </summary>
 		public IReadOnlyList<BgOptionDef> OptionDefs => _optionDefs;
 
+		static readonly BgOpcodeInfo[] _opcodeTable = CreateOpcodeLookup();
+
+		static BgOpcodeInfo[] CreateOpcodeLookup()
+		{
+			BgOpcodeInfo[] opcodes = new BgOpcodeInfo[256];
+
+			opcodes[(int)BgOpcode.BoolFalse] = new BgOpcodeInfo();
+			opcodes[(int)BgOpcode.BoolTrue] = new BgOpcodeInfo();
+			opcodes[(int)BgOpcode.BoolNot] = new BgOpcodeInfo(BgArg.Value);
+			opcodes[(int)BgOpcode.BoolAnd] = new BgOpcodeInfo(BgArg.Value, BgArg.Value);
+			opcodes[(int)BgOpcode.BoolOr] = new BgOpcodeInfo(BgArg.Value, BgArg.Value);
+			opcodes[(int)BgOpcode.BoolXor] = new BgOpcodeInfo(BgArg.Value, BgArg.Value);
+			opcodes[(int)BgOpcode.BoolEq] = new BgOpcodeInfo(BgArg.Value, BgArg.Value);
+			opcodes[(int)BgOpcode.BoolOption] = new BgOpcodeInfo(BgArg.Value);
+
+			opcodes[(int)BgOpcode.IntLiteral] = new BgOpcodeInfo(BgArg.VarIntSigned);
+			opcodes[(int)BgOpcode.IntEq] = new BgOpcodeInfo(BgArg.Value, BgArg.Value);
+			opcodes[(int)BgOpcode.IntLt] = new BgOpcodeInfo(BgArg.Value, BgArg.Value);
+			opcodes[(int)BgOpcode.IntGt] = new BgOpcodeInfo(BgArg.Value, BgArg.Value);
+			opcodes[(int)BgOpcode.IntAdd] = new BgOpcodeInfo(BgArg.Value, BgArg.Value);
+			opcodes[(int)BgOpcode.IntMultiply] = new BgOpcodeInfo(BgArg.Value, BgArg.Value);
+			opcodes[(int)BgOpcode.IntDivide] = new BgOpcodeInfo(BgArg.Value, BgArg.Value);
+			opcodes[(int)BgOpcode.IntModulo] = new BgOpcodeInfo(BgArg.Value, BgArg.Value);
+			opcodes[(int)BgOpcode.IntNegate] = new BgOpcodeInfo(BgArg.Value);
+			opcodes[(int)BgOpcode.IntOption] = new BgOpcodeInfo(BgArg.Value);
+
+			opcodes[(int)BgOpcode.StrEmpty] = new BgOpcodeInfo();
+			opcodes[(int)BgOpcode.StrLiteral] = new BgOpcodeInfo(BgArg.StringLiteral);
+			opcodes[(int)BgOpcode.StrCompare] = new BgOpcodeInfo(BgArg.Value, BgArg.Value, BgArg.VarIntUnsigned);
+			opcodes[(int)BgOpcode.StrConcat] = new BgOpcodeInfo(BgArg.Value, BgArg.Value);
+			opcodes[(int)BgOpcode.StrFormat] = new BgOpcodeInfo(BgArg.Value, BgArg.ArgList);
+			opcodes[(int)BgOpcode.StrSplit] = new BgOpcodeInfo(BgArg.Value, BgArg.Value);
+			opcodes[(int)BgOpcode.StrJoin] = new BgOpcodeInfo(BgArg.Value, BgArg.Value);
+			opcodes[(int)BgOpcode.StrMatch] = new BgOpcodeInfo(BgArg.Value, BgArg.Value);
+			opcodes[(int)BgOpcode.StrReplace] = new BgOpcodeInfo(BgArg.Value, BgArg.Value, BgArg.Value);
+			opcodes[(int)BgOpcode.StrOption] = new BgOpcodeInfo(BgArg.Value);
+
+			opcodes[(int)BgOpcode.EnumConstant] = new BgOpcodeInfo(BgArg.VarIntSigned);
+			opcodes[(int)BgOpcode.EnumParse] = new BgOpcodeInfo(BgArg.Value, BgArg.Value, BgArg.Value);
+			opcodes[(int)BgOpcode.EnumToString] = new BgOpcodeInfo(BgArg.Value, BgArg.Value, BgArg.Value);
+
+			opcodes[(int)BgOpcode.ListEmpty] = new BgOpcodeInfo();
+			opcodes[(int)BgOpcode.ListPush] = new BgOpcodeInfo(BgArg.Value, BgArg.Value);
+			opcodes[(int)BgOpcode.ListPushLazy] = new BgOpcodeInfo(BgArg.Value, BgArg.Fragment);
+			opcodes[(int)BgOpcode.ListCount] = new BgOpcodeInfo(BgArg.Value);
+			opcodes[(int)BgOpcode.ListElement] = new BgOpcodeInfo(BgArg.Value, BgArg.Value);
+			opcodes[(int)BgOpcode.ListConcat] = new BgOpcodeInfo(BgArg.Value, BgArg.Value);
+			opcodes[(int)BgOpcode.ListUnion] = new BgOpcodeInfo(BgArg.Value, BgArg.Value);
+			opcodes[(int)BgOpcode.ListExcept] = new BgOpcodeInfo(BgArg.Value, BgArg.Value);
+			opcodes[(int)BgOpcode.ListSelect] = new BgOpcodeInfo(BgArg.Value, BgArg.Fragment);
+			opcodes[(int)BgOpcode.ListWhere] = new BgOpcodeInfo(BgArg.Value, BgArg.Fragment);
+			opcodes[(int)BgOpcode.ListDistinct] = new BgOpcodeInfo(BgArg.Value);
+			opcodes[(int)BgOpcode.ListContains] = new BgOpcodeInfo(BgArg.Value, BgArg.Value);
+			opcodes[(int)BgOpcode.ListLazy] = new BgOpcodeInfo(BgArg.Fragment);
+			opcodes[(int)BgOpcode.ListOption] = new BgOpcodeInfo(BgArg.Value);
+
+			opcodes[(int)BgOpcode.ObjEmpty] = new BgOpcodeInfo();
+			opcodes[(int)BgOpcode.ObjGet] = new BgOpcodeInfo(BgArg.Value, BgArg.Name, BgArg.Value);
+			opcodes[(int)BgOpcode.ObjSet] = new BgOpcodeInfo(BgArg.Value, BgArg.Name, BgArg.Value);
+
+			opcodes[(int)BgOpcode.Call] = new BgOpcodeInfo(BgArg.Fragment, BgArg.ArgList);
+			opcodes[(int)BgOpcode.Argument] = new BgOpcodeInfo(BgArg.VarIntUnsigned);
+			opcodes[(int)BgOpcode.Jump] = new BgOpcodeInfo(BgArg.Fragment);
+
+			opcodes[(int)BgOpcode.Choose] = new BgOpcodeInfo(BgArg.Value, BgArg.Fragment, BgArg.Fragment);
+			opcodes[(int)BgOpcode.Throw] = new BgOpcodeInfo(BgArg.StringLiteral, BgArg.VarIntUnsigned, BgArg.Value);
+			opcodes[(int)BgOpcode.Null] = new BgOpcodeInfo();
+			opcodes[(int)BgOpcode.Thunk] = new BgOpcodeInfo(BgArg.Thunk, BgArg.ArgList);
+
+			return opcodes;
+		}
+
 		/// <summary>
 		/// Constructor
 		/// </summary>
@@ -136,6 +243,54 @@ namespace EpicGames.BuildGraph
 		object Evaluate(Frame frame)
 		{
 			BgOpcode opcode = ReadOpcode(frame);
+			BgOpcodeInfo opcodeInfo = _opcodeTable[(int)opcode];
+
+			object arg0 = ReadArgument(frame, opcodeInfo.Arg0);
+			object arg1 = ReadArgument(frame, opcodeInfo.Arg1);
+			object arg2 = ReadArgument(frame, opcodeInfo.Arg2);
+
+			return Evaluate(frame, opcode, arg0, arg1, arg2);
+		}
+
+		object ReadArgument(Frame frame, BgArg type)
+		{
+			switch (type)
+			{
+				case BgArg.None:
+					return null!;
+				case BgArg.Value:
+					return Evaluate(frame);
+				case BgArg.Name:
+					return ReadName(frame);
+				case BgArg.Fragment:
+					return ReadFragment(frame);
+				case BgArg.VarIntSigned:
+					return ReadSignedInteger(frame);
+				case BgArg.VarIntUnsigned:
+					return ReadUnsignedInteger(frame);
+				case BgArg.StringLiteral:
+					return ReadString(frame);
+				case BgArg.ArgList:
+					{
+						int argCount = (int)ReadUnsignedInteger(frame);
+						object[] argList = new object[argCount];
+
+						for (int argIdx = 0; argIdx < argCount; argIdx++)
+						{
+							argList[argIdx] = Evaluate(frame);
+						}
+
+						return argList;
+					}
+				case BgArg.Thunk:
+					return _thunks[(int)ReadUnsignedInteger(frame)];
+				default:
+					throw new InvalidDataException();
+			}
+		}
+
+		object Evaluate(Frame frame, BgOpcode opcode, object arg0, object arg1, object arg2)
+		{
 			switch (opcode)
 			{
 				#region Bool opcodes
@@ -145,34 +300,34 @@ namespace EpicGames.BuildGraph
 				case BgOpcode.BoolTrue:
 					return true;
 				case BgOpcode.BoolNot:
-					return !(bool)Evaluate(frame);
+					return !(bool)arg0;
 				case BgOpcode.BoolAnd:
 					{
-						bool lhs = (bool)Evaluate(frame);
-						bool rhs = (bool)Evaluate(frame);
+						bool lhs = (bool)arg0;
+						bool rhs = (bool)arg1;
 						return lhs & rhs;
 					}
 				case BgOpcode.BoolOr:
 					{
-						bool lhs = (bool)Evaluate(frame);
-						bool rhs = (bool)Evaluate(frame);
+						bool lhs = (bool)arg0;
+						bool rhs = (bool)arg1;
 						return lhs | rhs;
 					}
 				case BgOpcode.BoolXor:
 					{
-						bool lhs = (bool)Evaluate(frame);
-						bool rhs = (bool)Evaluate(frame);
+						bool lhs = (bool)arg0;
+						bool rhs = (bool)arg1;
 						return lhs ^ rhs;
 					}
 				case BgOpcode.BoolEq:
 					{
-						bool lhs = (bool)Evaluate(frame);
-						bool rhs = (bool)Evaluate(frame);
+						bool lhs = (bool)arg0;
+						bool rhs = (bool)arg1;
 						return lhs == rhs;
 					}
 				case BgOpcode.BoolOption:
 					{
-						BgObjectDef obj = (BgObjectDef)Evaluate(frame);
+						BgObjectDef obj = (BgObjectDef)arg0;
 
 						BgBoolOptionDef option = obj.Deserialize<BgBoolOptionDef>();
 						_optionDefs.Add(option);
@@ -190,54 +345,54 @@ namespace EpicGames.BuildGraph
 				#region Integer opcodes
 
 				case BgOpcode.IntLiteral:
-					return ReadSignedInteger(frame);
+					return arg0;
 				case BgOpcode.IntEq:
 					{
-						int lhs = (int)Evaluate(frame);
-						int rhs = (int)Evaluate(frame);
+						int lhs = (int)arg0;
+						int rhs = (int)arg1;
 						return lhs == rhs;
 					}
 				case BgOpcode.IntLt:
 					{
-						int lhs = (int)Evaluate(frame);
-						int rhs = (int)Evaluate(frame);
+						int lhs = (int)arg0;
+						int rhs = (int)arg1;
 						return lhs < rhs;
 					}
 				case BgOpcode.IntGt:
 					{
-						int lhs = (int)Evaluate(frame);
-						int rhs = (int)Evaluate(frame);
+						int lhs = (int)arg0;
+						int rhs = (int)arg1;
 						return lhs > rhs;
 					}
 				case BgOpcode.IntAdd:
 					{
-						int lhs = (int)Evaluate(frame);
-						int rhs = (int)Evaluate(frame);
+						int lhs = (int)arg0;
+						int rhs = (int)arg1;
 						return lhs + rhs;
 					}
 				case BgOpcode.IntMultiply:
 					{
-						int lhs = (int)Evaluate(frame);
-						int rhs = (int)Evaluate(frame);
+						int lhs = (int)arg0;
+						int rhs = (int)arg1;
 						return lhs * rhs;
 					}
 				case BgOpcode.IntDivide:
 					{
-						int lhs = (int)Evaluate(frame);
-						int rhs = (int)Evaluate(frame);
+						int lhs = (int)arg0;
+						int rhs = (int)arg1;
 						return lhs / rhs;
 					}
 				case BgOpcode.IntModulo:
 					{
-						int lhs = (int)Evaluate(frame);
-						int rhs = (int)Evaluate(frame);
+						int lhs = (int)arg0;
+						int rhs = (int)arg1;
 						return lhs % rhs;
 					}
 				case BgOpcode.IntNegate:
-					return -(int)Evaluate(frame);
+					return -(int)arg0;
 				case BgOpcode.IntOption:
 					{
-						BgObjectDef obj = (BgObjectDef)Evaluate(frame);
+						BgObjectDef obj = (BgObjectDef)arg0;
 
 						BgIntOptionDef option = obj.Deserialize<BgIntOptionDef>();
 						_optionDefs.Add(option);
@@ -257,63 +412,55 @@ namespace EpicGames.BuildGraph
 				case BgOpcode.StrEmpty:
 					return String.Empty;
 				case BgOpcode.StrLiteral:
-					return ReadString(frame);
+					return arg0;
 				case BgOpcode.StrCompare:
 					{
-						string lhs = (string)Evaluate(frame);
-						string rhs = (string)Evaluate(frame);
-						StringComparison comparison = (StringComparison)ReadUnsignedInteger(frame);
+						string lhs = (string)arg0;
+						string rhs = (string)arg1;
+						StringComparison comparison = (StringComparison)(ulong)arg2;
 						return String.Compare(lhs, rhs, comparison);
 					}
 				case BgOpcode.StrConcat:
 					{
-						string lhs = (string)Evaluate(frame);
-						string rhs = (string)Evaluate(frame);
+						string lhs = (string)arg0;
+						string rhs = (string)arg1;
 						return lhs + rhs;
 					}
 				case BgOpcode.StrFormat:
 					{
-						string format = (string)Evaluate(frame);
-
-						int count = (int)ReadUnsignedInteger(frame);
-						object[] arguments = new object[count];
-
-						for (int idx = 0; idx < count; idx++)
-						{
-							object argument = Evaluate(frame);
-							arguments[idx] = argument;
-						}
+						string format = (string)arg0;
+						object?[] arguments = (object?[])arg1;
 
 						return String.Format(format, arguments);
 					}
 				case BgOpcode.StrSplit:
 					{
-						string source = (string)Evaluate(frame);
-						string separator = (string)Evaluate(frame);
+						string source = (string)arg0;
+						string separator = (string)arg1;
 						return source.Split(separator, StringSplitOptions.RemoveEmptyEntries);
 					}
 				case BgOpcode.StrJoin:
 					{
-						string lhs = (string)Evaluate(frame);
-						IEnumerable<object> rhs = (IEnumerable<object>)Evaluate(frame);
+						string lhs = (string)arg0;
+						IEnumerable<object> rhs = (IEnumerable<object>)arg1;
 						return String.Join(lhs, rhs);
 					}
 				case BgOpcode.StrMatch:
 					{
-						string input = (string)Evaluate(frame);
-						string pattern = (string)Evaluate(frame);
+						string input = (string)arg0;
+						string pattern = (string)arg1;
 						return Regex.IsMatch(input, pattern);
 					}
 				case BgOpcode.StrReplace:
 					{
-						string input = (string)Evaluate(frame);
-						string pattern = (string)Evaluate(frame);
-						string replacement = (string)Evaluate(frame);
+						string input = (string)arg0;
+						string pattern = (string)arg1;
+						string replacement = (string)arg2;
 						return Regex.Replace(input, pattern, replacement);
 					}
 				case BgOpcode.StrOption:
 					{
-						BgObjectDef obj = (BgObjectDef)Evaluate(frame);
+						BgObjectDef obj = (BgObjectDef)arg0;
 
 						BgStringOptionDef option = obj.Deserialize<BgStringOptionDef>();
 						_optionDefs.Add(option);
@@ -331,12 +478,12 @@ namespace EpicGames.BuildGraph
 				#region Enum opcodes
 
 				case BgOpcode.EnumConstant:
-					return ReadSignedInteger(frame);
+					return (int)arg0;
 				case BgOpcode.EnumParse:
 					{
-						string name = (string)Evaluate(frame);
-						List<string> names = EvaluateList<string>(frame);
-						List<int> values = EvaluateList<int>(frame);
+						string name = (string)arg0;
+						List<string> names = GetListArg<string>(arg1);
+						List<int> values = GetListArg<int>(arg2);
 
 						for (int idx = 0; idx < names.Count; idx++)
 						{
@@ -350,9 +497,9 @@ namespace EpicGames.BuildGraph
 					}
 				case BgOpcode.EnumToString:
 					{
-						int value = (int)Evaluate(frame);
-						List<string> names = EvaluateList<string>(frame);
-						List<int> values = EvaluateList<int>(frame);
+						int value = (int)arg0;
+						List<string> names = GetListArg<string>(arg1);
+						List<int> values = GetListArg<int>(arg2);
 
 						for (int idx = 0; idx < names.Count; idx++)
 						{
@@ -372,81 +519,81 @@ namespace EpicGames.BuildGraph
 					return Enumerable.Empty<object>();
 				case BgOpcode.ListPush:
 					{
-						IEnumerable<object> list = (IEnumerable<object>)Evaluate(frame);
-						object item = Evaluate(frame);
+						IEnumerable<object> list = (IEnumerable<object>)arg0;
+						object item = arg1;
 
 						return list.Concat(new[] { item });
 					}
 				case BgOpcode.ListPushLazy:
 					{
-						IEnumerable<object> list = (IEnumerable<object>)Evaluate(frame);
-						int fragment = ReadFragment(frame);
+						IEnumerable<object> list = (IEnumerable<object>)arg0;
+						int fragment = (int)arg1;
 
 						IEnumerable<object> item = LazyEvaluateItem(new Frame(frame), fragment);
 						return list.Concat(item);
 					}
 				case BgOpcode.ListCount:
 					{
-						IEnumerable<object> list = (IEnumerable<object>)Evaluate(frame);
+						IEnumerable<object> list = (IEnumerable<object>)arg0;
 						return list.Count();
 					}
 				case BgOpcode.ListElement:
 					{
-						IEnumerable<object> list = (IEnumerable<object>)Evaluate(frame);
-						int index = (int)Evaluate(frame);
+						IEnumerable<object> list = (IEnumerable<object>)arg0;
+						int index = (int)arg1;
 						return list.ElementAt(index);
 					}
 				case BgOpcode.ListConcat:
 					{
-						IEnumerable<object> lhs = (IEnumerable<object>)Evaluate(frame);
-						IEnumerable<object> rhs = (IEnumerable<object>)Evaluate(frame);
+						IEnumerable<object> lhs = (IEnumerable<object>)arg0;
+						IEnumerable<object> rhs = (IEnumerable<object>)arg1;
 						return Enumerable.Concat(lhs, rhs);
 					}
 				case BgOpcode.ListUnion:
 					{
-						IEnumerable<object> lhs = (IEnumerable<object>)Evaluate(frame);
-						IEnumerable<object> rhs = (IEnumerable<object>)Evaluate(frame);
+						IEnumerable<object> lhs = (IEnumerable<object>)arg0;
+						IEnumerable<object> rhs = (IEnumerable<object>)arg1;
 						return lhs.Union(rhs);
 					}
 				case BgOpcode.ListExcept:
 					{
-						IEnumerable<object> lhs = (IEnumerable<object>)Evaluate(frame);
-						IEnumerable<object> rhs = (IEnumerable<object>)Evaluate(frame);
+						IEnumerable<object> lhs = (IEnumerable<object>)arg0;
+						IEnumerable<object> rhs = (IEnumerable<object>)arg1;
 						return lhs.Except(rhs);
 					}
 				case BgOpcode.ListSelect:
 					{
-						IEnumerable<object> source = (IEnumerable<object>)Evaluate(frame);
-						int fragment = ReadFragment(frame);
+						IEnumerable<object> source = (IEnumerable<object>)arg0;
+						int fragment = (int)arg1;
 
 						return source.Select(x => Call(fragment, new object[] { x }));
 					}
 				case BgOpcode.ListWhere:
 					{
-						IEnumerable<object> source = (IEnumerable<object>)Evaluate(frame);
-						int fragment = ReadFragment(frame);
+						IEnumerable<object> source = (IEnumerable<object>)arg0;
+						int fragment = (int)arg1;
 
 						return source.Where(x => (bool)Call(fragment, new object[] { x }));
 					}
 				case BgOpcode.ListDistinct:
 					{
-						IEnumerable<object> source = (IEnumerable<object>)Evaluate(frame);
+						IEnumerable<object> source = (IEnumerable<object>)arg0;
 						return source.Distinct();
 					}
 				case BgOpcode.ListContains:
 					{
-						IEnumerable<object> source = (IEnumerable<object>)Evaluate(frame);
-						object item = Evaluate(frame);
+						IEnumerable<object> source = (IEnumerable<object>)arg0;
+						object item = arg1;
 						return source.Contains(item);
 					}
 				case BgOpcode.ListLazy:
 					{
-						int fragment = ReadFragment(frame);
+						int fragment = (int)arg0;
 						return LazyEvaluateList(new Frame(frame), fragment);
 					}
 				case BgOpcode.ListOption:
 					{
-						BgObjectDef obj = (BgObjectDef)Evaluate(frame);
+						BgObjectDef obj = (BgObjectDef)arg0;
 
 						BgListOptionDef option = obj.Deserialize<BgListOptionDef>();
 						_optionDefs.Add(option);
@@ -467,9 +614,9 @@ namespace EpicGames.BuildGraph
 					return new BgObjectDef();
 				case BgOpcode.ObjGet:
 					{
-						BgObjectDef obj = (BgObjectDef)Evaluate(frame);
-						string name = ReadName(frame);
-						object defaultValue = Evaluate(frame);
+						BgObjectDef obj = (BgObjectDef)arg0;
+						string name = (string)arg1;
+						object defaultValue = arg2;
 
 						object? value;
 						if (!obj.Properties.TryGetValue(name, out value))
@@ -481,9 +628,9 @@ namespace EpicGames.BuildGraph
 					}
 				case BgOpcode.ObjSet:
 					{
-						BgObjectDef obj = (BgObjectDef)Evaluate(frame);
-						string name = ReadName(frame);
-						object value = Evaluate(frame);
+						BgObjectDef obj = (BgObjectDef)arg0;
+						string name = (string)arg1;
+						object value = arg2;
 						return obj.Set(name, value);
 					}
 
@@ -492,25 +639,18 @@ namespace EpicGames.BuildGraph
 
 				case BgOpcode.Call:
 					{
-						int count = (int)ReadUnsignedInteger(frame);
-
-						object[] arguments = new object[count];
-						for (int idx = 0; idx < count; idx++)
-						{
-							arguments[idx] = Evaluate(frame);
-						}
-
-						int function = (int)ReadUnsignedInteger(frame);
+						int function = (int)arg0;
+						object[] arguments = (object[])arg1;
 						return Evaluate(new Frame(_fragments[function], arguments));
 					}
 				case BgOpcode.Argument:
 					{
-						int index = (int)ReadUnsignedInteger(frame);
+						int index = (int)(ulong)arg0;
 						return frame.Arguments[index];
 					}
 				case BgOpcode.Jump:
 					{
-						int fragment = (int)ReadUnsignedInteger(frame);
+						int fragment = (int)arg0;
 						return Jump(frame, fragment);
 					}
 
@@ -519,44 +659,44 @@ namespace EpicGames.BuildGraph
 
 				case BgOpcode.Choose:
 					{
-						bool condition = (bool)Evaluate(frame);
+						bool condition = (bool)arg0;
 
-						int fragmentIfTrue = (int)ReadUnsignedInteger(frame);
-						int fragmentIfFalse = (int)ReadUnsignedInteger(frame);
+						int fragmentIfTrue = (int)arg1;
+						int fragmentIfFalse = (int)arg2;
 
 						return Jump(frame, condition ? fragmentIfTrue : fragmentIfFalse);
 					}
 				case BgOpcode.Throw:
 					{
-						string sourceFile = ReadString(frame);
-						int sourceLine = (int)ReadUnsignedInteger(frame);
-						string message = (string)Evaluate(frame);
+						string sourceFile = (string)arg0;
+						int sourceLine = (int)(ulong)arg1;
+						string message = (string)arg2;
 						throw new BgBytecodeException(sourceFile, sourceLine, message);
 					}
 				case BgOpcode.Null:
 					return null!;
 				case BgOpcode.Thunk:
 					{
-						BgThunkDef method = _thunks[(int)ReadUnsignedInteger(frame)];
+						BgThunkDef method = (BgThunkDef)arg0;
+						object[] arguments = (object[])arg1;
 
-						int count = (int)ReadUnsignedInteger(frame);
-						object?[] arguments = new object?[count];
+						object[] thunkArgs = new object[arguments.Length];
 
 						ParameterInfo[] parameters = method.Method.GetParameters();
-						for (int idx = 0; idx < count; idx++)
+						for (int idx = 0; idx < arguments.Length; idx++)
 						{
-							object value = Evaluate(frame);
+							object value = arguments[idx];
 							if (typeof(BgExpr).IsAssignableFrom(parameters[idx].ParameterType))
 							{
-								arguments[idx] = BgType.Constant(parameters[idx].ParameterType, value);
+								thunkArgs[idx] = BgType.Constant(parameters[idx].ParameterType, value);
 							}
 							else
 							{
-								arguments[idx] = method.Arguments[idx];
+								thunkArgs[idx] = method.Arguments[idx]!;
 							}
 						}
 
-						return new BgThunkDef(method.Method, arguments);
+						return new BgThunkDef(method.Method, thunkArgs);
 					}
 
 				#endregion
@@ -566,9 +706,9 @@ namespace EpicGames.BuildGraph
 			}
 		}
 
-		List<T> EvaluateList<T>(Frame frame)
+		static List<T> GetListArg<T>(object value)
 		{
-			return ((IEnumerable<object>)Evaluate(frame)).Select(x => (T)x).ToList();
+			return ((IEnumerable<object>)value).Select(x => (T)x).ToList();
 		}
 
 		IEnumerable<object> LazyEvaluateItem(Frame frame, int fragment)
@@ -704,208 +844,46 @@ namespace EpicGames.BuildGraph
 		void Disassemble(Frame frame, ILogger logger)
 		{
 			BgOpcode opcode = Trace(frame, null, ReadOpcode, logger);
-			switch (opcode)
-			{
-				#region Bool opcodes
+			BgOpcodeInfo opcodeInfo = _opcodeTable[(int)opcode];
 
-				case BgOpcode.BoolFalse:
-				case BgOpcode.BoolTrue:
-					break;
-				case BgOpcode.BoolNot:
-					Disassemble(frame, logger);
-					break;
-				case BgOpcode.BoolAnd:
-				case BgOpcode.BoolOr:
-				case BgOpcode.BoolXor:
-				case BgOpcode.BoolEq:
-					Disassemble(frame, logger);
-					Disassemble(frame, logger);
-					break;
-				case BgOpcode.BoolOption:
-					Disassemble(frame, logger);
-					break;
-
-				#endregion
-				#region Integer opcodes
-
-				case BgOpcode.IntLiteral:
-					Trace(frame, null, ReadSignedInteger, logger);
-					break;
-				case BgOpcode.IntEq:
-				case BgOpcode.IntLt:
-				case BgOpcode.IntGt:
-				case BgOpcode.IntAdd:
-				case BgOpcode.IntMultiply:
-				case BgOpcode.IntDivide:
-				case BgOpcode.IntModulo:
-					Disassemble(frame, logger);
-					Disassemble(frame, logger);
-					break;
-				case BgOpcode.IntNegate:
-					Disassemble(frame, logger);
-					break;
-				case BgOpcode.IntOption:
-					Disassemble(frame, logger);
-					break;
-
-				#endregion
-				#region String opcodes
-
-				case BgOpcode.StrEmpty:
-					break;
-				case BgOpcode.StrLiteral:
-					Trace(frame, ReadString, x => $"\"{x}\"", logger);
-					break;
-				case BgOpcode.StrCompare:
-					Disassemble(frame, logger);
-					Disassemble(frame, logger);
-					Trace(frame, "type", ReadUnsignedInteger, logger);
-					break;
-				case BgOpcode.StrConcat:
-				case BgOpcode.StrSplit:
-				case BgOpcode.StrJoin:
-					Disassemble(frame, logger);
-					Disassemble(frame, logger);
-					break;
-				case BgOpcode.StrFormat:
-					{
+			for (int idx = 0; idx < opcodeInfo.ArgCount; idx++)
+			{
+				BgArg arg = opcodeInfo.GetArg(idx);
+				switch (arg)
+				{
+					case BgArg.Value:
 						Disassemble(frame, logger);
-
-						int count = (int)Trace(frame, "count", ReadUnsignedInteger, logger);
-						for (int idx = 0; idx < count; idx++)
-						{
-							Disassemble(frame, logger);
-						}
-
 						break;
-					}
-				case BgOpcode.StrMatch:
-					Disassemble(frame, logger);
-					Disassemble(frame, logger);
-					break;
-				case BgOpcode.StrReplace:
-					Disassemble(frame, logger);
-					Disassemble(frame, logger);
-					Disassemble(frame, logger);
-					break;
-				case BgOpcode.StrOption:
-					Disassemble(frame, logger);
-					break;
-
-				#endregion
-				#region Enum opcodes
-
-				case BgOpcode.EnumConstant:
-					Trace(frame, "value", ReadSignedInteger, logger);
-					break;
-				case BgOpcode.EnumParse:
-				case BgOpcode.EnumToString:
-					Disassemble(frame, logger);
-					Disassemble(frame, logger);
-					Disassemble(frame, logger);
-					break;
-
-				#endregion
-				#region List opcodes
-
-				case BgOpcode.ListEmpty:
-					break;
-				case BgOpcode.ListPush:
-					Disassemble(frame, logger);
-					Disassemble(frame, logger);
-					break;
-				case BgOpcode.ListPushLazy:
-				case BgOpcode.ListSelect:
-				case BgOpcode.ListWhere:
-					Disassemble(frame, logger);
-					TraceFragment(frame, logger);
-					break;
-				case BgOpcode.ListCount:
-				case BgOpcode.ListDistinct:
-					Disassemble(frame, logger);
-					break;
-				case BgOpcode.ListElement:
-				case BgOpcode.ListConcat:
-				case BgOpcode.ListUnion:
-				case BgOpcode.ListExcept:
-				case BgOpcode.ListContains:
-					Disassemble(frame, logger);
-					Disassemble(frame, logger);
-					break;
-				case BgOpcode.ListLazy:
-					TraceFragment(frame, logger);
-					break;
-				case BgOpcode.ListOption:
-					Disassemble(frame, logger);
-					break;
-
-				#endregion
-
-				#region Object opcodes
-
-				case BgOpcode.ObjEmpty:
-					break;
-				case BgOpcode.ObjGet:
-				case BgOpcode.ObjSet:
-					Disassemble(frame, logger);
-					Trace(frame, ReadName, x => $"name: {x}", logger);
-					Disassemble(frame, logger);
-					break;
-
-				#endregion
-
-				#region Function opcodes
-
-				case BgOpcode.Call:
-					{
-						int count = (int)Trace(frame, "count", ReadUnsignedInteger, logger);
-						for (int idx = 0; idx < count; idx++)
-						{
-							Disassemble(frame, logger);
-						}
-
+					case BgArg.Name:
+						Trace(frame, ReadName, x => $"name: {x}", logger);
+						break;
+					case BgArg.Fragment:
 						TraceFragment(frame, logger);
 						break;
-					}
-				case BgOpcode.Argument:
-					Trace(frame, "arg", ReadUnsignedInteger, logger);
-					break;
-				case BgOpcode.Jump:
-					TraceFragment(frame, logger);
-					break;
-
-				#endregion
-				#region Generic opcodes
-
-				case BgOpcode.Choose:
-					Disassemble(frame, logger);
-					TraceFragment(frame, logger);
-					TraceFragment(frame, logger);
-					break;
-				case BgOpcode.Throw:
-					Trace(frame, "file", ReadString, logger);
-					Trace(frame, "line", ReadUnsignedInteger, logger);
-					Trace(frame, "message", ReadString, logger);
-					break;
-				case BgOpcode.Null:
-					break;
-				case BgOpcode.Thunk:
-					{
-						Trace(frame, "method", ReadUnsignedInteger, logger);
-
-						int count = (int)Trace(frame, "args", ReadUnsignedInteger, logger);
-						for (int idx = 0; idx < count; idx++)
+					case BgArg.VarIntSigned:
+						Trace(frame, "int:", ReadSignedInteger, logger);
+						break;
+					case BgArg.VarIntUnsigned:
+						Trace(frame, "uint:", ReadUnsignedInteger, logger);
+						break;
+					case BgArg.StringLiteral:
+						Trace(frame, ReadString, x => $"\"{x}\"", logger);
+						break;
+					case BgArg.ArgList:
 						{
-							Disassemble(frame, logger);
+							int argCount = (int)Trace(frame, "count", ReadUnsignedInteger, logger);
+							for (int argIdx = 0; argIdx < argCount; argIdx++)
+							{
+								Disassemble(frame, logger);
+							}
+							break;
 						}
-
+					case BgArg.Thunk:
+						Trace(frame, "thunk", ReadUnsignedInteger, logger);
 						break;
-					}
-
-				#endregion
-
-				default:
-					throw new InvalidDataException($"Invalid opcode: {opcode}");
+					default:
+						throw new InvalidDataException();
+				}
 			}
 		}
 
