Horde: Add inputs and outputs to exported Horde job definition.

diff --git a/Engine/Source/Programs/Shared/EpicGames.BuildGraph/BgGraphDef.cs b/Engine/Source/Programs/Shared/EpicGames.BuildGraph/BgGraphDef.cs
--- a/Engine/Source/Programs/Shared/EpicGames.BuildGraph/BgGraphDef.cs
+++ b/Engine/Source/Programs/Shared/EpicGames.BuildGraph/BgGraphDef.cs
@@ -440,6 +440,20 @@ namespace EpicGames.BuildGraph
 						jsonWriter.WriteValue("RunEarly", node.RunEarly);
 						jsonWriter.WriteValue("Warnings", node.NotifyOnWarnings);
 
+						jsonWriter.WriteArrayStart("Inputs");
+						foreach (BgNodeOutput input in node.Inputs)
+						{
+							jsonWriter.WriteValue($"{input.ProducingNode.Name}/{input.TagName}");
+						}
+						jsonWriter.WriteArrayEnd();
+
+						jsonWriter.WriteArrayStart("Outputs");
+						foreach (BgNodeOutput output in node.Outputs)
+						{
+							jsonWriter.WriteValue($"{output.ProducingNode.Name}/{output.TagName}");
+						}
+						jsonWriter.WriteArrayEnd();
+
 						jsonWriter.WriteArrayStart("InputDependencies");
 						foreach (string inputDependency in node.GetDirectInputDependencies().Select(x => x.Name))
 						{
