Horde: Disable exception when waiting for processes to terminate.

diff --git a/Engine/Source/Programs/Shared/EpicGames.Core/ProcessUtils.cs b/Engine/Source/Programs/Shared/EpicGames.Core/ProcessUtils.cs
--- a/Engine/Source/Programs/Shared/EpicGames.Core/ProcessUtils.cs
+++ b/Engine/Source/Programs/Shared/EpicGames.Core/ProcessUtils.cs
@@ -141,7 +141,7 @@ namespace EpicGames.Core
 					// Wait for them all to complete
 					if (WaitForMultipleObjects(waitHandles.Count, waitHandles.Select(x => x.DangerousGetHandle()).ToArray(), true, 10 * 1000) == WAIT_FAILED)
 					{
-						throw new Win32Exception();
+						logger.LogInformation("Failed call to WaitForMultipleObjects ({Code})", Marshal.GetLastWin32Error());
 					}
 				}
 				finally
