Horde: Fix tracking of modified TreeNode objects. TreeNodeRef and TreeNode objects now track their owner, and the dirty flag is set on any ancestors when a node is modified.

Also add an TreeNode.EnumerateRefs() method to allow depth-first serialization of nodes without having to keep serialized data in oversized buffers.

diff --git a/Engine/Source/Programs/Horde/Horde.Build/Perforce/ReplicationService.cs b/Engine/Source/Programs/Horde/Horde.Build/Perforce/ReplicationService.cs
--- a/Engine/Source/Programs/Horde/Horde.Build/Perforce/ReplicationService.cs
+++ b/Engine/Source/Programs/Horde/Horde.Build/Perforce/ReplicationService.cs
@@ -108,6 +108,12 @@ namespace Horde.Build.Perforce
 			writer.WriteList(Paths, x => writer.WriteUtf8String(x));
 			writer.WriteRef(Contents);
 		}
+
+		/// <inheritdoc/>
+		public override IEnumerable<TreeNodeRef> EnumerateRefs()
+		{
+			yield return Contents;
+		}
 	}
 
 	/// <summary>
diff --git a/Engine/Source/Programs/Shared/EpicGames.Horde.Tests/BundleStoreTests.cs b/Engine/Source/Programs/Shared/EpicGames.Horde.Tests/BundleStoreTests.cs
--- a/Engine/Source/Programs/Shared/EpicGames.Horde.Tests/BundleStoreTests.cs
+++ b/Engine/Source/Programs/Shared/EpicGames.Horde.Tests/BundleStoreTests.cs
@@ -68,6 +68,8 @@ namespace EpicGames.Horde.Tests
 				writer.WriteVariableLengthBytes(Data);
 				writer.WriteVariableLengthArray(Refs, x => writer.WriteRef(x));
 			}
+
+			public override IEnumerable<TreeNodeRef> EnumerateRefs() => Refs;
 		}
 
 		static async Task TestTreeAsync(MemoryStorageClient store, TreeOptions options)
diff --git a/Engine/Source/Programs/Shared/EpicGames.Horde/Storage/Nodes/DirectoryNode.cs b/Engine/Source/Programs/Shared/EpicGames.Horde/Storage/Nodes/DirectoryNode.cs
--- a/Engine/Source/Programs/Shared/EpicGames.Horde/Storage/Nodes/DirectoryNode.cs
+++ b/Engine/Source/Programs/Shared/EpicGames.Horde/Storage/Nodes/DirectoryNode.cs
@@ -1,6 +1,7 @@
 // Copyright Epic Games, Inc. All Rights Reserved.
 
 using EpicGames.Core;
+using EpicGames.Serialization;
 using Microsoft.Extensions.Logging;
 using Microsoft.Extensions.Options;
 using System;
@@ -149,6 +150,29 @@ namespace EpicGames.Horde.Storage.Nodes
 			Target = await node.AppendAsync(stream, options, writer, cancellationToken);
 		}
 
+		/// <summary>
+		/// Copies the contents of this node and its children to the given output stream
+		/// </summary>
+		/// <param name="outputStream">The output stream to receive the data</param>
+		/// <param name="cancellationToken">Cancellation token for the operation</param>
+		public async Task CopyToStreamAsync(Stream outputStream, CancellationToken cancellationToken)
+		{
+			FileNode node = await ExpandAsync(cancellationToken);
+			await node.CopyToStreamAsync(outputStream, cancellationToken);
+		}
+
+		/// <summary>
+		/// Extracts the contents of this node to a file
+		/// </summary>
+		/// <param name="file">File to write with the contents of this node</param>
+		/// <param name="cancellationToken">Cancellation token for the operation</param>
+		/// <returns></returns>
+		public async Task CopyToFileAsync(FileInfo file, CancellationToken cancellationToken)
+		{
+			FileNode node = await ExpandAsync(cancellationToken);
+			await node.CopyToFileAsync(file, cancellationToken);
+		}
+
 		/// <inheritdoc/>
 		public override string ToString() => Name.ToString();
 	}
@@ -173,6 +197,14 @@ namespace EpicGames.Horde.Storage.Nodes
 		/// </summary>
 		readonly long _cachedLength;
 
+		/// <summary>
+		/// Constructor
+		/// </summary>
+		public DirectoryEntry(Utf8String name)
+			: this(name, new DirectoryNode())
+		{
+		}
+
 		/// <summary>
 		/// Constructor
 		/// </summary>
@@ -233,7 +265,7 @@ namespace EpicGames.Horde.Storage.Nodes
 		/// <summary>
 		/// Total size of this directory
 		/// </summary>
-		public long Length => Files.Sum(x => x.Length) + Directories.Sum(x => x.Length);
+		public long Length => _nameToFileEntry.Values.Sum(x => x.Length) + _nameToDirectoryEntry.Values.Sum(x => x.Length);
 
 		/// <summary>
 		/// Flags for this directory 
@@ -245,11 +277,21 @@ namespace EpicGames.Horde.Storage.Nodes
 		/// </summary>
 		public IReadOnlyCollection<FileEntry> Files => _nameToFileEntry.Values;
 
+		/// <summary>
+		/// Map of name to file entry
+		/// </summary>
+		public IReadOnlyDictionary<Utf8String, FileEntry> NameToFile => _nameToFileEntry;
+
 		/// <summary>
 		/// All the subdirectories within this directory
 		/// </summary>
 		public IReadOnlyCollection<DirectoryEntry> Directories => _nameToDirectoryEntry.Values;
 
+		/// <summary>
+		/// Map of name to file entry
+		/// </summary>
+		public IReadOnlyDictionary<Utf8String, DirectoryEntry> NameToDirectory => _nameToDirectoryEntry;
+
 		/// <summary>
 		/// Constructor
 		/// </summary>
@@ -299,18 +341,31 @@ namespace EpicGames.Horde.Storage.Nodes
 			writer.WriteUnsignedVarInt((ulong)Flags);
 
 			writer.WriteUnsignedVarInt(Files.Count);
-			foreach (FileEntry fileEntry in Files)
+			foreach (FileEntry fileEntry in _nameToFileEntry.Values)
 			{
 				fileEntry.Serialize(writer);
 			}
 
 			writer.WriteUnsignedVarInt(Directories.Count);
-			foreach (DirectoryEntry directoryEntry in Directories)
+			foreach (DirectoryEntry directoryEntry in _nameToDirectoryEntry.Values)
 			{
 				directoryEntry.Serialize(writer);
 			}
 		}
 
+		/// <inheritdoc/>
+		public override IEnumerable<TreeNodeRef> EnumerateRefs()
+		{
+			foreach (FileEntry fileEntry in _nameToFileEntry.Values)
+			{
+				yield return fileEntry;
+			}
+			foreach (DirectoryEntry directoryEntry in _nameToDirectoryEntry.Values)
+			{
+				yield return directoryEntry;
+			}
+		}
+
 		/// <summary>
 		/// Clear the contents of this directory
 		/// </summary>
@@ -336,12 +391,7 @@ namespace EpicGames.Horde.Storage.Nodes
 		/// <param name="entry">The entry to add</param>
 		public void AddFile(FileEntry entry)
 		{
-			if (TryGetDirectoryEntry(entry.Name, out _))
-			{
-				throw new ArgumentException($"A directory with the name {entry.Name} already exists");
-			}
-
-			_nameToFileEntry.Add(entry.Name, entry);
+			_nameToFileEntry[entry.Name] = entry;
 			MarkAsDirty();
 		}
 
@@ -512,6 +562,23 @@ namespace EpicGames.Horde.Storage.Nodes
 		/// <returns>True if the directory was found</returns>
 		public bool TryGetDirectoryEntry(Utf8String name, [NotNullWhen(true)] out DirectoryEntry? entry) => _nameToDirectoryEntry.TryGetValue(name, out entry);
 
+		/// <summary>
+		/// Attempts to get a directory entry with the given name, and adds one if it does not already exist.
+		/// </summary>
+		/// <param name="name">Name of the directory</param>
+		/// <returns>Directory entry with the given name</returns>
+		public DirectoryEntry FindOrAddDirectoryEntry(Utf8String name)
+		{
+			DirectoryEntry? entry;
+			if (!_nameToDirectoryEntry.TryGetValue(name, out entry))
+			{
+				entry = new DirectoryEntry(name);
+				_nameToDirectoryEntry.Add(name, entry);
+				MarkAsDirty();
+			}
+			return entry;
+		}
+
 		/// <summary>
 		/// Tries to get a directory with the given name
 		/// </summary>
@@ -573,12 +640,24 @@ namespace EpicGames.Horde.Storage.Nodes
 		/// <returns></returns>
 		public async Task CopyFromDirectoryAsync(DirectoryInfo directoryInfo, ChunkingOptions options, TreeWriter writer, CancellationToken cancellationToken)
 		{
-			const int MaxWriters = 32;
-			const long MinSizePerWriter = 1024 * 1024;
-
 			// Enumerate all the files below this directory
 			List<(DirectoryNode DirectoryNode, FileInfo FileInfo)> files = new List<(DirectoryNode, FileInfo)>();
 			FindFilesToCopy(directoryInfo, files);
+			await CopyFromDirectoryAsync(files, options, writer, cancellationToken);
+		}
+
+		/// <summary>
+		/// Adds files from a directory on disk
+		/// </summary>
+		/// <param name="files"></param>
+		/// <param name="options">Options for chunking file content</param>
+		/// <param name="writer">Writer for new node data</param>
+		/// <param name="cancellationToken">Cancellation token for the operation</param>
+		/// <returns></returns>
+		public static async Task CopyFromDirectoryAsync(List<(DirectoryNode DirectoryNode, FileInfo FileInfo)> files, ChunkingOptions options, TreeWriter writer, CancellationToken cancellationToken)
+		{
+			const int MaxWriters = 32;
+			const long MinSizePerWriter = 1024 * 1024;
 
 			// Compute the total size
 			long totalSize = files.Sum(x => x.Item2.Length);
@@ -659,13 +738,13 @@ namespace EpicGames.Horde.Storage.Nodes
 			directoryInfo.Create();
 
 			List<Task> tasks = new List<Task>();
-			foreach (FileEntry fileEntry in Files)
+			foreach (FileEntry fileEntry in _nameToFileEntry.Values)
 			{
 				FileInfo fileInfo = new FileInfo(Path.Combine(directoryInfo.FullName, fileEntry.Name.ToString()));
 				FileNode fileNode = await fileEntry.ExpandAsync(cancellationToken);
 				tasks.Add(Task.Run(() => fileNode.CopyToFileAsync(fileInfo, cancellationToken), cancellationToken));
 			}
-			foreach (DirectoryEntry directoryEntry in Directories)
+			foreach (DirectoryEntry directoryEntry in _nameToDirectoryEntry.Values)
 			{
 				DirectoryInfo subDirectoryInfo = directoryInfo.CreateSubdirectory(directoryEntry.Name.ToString());
 				DirectoryNode subDirectoryNode = await directoryEntry.ExpandAsync(cancellationToken);
diff --git a/Engine/Source/Programs/Shared/EpicGames.Horde/Storage/Nodes/FileNode.cs b/Engine/Source/Programs/Shared/EpicGames.Horde/Storage/Nodes/FileNode.cs
--- a/Engine/Source/Programs/Shared/EpicGames.Horde/Storage/Nodes/FileNode.cs
+++ b/Engine/Source/Programs/Shared/EpicGames.Horde/Storage/Nodes/FileNode.cs
@@ -333,6 +333,9 @@ namespace EpicGames.Horde.Storage.Nodes
 			writer.WriteVariableLengthBytes(_writtenSequence);
 		}
 
+		/// <inheritdoc/>
+		public override IEnumerable<TreeNodeRef> EnumerateRefs() => Enumerable.Empty<TreeNodeRef>();
+
 		/// <inheritdoc/>
 		public override bool IsReadOnly() => _isReadOnly;
 
@@ -545,6 +548,9 @@ namespace EpicGames.Horde.Storage.Nodes
 			writer.WriteList(_children, x => writer.WriteRef(x));
 		}
 
+		/// <inheritdoc/>
+		public override IEnumerable<TreeNodeRef> EnumerateRefs() => _children;
+
 		/// <inheritdoc/>
 		public override async ValueTask<ReadOnlyMemory<byte>> AppendDataAsync(ReadOnlyMemory<byte> newData, ChunkingOptions options, TreeWriter writer, CancellationToken cancellationToken)
 		{
diff --git a/Engine/Source/Programs/Shared/EpicGames.Horde/Storage/StorageClientBase.cs b/Engine/Source/Programs/Shared/EpicGames.Horde/Storage/StorageClientBase.cs
--- a/Engine/Source/Programs/Shared/EpicGames.Horde/Storage/StorageClientBase.cs
+++ b/Engine/Source/Programs/Shared/EpicGames.Horde/Storage/StorageClientBase.cs
@@ -565,10 +565,10 @@ namespace EpicGames.Horde.Storage
 				_refs = refs;
 			}
 
-			public void ReadRef(TreeNodeRef treeNodeRef)
+			public TreeNodeRefData ReadRef()
 			{
 				(IoHash hash, NodeLocator locator) = _refs[(int)this.ReadUnsignedVarInt()];
-				treeNodeRef.MarkAsClean(_store, hash, locator, 0);
+				return new TreeNodeRefData(_store, hash, locator);
 			}
 		}
 
diff --git a/Engine/Source/Programs/Shared/EpicGames.Horde/Storage/TreeNode.cs b/Engine/Source/Programs/Shared/EpicGames.Horde/Storage/TreeNode.cs
--- a/Engine/Source/Programs/Shared/EpicGames.Horde/Storage/TreeNode.cs
+++ b/Engine/Source/Programs/Shared/EpicGames.Horde/Storage/TreeNode.cs
@@ -1,6 +1,8 @@
 // Copyright Epic Games, Inc. All Rights Reserved.
 
 using System;
+using System.Collections.Generic;
+using System.Diagnostics;
 using System.Reflection;
 using System.Reflection.Emit;
 using EpicGames.Core;
@@ -13,21 +15,15 @@ namespace EpicGames.Horde.Storage
 	public abstract class TreeNode
 	{
 		/// <summary>
-		/// Revision number for the node. Incremented whenever the node is modified; used to detect changes between a serialized ref and live instance.
+		/// The ref which points to this node
 		/// </summary>
-		int _revision;
-
-		/// <summary>
-		/// Revision number for this node. Incremented Whether the node is dirty or not
-		/// </summary>
-		public int Revision => _revision;
+		public TreeNodeRef? IncomingRef { get; internal set; }
 
 		/// <summary>
 		/// Default constructor
 		/// </summary>
 		protected TreeNode()
 		{
-			_revision = 1;
 		}
 
 		/// <summary>
@@ -44,7 +40,10 @@ namespace EpicGames.Horde.Storage
 		/// </summary>
 		protected void MarkAsDirty()
 		{
-			_revision++;
+			if (IncomingRef != null)
+			{
+				IncomingRef.MarkAsDirty();
+			}
 		}
 
 		/// <summary>
@@ -53,6 +52,12 @@ namespace EpicGames.Horde.Storage
 		/// <returns>Data for the node</returns>
 		public abstract void Serialize(ITreeNodeWriter writer);
 
+		/// <summary>
+		/// Enumerate all outward references from this node
+		/// </summary>
+		/// <returns>References to other nodes</returns>
+		public abstract IEnumerable<TreeNodeRef> EnumerateRefs();
+
 		/// <summary>
 		/// Static instance of the serializer for a particular <see cref="TreeNode"/> type.
 		/// </summary>
@@ -80,7 +85,16 @@ namespace EpicGames.Horde.Storage
 		/// <typeparam name="T"></typeparam>
 		/// <param name="reader">Reader to deserialize data from</param>
 		/// <returns></returns>
-		public static T Deserialize<T>(ITreeNodeReader reader) where T : TreeNode => SerializerInstance<T>.Serializer.Deserialize(reader);
+		public static T Deserialize<T>(ITreeNodeReader reader) where T : TreeNode
+		{
+			T node = SerializerInstance<T>.Serializer.Deserialize(reader);
+			foreach (TreeNodeRef nodeRef in node.EnumerateRefs())
+			{
+				Debug.Assert(nodeRef._owner == null);
+				nodeRef._owner = node;
+			}
+			return node;
+		}
 	}
 
 	/// <summary>
@@ -91,7 +105,7 @@ namespace EpicGames.Horde.Storage
 		/// <summary>
 		/// Reads a reference to another node
 		/// </summary>
-		void ReadRef(TreeNodeRef target);
+		TreeNodeRefData ReadRef();
 	}
 
 	/// <summary>
diff --git a/Engine/Source/Programs/Shared/EpicGames.Horde/Storage/TreeNodeRef.cs b/Engine/Source/Programs/Shared/EpicGames.Horde/Storage/TreeNodeRef.cs
--- a/Engine/Source/Programs/Shared/EpicGames.Horde/Storage/TreeNodeRef.cs
+++ b/Engine/Source/Programs/Shared/EpicGames.Horde/Storage/TreeNodeRef.cs
@@ -1,6 +1,7 @@
 // Copyright Epic Games, Inc. All Rights Reserved.
 
 using System;
+using System.Diagnostics;
 using System.Threading;
 using System.Threading.Tasks;
 using EpicGames.Core;
@@ -12,6 +13,11 @@ namespace EpicGames.Horde.Storage
 	/// </summary>
 	public class TreeNodeRef
 	{
+		/// <summary>
+		/// Node that owns this ref
+		/// </summary>
+		internal TreeNode? _owner;
+
 		/// <summary>
 		/// Store containing the node data. May be null for nodes in memory.
 		/// </summary>
@@ -28,23 +34,45 @@ namespace EpicGames.Horde.Storage
 		public NodeLocator Locator { get; private set; }
 
 		/// <summary>
-		/// Revision number of the target node
+		/// The target node, or null if the node is not resident in memory.
 		/// </summary>
-		private int _revision;
+		public TreeNode? Target
+		{
+			get => _target;
+			set
+			{
+				if (value != _target)
+				{
+					if (value == null)
+					{
+						_target = value;
+					}
+					else
+					{
+						if (value != null && value.IncomingRef != null)
+						{
+							throw new ArgumentException("Target node may not be part of an existing tree.");
+						}
+						if (_target != null)
+						{
+							_target.IncomingRef = null;
+						}
+						_target = value;
+						MarkAsDirty();
+					}
+				}
+			}
+		}
 
 		/// <summary>
-		/// The target node in memory
+		/// Node pointed to by this ref
 		/// </summary>
 		private TreeNode? _target;
 
 		/// <summary>
-		/// The target node, or null if the node is not resident in memory.
+		/// Whether this ref is dirty
 		/// </summary>
-		public TreeNode? Target
-		{
-			get => _target;
-			set => MarkAsDirty(value);
-		}
+		private bool _dirty;
 
 		/// <summary>
 		/// Creates a reference to a node in memory.
@@ -52,20 +80,28 @@ namespace EpicGames.Horde.Storage
 		/// <param name="target">Node to reference</param>
 		protected internal TreeNodeRef(TreeNode target)
 		{
+			Debug.Assert(target != null);
 			Target = target;
+			_dirty = true;
 		}
 
 		/// <summary>
 		/// Creates a reference to a node with the given hash
 		/// </summary>
-		/// <param name="store">Storage client containing the data</param>
+		/// <param name="owner">Node which owns the ref</param>
+		/// <param name="store">Store to fetch this node from</param>
 		/// <param name="hash">Hash of the referenced node</param>
 		/// <param name="locator">Locator for the node</param>
-		internal TreeNodeRef(IStorageClient store, IoHash hash, NodeLocator locator)
+		internal TreeNodeRef(TreeNode owner, IStorageClient store, IoHash hash, NodeLocator locator)
 		{
+			Debug.Assert(store != null);
+			_owner = owner;
+
 			Store = store;
 			Hash = hash;
 			Locator = locator;
+
+			_dirty = false;
 		}
 
 		/// <summary>
@@ -74,37 +110,40 @@ namespace EpicGames.Horde.Storage
 		/// <param name="reader"></param>
 		public TreeNodeRef(ITreeNodeReader reader)
 		{
-			reader.ReadRef(this);
+			TreeNodeRefData data = reader.ReadRef();
+			Store = data.Store;
+			Hash = data.Hash;
+			Locator = data.Locator;
+
+			_dirty = false;
+			Debug.Assert(Hash != IoHash.Zero);
 		}
 
 		/// <summary>
 		/// Determines whether the the referenced node has modified from the last version written to storage
 		/// </summary>
 		/// <returns></returns>
-		public bool IsDirty() => _target != null && _revision != _target.Revision;
+		public bool IsDirty() => _dirty;
 
 		/// <summary>
 		/// Update the reference to refer to a node in memory.
 		/// </summary>
-		/// <param name="target">The target node</param>
-		public void MarkAsDirty(TreeNode? target)
+		public void MarkAsDirty()
 		{
-			if (target == null)
+			Debug.Assert(Target != null);
+
+			Store = null;
+			Hash = default;
+			Locator = default;
+
+			if (!_dirty)
 			{
-				if (!Locator.IsValid())
+				_dirty = true;
+				if (_owner != null && _owner.IncomingRef != null)
 				{
-					throw new InvalidOperationException("Node has not been serialized to disk; cannot clear target reference.");
+					_owner.IncomingRef.MarkAsDirty();
 				}
 			}
-			else
-			{
-				Store = null;
-				Hash = default;
-				Locator = default;
-				_revision = 0;
-			}
-
-			_target = target;
 		}
 
 		/// <summary>
@@ -113,19 +152,14 @@ namespace EpicGames.Horde.Storage
 		/// <param name="store">The storage client</param>
 		/// <param name="hash">Hash of the node</param>
 		/// <param name="locator">Location of the node</param>
-		/// <param name="revision">Revision number for the node</param>
-		public bool MarkAsClean(IStorageClient store, IoHash hash, NodeLocator locator, int revision)
+		internal void MarkAsWritten(IStorageClient store, IoHash hash, NodeLocator locator)
 		{
-			bool result = false;
-			if (_target == null || _target.Revision == revision)
+			if (hash == Hash)
 			{
 				Store = store;
-				Hash = hash;
 				Locator = locator;
-				_revision = revision;
-				result = true;
+				_dirty = false;
 			}
-			return result;
 		}
 
 		/// <summary>
@@ -133,10 +167,10 @@ namespace EpicGames.Horde.Storage
 		/// </summary>
 		/// <param name="hash">Hash of the node</param>
 		/// <returns></returns>
-		public void MarkAsPendingWrite(IoHash hash)
+		internal void MarkAsPendingWrite(IoHash hash)
 		{
 			Hash = hash;
-			_revision = Target!.Revision;
+			_dirty = false;
 		}
 
 		/// <summary>
@@ -172,11 +206,12 @@ namespace EpicGames.Horde.Storage
 		/// <summary>
 		/// Constructor
 		/// </summary>
+		/// <param name="owner">Node which owns the ref</param>
 		/// <param name="store">Storage client containing the node data</param>
 		/// <param name="hash">Hash of the referenced node</param>
 		/// <param name="locator">Locator for the node</param>
-		internal TreeNodeRef(IStorageClient store, IoHash hash, NodeLocator locator) 
-			: base(store, hash, locator)
+		internal TreeNodeRef(TreeNode owner, IStorageClient store, IoHash hash, NodeLocator locator) 
+			: base(owner, store, hash, locator)
 		{
 		}
 
@@ -196,13 +231,27 @@ namespace EpicGames.Horde.Storage
 		/// <returns></returns>
 		public async ValueTask<T> ExpandAsync(CancellationToken cancellationToken = default)
 		{
-			T? result = Target;
-			if (result == null)
+			if (base.Target == null)
 			{
-				Target = await Store!.ReadNodeAsync<T>(Locator, cancellationToken);
-				result = Target;
+				base.Target = await Store!.ReadNodeAsync<T>(Locator, cancellationToken);
+				Target!.IncomingRef = this;
 			}
-			return result;
+			return Target!;
+		}
+
+		/// <summary>
+		/// Resolve this reference to a concrete node
+		/// </summary>
+		/// <param name="cancellationToken">Cancellation token for the operation</param>
+		/// <returns></returns>
+		public async ValueTask<T> ExpandCopyAsync(CancellationToken cancellationToken = default)
+		{
+			return await Store!.ReadNodeAsync<T>(Locator, cancellationToken);
 		}
 	}
+
+	/// <summary>
+	/// Deserialized ref data
+	/// </summary>
+	public record struct TreeNodeRefData(IStorageClient Store, IoHash Hash, NodeLocator Locator);
 }
diff --git a/Engine/Source/Programs/Shared/EpicGames.Horde/Storage/TreeWriter.cs b/Engine/Source/Programs/Shared/EpicGames.Horde/Storage/TreeWriter.cs
--- a/Engine/Source/Programs/Shared/EpicGames.Horde/Storage/TreeWriter.cs
+++ b/Engine/Source/Programs/Shared/EpicGames.Horde/Storage/TreeWriter.cs
@@ -41,37 +41,38 @@ namespace EpicGames.Horde.Storage
 		{
 			public readonly IoHash Hash;
 			public readonly int Length;
-			public readonly int Revision;
 			public readonly IReadOnlyList<IoHash> RefHashes;
 			public readonly TreeNodeRef NodeRef;
 
-			public NodeInfo(IoHash hash, int length, IReadOnlyList<IoHash> refs, TreeNodeRef nodeRef, int revision)
+			public NodeInfo(IoHash hash, int length, IReadOnlyList<IoHash> refs, TreeNodeRef nodeRef)
 			{
 				Hash = hash;
 				Length = length;
 				RefHashes = refs;
 				NodeRef = nodeRef;
-				Revision = revision;
 			}
 		}
 
 		class NodeWriter : ITreeNodeWriter
 		{
 			public ByteArrayBuilder Data = new ByteArrayBuilder();
-			public readonly List<TreeNodeRef> Refs = new List<TreeNodeRef>();
+
+			readonly Dictionary<TreeNodeRef, int> _refToIndex;
 
 			public int Length => Data.Length;
 
-			public void Clear()
+			public NodeWriter(List<TreeNodeRef> nextRefs)
 			{
-				Data.Clear();
-				Refs.Clear();
+				_refToIndex = new Dictionary<TreeNodeRef, int>(nextRefs.Count);
+				for (int idx = 0; idx < nextRefs.Count; idx++)
+				{
+					_refToIndex.TryAdd(nextRefs[idx], idx);
+				}
 			}
 
 			public void WriteRef(TreeNodeRef target)
 			{
-				Refs.Add(target);
-				Data.WriteUnsignedVarInt(Refs.Count - 1);
+				Data.WriteUnsignedVarInt(_refToIndex[target]);
 			}
 
 			public Memory<byte> GetMemory(int minSize = 1) => Data.GetMemory(minSize);
@@ -164,31 +165,54 @@ namespace EpicGames.Horde.Storage
 				return nodeRef.Hash;
 			}
 
-			// Serialize the node
-			NodeWriter nodeWriter = new NodeWriter();
-			nodeRef.Target!.Serialize(nodeWriter);
+			// Check we actually have a target node
+			TreeNode? node = nodeRef.Target;
+			if (node == null)
+			{
+				Debug.Assert(nodeRef.Hash != IoHash.Zero);
+				_hashToNode.TryAdd(nodeRef.Hash, nodeRef.Locator);
+				return nodeRef.Hash;
+			}
 
-			// Write all the references first
-			IoHash[] nextRefs = new IoHash[nodeWriter.Refs.Count];
-			for (int idx = 0; idx < nodeWriter.Refs.Count; idx++)
+			// Update the incoming ref
+			if (node.IncomingRef == null)
+			{
+				node.IncomingRef = nodeRef;
+			}
+			else if(node.IncomingRef != nodeRef)
 			{
-				TreeNodeRef reference = nodeWriter.Refs[idx];
-				nextRefs[idx] = await WriteAsync(reference, cancellationToken);
+				throw new InvalidOperationException("Node cannot have multiple incoming references");
 			}
 
+			// Enumerate all the outward references
+			List<TreeNodeRef> nextRefs = new List<TreeNodeRef>(node.EnumerateRefs());
+
+			// Write all the referenced nodes first
+			IoHash[] nextHashes = new IoHash[nextRefs.Count];
+			for (int idx = 0; idx < nextRefs.Count; idx++)
+			{
+				nextHashes[idx] = await WriteAsync(nextRefs[idx], cancellationToken);
+				Debug.Assert(nextHashes[idx] != IoHash.Zero);
+			}
+
+			// Serialize the node
+			NodeWriter nodeWriter = new NodeWriter(nextRefs);
+			nodeRef.Target!.Serialize(nodeWriter);
+
 			// Get the hash for the new blob
 			ReadOnlySequence<byte> data = nodeWriter.Data.AsSequence();
-			IoHash hash = ComputeHash(data, nextRefs);
+			IoHash hash = ComputeHash(data, nextHashes);
+			nodeRef.MarkAsPendingWrite(hash);
 
 			// Check if we're already tracking a node with the same hash
-			NodeInfo nodeInfo = new NodeInfo(hash, (int)data.Length, nextRefs, nodeRef, nodeRef.Target.Revision);
+			NodeInfo nodeInfo = new NodeInfo(hash, (int)data.Length, nextHashes, nodeRef);
 			if (_hashToNode.TryGetValue(hash, out NodeLocator locator))
 			{
 				// If the node is in the lookup but not currently valid, it's already queued for writing. Add this ref to the list of refs that needs fixing up,
 				// so we can update it after flushing.
 				if (locator.IsValid())
 				{
-					nodeRef.MarkAsClean(_store, hash, locator, nodeRef.Target.Revision);
+					nodeRef.MarkAsWritten(_store, hash, locator);
 				}
 				else
 				{
@@ -214,8 +238,6 @@ namespace EpicGames.Horde.Storage
 				}
 			}
 
-			// Mark the node as queued for writing
-			nodeRef.MarkAsPendingWrite(hash);
 			return hash;
 		}
 
@@ -340,7 +362,7 @@ namespace EpicGames.Horde.Storage
 					NodeLocator nodeLocator = new NodeLocator(locator, idx);
 
 					NodeInfo nodeInfo = _queue[idx];
-					nodeInfo.NodeRef.MarkAsClean(_store, nodeInfo.Hash, nodeLocator, nodeInfo.Revision);
+					nodeInfo.NodeRef.MarkAsWritten(_store, nodeInfo.Hash, nodeLocator);
 
 					_hashToNode[nodeInfo.Hash] = nodeLocator;
 				}
@@ -352,7 +374,7 @@ namespace EpicGames.Horde.Storage
 					NodeInfo nodeInfo = _secondaryQueue[refIdx];
 					if (_hashToNode.TryGetValue(nodeInfo.Hash, out NodeLocator refLocator))
 					{
-						nodeInfo.NodeRef.MarkAsClean(_store, nodeInfo.Hash, refLocator, nodeInfo.Revision);
+						nodeInfo.NodeRef.MarkAsWritten(_store, nodeInfo.Hash, refLocator);
 					}
 				}
 				_secondaryQueue.RemoveRange(0, refIdx);
