Horde: Fix static analysis warning.

diff --git a/Engine/Source/Programs/Shared/EpicGames.Horde/Storage/TreeNodeRef.cs b/Engine/Source/Programs/Shared/EpicGames.Horde/Storage/TreeNodeRef.cs
--- a/Engine/Source/Programs/Shared/EpicGames.Horde/Storage/TreeNodeRef.cs
+++ b/Engine/Source/Programs/Shared/EpicGames.Horde/Storage/TreeNodeRef.cs
@@ -49,7 +49,7 @@ namespace EpicGames.Horde.Storage
 					}
 					else
 					{
-						if (value != null && value.IncomingRef != null)
+						if (value.IncomingRef != null)
 						{
 							throw new ArgumentException("Target node may not be part of an existing tree.");
 						}
