Disable native Perforce library on Windows 11 ARM. P4.EXE runs under emulation, but in-proc library fails to load (naturally).

diff --git a/Engine/Source/Programs/Shared/EpicGames.Perforce/NativePerforceConnection.cs b/Engine/Source/Programs/Shared/EpicGames.Perforce/NativePerforceConnection.cs
--- a/Engine/Source/Programs/Shared/EpicGames.Perforce/NativePerforceConnection.cs
+++ b/Engine/Source/Programs/Shared/EpicGames.Perforce/NativePerforceConnection.cs
@@ -337,7 +337,18 @@ namespace EpicGames.Perforce
 		/// Check whether the native client is supported on the current platform
 		/// </summary>
 		/// <returns></returns>
-		public static bool IsSupported() => !(RuntimeInformation.IsOSPlatform(OSPlatform.OSX) && RuntimeInformation.OSArchitecture == Architecture.Arm64);
+		public static bool IsSupported()
+		{
+			if (RuntimeInformation.IsOSPlatform(OSPlatform.OSX))
+			{
+				return RuntimeInformation.OSArchitecture != Architecture.Arm64;
+			}
+			else if (RuntimeInformation.IsOSPlatform(OSPlatform.Windows))
+			{
+				return RuntimeInformation.OSArchitecture != Architecture.Arm64;
+			}
+			return true;
+		}
 
 		void GetNextWriteBuffer(NativeWriteBuffer nativeWriteBuffer, int minSize)
 		{
