Changes to allow compiling engine modules with Iris replication code by default, and link against modules/projects that don't depend on Iris.
ModuleRules.SetupIrisSupport defines a new macro to enable or disable the Iris fast array bindings for a module, added UHT support for this. This allows the dependency on the Iris bindings to be compiled out of modules that don't need it and still link with Iris-enabled modules.
Iris plugin disabled by default, should be enabled for specific projects that use it.

diff --git a/Engine/Plugins/Experimental/Iris/Iris.uplugin b/Engine/Plugins/Experimental/Iris/Iris.uplugin
--- a/Engine/Plugins/Experimental/Iris/Iris.uplugin
+++ b/Engine/Plugins/Experimental/Iris/Iris.uplugin
@@ -1,24 +1,24 @@
 {
-    "FileVersion" : 3,
-    "Version" : 1,
-    "VersionName" : "0.1",
-    "FriendlyName" : "Iris",
-    "Description" : "Iris networking.",
-    "Category": "Networking",
-    "CreatedBy" : "Epic Games, Inc.",
-    "CreatedByURL" : "https://epicgames.com",
-    "DocsURL" : "",
-    "MarketplaceURL" : "",
-    "SupportURL" : "",
-    "EnabledByDefault" : true,
-    "CanContainContent" : false,
-    "IsExperimentalVersion" : true,
-    "Installed" : false,
-    "Modules": [
-      {
-        "Name": "Iris",
-        "Type": "Runtime",
-        "LoadingPhase": "Default"
-      }
-    ]
+	"FileVersion": 3,
+	"Version": 1,
+	"VersionName": "0.1",
+	"FriendlyName": "Iris",
+	"Description": "Iris networking.",
+	"Category": "Networking",
+	"CreatedBy": "Epic Games, Inc.",
+	"CreatedByURL": "https://epicgames.com",
+	"DocsURL": "",
+	"MarketplaceURL": "",
+	"SupportURL": "",
+	"EnabledByDefault": false,
+	"CanContainContent": false,
+	"IsExperimentalVersion": true,
+	"Installed": false,
+	"Modules": [
+		{
+			"Name": "Iris",
+			"Type": "Runtime",
+			"LoadingPhase": "Default"
+		}
+	]
 }
diff --git a/Engine/Source/Programs/Shared/EpicGames.UHT/Exporters/CodeGen/UhtHeaderCodeGeneratorCppFile.cs b/Engine/Source/Programs/Shared/EpicGames.UHT/Exporters/CodeGen/UhtHeaderCodeGeneratorCppFile.cs
--- a/Engine/Source/Programs/Shared/EpicGames.UHT/Exporters/CodeGen/UhtHeaderCodeGeneratorCppFile.cs
+++ b/Engine/Source/Programs/Shared/EpicGames.UHT/Exporters/CodeGen/UhtHeaderCodeGeneratorCppFile.cs
@@ -597,7 +597,13 @@ namespace EpicGames.UHT.Exporters.CodeGen
 				// Inject implementation needed to support auto bindings of fast arrays
 				if (ObjectInfos[scriptStruct.ObjectTypeIndex].FastArrayProperty != null)
 				{
-					builder.Append("UE_NET_IMPLEMENT_FASTARRAY(").Append(scriptStruct.SourceName).Append(");\r\n");
+					// The preprocessor conditional is written here instead of in FastArraySerializerImplementation.h
+					// since it may evaluate differently in different modules, triggering warnings in IncludeTool.
+					builder.Append("#if defined(UE_NET_HAS_IRIS_FASTARRAY_BINDING) && UE_NET_HAS_IRIS_FASTARRAY_BINDING\r\n");
+					builder.Append("\tUE_NET_IMPLEMENT_FASTARRAY(").Append(scriptStruct.SourceName).Append(");\r\n");
+					builder.Append("#else\r\n");
+					builder.Append("\tUE_NET_IMPLEMENT_FASTARRAY_STUB(").Append(scriptStruct.SourceName).Append(");\r\n");
+					builder.Append("#endif\r\n");
 				}
 			}
 
diff --git a/Engine/Source/Programs/UnrealBuildTool/Configuration/ModuleRules.cs b/Engine/Source/Programs/UnrealBuildTool/Configuration/ModuleRules.cs
--- a/Engine/Source/Programs/UnrealBuildTool/Configuration/ModuleRules.cs
+++ b/Engine/Source/Programs/UnrealBuildTool/Configuration/ModuleRules.cs
@@ -1327,6 +1327,7 @@ namespace UnrealBuildTool
 			if (Target.bUseIris == true)
 			{
 				PublicDefinitions.Add("UE_WITH_IRIS=1");
+				PrivateDefinitions.Add("UE_NET_HAS_IRIS_FASTARRAY_BINDING=1");
 
 				if (bAddAsPublicDependency)
 				{
diff --git a/Engine/Source/Programs/UnrealHeaderTool/Private/CodeGenerator.cpp b/Engine/Source/Programs/UnrealHeaderTool/Private/CodeGenerator.cpp
--- a/Engine/Source/Programs/UnrealHeaderTool/Private/CodeGenerator.cpp
+++ b/Engine/Source/Programs/UnrealHeaderTool/Private/CodeGenerator.cpp
@@ -3636,7 +3636,13 @@ void FNativeClassHeaderGenerator::ExportGeneratedStructBodyMacros(FOutputDevice&
 		if (bGenerateFastArraySerializerTypeDefinition)
 		{
 			OutFlags |= EExportClassOutFlags::NeedsFastArrayHeaders;
-			Out.Logf(TEXT("UE_NET_IMPLEMENT_FASTARRAY(%s);\r\n"), *StructNameCPP);
+			// The preprocessor conditional is written here instead of in FastArraySerializerImplementation.h
+			// since it may evaluate differently in different modules, triggering warnings in IncludeTool.
+			Out.Logf(TEXT("#if defined(UE_NET_HAS_IRIS_FASTARRAY_BINDING) && UE_NET_HAS_IRIS_FASTARRAY_BINDING\r\n"));
+			Out.Logf(TEXT("\tUE_NET_IMPLEMENT_FASTARRAY(%s);\r\n"), *StructNameCPP);
+			Out.Logf(TEXT("#else\r\n"));
+			Out.Logf(TEXT("\tUE_NET_IMPLEMENT_FASTARRAY_STUB(%s);\r\n"), *StructNameCPP);
+			Out.Logf(TEXT("#endif\r\n"));
 		}
 	}
 
diff --git a/Engine/Source/Runtime/Engine/Private/UnrealEngine.cpp b/Engine/Source/Runtime/Engine/Private/UnrealEngine.cpp
--- a/Engine/Source/Runtime/Engine/Private/UnrealEngine.cpp
+++ b/Engine/Source/Runtime/Engine/Private/UnrealEngine.cpp
@@ -13199,6 +13199,15 @@ namespace UE::Private
 				});
 		}
 
+		// Only use Iris if the module is loaded (this happens automatically if the Iris plugin is enabled)
+		if (IrisConfig && IrisConfig->bEnableIris)
+		{
+			if (!ensureMsgf(FModuleManager::Get().IsModuleLoaded("IrisCore"), TEXT("%s is not using Iris because the IrisCore module isn't loaded. Check whether the Iris plugin is enabled."), *InNetDriverName.ToString()))
+			{
+				return false;
+			}
+		}
+
 		return IrisConfig ? IrisConfig->bEnableIris : false;
 #else
 		return false;
diff --git a/Engine/Source/Runtime/Experimental/Iris/Core/Private/Iris/IrisCoreModule.cpp b/Engine/Source/Runtime/Experimental/Iris/Core/Private/Iris/IrisCoreModule.cpp
--- a/Engine/Source/Runtime/Experimental/Iris/Core/Private/Iris/IrisCoreModule.cpp
+++ b/Engine/Source/Runtime/Experimental/Iris/Core/Private/Iris/IrisCoreModule.cpp
@@ -31,16 +31,12 @@ private:
 
 	virtual void StartupModule() override
 	{
-		// Check command line if we should disable Iris as we need to do that early
+		// Check command line for whether we should override the net.Iris.UseIrisReplication cvar, as we need to do that early
 		int32 UseIrisReplication;
 		if(FParse::Value(FCommandLine::Get(), TEXT("UseIrisReplication="), UseIrisReplication))
 		{
 			UE::Net::SetUseIrisReplication(UseIrisReplication > 0);
 		}
-		else
-		{
-			UE::Net::SetUseIrisReplication(true);
-		}
 		
 		RegisterPropertyNetSerializerSelectorTypes();
 
diff --git a/Engine/Source/Runtime/Experimental/Iris/Core/Private/Iris/ReplicationSystem/LegacyPushModel.cpp b/Engine/Source/Runtime/Experimental/Iris/Core/Private/Iris/ReplicationSystem/LegacyPushModel.cpp
--- a/Engine/Source/Runtime/Experimental/Iris/Core/Private/Iris/ReplicationSystem/LegacyPushModel.cpp
+++ b/Engine/Source/Runtime/Experimental/Iris/Core/Private/Iris/ReplicationSystem/LegacyPushModel.cpp
@@ -17,7 +17,7 @@ namespace UE::Net::Private
 {
 
 bool bIsIrisPushModelForceEnabled = false;
-int IrisPushModelMode = 1;
+int IrisPushModelMode = 2;
 static FAutoConsoleVariableRef CVarIrisPushModelMode(
 		TEXT("net.Iris.PushModelMode"),
 		IrisPushModelMode,
diff --git a/Engine/Source/Runtime/Net/Core/Classes/Net/Serialization/FastArraySerializerImplementation.h b/Engine/Source/Runtime/Net/Core/Classes/Net/Serialization/FastArraySerializerImplementation.h
--- a/Engine/Source/Runtime/Net/Core/Classes/Net/Serialization/FastArraySerializerImplementation.h
+++ b/Engine/Source/Runtime/Net/Core/Classes/Net/Serialization/FastArraySerializerImplementation.h
@@ -15,6 +15,13 @@ UE::Net::CreateAndRegisterReplicationFragmentFunc FastArrayType::GetFastArrayCre
 	return CreateFastArrayReplicationFragmentFunction; \
 }
 
+#define UE_NET_IMPLEMENT_FASTARRAY_STUB(FastArrayType) \
+UE::Net::CreateAndRegisterReplicationFragmentFunc FastArrayType::GetFastArrayCreateReplicationFragmentFunction() \
+{ \
+	return nullptr; \
+}
+
 #else
 #define UE_NET_IMPLEMENT_FASTARRAY(...)
+#define UE_NET_IMPLEMENT_FASTARRAY_STUB(...)
 #endif
