Handle Xcode version numbers more robustly

#rnx

diff --git a/Engine/Source/Programs/Shared/EpicGames.Build/Platform/Mac/ApplePlatformSDK.cs b/Engine/Source/Programs/Shared/EpicGames.Build/Platform/Mac/ApplePlatformSDK.cs
--- a/Engine/Source/Programs/Shared/EpicGames.Build/Platform/Mac/ApplePlatformSDK.cs
+++ b/Engine/Source/Programs/Shared/EpicGames.Build/Platform/Mac/ApplePlatformSDK.cs
@@ -46,12 +46,27 @@ namespace UnrealBuildBase
 
 		private static string? GetInstalledSDKVersion()
 		{
-			// get xcode version on Mac
+			// get Xcode version on Mac
 			if (OperatingSystem.IsMacOS())
 			{
-				string Output = RunLocalProcessAndReturnStdOut("sh", "-c 'xcodebuild -version'");
-				Match Result = Regex.Match(Output, @"Xcode (\S*)");
-				return Result.Success ? Result.Groups[1].Value : null;
+				int ExitCode = 0;
+				string Output = RunLocalProcessAndReturnStdOut("sh", "-c 'xcodebuild -version'", out ExitCode);
+
+				// For macOS, only an ExitCode of 0 means there was no issues running the local process
+				if (ExitCode == 0)
+				{
+					Match Result = Regex.Match(Output, @"Xcode (\S*)");
+
+					// If Xcode is not installed (or xcode-select is pointing to an invalid dir), "Output" can return random text
+					// (generally an error message with the word "Xcode " within it that can successfully Regex match), 
+					// so verify we got back an Int that is in String format.
+					if (Result.Success && TryConvertVersionToInt(Result.Groups[1].Value, out _))
+					{
+						return Result.Groups[1].Value;
+					}
+				}
+
+				return null;
 			}
 
 			if (OperatingSystem.IsWindows())
