Horde: Fix warning generating documentation due to missing argument comment.

diff --git a/Engine/Source/Programs/Shared/EpicGames.Core/MemoryReader.cs b/Engine/Source/Programs/Shared/EpicGames.Core/MemoryReader.cs
--- a/Engine/Source/Programs/Shared/EpicGames.Core/MemoryReader.cs
+++ b/Engine/Source/Programs/Shared/EpicGames.Core/MemoryReader.cs
@@ -327,6 +327,7 @@ namespace EpicGames.Core
 		/// <param name="reader">Reader to serialize from</param>
 		/// <param name="readKey">Delegate to write an individual key</param>
 		/// <param name="readValue">Delegate to write an individual value</param>
+		/// <param name="comparer">Comparer for the new dictionary</param>
 		public static Dictionary<TKey, TValue> ReadDictionary<TKey, TValue>(this IMemoryReader reader, Func<TKey> readKey, Func<TValue> readValue, IEqualityComparer<TKey>? comparer = null) where TKey : notnull
 		{
 			Dictionary<TKey, TValue> dictionary = new Dictionary<TKey, TValue>(comparer);
