[UHT]
* Fixed so SparseClassDataType code generation is not inlined. This makes it possible to forward declare the type used for sparse class data

diff --git a/Engine/Source/Programs/Shared/EpicGames.UHT/Exporters/CodeGen/UhtHeaderCodeGeneratorCppFile.cs b/Engine/Source/Programs/Shared/EpicGames.UHT/Exporters/CodeGen/UhtHeaderCodeGeneratorCppFile.cs
--- a/Engine/Source/Programs/Shared/EpicGames.UHT/Exporters/CodeGen/UhtHeaderCodeGeneratorCppFile.cs
+++ b/Engine/Source/Programs/Shared/EpicGames.UHT/Exporters/CodeGen/UhtHeaderCodeGeneratorCppFile.cs
@@ -1014,6 +1014,9 @@ namespace EpicGames.UHT.Exporters.CodeGen
 			// Add the accessors
 			AppendPropertyAccessors(builder, classObj);
 
+			// Add sparse accessors
+			AppendSparseAccessors(builder, classObj);
+
 			// Collect the callback function and sort by name to make the order stable
 			List<UhtFunction> callbackFunctions = new(classObj.Functions.Where(x => x.FunctionFlags.HasAnyFlags(EFunctionFlags.Event) && x.SuperFunction == null));
 			callbackFunctions.Sort((x, y) => StringComparerUE.OrdinalIgnoreCase.Compare(x.EngineName, y.EngineName));
@@ -1207,6 +1210,45 @@ namespace EpicGames.UHT.Exporters.CodeGen
 			return builder;
 		}
 
+		private StringBuilder AppendSparseAccessors(StringBuilder builder, UhtClass classObj)
+		{
+			string[]? sparseDataTypes = classObj.MetaData.GetStringArray(UhtNames.SparseClassDataTypes);
+			if (sparseDataTypes == null)
+			{
+				return builder;
+			}
+
+			string[]? baseSparseDataTypes = classObj.SuperClass?.MetaData.GetStringArray(UhtNames.SparseClassDataTypes);
+			if (baseSparseDataTypes != null && Enumerable.SequenceEqual(sparseDataTypes, baseSparseDataTypes))
+			{
+				return builder;
+			}
+
+			foreach (string sparseDataType in sparseDataTypes)
+			{
+				builder.Append('F').Append(sparseDataType).Append("* ").Append(classObj.SourceName).Append("::Get").Append(sparseDataType).Append("() \r\n");
+				builder.Append("{ \r\n");
+				builder.Append("\treturn static_cast<F").Append(sparseDataType).Append("*>(GetClass()->GetOrCreateSparseClassData()); \r\n");
+				builder.Append("} \r\n");
+
+				builder.Append('F').Append(sparseDataType).Append("* ").Append(classObj.SourceName).Append("::Get").Append(sparseDataType).Append("() const \r\n");
+				builder.Append("{ \r\n");
+				builder.Append("\treturn static_cast<F").Append(sparseDataType).Append("*>(GetClass()->GetOrCreateSparseClassData()); \r\n");
+				builder.Append("} \r\n");
+
+				builder.Append("const F").Append(sparseDataType).Append("* ").Append(classObj.SourceName).Append("::Get").Append(sparseDataType).Append("(EGetSparseClassDataMethod GetMethod) const \r\n");
+				builder.Append("{ \r\n");
+				builder.Append("\treturn static_cast<const F").Append(sparseDataType).Append("*>(GetClass()->GetSparseClassData(GetMethod)); \r\n");
+				builder.Append("} \r\n");
+
+				builder.Append("UScriptStruct* ").Append(classObj.SourceName).Append("::StaticGet").Append(sparseDataType).Append("ScriptStruct()\r\n");
+				builder.Append("{ \r\n");
+				builder.Append("\treturn F").Append(sparseDataType).Append("::StaticStruct(); \r\n");
+				builder.Append("} \r\n");
+			}
+			return builder;
+		}
+
 		private StringBuilder AppendInterfaceCallFunctions(StringBuilder builder, UhtClass classObj, List<UhtFunction> callbackFunctions)
 		{
 			foreach (UhtFunction function in callbackFunctions)
@@ -1512,7 +1554,7 @@ namespace EpicGames.UHT.Exporters.CodeGen
 				{
 					foreach (string sparseClass in sparseDataTypes)
 					{
-						builder.Append("\t\t\t").Append(registrationName).Append(".OuterSingleton->SetSparseClassDataStruct(F").Append(sparseClass).Append("::StaticStruct());\r\n");
+						builder.Append("\t\t\t").Append(registrationName).Append(".OuterSingleton->SetSparseClassDataStruct(").Append(classObj.SourceName).Append("::StaticGet").Append(sparseClass).Append("ScriptStruct());\r\n");
 					}
 				}
 				builder.Append("\t\t}\r\n");
diff --git a/Engine/Source/Programs/Shared/EpicGames.UHT/Exporters/CodeGen/UhtHeaderCodeGeneratorHFile.cs b/Engine/Source/Programs/Shared/EpicGames.UHT/Exporters/CodeGen/UhtHeaderCodeGeneratorHFile.cs
--- a/Engine/Source/Programs/Shared/EpicGames.UHT/Exporters/CodeGen/UhtHeaderCodeGeneratorHFile.cs
+++ b/Engine/Source/Programs/Shared/EpicGames.UHT/Exporters/CodeGen/UhtHeaderCodeGeneratorHFile.cs
@@ -537,69 +537,65 @@ namespace EpicGames.UHT.Exporters.CodeGen
 				string[]? sparseDataTypes = classObj.MetaData.GetStringArray(UhtNames.SparseClassDataTypes);
 				if (sparseDataTypes != null)
 				{
-					foreach (string sparseDataType in sparseDataTypes)
+					string[]? baseSparseDataTypes = classObj.SuperClass?.MetaData.GetStringArray(UhtNames.SparseClassDataTypes);
+					if (baseSparseDataTypes == null || !Enumerable.SequenceEqual(sparseDataTypes, baseSparseDataTypes))
 					{
-						builder.Append('F').Append(sparseDataType).Append("* Get").Append(sparseDataType).Append("() \\\r\n");
-						builder.Append("{ \\\r\n");
-						builder.Append("\treturn (F").Append(sparseDataType).Append("*)(GetClass()->GetOrCreateSparseClassData()); \\\r\n");
-						builder.Append("} \\\r\n");
-
-						builder.Append('F').Append(sparseDataType).Append("* Get").Append(sparseDataType).Append("() const \\\r\n");
-						builder.Append("{ \\\r\n");
-						builder.Append("\treturn (F").Append(sparseDataType).Append("*)(GetClass()->GetOrCreateSparseClassData()); \\\r\n");
-						builder.Append("} \\\r\n");
-
-						builder.Append("const F").Append(sparseDataType).Append("* Get").Append(sparseDataType).Append("(EGetSparseClassDataMethod GetMethod) const \\\r\n");
-						builder.Append("{ \\\r\n");
-						builder.Append("\treturn (const F").Append(sparseDataType).Append("*)(GetClass()->GetSparseClassData(GetMethod)); \\\r\n");
-						builder.Append("} \\\r\n");
-					}
+						string api = classObj.ClassFlags.HasAnyFlags(EClassFlags.MinimalAPI) ? PackageApi : "";
 
-					foreach (string sparseDataType in sparseDataTypes)
-					{
-						UhtScriptStruct? sparseScriptStruct = classObj.FindType(UhtFindOptions.EngineName | UhtFindOptions.ScriptStruct | UhtFindOptions.NoSelf, sparseDataType) as UhtScriptStruct;
-						while (sparseScriptStruct != null)
+						foreach (string sparseDataType in sparseDataTypes)
+						{
+							builder.Append(api).Append('F').Append(sparseDataType).Append("* Get").Append(sparseDataType).Append("(); \\\r\n");
+							builder.Append(api).Append('F').Append(sparseDataType).Append("* Get").Append(sparseDataType).Append("() const; \\\r\n");
+							builder.Append(api).Append("const F").Append(sparseDataType).Append("* Get").Append(sparseDataType).Append("(EGetSparseClassDataMethod GetMethod) const; \\\r\n");
+							builder.Append(api).Append("static UScriptStruct* StaticGet").Append(sparseDataType).Append("ScriptStruct(); \\\r\n");
+						}
+
+						foreach (string sparseDataType in sparseDataTypes)
 						{
-							foreach (UhtProperty sparseProperty in sparseScriptStruct.Properties)
+							UhtScriptStruct? sparseScriptStruct = classObj.FindType(UhtFindOptions.EngineName | UhtFindOptions.ScriptStruct | UhtFindOptions.NoSelf, sparseDataType) as UhtScriptStruct;
+							while (sparseScriptStruct != null)
 							{
-								if (!sparseProperty.MetaData.ContainsKey(UhtNames.NoGetter))
+								foreach (UhtProperty sparseProperty in sparseScriptStruct.Properties)
 								{
-									string propertyName = sparseProperty.SourceName;
-									string cleanPropertyName = propertyName;
-									if (sparseProperty is UhtBoolProperty && propertyName.StartsWith("b", StringComparison.Ordinal))
-									{
-										cleanPropertyName = propertyName[1..];
-									}
-
-									bool getByRef = sparseProperty.MetaData.ContainsKey(UhtNames.GetByRef);
-
-									if (getByRef)
-									{
-										builder.Append("const ").AppendSparse(sparseProperty).Append("& Get").Append(cleanPropertyName).Append("() \\\r\n");
-									}
-									else
-									{
-										builder.AppendSparse(sparseProperty).Append(" Get").Append(cleanPropertyName).Append("() \\\r\n");
-									}
-									builder.Append("{ \\\r\n");
-									builder.Append("\treturn Get").Append(sparseDataType).Append("()->").Append(propertyName).Append("; \\\r\n");
-									builder.Append("} \\\r\n");
-
-									if (getByRef)
+									if (!sparseProperty.MetaData.ContainsKey(UhtNames.NoGetter))
 									{
-										builder.Append("const ").AppendSparse(sparseProperty).Append("& Get").Append(cleanPropertyName).Append("() const \\\r\n");
+										string propertyName = sparseProperty.SourceName;
+										string cleanPropertyName = propertyName;
+										if (sparseProperty is UhtBoolProperty && propertyName.StartsWith("b", StringComparison.Ordinal))
+										{
+											cleanPropertyName = propertyName[1..];
+										}
+
+										bool getByRef = sparseProperty.MetaData.ContainsKey(UhtNames.GetByRef);
+
+										if (getByRef)
+										{
+											builder.Append("const ").AppendSparse(sparseProperty).Append("& Get").Append(cleanPropertyName).Append("() \\\r\n");
+										}
+										else
+										{
+											builder.AppendSparse(sparseProperty).Append(" Get").Append(cleanPropertyName).Append("() \\\r\n");
+										}
+										builder.Append("{ \\\r\n");
+										builder.Append("\treturn Get").Append(sparseDataType).Append("()->").Append(propertyName).Append("; \\\r\n");
+										builder.Append("} \\\r\n");
+
+										if (getByRef)
+										{
+											builder.Append("const ").AppendSparse(sparseProperty).Append("& Get").Append(cleanPropertyName).Append("() const \\\r\n");
+										}
+										else
+										{
+											builder.AppendSparse(sparseProperty).Append(" Get").Append(cleanPropertyName).Append("() const \\\r\n");
+										}
+										builder.Append("{ \\\r\n");
+										builder.Append("\treturn Get").Append(sparseDataType).Append("()->").Append(propertyName).Append("; \\\r\n");
+										builder.Append("} \\\r\n");
 									}
-									else
-									{
-										builder.AppendSparse(sparseProperty).Append(" Get").Append(cleanPropertyName).Append("() const \\\r\n");
-									}
-									builder.Append("{ \\\r\n");
-									builder.Append("\treturn Get").Append(sparseDataType).Append("()->").Append(propertyName).Append("; \\\r\n");
-									builder.Append("} \\\r\n");
 								}
-							}
 
-							sparseScriptStruct = sparseScriptStruct.SuperScriptStruct;
+								sparseScriptStruct = sparseScriptStruct.SuperScriptStruct;
+							}
 						}
 					}
 				}
diff --git a/Engine/Source/Programs/UnrealHeaderTool/Private/CodeGenerator.cpp b/Engine/Source/Programs/UnrealHeaderTool/Private/CodeGenerator.cpp
--- a/Engine/Source/Programs/UnrealHeaderTool/Private/CodeGenerator.cpp
+++ b/Engine/Source/Programs/UnrealHeaderTool/Private/CodeGenerator.cpp
@@ -2256,7 +2256,7 @@ void FNativeClassHeaderGenerator::ExportNativeGeneratedInitCode(FOutputDevice& O
 		
 		for (const FString& SparseClassDataString : SparseClassDataTypes)
 		{
-			GeneratedClassRegisterFunctionText.Logf(TEXT("\t\t\t%s->SetSparseClassDataStruct(F%s::StaticStruct());\r\n"), *OuterSingletonName, *SparseClassDataString);
+			GeneratedClassRegisterFunctionText.Logf(TEXT("\t\t\t%s->SetSparseClassDataStruct(%s::StaticGet%sScriptStruct());\r\n"), *OuterSingletonName, *ClassNameCPP, *SparseClassDataString);
 		}
 
 		GeneratedClassRegisterFunctionText.Logf(TEXT("\t\t}\r\n"));
@@ -4828,20 +4828,22 @@ void FNativeClassHeaderGenerator::ExportNativeFunctions(FOutputDevice& OutGenera
 	FString FullClassName = ClassDef.GetNameWithPrefix();
 	for (const FString& SparseClassDataString : SparseClassDataTypes)
 	{
-		RuntimeStringBuilders.AutogeneratedStaticData.Logf(TEXT("F%s* Get%s()\r\n"), *SparseClassDataString, *SparseClassDataString);
-		RuntimeStringBuilders.AutogeneratedStaticData += TEXT("{\r\n");
-		RuntimeStringBuilders.AutogeneratedStaticData.Logf(TEXT("\treturn (F%s*)(GetClass()->GetOrCreateSparseClassData());\r\n"), *SparseClassDataString);
-		RuntimeStringBuilders.AutogeneratedStaticData += TEXT("}\r\n");
+		TArray<FString> SuperSparseClassDataTypes;
+		if (FUnrealClassDefinitionInfo* SuperClass = ClassDef.GetSuperClass())
+		{
+			SuperClass->GetSparseClassDataTypes(SuperSparseClassDataTypes);
+		}
 
-		RuntimeStringBuilders.AutogeneratedStaticData.Logf(TEXT("F%s* Get%s() const\r\n"), *SparseClassDataString, *SparseClassDataString);
-		RuntimeStringBuilders.AutogeneratedStaticData += TEXT("{\r\n");
-		RuntimeStringBuilders.AutogeneratedStaticData.Logf(TEXT("\treturn (F%s*)(GetClass()->GetOrCreateSparseClassData());\r\n"), *SparseClassDataString);
-		RuntimeStringBuilders.AutogeneratedStaticData += TEXT("}\r\n");
+		if (SparseClassDataTypes != SuperSparseClassDataTypes)
+		{
+			FString Empty;
+			const FString& API = ClassDef.HasAnyClassFlags(CLASS_MinimalAPI) ? GetAPIString() : Empty;
 
-		RuntimeStringBuilders.AutogeneratedStaticData.Logf(TEXT("const F%s* Get%s(EGetSparseClassDataMethod GetMethod) const\r\n"), *SparseClassDataString, *SparseClassDataString);
-		RuntimeStringBuilders.AutogeneratedStaticData += TEXT("{\r\n");
-		RuntimeStringBuilders.AutogeneratedStaticData.Logf(TEXT("\treturn (const F%s*)(GetClass()->GetSparseClassData(GetMethod));\r\n"), *SparseClassDataString);
-		RuntimeStringBuilders.AutogeneratedStaticData += TEXT("}\r\n");
+			RuntimeStringBuilders.AutogeneratedStaticData.Logf(TEXT("%sF%s* Get%s();\r\n"), *API, *SparseClassDataString, *SparseClassDataString);
+			RuntimeStringBuilders.AutogeneratedStaticData.Logf(TEXT("%sF%s* Get%s() const;\r\n"), *API, *SparseClassDataString, *SparseClassDataString);
+			RuntimeStringBuilders.AutogeneratedStaticData.Logf(TEXT("%sconst F%s* Get%s(EGetSparseClassDataMethod GetMethod) const;\r\n"), *API, *SparseClassDataString, *SparseClassDataString);
+			RuntimeStringBuilders.AutogeneratedStaticData.Logf(TEXT("%sstatic UScriptStruct* StaticGet%sScriptStruct();\r\n"), *API, *SparseClassDataString);
+		}
 
 		FUnrealScriptStructDefinitionInfo* SparseClassDataStructDef = GTypeDefinitionInfoMap.FindByName<FUnrealScriptStructDefinitionInfo>(*SparseClassDataString);
 		while (SparseClassDataStructDef != nullptr)
@@ -5154,6 +5156,41 @@ void FNativeClassHeaderGenerator::ExportNativeFunctions(FOutputDevice& OutGenera
 
 		OutGeneratedCPPText.Log(RuntimeStringBuilders.AccessorWrapperImplementations);
 	}
+
+	if (!SparseClassDataTypes.IsEmpty())
+	{
+		TArray<FString> SuperSparseClassDataTypes;
+		if (FUnrealClassDefinitionInfo* SuperClass = ClassDef.GetSuperClass())
+		{
+			SuperClass->GetSparseClassDataTypes(SuperSparseClassDataTypes);
+		}
+
+		if (SuperSparseClassDataTypes != SparseClassDataTypes)
+		{
+			for (const FString& SparseClassDataString : SparseClassDataTypes)
+			{
+				OutGeneratedCPPText.Logf(TEXT("F%s* %s::Get%s() \r\n"), *SparseClassDataString, *ClassCPPName, *SparseClassDataString);
+				OutGeneratedCPPText.Log(TEXT("{ \r\n"));
+				OutGeneratedCPPText.Logf(TEXT("\treturn static_cast<F%s*>(GetClass()->GetOrCreateSparseClassData()); \r\n"), *SparseClassDataString);
+				OutGeneratedCPPText.Log(TEXT("} \r\n"));
+
+				OutGeneratedCPPText.Logf(TEXT("F%s* %s::Get%s() const \r\n"), *SparseClassDataString, *ClassCPPName, *SparseClassDataString);
+				OutGeneratedCPPText.Log(TEXT("{ \r\n"));
+				OutGeneratedCPPText.Logf(TEXT("\treturn static_cast<F%s*>(GetClass()->GetOrCreateSparseClassData()); \r\n"), *SparseClassDataString);
+				OutGeneratedCPPText.Log(TEXT("} \r\n"));
+
+				OutGeneratedCPPText.Logf(TEXT("const F%s* %s::Get%s(EGetSparseClassDataMethod GetMethod) const \r\n"), *SparseClassDataString, *ClassCPPName, *SparseClassDataString);
+				OutGeneratedCPPText.Log(TEXT("{ \r\n"));
+				OutGeneratedCPPText.Logf(TEXT("\treturn static_cast<const F%s*>(GetClass()->GetSparseClassData(GetMethod)); \r\n"), *SparseClassDataString);
+				OutGeneratedCPPText.Log(TEXT("} \r\n"));
+
+				OutGeneratedCPPText.Logf(TEXT("UScriptStruct* %s::StaticGet%sScriptStruct()\r\n"), *ClassCPPName, *SparseClassDataString);
+				OutGeneratedCPPText.Log(TEXT("{ \r\n"));
+				OutGeneratedCPPText.Logf(TEXT("\treturn F%s::StaticStruct(); \r\n"), *SparseClassDataString);
+				OutGeneratedCPPText.Log(TEXT("} \r\n"));
+			}
+		}
+	}
 }
 
 
