[UHT]
* Moved event parameter structs to .gen.cpp

diff --git a/Engine/Source/Programs/Shared/EpicGames.UHT/Exporters/CodeGen/UhtHeaderCodeGeneratorCppFile.cs b/Engine/Source/Programs/Shared/EpicGames.UHT/Exporters/CodeGen/UhtHeaderCodeGeneratorCppFile.cs
--- a/Engine/Source/Programs/Shared/EpicGames.UHT/Exporters/CodeGen/UhtHeaderCodeGeneratorCppFile.cs
+++ b/Engine/Source/Programs/Shared/EpicGames.UHT/Exporters/CodeGen/UhtHeaderCodeGeneratorCppFile.cs
@@ -1021,6 +1021,9 @@ namespace EpicGames.UHT.Exporters.CodeGen
 			List<UhtFunction> callbackFunctions = new(classObj.Functions.Where(x => x.FunctionFlags.HasAnyFlags(EFunctionFlags.Event) && x.SuperFunction == null));
 			callbackFunctions.Sort((x, y) => StringComparerUE.OrdinalIgnoreCase.Compare(x.EngineName, y.EngineName));
 
+			// Generate the callback parameter structures
+			AppendCallbackParametersDecls(builder, classObj, callbackFunctions);
+
 			// VM -> C++ proxies (events and delegates).
 			if (!classObj.ClassExportFlags.HasAnyFlags(UhtClassExportFlags.NoExport))
 			{
@@ -1713,6 +1716,14 @@ namespace EpicGames.UHT.Exporters.CodeGen
 			return builder;
 		}
 
+		private StringBuilder AppendCallbackParametersDecls(StringBuilder builder, UhtClass classObj, List<UhtFunction> callbackFunctions)
+		{
+			foreach (UhtFunction function in callbackFunctions)
+			{
+				AppendEventParameter(builder, function, function.StrippedFunctionName, UhtPropertyTextType.EventParameterMember, true, 1, "\r\n");
+			}
+			return builder;
+		}
 		private StringBuilder AppendCallbackFunctions(StringBuilder builder, UhtClass classObj, List<UhtFunction> callbackFunctions)
 		{
 			if (callbackFunctions.Count > 0)
diff --git a/Engine/Source/Programs/Shared/EpicGames.UHT/Exporters/CodeGen/UhtHeaderCodeGeneratorHFile.cs b/Engine/Source/Programs/Shared/EpicGames.UHT/Exporters/CodeGen/UhtHeaderCodeGeneratorHFile.cs
--- a/Engine/Source/Programs/Shared/EpicGames.UHT/Exporters/CodeGen/UhtHeaderCodeGeneratorHFile.cs
+++ b/Engine/Source/Programs/Shared/EpicGames.UHT/Exporters/CodeGen/UhtHeaderCodeGeneratorHFile.cs
@@ -363,8 +363,8 @@ namespace EpicGames.UHT.Exporters.CodeGen
 			List<UhtFunction> callbackFunctions = new(classObj.Functions.Where(x => x.FunctionFlags.HasAnyFlags(EFunctionFlags.Event) && x.SuperFunction == null));
 			callbackFunctions.Sort((x, y) => StringComparerUE.OrdinalIgnoreCase.Compare(x.EngineName, y.EngineName));
 
-			// Generate the callback parameter structures
-			AppendCallbackParametersDecls(builder, classObj, callbackFunctions);
+			// This is now done in.gen.cpp - Generate the callback parameter structures
+			//AppendCallbackParametersDecls(builder, classObj, callbackFunctions);
 
 			// Generate the RPC wrappers for the callbacks
 			AppendCallbackRpcWrapperDecls(builder, classObj, callbackFunctions);
@@ -435,10 +435,10 @@ namespace EpicGames.UHT.Exporters.CodeGen
 
 			using (UhtMacroCreator macro = new(builder, this, classObj.PrologLineNumber, PrologMacroSuffix))
 			{
-				if (callbackFunctions.Count > 0)
-				{
-					builder.Append('\t').AppendMacroName(this, classObj, EventParamsMacroSuffix).Append(" \\\r\n");
-				}
+				//if (callbackFunctions.Count > 0)
+				//{
+				//	builder.Append('\t').AppendMacroName(this, classObj, EventParamsMacroSuffix).Append(" \\\r\n");
+				//}
 			}
 
 			bool hasCallbacks = callbackFunctions.Count > 0;
diff --git a/Engine/Source/Programs/UnrealHeaderTool/Private/CodeGenerator.cpp b/Engine/Source/Programs/UnrealHeaderTool/Private/CodeGenerator.cpp
--- a/Engine/Source/Programs/UnrealHeaderTool/Private/CodeGenerator.cpp
+++ b/Engine/Source/Programs/UnrealHeaderTool/Private/CodeGenerator.cpp
@@ -2745,12 +2745,15 @@ void FNativeClassHeaderGenerator::ExportClassFromSourceFileInner(
 		// export parameters structs for all events and delegates
 		for (FUnrealFunctionDefinitionInfo* FunctionDef : CallbackFunctions)
 		{
-			ExportEventParm(UClassMacroContent, OutReferenceGatherers.ForwardDeclarations, *FunctionDef, /*Indent=*/ 1, /*bOutputConstructor=*/ true, EExportingState::Normal);
+			//ExportEventParm(UClassMacroContent, OutReferenceGatherers.ForwardDeclarations, *FunctionDef, /*Indent=*/ 1, /*bOutputConstructor=*/ true, EExportingState::Normal);
+			FUHTStringBuilder Output;
+			ExportEventParm(Output, OutReferenceGatherers.ForwardDeclarations, *FunctionDef, /*Indent=*/ 1, /*bOutputConstructor=*/ true, EExportingState::Normal);
+			OutCpp.Log(Output);
 		}
 
-		FString MacroName = SourceFile.GetGeneratedMacroName(ClassDef.GetGeneratedBodyLine(), TEXT("_EVENT_PARMS"));
-		WriteMacro(OutGeneratedHeaderText, MacroName, UClassMacroContent);
-		PrologMacroCalls.Logf(TEXT("\t%s\r\n"), *MacroName);
+		//FString MacroName = SourceFile.GetGeneratedMacroName(ClassDef.GetGeneratedBodyLine(), TEXT("_EVENT_PARMS"));
+		//WriteMacro(OutGeneratedHeaderText, MacroName, UClassMacroContent);
+		//PrologMacroCalls.Logf(TEXT("\t%s\r\n"), *MacroName);
 
 		// VM -> C++ proxies (events and delegates).
 		FOutputDeviceNull NullOutput;
