Fixing typo in diagnostic message.

#rnx

diff --git a/Engine/Source/Programs/Shared/EpicGames.UHT/Parsers/UhtHeaderFileParser.cs b/Engine/Source/Programs/Shared/EpicGames.UHT/Parsers/UhtHeaderFileParser.cs
--- a/Engine/Source/Programs/Shared/EpicGames.UHT/Parsers/UhtHeaderFileParser.cs
+++ b/Engine/Source/Programs/Shared/EpicGames.UHT/Parsers/UhtHeaderFileParser.cs
@@ -902,7 +902,7 @@ namespace EpicGames.UHT.Parsers
 				case UhtCompilerDirective.OneBlock: return "1";
 				case UhtCompilerDirective.WithHotReload: return "WITH_HOT_RELOAD";
 				case UhtCompilerDirective.WithEditor: return "WITH_EDITOR";
-				case UhtCompilerDirective.WithEditorOnlyData: return "WITH_EDITOR_DATA";
+				case UhtCompilerDirective.WithEditorOnlyData: return "WITH_EDITORONLY_DATA";
 				default: return "<unrecognized>";
 			}
 		}
