AutomationTool: ScriptModules not in the Engine directory should stil be compiled when using an installed engine

#rnx

diff --git a/Engine/Source/Programs/Shared/EpicGames.Build/System/CompileScriptModules.cs b/Engine/Source/Programs/Shared/EpicGames.Build/System/CompileScriptModules.cs
--- a/Engine/Source/Programs/Shared/EpicGames.Build/System/CompileScriptModules.cs
+++ b/Engine/Source/Programs/Shared/EpicGames.Build/System/CompileScriptModules.cs
@@ -172,16 +172,48 @@ namespace UnrealBuildBase
 			{
 				BuildFlags |= BuildFlags.NoCompile;
 			}
+			if (bUseBuildRecords)
+			{
+				BuildFlags |= BuildFlags.UseBuildRecords;
+			}
+
+			Dictionary<FileReference, CsProjBuildRecordEntry> BuildResults;
 			if (Unreal.IsEngineInstalled())
 			{
-				BuildFlags |= BuildFlags.NoCompile | BuildFlags.ErrorOnMissingTarget;
+				BuildFlags EngineBuildFlags = BuildFlags;
+				EngineBuildFlags |= BuildFlags.NoCompile | BuildFlags.ErrorOnMissingTarget;
+
+				bool bEngineBuildSuccess = false;
+				Dictionary<FileReference, CsProjBuildRecordEntry> EngineBuildResults = Build(
+					RulesFileType,
+					FoundProjects.Where(x => x.IsUnderDirectory(Unreal.EngineDirectory)).ToHashSet(),
+					BaseDirectories.Where(x => x.IsUnderDirectory(Unreal.EngineDirectory)).ToList(),
+					null,
+					EngineBuildFlags,
+					out bEngineBuildSuccess,
+					OnBuildingProjects,
+					Logger);
+
+				bool bProjectBuildSuccess = false;
+				Dictionary<FileReference, CsProjBuildRecordEntry> ProjectBuildResults = Build(
+					RulesFileType,
+					FoundProjects.Where(x => !x.IsUnderDirectory(Unreal.EngineDirectory)).ToHashSet(),
+					BaseDirectories,
+					null,
+					BuildFlags,
+					out bProjectBuildSuccess,
+					OnBuildingProjects,
+					Logger);
+
+				bBuildSuccess = bEngineBuildSuccess && bProjectBuildSuccess;
+				BuildResults = EngineBuildResults.Union(ProjectBuildResults).Distinct().ToDictionary(x => x.Key, x => x.Value);
 			}
-			if (bUseBuildRecords)
+			else
 			{
-				BuildFlags |= BuildFlags.UseBuildRecords;
+				BuildResults = Build(RulesFileType, FoundProjects, BaseDirectories, null, BuildFlags, out bBuildSuccess, OnBuildingProjects, Logger);
 			}
 
-			return GetTargetPaths(Build(RulesFileType, FoundProjects, BaseDirectories, null, BuildFlags, out bBuildSuccess, OnBuildingProjects, Logger));
+			return GetTargetPaths(BuildResults);
 		}
 
 		/// <summary>
