Horde: Remove the remaining boolean variable names with a 'b' prefix.

diff --git a/Engine/Source/Programs/Horde/Horde.Agent/Utility/TempStorage.cs b/Engine/Source/Programs/Horde/Horde.Agent/Utility/TempStorage.cs
--- a/Engine/Source/Programs/Horde/Horde.Agent/Utility/TempStorage.cs
+++ b/Engine/Source/Programs/Horde/Horde.Agent/Utility/TempStorage.cs
@@ -867,12 +867,12 @@ namespace Horde.Storage.Utility
 			}
 
 			// Check all the local files are as expected
-			bool bAllMatch = true;
+			bool allMatch = true;
 			foreach(TempStorageFile File in manifest.Files)
 			{
-				bAllMatch &= File.Compare(rootDir, logger);
+				allMatch &= File.Compare(rootDir, logger);
 			}
-			if(!bAllMatch)
+			if(!allMatch)
 			{
 				throw new TempStorageException("Files have been modified");
 			}
diff --git a/Engine/Source/Programs/Horde/Horde.Build.Tests/NotificationServiceTests.cs b/Engine/Source/Programs/Horde/Horde.Build.Tests/NotificationServiceTests.cs
--- a/Engine/Source/Programs/Horde/Horde.Build.Tests/NotificationServiceTests.cs
+++ b/Engine/Source/Programs/Horde/Horde.Build.Tests/NotificationServiceTests.cs
@@ -192,8 +192,8 @@ namespace Horde.Build.Tests
 		//				return;
 		//			}
 
-		//			bool bFireTrigger = NewLabel.State == LabelState.Complete;
-		//			INotificationTrigger? Trigger = await GetNotificationTrigger(Job.LabelIdxToTriggerId[LabelIdx], bFireTrigger);
+		//			bool fireTrigger = NewLabel.State == LabelState.Complete;
+		//			INotificationTrigger? Trigger = await GetNotificationTrigger(Job.LabelIdxToTriggerId[LabelIdx], fireTrigger);
 		//			if (Trigger == null)
 		//			{
 		//				continue;
diff --git a/Engine/Source/Programs/Horde/Horde.Build/Agents/AgentCollection.cs b/Engine/Source/Programs/Horde/Horde.Build/Agents/AgentCollection.cs
--- a/Engine/Source/Programs/Horde/Horde.Build/Agents/AgentCollection.cs
+++ b/Engine/Source/Programs/Horde/Horde.Build/Agents/AgentCollection.cs
@@ -135,11 +135,11 @@ namespace Horde.Build.Agents
 			{
 			}
 
-			public AgentDocument(AgentId id, bool bEnabled, AgentSoftwareChannelName? channel, List<PoolId> pools)
+			public AgentDocument(AgentId id, bool enabled, AgentSoftwareChannelName? channel, List<PoolId> pools)
 			{
 				Id = id;
 				Acl = new Acl();
-				Enabled = bEnabled;
+				Enabled = enabled;
 				Channel = channel;
 				Pools = pools;
 			}
@@ -165,9 +165,9 @@ namespace Horde.Build.Agents
 		}
 
 		/// <inheritdoc/>
-		public async Task<IAgent> AddAsync(AgentId id, bool bEnabled, AgentSoftwareChannelName? channel, List<PoolId>? pools)
+		public async Task<IAgent> AddAsync(AgentId id, bool enabled, AgentSoftwareChannelName? channel, List<PoolId>? pools)
 		{
-			AgentDocument agent = new AgentDocument(id, bEnabled, channel, pools ?? new List<PoolId>());
+			AgentDocument agent = new AgentDocument(id, enabled, channel, pools ?? new List<PoolId>());
 			await _agents.InsertOneAsync(agent);
 			return agent;
 		}
@@ -523,8 +523,8 @@ namespace Horde.Build.Agents
 			update = update.Unset(x => x.Leases);
 			update = update.Set(x => x.Status, AgentStatus.Stopped);
 
-			bool bDeleted = agent.Deleted || agent.Ephemeral;
-			if (bDeleted != agent.Deleted)
+			bool deleted = agent.Deleted || agent.Ephemeral;
+			if (deleted != agent.Deleted)
 			{
 				update = update.Set(x => x.Deleted, agent.Deleted);
 			}
diff --git a/Engine/Source/Programs/Horde/Horde.Build/Agents/AgentMessages.cs b/Engine/Source/Programs/Horde/Horde.Build/Agents/AgentMessages.cs
--- a/Engine/Source/Programs/Horde/Horde.Build/Agents/AgentMessages.cs
+++ b/Engine/Source/Programs/Horde/Horde.Build/Agents/AgentMessages.cs
@@ -310,7 +310,7 @@ namespace Horde.Build.Agents
 			Identifier = workspace.Identifier;
 			Stream = workspace.Stream;
 			View = workspace.View;
-			BIncremental = workspace.BIncremental;
+			BIncremental = workspace.Incremental;
 		}
 	}
 
@@ -460,8 +460,8 @@ namespace Horde.Build.Agents
 		/// <param name="agent">The agent to construct from</param>
 		/// <param name="leases">Active leases</param>
 		/// <param name="rate">Rate for this agent</param>
-		/// <param name="bIncludeAcl">Whether to include the ACL in the response</param>
-		public GetAgentResponse(IAgent agent, List<GetAgentLeaseResponse> leases, double? rate, bool bIncludeAcl)
+		/// <param name="includeAcl">Whether to include the ACL in the response</param>
+		public GetAgentResponse(IAgent agent, List<GetAgentLeaseResponse> leases, double? rate, bool includeAcl)
 		{
 			Id = agent.Id.ToString();
 			Name = agent.Id.ToString();
@@ -492,7 +492,7 @@ namespace Horde.Build.Agents
 			Workspaces = agent.Workspaces.ConvertAll(x => new GetAgentWorkspaceResponse(x));
 			Capabilities = new { Devices = new[] { new { agent.Properties, agent.Resources } } };
 			Leases = leases;
-			Acl = (bIncludeAcl && agent.Acl != null) ? new GetAclResponse(agent.Acl) : null;
+			Acl = (includeAcl && agent.Acl != null) ? new GetAclResponse(agent.Acl) : null;
 			Comment = agent.Comment;
 		}
 	}
diff --git a/Engine/Source/Programs/Horde/Horde.Build/Agents/AgentService.cs b/Engine/Source/Programs/Horde/Horde.Build/Agents/AgentService.cs
--- a/Engine/Source/Programs/Horde/Horde.Build/Agents/AgentService.cs
+++ b/Engine/Source/Programs/Horde/Horde.Build/Agents/AgentService.cs
@@ -180,13 +180,13 @@ namespace Horde.Build.Agents
 		/// Register a new agent
 		/// </summary>
 		/// <param name="name">Name of the agent</param>
-		/// <param name="bEnabled">Whether the agent is currently enabled</param>
+		/// <param name="enabled">Whether the agent is currently enabled</param>
 		/// <param name="channel">Override for the desired software version</param>
 		/// <param name="pools">Pools for this agent</param>
 		/// <returns>Unique id for the agent</returns>
-		public Task<IAgent> CreateAgentAsync(string name, bool bEnabled, AgentSoftwareChannelName? channel, List<PoolId>? pools)
+		public Task<IAgent> CreateAgentAsync(string name, bool enabled, AgentSoftwareChannelName? channel, List<PoolId>? pools)
 		{
-			return Agents.AddAsync(new AgentId(name), bEnabled, channel, pools);
+			return Agents.AddAsync(new AgentId(name), enabled, channel, pools);
 		}
 
 		/// <summary>
@@ -217,11 +217,11 @@ namespace Horde.Build.Agents
 		/// </summary>
 		/// <param name="agent">The agent to update</param>
 		/// <param name="workspaces">Current list of workspaces</param>
-		/// <param name="bPendingConform">Whether the agent still needs to run another conform</param>
+		/// <param name="pendingConform">Whether the agent still needs to run another conform</param>
 		/// <returns>New agent state</returns>
-		public async Task<bool> TryUpdateWorkspacesAsync(IAgent agent, List<AgentWorkspace> workspaces, bool bPendingConform)
+		public async Task<bool> TryUpdateWorkspacesAsync(IAgent agent, List<AgentWorkspace> workspaces, bool pendingConform)
 		{
-			IAgent? newAgent = await Agents.TryUpdateWorkspacesAsync(agent, workspaces, bPendingConform);
+			IAgent? newAgent = await Agents.TryUpdateWorkspacesAsync(agent, workspaces, pendingConform);
 			return newAgent != null;
 		}
 
@@ -649,7 +649,7 @@ namespace Horde.Build.Agents
 				}
 
 				// Flag for whether the leases array should be updated
-				bool bUpdateLeases = false;
+				bool updateLeases = false;
 				List<AgentLease> leases = new List<AgentLease>(agent.Leases);
 
 				// Remove any completed leases from the agent
@@ -664,7 +664,7 @@ namespace Horde.Build.Agents
 						{
 							await RemoveLeaseAsync(agent, lease, utcNow, LeaseOutcome.Cancelled, null);
 							leases.RemoveAt(idx--);
-							bUpdateLeases = true;
+							updateLeases = true;
 						}
 					}
 					else
@@ -681,7 +681,7 @@ namespace Horde.Build.Agents
 							{
 								lease.State = LeaseState.Active;
 							}
-							bUpdateLeases = true;
+							updateLeases = true;
 						}
 					}
 				}
@@ -694,7 +694,7 @@ namespace Horde.Build.Agents
 						if (lease.State != LeaseState.Cancelled)
 						{
 							lease.State = LeaseState.Cancelled;
-							bUpdateLeases = true;
+							updateLeases = true;
 						}
 					}
 				}
@@ -703,7 +703,7 @@ namespace Horde.Build.Agents
 				List<PoolId> dynamicPools = await GetDynamicPoolsAsync(agent);
 
 				// Update the agent, and try to create new lease documents if we succeed
-				IAgent? newAgent = await Agents.TryUpdateSessionAsync(agent, status, sessionExpiresAt, properties, resources, dynamicPools, bUpdateLeases ? leases : null);
+				IAgent? newAgent = await Agents.TryUpdateSessionAsync(agent, status, sessionExpiresAt, properties, resources, dynamicPools, updateLeases ? leases : null);
 				if (newAgent != null)
 				{
 					agent = newAgent;
diff --git a/Engine/Source/Programs/Horde/Horde.Build/Agents/AgentsController.cs b/Engine/Source/Programs/Horde/Horde.Build/Agents/AgentsController.cs
--- a/Engine/Source/Programs/Horde/Horde.Build/Agents/AgentsController.cs
+++ b/Engine/Source/Programs/Horde/Horde.Build/Agents/AgentsController.cs
@@ -128,8 +128,8 @@ namespace Horde.Build.Agents
 				leases.Add(new GetAgentLeaseResponse(lease, details));
 			}
 
-			bool bIncludeAcl = await _agentService.AuthorizeAsync(agent, AclAction.ViewPermissions, User, permissionsCache);
-			return new GetAgentResponse(agent, leases, rate, bIncludeAcl).ApplyFilter(filter);
+			bool includeAcl = await _agentService.AuthorizeAsync(agent, AclAction.ViewPermissions, User, permissionsCache);
+			return new GetAgentResponse(agent, leases, rate, includeAcl).ApplyFilter(filter);
 		}
 
 		/// <summary>
diff --git a/Engine/Source/Programs/Horde/Horde.Build/Agents/Fleet/Providers/AwsFleetManager.cs b/Engine/Source/Programs/Horde/Horde.Build/Agents/Fleet/Providers/AwsFleetManager.cs
--- a/Engine/Source/Programs/Horde/Horde.Build/Agents/Fleet/Providers/AwsFleetManager.cs
+++ b/Engine/Source/Programs/Horde/Horde.Build/Agents/Fleet/Providers/AwsFleetManager.cs
@@ -249,7 +249,7 @@ namespace Horde.Build.Agents.Fleet.Providers
 			foreach (IAgent agent in agentsLimitedByCount)
 			{
 				IAuditLogChannel<AgentId> agentLogger = agentCollection.GetLogger(agent.Id);
-				if (await agentCollection.TryUpdateSettingsAsync(agent, bRequestShutdown: true, shutdownReason: "Autoscaler") != null)
+				if (await agentCollection.TryUpdateSettingsAsync(agent, requestShutdown: true, shutdownReason: "Autoscaler") != null)
 				{
 					agentLogger.LogInformation("Marked {AgentId} in pool {PoolName} for shutdown due to autoscaling (currently {NumLeases} leases outstanding)", agent.Id, pool.Name, agent.Leases.Count);
 				}
diff --git a/Engine/Source/Programs/Horde/Horde.Build/Agents/IAgent.cs b/Engine/Source/Programs/Horde/Horde.Build/Agents/IAgent.cs
--- a/Engine/Source/Programs/Horde/Horde.Build/Agents/IAgent.cs
+++ b/Engine/Source/Programs/Horde/Horde.Build/Agents/IAgent.cs
@@ -66,7 +66,7 @@ namespace Horde.Build.Agents
 		/// <summary>
 		/// Whether to use an incremental workspace
 		/// </summary>
-		public bool BIncremental { get; set; }
+		public bool Incremental { get; set; }
 
 		/// <summary>
 		/// Constructor
@@ -76,8 +76,8 @@ namespace Horde.Build.Agents
 		/// <param name="identifier">Identifier to distinguish this workspace from other workspaces</param>
 		/// <param name="stream">The stream to sync</param>
 		/// <param name="view">Custom view for the workspace</param>
-		/// <param name="bIncremental">Whether to use an incremental workspace</param>
-		public AgentWorkspace(string? cluster, string? userName, string identifier, string stream, List<string>? view, bool bIncremental)
+		/// <param name="incremental">Whether to use an incremental workspace</param>
+		public AgentWorkspace(string? cluster, string? userName, string identifier, string stream, List<string>? view, bool incremental)
 		{
 			if (!String.IsNullOrEmpty(cluster))
 			{
@@ -90,7 +90,7 @@ namespace Horde.Build.Agents
 			Identifier = identifier;
 			Stream = stream;
 			View = view;
-			BIncremental = bIncremental;
+			Incremental = incremental;
 		}
 
 		/// <summary>
@@ -125,7 +125,7 @@ namespace Horde.Build.Agents
 			{
 				return false;
 			}
-			if (Cluster != other.Cluster || UserName != other.UserName || Identifier != other.Identifier || Stream != other.Stream || BIncremental != other.BIncremental)
+			if (Cluster != other.Cluster || UserName != other.UserName || Identifier != other.Identifier || Stream != other.Stream || Incremental != other.Incremental)
 			{
 				return false;
 			}
@@ -139,7 +139,7 @@ namespace Horde.Build.Agents
 		/// <inheritdoc/>
 		public override int GetHashCode()
 		{
-			return HashCode.Combine(Cluster, UserName, Identifier, Stream, BIncremental); // Ignore 'View' for now
+			return HashCode.Combine(Cluster, UserName, Identifier, Stream, Incremental); // Ignore 'View' for now
 		}
 
 		/// <summary>
@@ -175,7 +175,7 @@ namespace Horde.Build.Agents
 			{
 				result.View.AddRange(View);
 			}
-			result.Incremental = BIncremental;
+			result.Incremental = Incremental;
 			return result;
 		}
 	}
@@ -766,7 +766,7 @@ namespace Horde.Build.Agents
 			{
 				result.View.AddRange(workspace.View);
 			}
-			result.Incremental = workspace.BIncremental;
+			result.Incremental = workspace.Incremental;
 			workspaceMessages.Add(result);
 
 			return true;
diff --git a/Engine/Source/Programs/Horde/Horde.Build/Agents/IAgentCollection.cs b/Engine/Source/Programs/Horde/Horde.Build/Agents/IAgentCollection.cs
--- a/Engine/Source/Programs/Horde/Horde.Build/Agents/IAgentCollection.cs
+++ b/Engine/Source/Programs/Horde/Horde.Build/Agents/IAgentCollection.cs
@@ -26,10 +26,10 @@ namespace Horde.Build.Agents
 		/// Adds a new agent with the given properties
 		/// </summary>
 		/// <param name="id">Id for the new agent</param>
-		/// <param name="bEnabled">Whether the agent is enabled or not</param>
+		/// <param name="enabled">Whether the agent is enabled or not</param>
 		/// <param name="channel">Channel to use for software run by this agent</param>
 		/// <param name="pools">Pools for the agent</param>
-		Task<IAgent> AddAsync(AgentId id, bool bEnabled, AgentSoftwareChannelName? channel = null, List<PoolId>? pools = null);
+		Task<IAgent> AddAsync(AgentId id, bool enabled, AgentSoftwareChannelName? channel = null, List<PoolId>? pools = null);
 
 		/// <summary>
 		/// Deletes an agent
@@ -76,18 +76,18 @@ namespace Horde.Build.Agents
 		/// Update an agent's settings
 		/// </summary>
 		/// <param name="agent">Agent instance</param>
-		/// <param name="bEnabled">Whether the agent is enabled or not</param>
-		/// <param name="bRequestConform">Whether to request a conform job be run</param>
-		/// <param name="bRequestFullConform">Whether to request a full conform job be run</param>
-		/// <param name="bRequestRestart">Whether to request the machine be restarted</param>
-		/// <param name="bRequestShutdown">Whether to request the machine be shut down</param>
+		/// <param name="enabled">Whether the agent is enabled or not</param>
+		/// <param name="requestConform">Whether to request a conform job be run</param>
+		/// <param name="requestFullConform">Whether to request a full conform job be run</param>
+		/// <param name="requestRestart">Whether to request the machine be restarted</param>
+		/// <param name="requestShutdown">Whether to request the machine be shut down</param>
 		/// <param name="shutdownReason">The reason for shutting down agent, ex. Autoscaler/Manual/Unexpected</param>
 		/// <param name="channel">Override for the desired software channel</param>
 		/// <param name="pools">List of pools for the agent</param>
 		/// <param name="acl">New ACL for this agent</param>
 		/// <param name="comment">New comment</param>
 		/// <returns>Version of the software that needs to be installed on the agent. Null if the agent is running the correct version.</returns>
-		Task<IAgent?> TryUpdateSettingsAsync(IAgent agent, bool? bEnabled = null, bool? bRequestConform = null, bool? bRequestFullConform = null, bool? bRequestRestart = null, bool? bRequestShutdown = null, string? shutdownReason = null, AgentSoftwareChannelName? channel = null, List<PoolId>? pools = null, Acl? acl = null, string? comment = null);
+		Task<IAgent?> TryUpdateSettingsAsync(IAgent agent, bool? enabled = null, bool? requestConform = null, bool? requestFullConform = null, bool? requestRestart = null, bool? requestShutdown = null, string? shutdownReason = null, AgentSoftwareChannelName? channel = null, List<PoolId>? pools = null, Acl? acl = null, string? comment = null);
 
 		/// <summary>
 		/// Update the current workspaces for an agent.
diff --git a/Engine/Source/Programs/Horde/Horde.Build/Agents/Pools/PoolService.cs b/Engine/Source/Programs/Horde/Horde.Build/Agents/Pools/PoolService.cs
--- a/Engine/Source/Programs/Horde/Horde.Build/Agents/Pools/PoolService.cs
+++ b/Engine/Source/Programs/Horde/Horde.Build/Agents/Pools/PoolService.cs
@@ -249,7 +249,7 @@ namespace Horde.Build.Agents.Pools
 		/// <returns>List of workspaces</returns>
 		public async Task<HashSet<AgentWorkspace>> GetWorkspacesAsync(IAgent agent, DateTime validAtTime)
 		{
-			bool bAddAutoSdkWorkspace = false;
+			bool addAutoSdkWorkspace = false;
 			HashSet<AgentWorkspace> workspaces = new HashSet<AgentWorkspace>();
 
 			Dictionary<PoolId, IPool> poolMapping = await GetPoolLookupAsync(validAtTime);
@@ -259,11 +259,11 @@ namespace Horde.Build.Agents.Pools
 				if (poolMapping.TryGetValue(poolId, out pool))
 				{
 					workspaces.UnionWith(pool.Workspaces);
-					bAddAutoSdkWorkspace |= pool.UseAutoSdk;
+					addAutoSdkWorkspace |= pool.UseAutoSdk;
 				}
 			}
 
-			if (bAddAutoSdkWorkspace)
+			if (addAutoSdkWorkspace)
 			{
 				IGlobals globals = await _globalsService.GetAsync();
 				workspaces.UnionWith(agent.GetAutoSdkWorkspaces(globals, workspaces.ToList()));
diff --git a/Engine/Source/Programs/Horde/Horde.Build/Agents/Pools/PoolUpdateService.cs b/Engine/Source/Programs/Horde/Horde.Build/Agents/Pools/PoolUpdateService.cs
--- a/Engine/Source/Programs/Horde/Horde.Build/Agents/Pools/PoolUpdateService.cs
+++ b/Engine/Source/Programs/Horde/Horde.Build/Agents/Pools/PoolUpdateService.cs
@@ -56,13 +56,13 @@ namespace Horde.Build.Agents.Pools
 			DateTime startTime = DateTime.UtcNow;
 
 			// Update the list
-			bool bRetryUpdate = true;
-			while (bRetryUpdate && !stoppingToken.IsCancellationRequested)
+			bool retryUpdate = true;
+			while (retryUpdate && !stoppingToken.IsCancellationRequested)
 			{
 				_logger.LogDebug("Updating pool->workspace map");
 
 				// Assume this will be the last iteration
-				bRetryUpdate = false;
+				retryUpdate = false;
 
 				// Capture the list of pools at the start of this update
 				List<IPool> currentPools = await _pools.GetAsync();
@@ -125,7 +125,7 @@ namespace Horde.Build.Agents.Pools
 						if (result == null)
 						{
 							_logger.LogInformation("Pool modified; will retry");
-							bRetryUpdate = true;
+							retryUpdate = true;
 						}
 					}
 				}
diff --git a/Engine/Source/Programs/Horde/Horde.Build/Issues/IssueService.cs b/Engine/Source/Programs/Horde/Horde.Build/Issues/IssueService.cs
--- a/Engine/Source/Programs/Horde/Horde.Build/Issues/IssueService.cs
+++ b/Engine/Source/Programs/Horde/Horde.Build/Issues/IssueService.cs
@@ -367,7 +367,7 @@ namespace Horde.Build.Issues
 		{
 			_ = issue;
 
-			bool bShowDesktopAlerts = false;
+			bool showDesktopAlerts = false;
 
 			HashSet<(StreamId, TemplateId)> checkedTemplates = new HashSet<(StreamId, TemplateId)>();
 			foreach(IIssueSpan span in spans)
@@ -377,13 +377,13 @@ namespace Horde.Build.Issues
 					HashSet<TemplateId>? templates;
 					if (_cachedDesktopAlerts.TryGetValue(span.StreamId, out templates) && templates.Contains(span.TemplateRefId))
 					{
-						bShowDesktopAlerts = true;
+						showDesktopAlerts = true;
 						break;
 					}
 				}
 			}
 
-			return bShowDesktopAlerts;
+			return showDesktopAlerts;
 		}
 
 		/// <summary>
@@ -1102,19 +1102,19 @@ namespace Horde.Build.Issues
 		/// </summary>
 		async ValueTask<bool> ContainsFixChange(StreamId streamId, int fixChange, Dictionary<(StreamId, int), bool> cachedContainsFixChange)
 		{
-			bool bContainsFixChange;
-			if (!cachedContainsFixChange.TryGetValue((streamId, fixChange), out bContainsFixChange) && fixChange > 0)
+			bool containsFixChange;
+			if (!cachedContainsFixChange.TryGetValue((streamId, fixChange), out containsFixChange) && fixChange > 0)
 			{
 				IStream? stream = await _streams.GetCachedStream(streamId);
 				if (stream != null)
 				{
 					_logger.LogInformation("Querying fix changelist {FixChange} in {StreamId}", fixChange, streamId);
 					ICommit? change = await _commitService.GetCollection(stream).FindAsync(fixChange, fixChange, 1).FirstOrDefaultAsync();
-					bContainsFixChange = change != null;
+					containsFixChange = change != null;
 				}
-				cachedContainsFixChange[(streamId, fixChange)] = bContainsFixChange;
+				cachedContainsFixChange[(streamId, fixChange)] = containsFixChange;
 			}
-			return bContainsFixChange;
+			return containsFixChange;
 		}
 
 		/// <summary>
diff --git a/Engine/Source/Programs/Horde/Horde.Build/Issues/IssuesController.cs b/Engine/Source/Programs/Horde/Horde.Build/Issues/IssuesController.cs
--- a/Engine/Source/Programs/Horde/Horde.Build/Issues/IssuesController.cs
+++ b/Engine/Source/Programs/Horde/Horde.Build/Issues/IssuesController.cs
@@ -263,8 +263,8 @@ namespace Horde.Build.Issues
 				IIssueDetails details = await _issueService.GetIssueDetailsAsync(issue);
 				if (await AuthorizeIssue(details, permissionsCache))
 				{
-					bool bShowDesktopAlerts = _issueService.ShowDesktopAlertsForIssue(issue, details.Spans);
-					GetIssueResponse response = await CreateIssueResponseAsync(details, bShowDesktopAlerts);
+					bool showDesktopAlerts = _issueService.ShowDesktopAlertsForIssue(issue, details.Spans);
+					GetIssueResponse response = await CreateIssueResponseAsync(details, showDesktopAlerts);
 					responses.Add(PropertyFilter.Apply(response, filter));
 				}
 			}
@@ -292,8 +292,8 @@ namespace Horde.Build.Issues
 				return Forbid();
 			}
 
-			bool bShowDesktopAlerts = _issueService.ShowDesktopAlertsForIssue(details.Issue, details.Spans);
-			return PropertyFilter.Apply(await CreateIssueResponseAsync(details, bShowDesktopAlerts), filter);
+			bool showDesktopAlerts = _issueService.ShowDesktopAlertsForIssue(details.Issue, details.Spans);
+			return PropertyFilter.Apply(await CreateIssueResponseAsync(details, showDesktopAlerts), filter);
 		}
 
 		/// <summary>
diff --git a/Engine/Source/Programs/Horde/Horde.Build/Jobs/Graphs/GraphCollection.cs b/Engine/Source/Programs/Horde/Horde.Build/Jobs/Graphs/GraphCollection.cs
--- a/Engine/Source/Programs/Horde/Horde.Build/Jobs/Graphs/GraphCollection.cs
+++ b/Engine/Source/Programs/Horde/Horde.Build/Jobs/Graphs/GraphCollection.cs
@@ -242,9 +242,9 @@ namespace Horde.Build.Jobs.Graphs
 							int nodeIdx = nodes.Count;
 
 							Priority priority = newNodeRequest.Priority ?? Priority.Normal;
-							bool bAllowRetry = newNodeRequest.AllowRetry ?? true;
-							bool bRunEarly = newNodeRequest.RunEarly ?? false;
-							bool bWarnings = newNodeRequest.Warnings ?? true;
+							bool allowRetry = newNodeRequest.AllowRetry ?? true;
+							bool runEarly = newNodeRequest.RunEarly ?? false;
+							bool warnings = newNodeRequest.Warnings ?? true;
 
 							NodeOutputRef[]? inputs = null;
 							if (newNodeRequest.Inputs != null && newNodeRequest.Inputs.Count > 0)
@@ -255,7 +255,7 @@ namespace Horde.Build.Jobs.Graphs
 							NodeRef[] inputDependencies = (newNodeRequest.InputDependencies == null) ? Array.Empty<NodeRef>() : newNodeRequest.InputDependencies.Select(x => nodeNameToRef[x]).ToArray();
 							NodeRef[] orderDependencies = (newNodeRequest.OrderDependencies == null) ? Array.Empty<NodeRef>() : newNodeRequest.OrderDependencies.Select(x => nodeNameToRef[x]).ToArray();
 							orderDependencies = orderDependencies.Union(inputDependencies).ToArray();
-							nodes.Add(new Node(newNodeRequest.Name, inputs, newNodeRequest.Outputs?.ToArray(), inputDependencies, orderDependencies, priority, bAllowRetry, bRunEarly, bWarnings, newNodeRequest.Credentials, newNodeRequest.Properties, newNodeRequest.Annotations));
+							nodes.Add(new Node(newNodeRequest.Name, inputs, newNodeRequest.Outputs?.ToArray(), inputDependencies, orderDependencies, priority, allowRetry, runEarly, warnings, newNodeRequest.Credentials, newNodeRequest.Properties, newNodeRequest.Annotations));
 
 							NodeRef nodeRef = new NodeRef(newGroups.Count, nodeIdx);
 							nodeNameToRef.Add(newNodeRequest.Name, nodeRef);
diff --git a/Engine/Source/Programs/Horde/Horde.Build/Jobs/IJob.cs b/Engine/Source/Programs/Horde/Horde.Build/Jobs/IJob.cs
--- a/Engine/Source/Programs/Horde/Horde.Build/Jobs/IJob.cs
+++ b/Engine/Source/Programs/Horde/Horde.Build/Jobs/IJob.cs
@@ -770,7 +770,7 @@ namespace Horde.Build.Jobs
 		/// <returns>Job state</returns>
 		public static JobState GetState(this IJob job)
 		{
-			bool bWaiting = false;
+			bool waiting = false;
 			foreach (IJobStepBatch batch in job.Batches)
 			{
 				foreach (IJobStep step in batch.Steps)
@@ -787,12 +787,12 @@ namespace Horde.Build.Jobs
 						}
 						else
 						{
-							bWaiting = true;
+							waiting = true;
 						}
 					}
 				}
 			}
-			return bWaiting ? JobState.Waiting : JobState.Complete;
+			return waiting ? JobState.Waiting : JobState.Complete;
 		}
 
 		/// <summary>
@@ -863,20 +863,20 @@ namespace Horde.Build.Jobs
 		/// <returns>The step outcome</returns>
 		public static (JobStepState, JobStepOutcome) GetTargetState(IEnumerable<IJobStep> steps)
 		{
-			bool bAnySkipped = false;
-			bool bAnyWarnings = false;
-			bool bAnyFailed = false;
-			bool bAnyPending = false;
+			bool anySkipped = false;
+			bool anyWarnings = false;
+			bool anyFailed = false;
+			bool anyPending = false;
 			foreach (IJobStep step in steps)
 			{
-				bAnyPending |= step.IsPending();
-				bAnySkipped |= step.State == JobStepState.Aborted || step.State == JobStepState.Skipped;
-				bAnyFailed |= (step.Outcome == JobStepOutcome.Failure);
-				bAnyWarnings |= (step.Outcome == JobStepOutcome.Warnings);
+				anyPending |= step.IsPending();
+				anySkipped |= step.State == JobStepState.Aborted || step.State == JobStepState.Skipped;
+				anyFailed |= (step.Outcome == JobStepOutcome.Failure);
+				anyWarnings |= (step.Outcome == JobStepOutcome.Warnings);
 			}
 
-			JobStepState newState = bAnyPending ? JobStepState.Running : JobStepState.Completed;
-			JobStepOutcome newOutcome = bAnyFailed ? JobStepOutcome.Failure : bAnyWarnings ? JobStepOutcome.Warnings : bAnySkipped ? JobStepOutcome.Unspecified : JobStepOutcome.Success;
+			JobStepState newState = anyPending ? JobStepState.Running : JobStepState.Completed;
+			JobStepOutcome newOutcome = anyFailed ? JobStepOutcome.Failure : anyWarnings ? JobStepOutcome.Warnings : anySkipped ? JobStepOutcome.Unspecified : JobStepOutcome.Success;
 			return (newState, newOutcome);
 		}
 
@@ -1312,25 +1312,25 @@ namespace Horde.Build.Jobs
 				if (label.RequiredNodes.Any(x => stepForNodeRef.ContainsKey(x)))
 				{
 					// Combine the state of the steps contributing towards this label
-					bool bAnySkipped = false;
-					bool bAnyWarnings = false;
-					bool bAnyFailed = false;
-					bool bAnyPending = false;
+					bool anySkipped = false;
+					bool anyWarnings = false;
+					bool anyFailed = false;
+					bool anyPending = false;
 					foreach (NodeRef includedNode in label.IncludedNodes)
 					{
 						IJobStep? step;
 						if (stepForNodeRef.TryGetValue(includedNode, out step))
 						{
-							bAnyPending |= step.IsPending();
-							bAnySkipped |= step.State == JobStepState.Aborted || step.State == JobStepState.Skipped;
-							bAnyFailed |= (step.Outcome == JobStepOutcome.Failure);
-							bAnyWarnings |= (step.Outcome == JobStepOutcome.Warnings);
+							anyPending |= step.IsPending();
+							anySkipped |= step.State == JobStepState.Aborted || step.State == JobStepState.Skipped;
+							anyFailed |= (step.Outcome == JobStepOutcome.Failure);
+							anyWarnings |= (step.Outcome == JobStepOutcome.Warnings);
 						}
 					}
 
 					// Figure out the overall label state
-					newState = bAnyPending ? LabelState.Running : LabelState.Complete;
-					newOutcome = bAnyFailed ? LabelOutcome.Failure : bAnyWarnings ? LabelOutcome.Warnings : bAnySkipped? LabelOutcome.Unspecified : LabelOutcome.Success;
+					newState = anyPending ? LabelState.Running : LabelState.Complete;
+					newOutcome = anyFailed ? LabelOutcome.Failure : anyWarnings ? LabelOutcome.Warnings : anySkipped? LabelOutcome.Unspecified : LabelOutcome.Success;
 				}
 
 				states.Add((newState, newOutcome));
diff --git a/Engine/Source/Programs/Horde/Horde.Build/Jobs/JobCollection.cs b/Engine/Source/Programs/Horde/Horde.Build/Jobs/JobCollection.cs
--- a/Engine/Source/Programs/Horde/Horde.Build/Jobs/JobCollection.cs
+++ b/Engine/Source/Programs/Horde/Horde.Build/Jobs/JobCollection.cs
@@ -561,7 +561,7 @@ namespace Horde.Build.Jobs
 			List<UpdateDefinition<JobDocument>> updates = new List<UpdateDefinition<JobDocument>>();
 
 			// Flag for whether to update batches
-			bool bUpdateBatches = false;
+			bool updateBatches = false;
 
 			// Build the update list
 			JobDocument jobDocument = Clone((JobDocument)inJob);
@@ -594,7 +594,7 @@ namespace Horde.Build.Jobs
 			{
 				jobDocument.AbortedByUserId = abortedByUserId;
 				updates.Add(updateBuilder.Set(x => x.AbortedByUserId, jobDocument.AbortedByUserId));
-				bUpdateBatches = true;
+				updateBatches = true;
 			}
 			if (notificationTriggerId != null)
 			{
@@ -639,7 +639,7 @@ namespace Horde.Build.Jobs
 				{
 					if (modifiedArgument.StartsWith(IJob.TargetArgumentPrefix, StringComparison.OrdinalIgnoreCase))
 					{
-						bUpdateBatches = true;
+						updateBatches = true;
 					}
 				}
 
@@ -648,7 +648,7 @@ namespace Horde.Build.Jobs
 			}
 
 			// Update the batches
-			if (bUpdateBatches)
+			if (updateBatches)
 			{
 				UpdateBatches(jobDocument, graph, updates, _logger);
 			}
@@ -711,7 +711,7 @@ namespace Horde.Build.Jobs
 				List<NodeRef> retriedNodes = jobDocument.RetriedNodes ?? new List<NodeRef>();
 
 				// Check if there are any steps that need to be run again
-				bool bUpdateState = false;
+				bool updateState = false;
 				foreach (JobStepDocument step in batch.Steps)
 				{
 					if (step.State == JobStepState.Running)
@@ -726,7 +726,7 @@ namespace Horde.Build.Jobs
 							retriedNodes.Add(new NodeRef(batch.GroupIdx, step.NodeIdx));
 						}
 
-						bUpdateState = true;
+						updateState = true;
 					}
 					else if (step.State == JobStepState.Ready || step.State == JobStepState.Waiting)
 					{
@@ -738,12 +738,12 @@ namespace Horde.Build.Jobs
 						{
 							step.State = JobStepState.Skipped;
 						}
-						bUpdateState = true;
+						updateState = true;
 					}
 				}
 
 				// Update the steps
-				if (bUpdateState)
+				if (updateState)
 				{
 					updates.Clear();
 					UpdateBatches(jobDocument, graph, updates, _logger);
@@ -767,8 +767,8 @@ namespace Horde.Build.Jobs
 			List<UpdateDefinition<JobDocument>> updates = new List<UpdateDefinition<JobDocument>>();
 
 			// Update the appropriate batch
-			bool bRefreshBatches = false;
-			bool bRefreshDependentJobSteps = false;
+			bool refreshBatches = false;
+			bool refreshDependentJobSteps = false;
 			for (int loopBatchIdx = 0; loopBatchIdx < jobDocument.Batches.Count; loopBatchIdx++)
 			{
 				int batchIdx = loopBatchIdx; // For lambda capture
@@ -795,7 +795,7 @@ namespace Horde.Build.Jobs
 									newOutcome = JobStepOutcome.Failure;
 								}
 
-								bRefreshDependentJobSteps = true;
+								refreshDependentJobSteps = true;
 							}
 
 							// Update the user that requested the abort
@@ -804,7 +804,7 @@ namespace Horde.Build.Jobs
 								step.AbortedByUserId = newAbortByUserId;
 								updates.Add(updateBuilder.Set(x => x.Batches[batchIdx].Steps[stepIdx].AbortedByUserId, step.AbortedByUserId));
 
-								bRefreshDependentJobSteps = true;
+								refreshDependentJobSteps = true;
 							}
 
 							// Update the state
@@ -824,7 +824,7 @@ namespace Horde.Build.Jobs
 									updates.Add(updateBuilder.Set(x => x.Batches[batchIdx].Steps[stepIdx].FinishTime, step.FinishTime));
 								}
 
-								bRefreshDependentJobSteps = true;
+								refreshDependentJobSteps = true;
 							}
 
 							// Update the job outcome
@@ -833,7 +833,7 @@ namespace Horde.Build.Jobs
 								step.Outcome = newOutcome;
 								updates.Add(updateBuilder.Set(x => x.Batches[batchIdx].Steps[stepIdx].Outcome, step.Outcome));
 
-								bRefreshDependentJobSteps = true;
+								refreshDependentJobSteps = true;
 							}
 
 							// Update the job step error
@@ -849,7 +849,7 @@ namespace Horde.Build.Jobs
 								step.LogId = newLogId.Value;
 								updates.Add(updateBuilder.Set(x => x.Batches[batchIdx].Steps[stepIdx].LogId, step.LogId));
 
-								bRefreshDependentJobSteps = true;
+								refreshDependentJobSteps = true;
 							}
 
 							// Update the notification trigger id
@@ -868,7 +868,7 @@ namespace Horde.Build.Jobs
 								step.RetriedByUserId = newRetryByUserId;
 								updates.Add(updateBuilder.Set(x => x.Batches[batchIdx].Steps[stepIdx].RetriedByUserId, step.RetriedByUserId));
 
-								bRefreshBatches = true;
+								refreshBatches = true;
 							}
 
 							// Update the priority
@@ -877,7 +877,7 @@ namespace Horde.Build.Jobs
 								step.Priority = newPriority.Value;
 								updates.Add(updateBuilder.Set(x => x.Batches[batchIdx].Steps[stepIdx].Priority, step.Priority));
 
-								bRefreshBatches = true;
+								refreshBatches = true;
 							}
 
 							// Add any new reports
@@ -934,14 +934,14 @@ namespace Horde.Build.Jobs
 			}
 
 			// Update the batches
-			if (bRefreshBatches)
+			if (refreshBatches)
 			{
 				updates.Clear(); // UpdateBatches will update the entire batches list. We need to remove all the individual step updates to avoid an exception.
 				UpdateBatches(jobDocument, graph, updates, _logger);
 			}
 
 			// Update the state of dependent jobsteps
-			if (bRefreshDependentJobSteps)
+			if (refreshDependentJobSteps)
 			{
 				RefreshDependentJobSteps(jobDocument, graph, updates, _logger);
 				RefreshJobPriority(jobDocument, updates);
diff --git a/Engine/Source/Programs/Horde/Horde.Build/Jobs/JobService.cs b/Engine/Source/Programs/Horde/Horde.Build/Jobs/JobService.cs
--- a/Engine/Source/Programs/Horde/Horde.Build/Jobs/JobService.cs
+++ b/Engine/Source/Programs/Horde/Horde.Build/Jobs/JobService.cs
@@ -869,7 +869,7 @@ namespace Horde.Build.Jobs
 								{
 									break;
 								}
-								if (await _agents.TryUpdateSettingsAsync(agent, bRequestConform: true) != null)
+								if (await _agents.TryUpdateSettingsAsync(agent, requestConform: true) != null)
 								{
 									_logger.LogError("Agent {AgentId} did not complete lease; marking for conform", agent.Id);
 									break;
diff --git a/Engine/Source/Programs/Horde/Horde.Build/Jobs/JobTaskSource.cs b/Engine/Source/Programs/Horde/Horde.Build/Jobs/JobTaskSource.cs
--- a/Engine/Source/Programs/Horde/Horde.Build/Jobs/JobTaskSource.cs
+++ b/Engine/Source/Programs/Horde/Horde.Build/Jobs/JobTaskSource.cs
@@ -835,7 +835,7 @@ namespace Horde.Build.Jobs
 			return null;
 		}
 
-		async Task<ExecuteJobTask?> CreateExecuteJobTaskAsync(IStream stream, IJob job, IJobStepBatch batch, IAgent agent, AgentWorkspace workspace, bool bUseAutoSdk, LogId logId)
+		async Task<ExecuteJobTask?> CreateExecuteJobTaskAsync(IStream stream, IJob job, IJobStepBatch batch, IAgent agent, AgentWorkspace workspace, bool useAutoSdk, LogId logId)
 		{
 			// Get the lease name
 			StringBuilder leaseName = new StringBuilder($"{stream.Name} - ");
@@ -867,7 +867,7 @@ namespace Horde.Build.Jobs
 				return null;
 			}
 
-			AgentWorkspace? autoSdkWorkspace = bUseAutoSdk ? agent.GetAutoSdkWorkspace(cluster) : null;
+			AgentWorkspace? autoSdkWorkspace = useAutoSdk ? agent.GetAutoSdkWorkspace(cluster) : null;
 			if (autoSdkWorkspace != null)
 			{
 				if (!await agent.TryAddWorkspaceMessage(autoSdkWorkspace, cluster, _perforceLoadBalancer, workspaces))
diff --git a/Engine/Source/Programs/Horde/Horde.Build/Jobs/JobsController.cs b/Engine/Source/Programs/Horde/Horde.Build/Jobs/JobsController.cs
--- a/Engine/Source/Programs/Horde/Horde.Build/Jobs/JobsController.cs
+++ b/Engine/Source/Programs/Horde/Horde.Build/Jobs/JobsController.cs
@@ -402,9 +402,9 @@ namespace Horde.Build.Jobs
 			}
 
 			IGraph graph = await _jobService.GetGraphAsync(job);
-			bool bIncludeAcl = await _jobService.AuthorizeAsync(job, AclAction.ViewPermissions, User, cache);
-			bool bIncludeCosts = await _jobService.AuthorizeAsync(job, AclAction.ViewCosts, User, cache);
-			return await CreateJobResponseAsync(job, graph, bIncludeAcl, bIncludeCosts, filter);
+			bool includeAcl = await _jobService.AuthorizeAsync(job, AclAction.ViewPermissions, User, cache);
+			bool includeCosts = await _jobService.AuthorizeAsync(job, AclAction.ViewCosts, User, cache);
+			return await CreateJobResponseAsync(job, graph, includeAcl, includeCosts, filter);
 		}
 
 		/// <summary>
@@ -412,19 +412,19 @@ namespace Horde.Build.Jobs
 		/// </summary>
 		/// <param name="job">The job document</param>
 		/// <param name="graph">The graph for this job</param>
-		/// <param name="bIncludeAcl">Whether to include the ACL in the response</param>
-		/// <param name="bIncludeCosts">Whether to include costs in the response</param>
+		/// <param name="includeAcl">Whether to include the ACL in the response</param>
+		/// <param name="includeCosts">Whether to include costs in the response</param>
 		/// <param name="filter">Filter for the properties to return</param>
 		/// <returns>Object containing the requested properties</returns>
-		async Task<object> CreateJobResponseAsync(IJob job, IGraph graph, bool bIncludeAcl, bool bIncludeCosts, PropertyFilter? filter)
+		async Task<object> CreateJobResponseAsync(IJob job, IGraph graph, bool includeAcl, bool includeCosts, PropertyFilter? filter)
 		{
 			if (filter == null)
 			{
-				return await CreateJobResponseAsync(job, graph, true, true, bIncludeAcl, bIncludeCosts);
+				return await CreateJobResponseAsync(job, graph, true, true, includeAcl, includeCosts);
 			}
 			else
 			{
-				return filter.ApplyTo(await CreateJobResponseAsync(job, graph, filter.Includes(nameof(GetJobResponse.Batches)), filter.Includes(nameof(GetJobResponse.Labels)) || filter.Includes(nameof(GetJobResponse.DefaultLabel)), bIncludeAcl, bIncludeCosts));
+				return filter.ApplyTo(await CreateJobResponseAsync(job, graph, filter.Includes(nameof(GetJobResponse.Batches)), filter.Includes(nameof(GetJobResponse.Labels)) || filter.Includes(nameof(GetJobResponse.DefaultLabel)), includeAcl, includeCosts));
 			}
 		}
 
@@ -433,12 +433,12 @@ namespace Horde.Build.Jobs
 		/// </summary>
 		/// <param name="job">The job document</param>
 		/// <param name="graph">The graph definition</param>
-		/// <param name="bIncludeBatches">Whether to include the job batches in the response</param>
-		/// <param name="bIncludeLabels">Whether to include the job aggregates in the response</param>
-		/// <param name="bIncludeAcl">Whether to include the ACL in the response</param>
-		/// <param name="bIncludeCosts">Whether to include costs of running particular agents</param>
+		/// <param name="includeBatches">Whether to include the job batches in the response</param>
+		/// <param name="includeLabels">Whether to include the job aggregates in the response</param>
+		/// <param name="includeAcl">Whether to include the ACL in the response</param>
+		/// <param name="includeCosts">Whether to include costs of running particular agents</param>
 		/// <returns>The response object</returns>
-		async ValueTask<GetJobResponse> CreateJobResponseAsync(IJob job, IGraph graph, bool bIncludeBatches, bool bIncludeLabels, bool bIncludeAcl, bool bIncludeCosts)
+		async ValueTask<GetJobResponse> CreateJobResponseAsync(IJob job, IGraph graph, bool includeBatches, bool includeLabels, bool includeAcl, bool includeCosts)
 		{
 			GetThinUserInfoResponse? startedByUserInfo = null;
 			if (job.StartedByUserId != null)
@@ -453,23 +453,23 @@ namespace Horde.Build.Jobs
 			}
 
 			GetAclResponse? aclResponse = null;
-			if (bIncludeAcl && job.Acl != null)
+			if (includeAcl && job.Acl != null)
 			{
 				aclResponse = new GetAclResponse(job.Acl);
 			}
 
 			GetJobResponse response = new GetJobResponse(job, startedByUserInfo, abortedByUserInfo, aclResponse);
-			if (bIncludeBatches || bIncludeLabels)
+			if (includeBatches || includeLabels)
 			{
-				if (bIncludeBatches)
+				if (includeBatches)
 				{
 					response.Batches = new List<GetBatchResponse>();
 					foreach (IJobStepBatch batch in job.Batches)
 					{
-						response.Batches.Add(await CreateBatchResponseAsync(batch, bIncludeCosts));
+						response.Batches.Add(await CreateBatchResponseAsync(batch, includeCosts));
 					}
 				}
-				if (bIncludeLabels)
+				if (includeLabels)
 				{
 					response.Labels = new List<GetLabelStateResponse>();
 					response.DefaultLabel = job.GetLabelStateResponses(graph, response.Labels);
@@ -482,9 +482,9 @@ namespace Horde.Build.Jobs
 		/// Get the response object for a batch
 		/// </summary>
 		/// <param name="batch"></param>
-		/// <param name="bIncludeCosts"></param>
+		/// <param name="includeCosts"></param>
 		/// <returns></returns>
-		async ValueTask<GetBatchResponse> CreateBatchResponseAsync(IJobStepBatch batch, bool bIncludeCosts)
+		async ValueTask<GetBatchResponse> CreateBatchResponseAsync(IJobStepBatch batch, bool includeCosts)
 		{
 			List<GetStepResponse> steps = new List<GetStepResponse>();
 			foreach (IJobStep step in batch.Steps)
@@ -493,7 +493,7 @@ namespace Horde.Build.Jobs
 			}
 
 			double? agentRate = null;
-			if (batch.AgentId != null && bIncludeCosts)
+			if (batch.AgentId != null && includeCosts)
 			{
 				agentRate = await _agentService.GetRateAsync(batch.AgentId.Value);
 			}
@@ -787,21 +787,21 @@ namespace Horde.Build.Jobs
 						graph = await _jobService.GetGraphAsync(job);
 					}
 
-					bool bIncludeAcl;
+					bool includeAcl;
 					using (IScope _ = GlobalTracer.Instance.BuildSpan("AuthorizeViewPermissions").StartActive())
 					{
-						bIncludeAcl = await _jobService.AuthorizeAsync(job, AclAction.ViewPermissions, User, permissionsCache);
+						includeAcl = await _jobService.AuthorizeAsync(job, AclAction.ViewPermissions, User, permissionsCache);
 					}
 
-					bool bIncludeCosts;
+					bool includeCosts;
 					using (IScope _ = GlobalTracer.Instance.BuildSpan("AuthorizeViewCosts").StartActive())
 					{
-						bIncludeCosts = await _jobService.AuthorizeAsync(job, AclAction.ViewCosts, User, permissionsCache);
+						includeCosts = await _jobService.AuthorizeAsync(job, AclAction.ViewCosts, User, permissionsCache);
 					}
 
 					using (IScope _ = GlobalTracer.Instance.BuildSpan("CreateResponse").StartActive())
 					{
-						responses.Add(await CreateJobResponseAsync(job, graph, bIncludeAcl, bIncludeCosts, filter));
+						responses.Add(await CreateJobResponseAsync(job, graph, includeAcl, includeCosts, filter));
 					}
 				}
 			}
@@ -867,21 +867,21 @@ namespace Horde.Build.Jobs
 						graph = await _jobService.GetGraphAsync(job);
 					}
 
-					bool bIncludeAcl;
+					bool includeAcl;
 					using (IScope _ = GlobalTracer.Instance.BuildSpan("AuthorizeViewPermissions").StartActive())
 					{
-						bIncludeAcl = await _jobService.AuthorizeAsync(job, AclAction.ViewPermissions, User, permissionsCache);
+						includeAcl = await _jobService.AuthorizeAsync(job, AclAction.ViewPermissions, User, permissionsCache);
 					}
 
-					bool bIncludeCosts;
+					bool includeCosts;
 					using (IScope _ = GlobalTracer.Instance.BuildSpan("AuthorizeViewCosts").StartActive())
 					{
-						bIncludeCosts = await _jobService.AuthorizeAsync(job, AclAction.ViewCosts, User, permissionsCache);
+						includeCosts = await _jobService.AuthorizeAsync(job, AclAction.ViewCosts, User, permissionsCache);
 					}
 
 					using (IScope _ = GlobalTracer.Instance.BuildSpan("CreateResponse").StartActive())
 					{
-						responses.Add(await CreateJobResponseAsync(job, graph, bIncludeAcl, bIncludeCosts, filter));
+						responses.Add(await CreateJobResponseAsync(job, graph, includeAcl, includeCosts, filter));
 					}
 				}
 			}
@@ -1067,12 +1067,12 @@ namespace Horde.Build.Jobs
 				return Forbid(AclAction.ViewJob, jobId);
 			}
 
-			bool bIncludeCosts = await _jobService.AuthorizeAsync(job, AclAction.ViewCosts, User, cache);
+			bool includeCosts = await _jobService.AuthorizeAsync(job, AclAction.ViewCosts, User, cache);
 
 			List<object> responses = new List<object>();
 			foreach (IJobStepBatch batch in job.Batches)
 			{
-				GetBatchResponse response = await CreateBatchResponseAsync(batch, bIncludeCosts);
+				GetBatchResponse response = await CreateBatchResponseAsync(batch, includeCosts);
 				responses.Add(response.ApplyFilter(filter));
 			}
 			return responses;
@@ -1136,13 +1136,13 @@ namespace Horde.Build.Jobs
 				return Forbid(AclAction.ViewJob, jobId);
 			}
 
-			bool bIncludeCosts = await _jobService.AuthorizeAsync(job, AclAction.ViewCosts, User, cache);
+			bool includeCosts = await _jobService.AuthorizeAsync(job, AclAction.ViewCosts, User, cache);
 
 			foreach (IJobStepBatch batch in job.Batches)
 			{
 				if (batch.Id == batchId)
 				{
-					GetBatchResponse response = await CreateBatchResponseAsync(batch, bIncludeCosts);
+					GetBatchResponse response = await CreateBatchResponseAsync(batch, includeCosts);
 					return response.ApplyFilter(filter);
 				}
 			}
diff --git a/Engine/Source/Programs/Horde/Horde.Build/Jobs/Templates/ITemplateCollection.cs b/Engine/Source/Programs/Horde/Horde.Build/Jobs/Templates/ITemplateCollection.cs
--- a/Engine/Source/Programs/Horde/Horde.Build/Jobs/Templates/ITemplateCollection.cs
+++ b/Engine/Source/Programs/Horde/Horde.Build/Jobs/Templates/ITemplateCollection.cs
@@ -17,15 +17,15 @@ namespace Horde.Build.Jobs.Templates
 		/// </summary>
 		/// <param name="name">Name of the template</param>
 		/// <param name="priority">Priority of this template</param>
-		/// <param name="bAllowPreflights">Whether to allow preflights of this job</param>
-		/// <param name="bUpdateIssues"> Whether to update issues for all jobs using this template</param>
-		/// <param name="bPromoteIssuesByDefault">Whether to promote issues by default for all jobs using this template</param>
+		/// <param name="allowPreflights">Whether to allow preflights of this job</param>
+		/// <param name="updateIssues"> Whether to update issues for all jobs using this template</param>
+		/// <param name="promoteIssuesByDefault">Whether to promote issues by default for all jobs using this template</param>
 		/// <param name="initialAgentType">The agent type to parse the buildgraph script</param>
 		/// <param name="submitNewChange">Path to a file within the stream to submit to generate a new changelist for jobs</param>
 		/// <param name="submitDescription">Description for new changes submitted to the stream</param>
 		/// <param name="arguments">List of arguments which are always specified</param>
 		/// <param name="parameters">List of template parameters</param>
-		Task<ITemplate> AddAsync(string name, Priority? priority = null, bool bAllowPreflights = true, bool bUpdateIssues = false, bool bPromoteIssuesByDefault = false, string? initialAgentType = null, string? submitNewChange = null, string? submitDescription = null, List<string>? arguments = null, List<Parameter>? parameters = null);
+		Task<ITemplate> AddAsync(string name, Priority? priority = null, bool allowPreflights = true, bool updateIssues = false, bool promoteIssuesByDefault = false, string? initialAgentType = null, string? submitNewChange = null, string? submitDescription = null, List<string>? arguments = null, List<Parameter>? parameters = null);
 
 		/// <summary>
 		/// Gets all the available templates
diff --git a/Engine/Source/Programs/Horde/Horde.Build/Jobs/Templates/TemplateCollection.cs b/Engine/Source/Programs/Horde/Horde.Build/Jobs/Templates/TemplateCollection.cs
--- a/Engine/Source/Programs/Horde/Horde.Build/Jobs/Templates/TemplateCollection.cs
+++ b/Engine/Source/Programs/Horde/Horde.Build/Jobs/Templates/TemplateCollection.cs
@@ -57,11 +57,11 @@ namespace Horde.Build.Jobs.Templates
 				Name = null!;
 			}
 
-			public TemplateDocument(string name, Priority? priority, bool bAllowPreflights, bool updateIssues, bool promoteIssuesByDefault, string? initialAgentType, string? submitNewChange, string? submitDescription, List<string>? arguments, List<Parameter>? parameters)
+			public TemplateDocument(string name, Priority? priority, bool allowPreflights, bool updateIssues, bool promoteIssuesByDefault, string? initialAgentType, string? submitNewChange, string? submitDescription, List<string>? arguments, List<Parameter>? parameters)
 			{
 				Name = name;
 				Priority = priority;
-				AllowPreflights = bAllowPreflights;
+				AllowPreflights = allowPreflights;
 				UpdateIssues = updateIssues;
 				PromoteIssuesByDefault = promoteIssuesByDefault;
 				InitialAgentType = initialAgentType;
@@ -104,9 +104,9 @@ namespace Horde.Build.Jobs.Templates
 		}
 
 		/// <inheritdoc/>
-		public async Task<ITemplate> AddAsync(string name, Priority? priority, bool bAllowPreflights, bool bUpdateIssues, bool bPromoteIssuesByDefault, string? initialAgentType, string? submitNewChange, string? submitDescription, List<string>? arguments, List<Parameter>? parameters)
+		public async Task<ITemplate> AddAsync(string name, Priority? priority, bool allowPreflights, bool updateIssues, bool promoteIssuesByDefault, string? initialAgentType, string? submitNewChange, string? submitDescription, List<string>? arguments, List<Parameter>? parameters)
 		{
-			TemplateDocument template = new TemplateDocument(name, priority, bAllowPreflights, bUpdateIssues, bPromoteIssuesByDefault, initialAgentType, submitNewChange, submitDescription, arguments, parameters);
+			TemplateDocument template = new TemplateDocument(name, priority, allowPreflights, updateIssues, promoteIssuesByDefault, initialAgentType, submitNewChange, submitDescription, arguments, parameters);
 			if (await GetAsync(template.Id) == null)
 			{
 				await _templates.ReplaceOneAsync(x => x.Id == template.Id, template, new ReplaceOptions { IsUpsert = true });
diff --git a/Engine/Source/Programs/Horde/Horde.Build/Logs/Data/LogIndexData.cs b/Engine/Source/Programs/Horde/Horde.Build/Logs/Data/LogIndexData.cs
--- a/Engine/Source/Programs/Horde/Horde.Build/Logs/Data/LogIndexData.cs
+++ b/Engine/Source/Programs/Horde/Horde.Build/Logs/Data/LogIndexData.cs
@@ -381,7 +381,7 @@ namespace Horde.Build.Logs.Data
 				for (; ; )
 				{
 					// Advance all the enumerators that are behind the current block index. If they are all equal, we have a match.
-					bool bMatch = true;
+					bool match = true;
 					foreach (IEnumerator<int> enumerator in enumerators)
 					{
 						while (enumerator.Current < blockIdx)
@@ -395,12 +395,12 @@ namespace Horde.Build.Logs.Data
 						if (enumerator.Current > blockIdx)
 						{
 							blockIdx = enumerator.Current;
-							bMatch = false;
+							match = false;
 						}
 					}
 
 					// Return the match and move to the next block
-					if (bMatch)
+					if (match)
 					{
 						if (predicates.All(predicate => predicate(blockIdx)))
 						{
@@ -416,15 +416,15 @@ namespace Horde.Build.Logs.Data
 		/// Gets predicates for matching a token that starts 
 		/// </summary>
 		/// <param name="text">The token text</param>
-		/// <param name="bAllowPartialMatch">Whether to allow a partial match of the token</param>
+		/// <param name="allowPartialMatch">Whether to allow a partial match of the token</param>
 		/// <param name="tokens">Set of aligned tokens that are required</param>
 		/// <param name="predicates">List of predicates for the search</param>
-		void GetAlignedTokenPredicate(ReadOnlySpan<byte> text, bool bAllowPartialMatch, HashSet<ulong> tokens, List<Predicate<int>> predicates)
+		void GetAlignedTokenPredicate(ReadOnlySpan<byte> text, bool allowPartialMatch, HashSet<ulong> tokens, List<Predicate<int>> predicates)
 		{
 			for (int offset = 0; offset < text.Length; offset += LogToken.MaxTokenBytes)
 			{
 				ulong token = LogToken.GetWindowedTokenValue(text, offset);
-				if (offset + LogToken.MaxTokenBytes > text.Length && bAllowPartialMatch)
+				if (offset + LogToken.MaxTokenBytes > text.Length && allowPartialMatch)
 				{
 					ulong tokenMask = LogToken.GetWindowedTokenMask(text, offset, true);
 					predicates.Add(blockIdx => BlockContainsToken(blockIdx, token, tokenMask));
@@ -438,9 +438,9 @@ namespace Horde.Build.Logs.Data
 		/// Generates a predicate for matching a token which may or may not start on a regular token boundary
 		/// </summary>
 		/// <param name="text">The token text</param>
-		/// <param name="bAllowPartialMatch">Whether to allow a partial match of the token</param>
+		/// <param name="allowPartialMatch">Whether to allow a partial match of the token</param>
 		/// <param name="predicates">List of predicates for the search</param>
-		void GetUnalignedTokenPredicate(ReadOnlySpan<byte> text, bool bAllowPartialMatch, List<Predicate<int>> predicates)
+		void GetUnalignedTokenPredicate(ReadOnlySpan<byte> text, bool allowPartialMatch, List<Predicate<int>> predicates)
 		{
 			byte[] textCopy = text.ToArray();
 
@@ -449,10 +449,10 @@ namespace Horde.Build.Logs.Data
 				HashSet<int> union = new HashSet<int>();
 				for (int shift = 0; shift < LogToken.MaxTokenBytes; shift++)
 				{
-					HashSet<int> blocks = new HashSet<int>(BlocksContainingToken(textCopy.AsSpan(), -shift, bAllowPartialMatch));
+					HashSet<int> blocks = new HashSet<int>(BlocksContainingToken(textCopy.AsSpan(), -shift, allowPartialMatch));
 					for (int offset = -shift + LogToken.MaxTokenBytes; offset < textCopy.Length && blocks.Count > 0; offset += LogToken.MaxTokenBytes)
 					{
-						blocks.IntersectWith(BlocksContainingToken(textCopy.AsSpan(), offset, bAllowPartialMatch));
+						blocks.IntersectWith(BlocksContainingToken(textCopy.AsSpan(), offset, allowPartialMatch));
 					}
 					if (blocks.Count > 0)
 					{
@@ -484,12 +484,12 @@ namespace Horde.Build.Logs.Data
 		/// </summary>
 		/// <param name="text">The token to test</param>
 		/// <param name="offset">Offset of the window into the token to test</param>
-		/// <param name="bAllowPartialMatch">Whether to allow a partial match of the token</param>
+		/// <param name="allowPartialMatch">Whether to allow a partial match of the token</param>
 		/// <returns>True if the given block contains a token</returns>
-		IEnumerable<int> BlocksContainingToken(ReadOnlySpan<byte> text, int offset, bool bAllowPartialMatch)
+		IEnumerable<int> BlocksContainingToken(ReadOnlySpan<byte> text, int offset, bool allowPartialMatch)
 		{
 			ulong token = LogToken.GetWindowedTokenValue(text, offset) << _numBlockBits;
-			ulong tokenMask = LogToken.GetWindowedTokenMask(text, offset, bAllowPartialMatch) << _numBlockBits;
+			ulong tokenMask = LogToken.GetWindowedTokenMask(text, offset, allowPartialMatch) << _numBlockBits;
 			ulong blockMask = (1UL << _numBlockBits) - 1;
 			return _cachedTrie!.EnumerateValues((value, valueMask) => (value & tokenMask) == (token & valueMask)).Select(x => (int)(x & blockMask)).Distinct();
 		}
diff --git a/Engine/Source/Programs/Horde/Horde.Build/Logs/Data/LogToken.cs b/Engine/Source/Programs/Horde/Horde.Build/Logs/Data/LogToken.cs
--- a/Engine/Source/Programs/Horde/Horde.Build/Logs/Data/LogToken.cs
+++ b/Engine/Source/Programs/Horde/Horde.Build/Logs/Data/LogToken.cs
@@ -142,15 +142,15 @@ namespace Horde.Build.Logs.Data
 		/// </summary>
 		/// <param name="text">The text to search</param>
 		/// <param name="offset">Offset of the window to read from the token</param>
-		/// <param name="bAllowPartialMatch">Whether to allow only matching the start of the string</param>
+		/// <param name="allowPartialMatch">Whether to allow only matching the start of the string</param>
 		/// <returns>Length of the first token</returns>
-		public static ulong GetWindowedTokenMask(ReadOnlySpan<byte> text, int offset, bool bAllowPartialMatch)
+		public static ulong GetWindowedTokenMask(ReadOnlySpan<byte> text, int offset, bool allowPartialMatch)
 		{
 			ulong token = 0;
 			for (int idx = 0; idx < MaxTokenBytes; idx++)
 			{
 				token <<= 8;
-				if (offset >= 0 && (offset < text.Length || !bAllowPartialMatch))
+				if (offset >= 0 && (offset < text.Length || !allowPartialMatch))
 				{
 					token |= 0xff;
 				}
diff --git a/Engine/Source/Programs/Horde/Horde.Build/Logs/LogFileService.cs b/Engine/Source/Programs/Horde/Horde.Build/Logs/LogFileService.cs
--- a/Engine/Source/Programs/Horde/Horde.Build/Logs/LogFileService.cs
+++ b/Engine/Source/Programs/Horde/Horde.Build/Logs/LogFileService.cs
@@ -610,7 +610,7 @@ namespace Horde.Build.Logs
 			}
 
 			// Otherwise keep appending subchunks
-			bool bResult = false;
+			bool result = false;
 			for (; ; )
 			{
 				// Flush the current sub-chunk if we're on a boundary
@@ -643,7 +643,7 @@ namespace Horde.Build.Logs
 					state._offset += length;
 					state._lineIndex += lineCount;
 					state._memory = state._memory.Slice(length);
-					bResult = true;
+					result = true;
 
 					// If this is the end of the data, bail out
 					if(state._memory.Length == 0)
@@ -661,7 +661,7 @@ namespace Horde.Build.Logs
 					break;
 				}
 			}
-			return bResult;
+			return result;
 		}
 
 		/// <summary>
@@ -670,19 +670,19 @@ namespace Horde.Build.Logs
 		/// <param name="span">Data to write</param>
 		/// <param name="maxLength">Maximum length of the data to write</param>
 		/// <param name="maxLineCount">Maximum number of lines to write</param>
-		/// <param name="bIsEmptyChunk">Whether the current chunk is empty</param>
+		/// <param name="isEmptyChunk">Whether the current chunk is empty</param>
 		/// <returns>A tuple consisting of the amount of data to write and number of lines in it</returns>
-		private static (int, int) GetWriteLength(ReadOnlySpan<byte> span, int maxLength, int maxLineCount, bool bIsEmptyChunk)
+		private static (int, int) GetWriteLength(ReadOnlySpan<byte> span, int maxLength, int maxLineCount, bool isEmptyChunk)
 		{
 			int length = 0;
 			int lineCount = 0;
-			for (int idx = 0; idx < maxLength || bIsEmptyChunk; idx++)
+			for (int idx = 0; idx < maxLength || isEmptyChunk; idx++)
 			{
 				if (span[idx] == '\n')
 				{
 					length = idx + 1;
 					lineCount++;
-					bIsEmptyChunk = false;
+					isEmptyChunk = false;
 
 					if (lineCount >= maxLineCount)
 					{
@@ -1087,8 +1087,8 @@ namespace Horde.Build.Logs
 		/// Adds tasks for writing a list of complete chunks
 		/// </summary>
 		/// <param name="chunksToWrite">List of chunks to write</param>
-		/// <param name="bCreateIndex">Create an index for the log</param>
-		private void WriteCompleteChunks(List<(LogId, long)> chunksToWrite, bool bCreateIndex)
+		/// <param name="createIndex">Create an index for the log</param>
+		private void WriteCompleteChunks(List<(LogId, long)> chunksToWrite, bool createIndex)
 		{
 			using IScope scope = GlobalTracer.Instance.BuildSpan("LogFileService.WriteCompleteChunks").StartActive();
 			int numTasksCreated = 0;
@@ -1113,7 +1113,7 @@ namespace Horde.Build.Logs
 				// Create the write task
 				if (offsets.Count > 0)
 				{
-					Task task = Task.Run(() => WriteCompleteChunksForLogAsync(logId, offsets, bCreateIndex));
+					Task task = Task.Run(() => WriteCompleteChunksForLogAsync(logId, offsets, createIndex));
 					numTasksCreated++;
 					lock (_writeLock)
 					{
@@ -1130,9 +1130,9 @@ namespace Horde.Build.Logs
 		/// Writes list of complete chunks
 		/// </summary>
 		/// <param name="chunksToWrite">List of chunks to write</param>
-		/// <param name="bCreateIndex">Create an index for the log</param>
+		/// <param name="createIndex">Create an index for the log</param>
 		/// <param name="cancellationToken">Cancellation token</param>
-		private async Task WriteCompleteChunksV2Async(List<(LogId, long)> chunksToWrite, bool bCreateIndex, CancellationToken cancellationToken)
+		private async Task WriteCompleteChunksV2Async(List<(LogId, long)> chunksToWrite, bool createIndex, CancellationToken cancellationToken)
 		{
 			using IScope scope = GlobalTracer.Instance.BuildSpan("LogFileService.WriteCompleteChunksV2Async").StartActive();
 			
@@ -1166,7 +1166,7 @@ namespace Horde.Build.Logs
 			await Parallel.ForEachAsync(offsetsToWrite, opts, async (x, innerCt) =>
 			{
 				(LogId logId, List<long> offsets) = x;
-				await WriteCompleteChunksForLogAsync(logId, offsets, bCreateIndex, innerCt);
+				await WriteCompleteChunksForLogAsync(logId, offsets, createIndex, innerCt);
 			});
 		}
 
@@ -1175,15 +1175,15 @@ namespace Horde.Build.Logs
 		/// </summary>
 		/// <param name="logId">Log file to update</param>
 		/// <param name="offsets">Chunks to write</param>
-		/// <param name="bCreateIndex">Whether to create the index for this log</param>
+		/// <param name="createIndex">Whether to create the index for this log</param>
 		/// <param name="cancellationToken">Cancellation token for the call</param>
 		/// <returns>Async task</returns>
-		private async Task<ILogFile?> WriteCompleteChunksForLogAsync(LogId logId, List<long> offsets, bool bCreateIndex, CancellationToken cancellationToken = default)
+		private async Task<ILogFile?> WriteCompleteChunksForLogAsync(LogId logId, List<long> offsets, bool createIndex, CancellationToken cancellationToken = default)
 		{
 			ILogFile? logFile = await _logFiles.GetLogFileAsync(logId);
 			if(logFile != null)
 			{
-				logFile = await WriteCompleteChunksForLogAsync(logFile, offsets, bCreateIndex, cancellationToken);
+				logFile = await WriteCompleteChunksForLogAsync(logFile, offsets, createIndex, cancellationToken);
 			}
 			return logFile;
 		}
@@ -1193,15 +1193,15 @@ namespace Horde.Build.Logs
 		/// </summary>
 		/// <param name="logFileInterface">Log file to update</param>
 		/// <param name="offsets">Chunks to write</param>
-		/// <param name="bCreateIndex">Whether to create the index for this log</param>
+		/// <param name="createIndex">Whether to create the index for this log</param>
 		/// <param name="cancellationToken">Cancellation token for the call</param>
 		/// <returns>Async task</returns>
-		private async Task<ILogFile?> WriteCompleteChunksForLogAsync(ILogFile logFileInterface, List<long> offsets, bool bCreateIndex, CancellationToken cancellationToken = default)
+		private async Task<ILogFile?> WriteCompleteChunksForLogAsync(ILogFile logFileInterface, List<long> offsets, bool createIndex, CancellationToken cancellationToken = default)
 		{
 			using IScope scope = GlobalTracer.Instance.BuildSpan("WriteCompleteChunksForLogAsync").StartActive();
 			scope.Span.SetTag("LogId", logFileInterface.Id.ToString());
 			scope.Span.SetTag("NumOffsets", offsets.Count);
-			scope.Span.SetTag("CreateIndex", bCreateIndex);
+			scope.Span.SetTag("CreateIndex", createIndex);
 			
 			// Write the data to the storage provider
 			List<Task<LogChunkData?>> chunkWriteTasks = new List<Task<LogChunkData?>>();
@@ -1265,7 +1265,7 @@ namespace Horde.Build.Logs
 			}
 
 			// Create the index if necessary
-			if (bCreateIndex && logFile != null)
+			if (createIndex && logFile != null)
 			{
 				try
 				{
diff --git a/Engine/Source/Programs/Horde/Horde.Build/Notifications/NotificationService.cs b/Engine/Source/Programs/Horde/Horde.Build/Notifications/NotificationService.cs
--- a/Engine/Source/Programs/Horde/Horde.Build/Notifications/NotificationService.cs
+++ b/Engine/Source/Programs/Horde/Horde.Build/Notifications/NotificationService.cs
@@ -457,9 +457,9 @@ namespace Horde.Build.Notifications
 		/// Gets the <see cref="INotificationTrigger"/> for a given trigger ID, if any.
 		/// </summary>
 		/// <param name="triggerId"></param>
-		/// <param name="bFireTrigger">If true, the trigger is fired and cannot be reused</param>
+		/// <param name="fireTrigger">If true, the trigger is fired and cannot be reused</param>
 		/// <returns></returns>
-		private async Task<INotificationTrigger?> GetNotificationTrigger(ObjectId? triggerId, bool bFireTrigger)
+		private async Task<INotificationTrigger?> GetNotificationTrigger(ObjectId? triggerId, bool fireTrigger)
 		{
 			if (triggerId == null)
 			{
@@ -472,7 +472,7 @@ namespace Horde.Build.Notifications
 				return null;
 			}
 
-			return bFireTrigger ? await _triggerCollection.FireAsync(trigger) : trigger;
+			return fireTrigger ? await _triggerCollection.FireAsync(trigger) : trigger;
 		}
 	
 		private async Task SendJobNotificationsAsync(IJob job, IGraph graph)
@@ -552,7 +552,7 @@ namespace Horde.Build.Notifications
 			return Task.CompletedTask;
 		}
 
-		private async Task<List<IUser>> GetUsersToNotify(EventRecord? eventRecord, ObjectId? notificationTriggerId, bool bFireTrigger)
+		private async Task<List<IUser>> GetUsersToNotify(EventRecord? eventRecord, ObjectId? notificationTriggerId, bool fireTrigger)
 		{
 			List<UserId> userIds = new List<UserId>();
 
@@ -572,7 +572,7 @@ namespace Horde.Build.Notifications
 			// Find the notifications for this particular step
 			if (notificationTriggerId != null)
 			{
-				INotificationTrigger? trigger = await GetNotificationTrigger(notificationTriggerId, bFireTrigger);
+				INotificationTrigger? trigger = await GetNotificationTrigger(notificationTriggerId, fireTrigger);
 				if (trigger != null)
 				{
 					foreach (INotificationSubscription subscription in trigger.Subscriptions)
@@ -736,9 +736,9 @@ namespace Horde.Build.Notifications
 						triggerId = null;
 					}
 
-					bool bFireTrigger = newLabel.State == LabelState.Complete;
+					bool fireTrigger = newLabel.State == LabelState.Complete;
 
-					List<IUser> usersToNotify = await GetUsersToNotify(eventId, triggerId, bFireTrigger);
+					List<IUser> usersToNotify = await GetUsersToNotify(eventId, triggerId, fireTrigger);
 
 					// filter preflight label notifications to only include initiator
 					if (usersToNotify.Count > 0 && job.PreflightChange != 0 && job.StartedByUserId != null)
diff --git a/Engine/Source/Programs/Horde/Horde.Build/Perforce/PerforceService.cs b/Engine/Source/Programs/Horde/Horde.Build/Perforce/PerforceService.cs
--- a/Engine/Source/Programs/Horde/Horde.Build/Perforce/PerforceService.cs
+++ b/Engine/Source/Programs/Horde/Horde.Build/Perforce/PerforceService.cs
@@ -779,7 +779,7 @@ namespace Horde.Build.Perforce
 				PerforceViewMap viewMap = await perforce.GetCachedStreamViewAsync(stream, cancellationToken);
 				List<string> mappedFiles = new List<string>();
 
-				bool bHasUnmappedFile = false;
+				bool hasUnmappedFile = false;
 				foreach (DescribeFileRecord shelvedFile in change.Files)
 				{
 					string mappedFile;
@@ -789,11 +789,11 @@ namespace Horde.Build.Perforce
 					}
 					else
 					{
-						bHasUnmappedFile = true;
+						hasUnmappedFile = true;
 					}
 				}
 
-				if (bHasUnmappedFile)
+				if (hasUnmappedFile)
 				{
 					return ((mappedFiles.Count > 0)? CheckShelfResult.MixedStream : CheckShelfResult.WrongStream, null);
 				}
diff --git a/Engine/Source/Programs/Horde/Horde.Build/Projects/IProject.cs b/Engine/Source/Programs/Horde/Horde.Build/Projects/IProject.cs
--- a/Engine/Source/Programs/Horde/Horde.Build/Projects/IProject.cs
+++ b/Engine/Source/Programs/Horde/Horde.Build/Projects/IProject.cs
@@ -88,21 +88,21 @@ namespace Horde.Build.Projects
 		/// Converts this object to a public response
 		/// </summary>
 		/// <param name="project">The project instance</param>
-		/// <param name="bIncludeStreams">Whether to include streams in the response</param>
-		/// <param name="bIncludeCategories">Whether to include categories in the response</param>
+		/// <param name="includeStreams">Whether to include streams in the response</param>
+		/// <param name="includeCategories">Whether to include categories in the response</param>
 		/// <param name="streams">The list of streams</param>
-		/// <param name="bIncludeAcl">Whether to include the ACL in the response</param>
+		/// <param name="includeAcl">Whether to include the ACL in the response</param>
 		/// <returns>Response instance</returns>
-		public static GetProjectResponse ToResponse(this IProject project, bool bIncludeStreams, bool bIncludeCategories, List<IStream>? streams, bool bIncludeAcl)
+		public static GetProjectResponse ToResponse(this IProject project, bool includeStreams, bool includeCategories, List<IStream>? streams, bool includeAcl)
 		{
 			List<GetProjectStreamResponse>? streamResponses = null;
-			if(bIncludeStreams)
+			if(includeStreams)
 			{
 				streamResponses = streams!.ConvertAll(x => new GetProjectStreamResponse(x.Id.ToString(), x.Name));
 			}
 
 			List<GetProjectCategoryResponse>? categoryResponses = null;
-			if (bIncludeCategories)
+			if (includeCategories)
 			{
 				categoryResponses = project.Config.Categories.ConvertAll(x => new GetProjectCategoryResponse(x));
 				if (streams != null)
@@ -131,7 +131,7 @@ namespace Horde.Build.Projects
 				}
 			}
 
-			GetAclResponse? aclResponse = (bIncludeAcl && project.Acl != null) ? new GetAclResponse(project.Acl) : null;
+			GetAclResponse? aclResponse = (includeAcl && project.Acl != null) ? new GetAclResponse(project.Acl) : null;
 			return new GetProjectResponse(project.Id.ToString(), project.Name, project.Order, streamResponses, categoryResponses, aclResponse);
 		}
 
diff --git a/Engine/Source/Programs/Horde/Horde.Build/Projects/ProjectsController.cs b/Engine/Source/Programs/Horde/Horde.Build/Projects/ProjectsController.cs
--- a/Engine/Source/Programs/Horde/Horde.Build/Projects/ProjectsController.cs
+++ b/Engine/Source/Programs/Horde/Horde.Build/Projects/ProjectsController.cs
@@ -67,8 +67,8 @@ namespace Horde.Build.Projects
 			{
 				if (await _projectService.AuthorizeAsync(project, AclAction.ViewProject, User, permissionsCache))
 				{
-					bool bIncludeAcl = await _projectService.AuthorizeAsync(project, AclAction.ViewPermissions, User, permissionsCache);
-					responses.Add(project.ToResponse(includeStreams, includeCategories, streams, bIncludeAcl).ApplyFilter(filter));
+					bool includeAcl = await _projectService.AuthorizeAsync(project, AclAction.ViewPermissions, User, permissionsCache);
+					responses.Add(project.ToResponse(includeStreams, includeCategories, streams, includeAcl).ApplyFilter(filter));
 				}
 			}
 			return responses;
@@ -97,11 +97,11 @@ namespace Horde.Build.Projects
 				return Forbid(AclAction.ViewProject, projectId);
 			}
 
-			bool bIncludeStreams = PropertyFilter.Includes(filter, nameof(GetProjectResponse.Streams));
-			bool bIncludeCategories = PropertyFilter.Includes(filter, nameof(GetProjectResponse.Categories));
+			bool includeStreams = PropertyFilter.Includes(filter, nameof(GetProjectResponse.Streams));
+			bool includeCategories = PropertyFilter.Includes(filter, nameof(GetProjectResponse.Categories));
 
 			List<IStream>? visibleStreams = null;
-			if (bIncludeStreams || bIncludeCategories)
+			if (includeStreams || includeCategories)
 			{
 				visibleStreams = new List<IStream>();
 
@@ -115,8 +115,8 @@ namespace Horde.Build.Projects
 				}
 			}
 
-			bool bIncludeAcl = await _projectService.AuthorizeAsync(project, AclAction.ViewPermissions, User, cache);
-			return project.ToResponse(bIncludeStreams, bIncludeCategories, visibleStreams, bIncludeAcl).ApplyFilter(filter);
+			bool includeAcl = await _projectService.AuthorizeAsync(project, AclAction.ViewPermissions, User, cache);
+			return project.ToResponse(includeStreams, includeCategories, visibleStreams, includeAcl).ApplyFilter(filter);
 		}
 
 		/// <summary>
diff --git a/Engine/Source/Programs/Horde/Horde.Build/Secrets/CredentialMessages.cs b/Engine/Source/Programs/Horde/Horde.Build/Secrets/CredentialMessages.cs
--- a/Engine/Source/Programs/Horde/Horde.Build/Secrets/CredentialMessages.cs
+++ b/Engine/Source/Programs/Horde/Horde.Build/Secrets/CredentialMessages.cs
@@ -94,13 +94,13 @@ namespace Horde.Build.Credentials
 		/// Constructor
 		/// </summary>
 		/// <param name="credential">The credential to construct from</param>
-		/// <param name="bIncludeAcl">Whether to include the ACL in the response</param>
-		public GetCredentialResponse(Credential credential, bool bIncludeAcl)
+		/// <param name="includeAcl">Whether to include the ACL in the response</param>
+		public GetCredentialResponse(Credential credential, bool includeAcl)
 		{
 			Id = credential.Id.ToString();
 			Name = credential.Name;
 			Properties = credential.Properties;
-			Acl = (bIncludeAcl && credential.Acl != null)? new GetAclResponse(credential.Acl) : null;
+			Acl = (includeAcl && credential.Acl != null)? new GetAclResponse(credential.Acl) : null;
 		}
 	}
 }
diff --git a/Engine/Source/Programs/Horde/Horde.Build/Secrets/CredentialsController.cs b/Engine/Source/Programs/Horde/Horde.Build/Secrets/CredentialsController.cs
--- a/Engine/Source/Programs/Horde/Horde.Build/Secrets/CredentialsController.cs
+++ b/Engine/Source/Programs/Horde/Horde.Build/Secrets/CredentialsController.cs
@@ -82,8 +82,8 @@ namespace Horde.Build.Secrets
 			{
 				if (await _credentialService.AuthorizeAsync(credential, AclAction.ViewCredential, User, cache))
 				{
-					bool bIncludeAcl = await _credentialService.AuthorizeAsync(credential, AclAction.ViewPermissions, User, cache);
-					responses.Add(new GetCredentialResponse(credential, bIncludeAcl).ApplyFilter(filter));
+					bool includeAcl = await _credentialService.AuthorizeAsync(credential, AclAction.ViewPermissions, User, cache);
+					responses.Add(new GetCredentialResponse(credential, includeAcl).ApplyFilter(filter));
 				}
 			}
 			return responses;
@@ -114,8 +114,8 @@ namespace Horde.Build.Secrets
 				return Forbid();
 			}
 
-			bool bIncludeAcl = await _credentialService.AuthorizeAsync(credential, AclAction.ViewPermissions, User, cache);
-			return new GetCredentialResponse(credential, bIncludeAcl).ApplyFilter(filter);
+			bool includeAcl = await _credentialService.AuthorizeAsync(credential, AclAction.ViewPermissions, User, cache);
+			return new GetCredentialResponse(credential, includeAcl).ApplyFilter(filter);
 		}
 
 		/// <summary>
diff --git a/Engine/Source/Programs/Horde/Horde.Build/Server/DowntimeService.cs b/Engine/Source/Programs/Horde/Horde.Build/Server/DowntimeService.cs
--- a/Engine/Source/Programs/Horde/Horde.Build/Server/DowntimeService.cs
+++ b/Engine/Source/Programs/Horde/Horde.Build/Server/DowntimeService.cs
@@ -82,7 +82,7 @@ namespace Horde.Build.Server
 			IGlobals globals = await _globalsService.GetAsync();
 
 			DateTimeOffset now = TimeZoneInfo.ConvertTime(DateTimeOffset.Now, _settings.CurrentValue.TimeZoneInfo);
-			bool bIsActive = globals.Config.Downtime.Any(x => x.IsActive(now));
+			bool isActive = globals.Config.Downtime.Any(x => x.IsActive(now));
 
 			DateTimeOffset? next = null;
 			foreach (ScheduledDowntime schedule in globals.Config.Downtime)
@@ -96,12 +96,12 @@ namespace Horde.Build.Server
 
 			if (next != null)
 			{
-				_logger.LogInformation("Server time: {Time}. Downtime: {Downtime}. Next: {Next}.", now, bIsActive, next.Value);
+				_logger.LogInformation("Server time: {Time}. Downtime: {Downtime}. Next: {Next}.", now, isActive, next.Value);
 			}
 
-			if (bIsActive != IsDowntimeActive)
+			if (isActive != IsDowntimeActive)
 			{
-				IsDowntimeActive = bIsActive;
+				IsDowntimeActive = isActive;
 				if (IsDowntimeActive)
 				{
 					_logger.LogInformation("Entering downtime");
diff --git a/Engine/Source/Programs/Horde/Horde.Build/Server/RpcService.cs b/Engine/Source/Programs/Horde/Horde.Build/Server/RpcService.cs
--- a/Engine/Source/Programs/Horde/Horde.Build/Server/RpcService.cs
+++ b/Engine/Source/Programs/Horde/Horde.Build/Server/RpcService.cs
@@ -219,13 +219,13 @@ namespace Horde.Build.Server
 
 				// Get the set of workspaces that are currently required
 				HashSet<AgentWorkspace> conformWorkspaces = await _poolService.GetWorkspacesAsync(agent, DateTime.UtcNow);
-				bool bPendingConform = !conformWorkspaces.SetEquals(newWorkspaces) || (agent.RequestFullConform && !request.RemoveUntrackedFiles);
+				bool pendingConform = !conformWorkspaces.SetEquals(newWorkspaces) || (agent.RequestFullConform && !request.RemoveUntrackedFiles);
 
 				// Update the workspaces
-				if (await _agentService.TryUpdateWorkspacesAsync(agent, newWorkspaces, bPendingConform))
+				if (await _agentService.TryUpdateWorkspacesAsync(agent, newWorkspaces, pendingConform))
 				{
 					UpdateAgentWorkspacesResponse response = new UpdateAgentWorkspacesResponse();
-					if (bPendingConform)
+					if (pendingConform)
 					{
 						response.Retry = await _conformTaskSource.GetWorkspacesAsync(agent, response.PendingWorkspaces);
 						response.RemoveUntrackedFiles = request.RemoveUntrackedFiles || agent.RequestFullConform;
diff --git a/Engine/Source/Programs/Horde/Horde.Build/Streams/IStream.cs b/Engine/Source/Programs/Horde/Horde.Build/Streams/IStream.cs
--- a/Engine/Source/Programs/Horde/Horde.Build/Streams/IStream.cs
+++ b/Engine/Source/Programs/Horde/Horde.Build/Streams/IStream.cs
@@ -269,15 +269,15 @@ namespace Horde.Build.Streams
 		/// Converts to a public response object
 		/// </summary>
 		/// <param name="stream">The stream object</param>
-		/// <param name="bIncludeAcl">Whether to include the ACL in the response object</param>
+		/// <param name="includeAcl">Whether to include the ACL in the response object</param>
 		/// <param name="apiTemplateRefs">The template refs for this stream. Passed separately because they have their own ACL.</param>
 		/// <returns>New response instance</returns>
-		public static GetStreamResponse ToApiResponse(this IStream stream, bool bIncludeAcl, List<GetTemplateRefResponse> apiTemplateRefs)
+		public static GetStreamResponse ToApiResponse(this IStream stream, bool includeAcl, List<GetTemplateRefResponse> apiTemplateRefs)
 		{
 			List<TabConfig> apiTabs = stream.Config.Tabs;
 			Dictionary<string, AgentConfig> apiAgentTypes = stream.Config.AgentTypes.ToDictionary(x => x.Key, x => x.Value);
 			Dictionary<string, WorkspaceConfig> apiWorkspaceTypes = stream.Config.WorkspaceTypes.ToDictionary(x => x.Key, x => x.Value);
-			GetAclResponse? apiAcl = (bIncludeAcl && stream.Acl != null)? new GetAclResponse(stream.Acl) : null;
+			GetAclResponse? apiAcl = (includeAcl && stream.Acl != null)? new GetAclResponse(stream.Acl) : null;
 			return new GetStreamResponse(stream.Id.ToString(), stream.ProjectId.ToString(), stream.Name, stream.ConfigRevision, stream.Config.Order, stream.Config.NotificationChannel, stream.Config.NotificationChannelFilter, stream.Config.TriageChannel, stream.Config.DefaultPreflight, apiTabs, apiAgentTypes, apiWorkspaceTypes, apiTemplateRefs, apiAcl, stream.PausedUntil, stream.PauseComment, stream.Config.Workflows);
 		}
 
diff --git a/Engine/Source/Programs/Horde/Horde.Build/Streams/StreamMessages.cs b/Engine/Source/Programs/Horde/Horde.Build/Streams/StreamMessages.cs
--- a/Engine/Source/Programs/Horde/Horde.Build/Streams/StreamMessages.cs
+++ b/Engine/Source/Programs/Horde/Horde.Build/Streams/StreamMessages.cs
@@ -153,9 +153,9 @@ namespace Horde.Build.Streams
 		/// <param name="identifier">Identifier to distinguish this workspace from other workspaces. Defaults to the workspace type name.</param>
 		/// <param name="stream">Override for the stream to sync</param>
 		/// <param name="view">Custom view for the workspace</param>
-		/// <param name="bIncremental">Whether to use an incrementally synced workspace</param>
-		/// <param name="bUseAutoSdk">Whether to use the AutoSDK</param>
-		public GetWorkspaceTypeResponse(string? cluster, string? serverAndPort, string? userName, string? identifier, string? stream, List<string>? view, bool bIncremental, bool bUseAutoSdk)
+		/// <param name="incremental">Whether to use an incrementally synced workspace</param>
+		/// <param name="useAutoSdk">Whether to use the AutoSDK</param>
+		public GetWorkspaceTypeResponse(string? cluster, string? serverAndPort, string? userName, string? identifier, string? stream, List<string>? view, bool incremental, bool useAutoSdk)
 		{
 			Cluster = cluster;
 			ServerAndPort = serverAndPort;
@@ -163,8 +163,8 @@ namespace Horde.Build.Streams
 			Identifier = identifier;
 			Stream = stream;
 			View = view;
-			Incremental = bIncremental;
-			UseAutoSdk = bUseAutoSdk;
+			Incremental = incremental;
+			UseAutoSdk = useAutoSdk;
 		}
 	}
 
@@ -278,8 +278,8 @@ namespace Horde.Build.Streams
 		/// <param name="templateRef">The template ref</param>
 		/// <param name="template">The template instance</param>
 		/// <param name="stepStates">The template step states</param>
-		/// <param name="bIncludeAcl">Whether to include the ACL in the response</param>
-		public GetTemplateRefResponse(TemplateId id, ITemplateRef templateRef, ITemplate template, List<GetTemplateStepStateResponse>? stepStates, bool bIncludeAcl)
+		/// <param name="includeAcl">Whether to include the ACL in the response</param>
+		public GetTemplateRefResponse(TemplateId id, ITemplateRef templateRef, ITemplate template, List<GetTemplateStepStateResponse>? stepStates, bool includeAcl)
 			: base(template)
 		{
 			Id = id.ToString();
@@ -292,7 +292,7 @@ namespace Horde.Build.Streams
 			ChainedJobs = (templateRef.Config.ChainedJobs != null && templateRef.Config.ChainedJobs.Count > 0) ? templateRef.Config.ChainedJobs : null;
 			StepStates = stepStates;
 			DefaultChange = templateRef.Config.DefaultChange;
-			Acl = (bIncludeAcl && templateRef.Acl != null) ? new GetAclResponse(templateRef.Acl) : null;
+			Acl = (includeAcl && templateRef.Acl != null) ? new GetAclResponse(templateRef.Acl) : null;
 		}
 	}
 
diff --git a/Engine/Source/Programs/Horde/Horde.Build/Streams/StreamsController.cs b/Engine/Source/Programs/Horde/Horde.Build/Streams/StreamsController.cs
--- a/Engine/Source/Programs/Horde/Horde.Build/Streams/StreamsController.cs
+++ b/Engine/Source/Programs/Horde/Horde.Build/Streams/StreamsController.cs
@@ -172,7 +172,7 @@ namespace Horde.Build.Streams
 			using IScope scope = GlobalTracer.Instance.BuildSpan("CreateGetStreamResponse").StartActive();
 			scope.Span.SetTag("streamId", stream.Id);
 			
-			bool bIncludeAcl = stream.Acl != null && await _streamService.AuthorizeAsync(stream, AclAction.ViewPermissions, User, cache);
+			bool includeAcl = stream.Acl != null && await _streamService.AuthorizeAsync(stream, AclAction.ViewPermissions, User, cache);
 
 			List<GetTemplateRefResponse> apiTemplateRefs = new List<GetTemplateRefResponse>();
 			foreach (KeyValuePair<TemplateId, ITemplateRef> pair in stream.Templates)
@@ -202,13 +202,13 @@ namespace Horde.Build.Streams
 							}
 						}
 
-						bool bIncludeTemplateAcl = pair.Value.Acl != null && await _streamService.AuthorizeAsync(stream, pair.Value, AclAction.ViewPermissions, User, cache);
-						apiTemplateRefs.Add(new GetTemplateRefResponse(pair.Key, pair.Value, template, stepStates, bIncludeTemplateAcl));
+						bool includeTemplateAcl = pair.Value.Acl != null && await _streamService.AuthorizeAsync(stream, pair.Value, AclAction.ViewPermissions, User, cache);
+						apiTemplateRefs.Add(new GetTemplateRefResponse(pair.Key, pair.Value, template, stepStates, includeTemplateAcl));
 					}
 				}
 			}
 
-			return stream.ToApiResponse(bIncludeAcl, apiTemplateRefs);
+			return stream.ToApiResponse(includeAcl, apiTemplateRefs);
 		}
 
 		/// <summary>
diff --git a/Engine/Source/Programs/Horde/Horde.Build/Ugs/UgsController.cs b/Engine/Source/Programs/Horde/Horde.Build/Ugs/UgsController.cs
--- a/Engine/Source/Programs/Horde/Horde.Build/Ugs/UgsController.cs
+++ b/Engine/Source/Programs/Horde/Horde.Build/Ugs/UgsController.cs
@@ -163,8 +163,8 @@ namespace Horde.Build.Ugs
 				foreach(IIssue issue in issues)
 				{
 					IIssueDetails details = await _issueService.GetIssueDetailsAsync(issue);
-					bool bNotify = userInfo != null && details.Suspects.Any(x => x.AuthorId == userInfo.Id);
-					responses.Add(await CreateIssueResponseAsync(details, bNotify));
+					bool notify = userInfo != null && details.Suspects.Any(x => x.AuthorId == userInfo.Id);
+					responses.Add(await CreateIssueResponseAsync(details, notify));
 				}
 			}
 			else
@@ -178,8 +178,8 @@ namespace Horde.Build.Ugs
 
 					if (cachedOpenIssue.ShowNotifications())
 					{
-						bool bNotify = userInfo != null && cachedOpenIssue.IncludeForUser(userInfo.Id);
-						responses.Add(await CreateIssueResponseAsync(cachedOpenIssue, bNotify));
+						bool notify = userInfo != null && cachedOpenIssue.IncludeForUser(userInfo.Id);
+						responses.Add(await CreateIssueResponseAsync(cachedOpenIssue, notify));
 					}
 				}
 			}
@@ -286,16 +286,16 @@ namespace Horde.Build.Ugs
 		/// Gets the URL for a failing step in the
 		/// </summary>
 		/// <param name="details">The issue to get a URL for</param>
-		/// <param name="bNotify">Whether to show notifications for this issue</param>
+		/// <param name="notify">Whether to show notifications for this issue</param>
 		/// <returns>The issue response</returns>
-		async Task<GetUgsIssueResponse> CreateIssueResponseAsync(IIssueDetails details, bool bNotify)
+		async Task<GetUgsIssueResponse> CreateIssueResponseAsync(IIssueDetails details, bool notify)
 		{
 			Uri? buildUrl = GetIssueBuildUrl(details);
 
 			IUser? owner = details.Issue.OwnerId.HasValue ? await _userCollection.GetCachedUserAsync(details.Issue.OwnerId.Value) : null;
 			IUser? nominatedBy = details.Issue.NominatedById.HasValue ? await _userCollection.GetCachedUserAsync(details.Issue.NominatedById.Value) : null;
 
-			return new GetUgsIssueResponse(details, owner, nominatedBy, bNotify, buildUrl);
+			return new GetUgsIssueResponse(details, owner, nominatedBy, notify, buildUrl);
 		}
 
 		/// <summary>
diff --git a/Engine/Source/Programs/Horde/Horde.Build/Ugs/UgsMessages.cs b/Engine/Source/Programs/Horde/Horde.Build/Ugs/UgsMessages.cs
--- a/Engine/Source/Programs/Horde/Horde.Build/Ugs/UgsMessages.cs
+++ b/Engine/Source/Programs/Horde/Horde.Build/Ugs/UgsMessages.cs
@@ -498,9 +498,9 @@ namespace Horde.Build.Ugs
 		/// <param name="details">Issue to construct from</param>
 		/// <param name="owner"></param>
 		/// <param name="nominatedBy"></param>
-		/// <param name="bNotify">Whether to notify the user about this issue</param>
+		/// <param name="notify">Whether to notify the user about this issue</param>
 		/// <param name="buildUrl">Link to the last build for this issue</param>
-		public GetUgsIssueResponse(IIssueDetails details, IUser? owner, IUser? nominatedBy, bool bNotify, Uri? buildUrl)
+		public GetUgsIssueResponse(IIssueDetails details, IUser? owner, IUser? nominatedBy, bool notify, Uri? buildUrl)
 		{
 			IIssue issue = details.Issue;
 			Id = issue.Id;
@@ -512,7 +512,7 @@ namespace Horde.Build.Ugs
 			AcknowledgedAt = issue.AcknowledgedAt;
 			FixChange = issue.FixChange;
 			ResolvedAt = issue.ResolvedAt;
-			BNotify = bNotify;
+			BNotify = notify;
 			BWarning = issue.Severity == IssueSeverity.Warning;
 			BuildUrl = buildUrl;
 			Streams = details.Spans.Select(x => x.StreamName ?? x.StreamId.ToString()).Distinct().ToList();
diff --git a/Engine/Source/Programs/Horde/Horde.Build/Utilities/AgentUtilities.cs b/Engine/Source/Programs/Horde/Horde.Build/Utilities/AgentUtilities.cs
--- a/Engine/Source/Programs/Horde/Horde.Build/Utilities/AgentUtilities.cs
+++ b/Engine/Source/Programs/Horde/Horde.Build/Utilities/AgentUtilities.cs
@@ -51,7 +51,7 @@ namespace Horde.Build.Utilities
 		/// <returns>New agent app data</returns>
 		public static byte[] UpdateAppSettings(byte[] data, Dictionary<string, object> settings)
 		{
-			bool bWrittenClientId = false;
+			bool writtenClientId = false;
 
 			MemoryStream outputStream = new MemoryStream();
 			using (ZipArchive outputArchive = new ZipArchive(outputStream, ZipArchiveMode.Create, true))
@@ -80,7 +80,7 @@ namespace Horde.Build.Utilities
 							using Utf8JsonWriter writer = new Utf8JsonWriter(outputEntryStream, new JsonWriterOptions { Indented = true });
 							JsonSerializer.Serialize<Dictionary<string, Dictionary<string, object>>>(writer, document, new JsonSerializerOptions { WriteIndented = true });
 
-							bWrittenClientId = true;
+							writtenClientId = true;
 						}
 						else
 						{
@@ -90,7 +90,7 @@ namespace Horde.Build.Utilities
 				}
 			}
 
-			if (!bWrittenClientId)
+			if (!writtenClientId)
 			{
 				throw new InvalidDataException("Missing appsettings.json file from zip archive");
 			}
diff --git a/Engine/Source/Programs/Shared/EpicGames.Perforce.Managed/CachedFileInfo.cs b/Engine/Source/Programs/Shared/EpicGames.Perforce.Managed/CachedFileInfo.cs
--- a/Engine/Source/Programs/Shared/EpicGames.Perforce.Managed/CachedFileInfo.cs
+++ b/Engine/Source/Programs/Shared/EpicGames.Perforce.Managed/CachedFileInfo.cs
@@ -20,17 +20,17 @@ namespace EpicGames.Perforce.Managed
 		public readonly ulong CacheId;
 		public readonly long Length;
 		public readonly long LastModifiedTicks;
-		public readonly bool BReadOnly;
+		public readonly bool ReadOnly;
 		public readonly uint SequenceNumber;
 
-		public CachedFileInfo(DirectoryReference cacheDir, FileContentId contentId, ulong cacheId, long length, long lastModifiedTicks, bool bReadOnly, uint sequenceNumber)
+		public CachedFileInfo(DirectoryReference cacheDir, FileContentId contentId, ulong cacheId, long length, long lastModifiedTicks, bool readOnly, uint sequenceNumber)
 		{
 			CacheDir = cacheDir;
 			ContentId = contentId;
 			CacheId = cacheId;
 			Length = length;
 			LastModifiedTicks = lastModifiedTicks;
-			BReadOnly = bReadOnly;
+			ReadOnly = readOnly;
 			SequenceNumber = sequenceNumber;
 		}
 
@@ -54,9 +54,9 @@ namespace EpicGames.Perforce.Managed
 				logger.LogWarning("warning: {File} was last modified at {Time}; expected {ExpectedTime}", location, info.LastWriteTimeUtc, new DateTime(LastModifiedTicks, DateTimeKind.Utc));
 				return false;
 			}
-			if (info.Attributes.HasFlag(FileAttributes.ReadOnly) != BReadOnly)
+			if (info.Attributes.HasFlag(FileAttributes.ReadOnly) != ReadOnly)
 			{
-				logger.LogWarning("warning: {File} readonly flag is {Flag}; expected {ExpectedFlag}", location, info.Attributes.HasFlag(FileAttributes.ReadOnly), BReadOnly);
+				logger.LogWarning("warning: {File} readonly flag is {Flag}; expected {ExpectedFlag}", location, info.Attributes.HasFlag(FileAttributes.ReadOnly), ReadOnly);
 				return false;
 			}
 
@@ -86,9 +86,9 @@ namespace EpicGames.Perforce.Managed
 			ulong cacheId = reader.ReadUInt64();
 			long length = reader.ReadInt64();
 			long lastModifiedTicks = reader.ReadInt64();
-			bool bReadOnly = reader.ReadBoolean();
+			bool readOnly = reader.ReadBoolean();
 			uint sequenceNumber = reader.ReadUInt32();
-			return new CachedFileInfo(cacheDir, contentId, cacheId, length, lastModifiedTicks, bReadOnly, sequenceNumber);
+			return new CachedFileInfo(cacheDir, contentId, cacheId, length, lastModifiedTicks, readOnly, sequenceNumber);
 		}
 
 		public static void WriteCachedFileInfo(this MemoryWriter writer, CachedFileInfo fileInfo)
@@ -97,7 +97,7 @@ namespace EpicGames.Perforce.Managed
 			writer.WriteUInt64(fileInfo.CacheId);
 			writer.WriteInt64(fileInfo.Length);
 			writer.WriteInt64(fileInfo.LastModifiedTicks);
-			writer.WriteBoolean(fileInfo.BReadOnly);
+			writer.WriteBoolean(fileInfo.ReadOnly);
 			writer.WriteUInt32(fileInfo.SequenceNumber);
 		}
 
diff --git a/Engine/Source/Programs/Shared/EpicGames.Perforce.Managed/ManagedWorkspace.cs b/Engine/Source/Programs/Shared/EpicGames.Perforce.Managed/ManagedWorkspace.cs
--- a/Engine/Source/Programs/Shared/EpicGames.Perforce.Managed/ManagedWorkspace.cs
+++ b/Engine/Source/Programs/Shared/EpicGames.Perforce.Managed/ManagedWorkspace.cs
@@ -148,7 +148,7 @@ namespace EpicGames.Perforce.Managed
 		/// <summary>
 		/// Whether a repair operation should be run on this workspace. Set whenever the state may be inconsistent.
 		/// </summary>
-		bool _bRequiresRepair;
+		bool _requiresRepair;
 
 		/// <summary>
 		/// The log output device
@@ -223,11 +223,11 @@ namespace EpicGames.Perforce.Managed
 		/// </summary>
 		/// <param name="hostName">Name of the current machine. Will be automatically detected from the host settings if not present.</param>
 		/// <param name="baseDir">The base directory for the repository</param>
-		/// <param name="bOverwrite">Whether to allow overwriting a repository that's not up to date</param>
+		/// <param name="overwrite">Whether to allow overwriting a repository that's not up to date</param>
 		/// <param name="logger">The logging interface</param>
 		/// <param name="cancellationToken">Cancellation token for this operation</param>
 		/// <returns></returns>
-		public static async Task<ManagedWorkspace> LoadOrCreateAsync(string hostName, DirectoryReference baseDir, bool bOverwrite, ILogger logger, CancellationToken cancellationToken)
+		public static async Task<ManagedWorkspace> LoadOrCreateAsync(string hostName, DirectoryReference baseDir, bool overwrite, ILogger logger, CancellationToken cancellationToken)
 		{
 			if (Exists(baseDir))
 			{
@@ -237,7 +237,7 @@ namespace EpicGames.Perforce.Managed
 				}
 				catch (Exception ex)
 				{
-					if (bOverwrite)
+					if (overwrite)
 					{
 						logger.LogWarning(ex, "Unable to load existing repository.");
 					}
@@ -353,11 +353,11 @@ namespace EpicGames.Perforce.Managed
 				throw new FatalErrorException("Unsupported data format (version {0}, current {1})", version, CurrentVersion);
 			}
 
-			bool bRequiresRepair = reader.ReadBoolean();
+			bool requiresRepair = reader.ReadBoolean();
 			uint nextSequenceNumber = reader.ReadUInt32();
 
 			ManagedWorkspace repo = new ManagedWorkspace(hostName, nextSequenceNumber, baseDir, logger);
-			repo._bRequiresRepair = bRequiresRepair;
+			repo._requiresRepair = requiresRepair;
 
 			int numTrackedFiles = reader.ReadInt32();
 			for (int idx = 0; idx < numTrackedFiles; idx++)
@@ -385,7 +385,7 @@ namespace EpicGames.Perforce.Managed
 			// Write the data to memory
 			MemoryWriter writer = new MemoryWriter(buffer.AsMemory());
 			writer.WriteInt32(CurrentVersion);
-			writer.WriteBoolean(_bRequiresRepair || (state != TransactionState.Clean));
+			writer.WriteBoolean(_requiresRepair || (state != TransactionState.Clean));
 			writer.WriteUInt32(_nextSequenceNumber);
 			writer.WriteInt32(_contentIdToTrackedFile.Count);
 			foreach (CachedFileInfo trackedFile in _contentIdToTrackedFile.Values)
@@ -407,16 +407,16 @@ namespace EpicGames.Perforce.Managed
 		/// <summary>
 		/// Cleans the current workspace
 		/// </summary>
-		/// <param name="bRemoveUntracked">Whether to remove untracked files</param>
+		/// <param name="removeUntracked">Whether to remove untracked files</param>
 		/// <param name="cancellationToken">Cancellation token</param>
-		public async Task CleanAsync(bool bRemoveUntracked, CancellationToken cancellationToken)
+		public async Task CleanAsync(bool removeUntracked, CancellationToken cancellationToken)
 		{
 			Stopwatch timer = Stopwatch.StartNew();
 
 			_logger.LogInformation("Cleaning workspace...");
 			using (_logger.BeginIndentScope("  "))
 			{
-				await CleanInternalAsync(bRemoveUntracked, cancellationToken);
+				await CleanInternalAsync(removeUntracked, cancellationToken);
 			}
 
 			_logger.LogInformation("Completed in {ElapsedTime}s", $"{timer.Elapsed.TotalSeconds:0.0}");
@@ -425,9 +425,9 @@ namespace EpicGames.Perforce.Managed
 		/// <summary>
 		/// Cleans the current workspace
 		/// </summary>
-		/// <param name="bRemoveUntracked">Whether to remove untracked files</param>
+		/// <param name="removeUntracked">Whether to remove untracked files</param>
 		/// <param name="cancellationToken">Cancellation token</param>
-		private async Task CleanInternalAsync(bool bRemoveUntracked, CancellationToken cancellationToken)
+		private async Task CleanInternalAsync(bool removeUntracked, CancellationToken cancellationToken)
 		{
 			FileInfo[] filesToDelete;
 			DirectoryInfo[] directoriesToDelete;
@@ -436,7 +436,7 @@ namespace EpicGames.Perforce.Managed
 			{
 				Stopwatch timer = Stopwatch.StartNew();
 
-				_workspace.Refresh(bRemoveUntracked, out filesToDelete, out directoriesToDelete);
+				_workspace.Refresh(removeUntracked, out filesToDelete, out directoriesToDelete);
 
 				status.Progress = $"({timer.Elapsed.TotalSeconds:0.0}s)";
 			}
@@ -513,7 +513,7 @@ namespace EpicGames.Perforce.Managed
 				_logger.LogDebug("  Workspace:");
 				foreach (WorkspaceFileInfo file in workspaceFiles)
 				{
-					_logger.LogDebug("    {File,-128} [{ContentId,-48}] [{Length,20:n0}] [{LastModified,20}]{Writable}", file.GetClientPath(), file.ContentId, file._length, file._lastModifiedTicks, file._bReadOnly ? "" : " [ writable ]");
+					_logger.LogDebug("    {File,-128} [{ContentId,-48}] [{Length,20:n0}] [{LastModified,20}]{Writable}", file.GetClientPath(), file.ContentId, file._length, file._lastModifiedTicks, file._readOnly ? "" : " [ writable ]");
 				}
 			}
 
@@ -522,7 +522,7 @@ namespace EpicGames.Perforce.Managed
 				_logger.LogDebug("  Cache:");
 				foreach (KeyValuePair<FileContentId, CachedFileInfo> pair in _contentIdToTrackedFile)
 				{
-					_logger.LogDebug("    {File,-128} [{ContentId,-48}] [{Length,20:n0}] [{LastModified,20}]{Writable}", pair.Value.GetLocation(), pair.Key, pair.Value.Length, pair.Value.LastModifiedTicks, pair.Value.BReadOnly ? "" : "[ writable ]");
+					_logger.LogDebug("    {File,-128} [{ContentId,-48}] [{Length,20:n0}] [{LastModified,20}]{Writable}", pair.Value.GetLocation(), pair.Key, pair.Value.Length, pair.Value.LastModifiedTicks, pair.Value.ReadOnly ? "" : "[ writable ]");
 				}
 			}
 
@@ -551,7 +551,7 @@ namespace EpicGames.Perforce.Managed
 				}
 
 				// Clear the repair flag
-				_bRequiresRepair = false;
+				_requiresRepair = false;
 
 				await SaveAsync(TransactionState.Clean, cancellationToken);
 
@@ -572,11 +572,11 @@ namespace EpicGames.Perforce.Managed
 		}
 
 		/// <summary>
-		/// Checks the bRequiresRepair flag, and repairs/resets it if set.
+		/// Checks the <see cref="_requiresRepair"/> flag, and repairs/resets it if set.
 		/// </summary>
 		private async Task RunOptionalRepairAsync(CancellationToken cancellationToken)
 		{
-			if (_bRequiresRepair)
+			if (_requiresRepair)
 			{
 				await RepairAsync(cancellationToken);
 			}
@@ -776,11 +776,11 @@ namespace EpicGames.Perforce.Managed
 		/// <param name="streamName">Name of the stream to sync</param>
 		/// <param name="changeNumber">Changelist number to sync. -1 to sync to latest.</param>
 		/// <param name="view">View of the workspace</param>
-		/// <param name="bRemoveUntracked">Whether to remove untracked files from the workspace</param>
-		/// <param name="bFakeSync">Whether to simulate the syncing operation rather than actually getting files from the server</param>
+		/// <param name="removeUntracked">Whether to remove untracked files from the workspace</param>
+		/// <param name="fakeSync">Whether to simulate the syncing operation rather than actually getting files from the server</param>
 		/// <param name="cacheFile">If set, uses the given file to cache the contents of the workspace. This can improve sync times when multiple machines sync the same workspace.</param>
 		/// <param name="cancellationToken">Cancellation token</param>
-		public async Task SyncAsync(IPerforceConnection perforce, string streamName, int changeNumber, IReadOnlyList<string> view, bool bRemoveUntracked, bool bFakeSync, FileReference? cacheFile, CancellationToken cancellationToken)
+		public async Task SyncAsync(IPerforceConnection perforce, string streamName, int changeNumber, IReadOnlyList<string> view, bool removeUntracked, bool fakeSync, FileReference? cacheFile, CancellationToken cancellationToken)
 		{
 			Stopwatch timer = Stopwatch.StartNew();
 			if (changeNumber == -1)
@@ -810,7 +810,7 @@ namespace EpicGames.Perforce.Managed
 				Task updateHaveTableTask = Task.Run(() => UpdateClientHaveTableAsync(perforce, changeNumber, view, cancellationToken), cancellationToken);
 
 				// Clean the current workspace
-				await CleanInternalAsync(bRemoveUntracked, cancellationToken);
+				await CleanInternalAsync(removeUntracked, cancellationToken);
 
 				// Wait for the have table update to finish
 				await updateHaveTableTask;
@@ -832,7 +832,7 @@ namespace EpicGames.Perforce.Managed
 
 				// Sync all the appropriate files
 				await RemoveFilesFromWorkspaceAsync(contents, cancellationToken);
-				await AddFilesToWorkspaceAsync(perforce, contents, bFakeSync, cancellationToken);
+				await AddFilesToWorkspaceAsync(perforce, contents, fakeSync, cancellationToken);
 			}
 
 			_logger.LogInformation("Completed in {ElapsedTime}s", $"{timer.Elapsed.TotalSeconds:0.0}");
@@ -938,13 +938,13 @@ namespace EpicGames.Perforce.Managed
 		/// <summary>
 		/// Populates the cache with the head revision of the given streams.
 		/// </summary>
-		public async Task PopulateAsync(List<PopulateRequest> requests, bool bFakeSync, CancellationToken cancellationToken)
+		public async Task PopulateAsync(List<PopulateRequest> requests, bool fakeSync, CancellationToken cancellationToken)
 		{
 			_logger.LogInformation("Populating with {NumStreams} streams", requests.Count);
 			using (_logger.BeginIndentScope("  "))
 			{
 				Tuple<int, StreamSnapshot>[] streamState = await PopulateCleanAsync(requests, cancellationToken);
-				await PopulateSyncAsync(requests, streamState, bFakeSync, cancellationToken);
+				await PopulateSyncAsync(requests, streamState, fakeSync, cancellationToken);
 			}
 		}
 
@@ -1046,7 +1046,7 @@ namespace EpicGames.Perforce.Managed
 		/// <summary>
 		/// Perform the sync part of a populate command
 		/// </summary>
-		public async Task PopulateSyncAsync(List<PopulateRequest> requests, Tuple<int, StreamSnapshot>[] streamState, bool bFakeSync, CancellationToken cancellationToken)
+		public async Task PopulateSyncAsync(List<PopulateRequest> requests, Tuple<int, StreamSnapshot>[] streamState, bool fakeSync, CancellationToken cancellationToken)
 		{
 			// Sync all the new files
 			for (int idx = 0; idx < requests.Count; idx++)
@@ -1065,7 +1065,7 @@ namespace EpicGames.Perforce.Managed
 
 					StreamSnapshot contents = streamState[idx].Item2;
 					await RemoveFilesFromWorkspaceAsync(contents, cancellationToken);
-					await AddFilesToWorkspaceAsync(request.PerforceClient, contents, bFakeSync, cancellationToken);
+					await AddFilesToWorkspaceAsync(request.PerforceClient, contents, fakeSync, cancellationToken);
 				}
 			}
 
@@ -1573,7 +1573,7 @@ namespace EpicGames.Perforce.Managed
 					foreach (KeyValuePair<FileContentId, WorkspaceFileInfo> fileToMove in filesToMove)
 					{
 						ulong cacheId = GetUniqueCacheId(fileToMove.Key);
-						CachedFileInfo newTrackingInfo = new CachedFileInfo(_cacheDir, fileToMove.Key, cacheId, fileToMove.Value._length, fileToMove.Value._lastModifiedTicks, fileToMove.Value._bReadOnly, _nextSequenceNumber);
+						CachedFileInfo newTrackingInfo = new CachedFileInfo(_cacheDir, fileToMove.Key, cacheId, fileToMove.Value._length, fileToMove.Value._lastModifiedTicks, fileToMove.Value._readOnly, _nextSequenceNumber);
 						_contentIdToTrackedFile.Add(fileToMove.Key, newTrackingInfo);
 						sourceAndTargetFiles.Add(new KeyValuePair<FileReference, FileReference>(fileToMove.Value.GetLocation(), newTrackingInfo.GetLocation()));
 					}
@@ -1640,7 +1640,7 @@ namespace EpicGames.Perforce.Managed
 			catch (Exception ex)
 			{
 				_logger.LogWarning(ex, "warning: Unable to delete file {FileName}.", fileToDelete.GetFullName());
-				_bRequiresRepair = true;
+				_requiresRepair = true;
 			}
 		}
 
@@ -1657,7 +1657,7 @@ namespace EpicGames.Perforce.Managed
 			catch (Exception ex)
 			{
 				_logger.LogWarning(ex, "warning: Unable to delete directory {Directory}", directoryToDelete);
-				_bRequiresRepair = true;
+				_requiresRepair = true;
 			}
 		}
 
@@ -1666,9 +1666,9 @@ namespace EpicGames.Perforce.Managed
 		/// </summary>
 		/// <param name="client">The client connection</param>
 		/// <param name="stream">Contents of the stream</param>
-		/// <param name="bFakeSync">Whether to simulate the sync operation, rather than actually syncing files</param>
+		/// <param name="fakeSync">Whether to simulate the sync operation, rather than actually syncing files</param>
 		/// <param name="cancellationToken">Cancellation token</param>
-		private async Task AddFilesToWorkspaceAsync(IPerforceConnection client, StreamSnapshot stream, bool bFakeSync, CancellationToken cancellationToken)
+		private async Task AddFilesToWorkspaceAsync(IPerforceConnection client, StreamSnapshot stream, bool fakeSync, CancellationToken cancellationToken)
 		{
 			// Make sure the repair flag is reset
 			await RunOptionalRepairAsync(cancellationToken);
@@ -1776,7 +1776,7 @@ namespace EpicGames.Perforce.Managed
 							{
 								(int batchBeginIdx, int batchEndIdx) = batches[nextBatchIdx];
 
-								Task task = Task.Run(() => SyncBatch(client, filesToSync, batchBeginIdx, batchEndIdx, bFakeSync, cancellationToken), cancellationToken);
+								Task task = Task.Run(() => SyncBatch(client, filesToSync, batchBeginIdx, batchEndIdx, fakeSync, cancellationToken), cancellationToken);
 								tasks[task] = nextBatchIdx++;
 							}
 
@@ -1827,12 +1827,12 @@ namespace EpicGames.Perforce.Managed
 		/// <param name="filesToSync">List of files to sync</param>
 		/// <param name="beginIdx">First file to sync</param>
 		/// <param name="endIdx">Index of the last file to sync (exclusive)</param>
-		/// <param name="bFakeSync">Whether to fake a sync</param>
+		/// <param name="fakeSync">Whether to fake a sync</param>
 		/// <param name="cancellationToken">Cancellation token for the request</param>
 		/// <returns>Async task</returns>
-		async Task SyncBatch(IPerforceConnection client, WorkspaceFileToSync[] filesToSync, int beginIdx, int endIdx, bool bFakeSync, CancellationToken cancellationToken)
+		async Task SyncBatch(IPerforceConnection client, WorkspaceFileToSync[] filesToSync, int beginIdx, int endIdx, bool fakeSync, CancellationToken cancellationToken)
 		{
-			if (bFakeSync)
+			if (fakeSync)
 			{
 				for (int idx = beginIdx; idx < endIdx; idx++)
 				{
@@ -1874,7 +1874,7 @@ namespace EpicGames.Perforce.Managed
 			{
 				_logger.LogWarning(ex, "warning: Unable to move {CacheFile} from cache to {WorkspaceFile}. Syncing instead.", fileToMove._trackedFile.GetLocation(), fileToMove._workspaceFile.GetLocation());
 				filesToSync.Enqueue(new WorkspaceFileToSync(fileToMove._streamFile, fileToMove._workspaceFile));
-				_bRequiresRepair = true;
+				_requiresRepair = true;
 			}
 		}
 
@@ -1894,7 +1894,7 @@ namespace EpicGames.Perforce.Managed
 			{
 				_logger.LogWarning(ex, "warning: Unable to copy {SourceFile} to {TargetFile}. Syncing instead.", fileToCopy._sourceWorkspaceFile.GetLocation(), fileToCopy._targetWorkspaceFile.GetLocation());
 				filesToSync.Enqueue(new WorkspaceFileToSync(fileToCopy._streamFile, fileToCopy._targetWorkspaceFile));
-				_bRequiresRepair = true;
+				_requiresRepair = true;
 			}
 		}
 
diff --git a/Engine/Source/Programs/Shared/EpicGames.Perforce.Managed/Transactions/AddTransaction.cs b/Engine/Source/Programs/Shared/EpicGames.Perforce.Managed/Transactions/AddTransaction.cs
--- a/Engine/Source/Programs/Shared/EpicGames.Perforce.Managed/Transactions/AddTransaction.cs
+++ b/Engine/Source/Programs/Shared/EpicGames.Perforce.Managed/Transactions/AddTransaction.cs
@@ -155,7 +155,7 @@ namespace EpicGames.Perforce.Managed
 			{
 				if (_filesToMove.TryAdd(trackedFile, new WorkspaceFileToMove(streamFile, trackedFile, workspaceFile)))
 				{
-					workspaceFile.SetMetadata(trackedFile.Length, trackedFile.LastModifiedTicks, trackedFile.BReadOnly);
+					workspaceFile.SetMetadata(trackedFile.Length, trackedFile.LastModifiedTicks, trackedFile.ReadOnly);
 				}
 				else
 				{
diff --git a/Engine/Source/Programs/Shared/EpicGames.Perforce.Managed/WorkspaceDirectoryInfo.cs b/Engine/Source/Programs/Shared/EpicGames.Perforce.Managed/WorkspaceDirectoryInfo.cs
--- a/Engine/Source/Programs/Shared/EpicGames.Perforce.Managed/WorkspaceDirectoryInfo.cs
+++ b/Engine/Source/Programs/Shared/EpicGames.Perforce.Managed/WorkspaceDirectoryInfo.cs
@@ -70,16 +70,16 @@ namespace EpicGames.Perforce.Managed
 		/// <param name="path">Relative path to the file, using forward slashes, and without a leading slash</param>
 		/// <param name="length">Length of the file on disk</param>
 		/// <param name="lastModifiedTicks">Last modified time of the file</param>
-		/// <param name="bReadOnly">Whether the file is read only</param>
+		/// <param name="readOnly">Whether the file is read only</param>
 		/// <param name="contentId">Unique identifier for the server content</param>
-		public void AddFile(Utf8String path, long length, long lastModifiedTicks, bool bReadOnly, FileContentId contentId)
+		public void AddFile(Utf8String path, long length, long lastModifiedTicks, bool readOnly, FileContentId contentId)
 		{
 			StreamDirectoryDigest = IoHash.Zero;
 
 			int idx = path.Span.IndexOf((byte)'/');
 			if (idx == -1)
 			{
-				NameToFile[path] = new WorkspaceFileInfo(this, path, length, lastModifiedTicks, bReadOnly, contentId);
+				NameToFile[path] = new WorkspaceFileInfo(this, path, length, lastModifiedTicks, readOnly, contentId);
 			}
 			else
 			{
@@ -92,7 +92,7 @@ namespace EpicGames.Perforce.Managed
 					NameToSubDirectory[name] = subDirectory;
 				}
 
-				subDirectory.AddFile(path.Slice(idx + 1), length, lastModifiedTicks, bReadOnly, contentId);
+				subDirectory.AddFile(path.Slice(idx + 1), length, lastModifiedTicks, readOnly, contentId);
 			}
 		}
 
@@ -124,16 +124,16 @@ namespace EpicGames.Perforce.Managed
 		/// <summary>
 		/// Refresh the state of the workspace on disk
 		/// </summary>
-		/// <param name="bRemoveUntracked">Whether to remove files that are not part of the stream</param>
+		/// <param name="removeUntracked">Whether to remove files that are not part of the stream</param>
 		/// <param name="filesToDelete">Receives an array of files to delete</param>
 		/// <param name="directoriesToDelete">Recevies an array of directories to delete</param>
-		public void Refresh(bool bRemoveUntracked, out FileInfo[] filesToDelete, out DirectoryInfo[] directoriesToDelete)
+		public void Refresh(bool removeUntracked, out FileInfo[] filesToDelete, out DirectoryInfo[] directoriesToDelete)
 		{
 			ConcurrentQueue<FileInfo> concurrentFilesToDelete = new ConcurrentQueue<FileInfo>();
 			ConcurrentQueue<DirectoryInfo> concurrentDirectoriesToDelete = new ConcurrentQueue<DirectoryInfo>();
 			using (ThreadPoolWorkQueue queue = new ThreadPoolWorkQueue())
 			{
-				queue.Enqueue(() => Refresh(new DirectoryInfo(GetFullName()), bRemoveUntracked, concurrentFilesToDelete, concurrentDirectoriesToDelete, queue));
+				queue.Enqueue(() => Refresh(new DirectoryInfo(GetFullName()), removeUntracked, concurrentFilesToDelete, concurrentDirectoriesToDelete, queue));
 			}
 			directoriesToDelete = concurrentDirectoriesToDelete.ToArray();
 			filesToDelete = concurrentFilesToDelete.ToArray();
@@ -143,11 +143,11 @@ namespace EpicGames.Perforce.Managed
 		/// Recursive method for querying the workspace state
 		/// </summary>
 		/// <param name="info"></param>
-		/// <param name="bRemoveUntracked"></param>
+		/// <param name="removeUntracked"></param>
 		/// <param name="filesToDelete"></param>
 		/// <param name="directoriesToDelete"></param>
 		/// <param name="queue"></param>
-		void Refresh(DirectoryInfo info, bool bRemoveUntracked, ConcurrentQueue<FileInfo> filesToDelete, ConcurrentQueue<DirectoryInfo> directoriesToDelete, ThreadPoolWorkQueue queue)
+		void Refresh(DirectoryInfo info, bool removeUntracked, ConcurrentQueue<FileInfo> filesToDelete, ConcurrentQueue<DirectoryInfo> directoriesToDelete, ThreadPoolWorkQueue queue)
 		{
 			// Recurse through subdirectories
 			Dictionary<Utf8String, WorkspaceDirectoryInfo> newNameToSubDirectory = new Dictionary<Utf8String, WorkspaceDirectoryInfo>(NameToSubDirectory.Count, NameToSubDirectory.Comparer);
@@ -157,9 +157,9 @@ namespace EpicGames.Perforce.Managed
 				if (NameToSubDirectory.TryGetValue(subDirectoryInfo.Name, out subDirectory))
 				{
 					newNameToSubDirectory.Add(subDirectory.Name, subDirectory);
-					queue.Enqueue(() => subDirectory.Refresh(subDirectoryInfo, bRemoveUntracked, filesToDelete, directoriesToDelete, queue));
+					queue.Enqueue(() => subDirectory.Refresh(subDirectoryInfo, removeUntracked, filesToDelete, directoriesToDelete, queue));
 				}
-				else if (bRemoveUntracked)
+				else if (removeUntracked)
 				{
 					directoriesToDelete.Enqueue(subDirectoryInfo);
 				}
@@ -184,7 +184,7 @@ namespace EpicGames.Perforce.Managed
 				}
 				else
 				{
-					if (bRemoveUntracked)
+					if (removeUntracked)
 					{
 						filesToDelete.Enqueue(file);
 					}
diff --git a/Engine/Source/Programs/Shared/EpicGames.Perforce.Managed/WorkspaceFileInfo.cs b/Engine/Source/Programs/Shared/EpicGames.Perforce.Managed/WorkspaceFileInfo.cs
--- a/Engine/Source/Programs/Shared/EpicGames.Perforce.Managed/WorkspaceFileInfo.cs
+++ b/Engine/Source/Programs/Shared/EpicGames.Perforce.Managed/WorkspaceFileInfo.cs
@@ -15,7 +15,7 @@ namespace EpicGames.Perforce.Managed
 		public readonly Utf8String Name;
 		public long _length;
 		public long _lastModifiedTicks;
-		public bool _bReadOnly;
+		public bool _readOnly;
 		public readonly FileContentId ContentId;
 
 		public WorkspaceFileInfo(WorkspaceDirectoryInfo directory, Utf8String name, FileContentId contentId)
@@ -30,21 +30,21 @@ namespace EpicGames.Perforce.Managed
 		{
 		}
 
-		public WorkspaceFileInfo(WorkspaceDirectoryInfo directory, Utf8String name, long length, long lastModifiedTicks, bool bReadOnly, FileContentId contentId)
+		public WorkspaceFileInfo(WorkspaceDirectoryInfo directory, Utf8String name, long length, long lastModifiedTicks, bool readOnly, FileContentId contentId)
 		{
 			Directory = directory;
 			Name = name;
 			_length = length;
 			_lastModifiedTicks = lastModifiedTicks;
-			_bReadOnly = bReadOnly;
+			_readOnly = readOnly;
 			ContentId = contentId;
 		}
 
-		public void SetMetadata(long length, long lastModifiedTicks, bool bReadOnly)
+		public void SetMetadata(long length, long lastModifiedTicks, bool readOnly)
 		{
 			_length = length;
 			_lastModifiedTicks = lastModifiedTicks;
-			_bReadOnly = bReadOnly;
+			_readOnly = readOnly;
 		}
 
 		public void UpdateMetadata()
@@ -54,13 +54,13 @@ namespace EpicGames.Perforce.Managed
 			{
 				_length = info.Length;
 				_lastModifiedTicks = info.LastWriteTimeUtc.Ticks;
-				_bReadOnly = info.Attributes.HasFlag(FileAttributes.ReadOnly);
+				_readOnly = info.Attributes.HasFlag(FileAttributes.ReadOnly);
 			}
 		}
 
 		public bool MatchesAttributes(FileInfo info)
 		{
-			return _length == info.Length && _lastModifiedTicks == info.LastWriteTimeUtc.Ticks && (info.Attributes.HasFlag(FileAttributes.ReadOnly) == _bReadOnly);
+			return _length == info.Length && _lastModifiedTicks == info.LastWriteTimeUtc.Ticks && (info.Attributes.HasFlag(FileAttributes.ReadOnly) == _readOnly);
 		}
 
 		public string GetClientPath()
@@ -98,9 +98,9 @@ namespace EpicGames.Perforce.Managed
 			Utf8String name = reader.ReadNullTerminatedUtf8String();
 			long length = reader.ReadInt64();
 			long lastModifiedTicks = reader.ReadInt64();
-			bool bReadOnly = reader.ReadBoolean();
+			bool readOnly = reader.ReadBoolean();
 			FileContentId contentId = reader.ReadFileContentId();
-			return new WorkspaceFileInfo(directory, name, length, lastModifiedTicks, bReadOnly, contentId);
+			return new WorkspaceFileInfo(directory, name, length, lastModifiedTicks, readOnly, contentId);
 		}
 
 		public static void WriteWorkspaceFileInfo(this MemoryWriter writer, WorkspaceFileInfo fileInfo)
@@ -108,7 +108,7 @@ namespace EpicGames.Perforce.Managed
 			writer.WriteNullTerminatedUtf8String(fileInfo.Name);
 			writer.WriteInt64(fileInfo._length);
 			writer.WriteInt64(fileInfo._lastModifiedTicks);
-			writer.WriteBoolean(fileInfo._bReadOnly);
+			writer.WriteBoolean(fileInfo._readOnly);
 			writer.WriteFileContentId(fileInfo.ContentId);
 		}
 
diff --git a/Engine/Source/Programs/Shared/EpicGames.Perforce/PerforceConnection.cs b/Engine/Source/Programs/Shared/EpicGames.Perforce/PerforceConnection.cs
--- a/Engine/Source/Programs/Shared/EpicGames.Perforce/PerforceConnection.cs
+++ b/Engine/Source/Programs/Shared/EpicGames.Perforce/PerforceConnection.cs
@@ -3326,12 +3326,12 @@ namespace EpicGames.Perforce
 		/// </summary>
 		/// <param name="connection">Connection to the Perforce server</param>
 		/// <param name="streamName">Name of the stream to query</param>
-		/// <param name="bIncludeView">Whether to include the stream view in the output</param>
+		/// <param name="includeView">Whether to include the stream view in the output</param>
 		/// <param name="cancellationToken">Token used to cancel the operation</param>
 		/// <returns>Stream information record</returns>
-		public static async Task<StreamRecord> GetStreamAsync(this IPerforceConnection connection, string streamName, bool bIncludeView, CancellationToken cancellationToken = default)
+		public static async Task<StreamRecord> GetStreamAsync(this IPerforceConnection connection, string streamName, bool includeView, CancellationToken cancellationToken = default)
 		{
-			return (await TryGetStreamAsync(connection, streamName, bIncludeView, cancellationToken)).Data;
+			return (await TryGetStreamAsync(connection, streamName, includeView, cancellationToken)).Data;
 		}
 
 		/// <summary>
@@ -3339,13 +3339,13 @@ namespace EpicGames.Perforce
 		/// </summary>
 		/// <param name="connection">Connection to the Perforce server</param>
 		/// <param name="streamName">Name of the stream to query</param>
-		/// <param name="bIncludeView">Whether to include the stream view in the output</param>
+		/// <param name="includeView">Whether to include the stream view in the output</param>
 		/// <param name="cancellationToken">Token used to cancel the operation</param>
 		/// <returns>Stream information record</returns>
-		public static Task<PerforceResponse<StreamRecord>> TryGetStreamAsync(this IPerforceConnection connection, string streamName, bool bIncludeView, CancellationToken cancellationToken = default)
+		public static Task<PerforceResponse<StreamRecord>> TryGetStreamAsync(this IPerforceConnection connection, string streamName, bool includeView, CancellationToken cancellationToken = default)
 		{
 			List<string> arguments = new List<string> { "-o" };
-			if (bIncludeView)
+			if (includeView)
 			{
 				arguments.Add("-v");
 			}
@@ -3389,12 +3389,12 @@ namespace EpicGames.Perforce
 		/// <param name="streamPath">The path for streams to enumerate (eg. "//UE4/...")</param>
 		/// <param name="maxResults">Maximum number of results to return</param>
 		/// <param name="filter">Additional filter to be applied to the results</param>
-		/// <param name="bUnloaded">Whether to enumerate unloaded workspaces</param>
+		/// <param name="unloaded">Whether to enumerate unloaded workspaces</param>
 		/// <param name="cancellationToken">Token used to cancel the operation</param>
 		/// <returns>List of streams matching the given criteria</returns>
-		public static async Task<List<StreamsRecord>> GetStreamsAsync(this IPerforceConnection connection, string? streamPath, int maxResults, string? filter, bool bUnloaded, CancellationToken cancellationToken = default)
+		public static async Task<List<StreamsRecord>> GetStreamsAsync(this IPerforceConnection connection, string? streamPath, int maxResults, string? filter, bool unloaded, CancellationToken cancellationToken = default)
 		{
-			return (await TryGetStreamsAsync(connection, streamPath, maxResults, filter, bUnloaded, cancellationToken)).Data;
+			return (await TryGetStreamsAsync(connection, streamPath, maxResults, filter, unloaded, cancellationToken)).Data;
 		}
 
 		/// <summary>
@@ -3404,14 +3404,14 @@ namespace EpicGames.Perforce
 		/// <param name="streamPath">The path for streams to enumerate (eg. "//UE4/...")</param>
 		/// <param name="maxResults">Maximum number of results to return</param>
 		/// <param name="filter">Additional filter to be applied to the results</param>
-		/// <param name="bUnloaded">Whether to enumerate unloaded workspaces</param>
+		/// <param name="unloaded">Whether to enumerate unloaded workspaces</param>
 		/// <param name="cancellationToken">Token used to cancel the operation</param>
 		/// <returns>List of streams matching the given criteria</returns>
-		public static async Task<PerforceResponseList<StreamsRecord>> TryGetStreamsAsync(this IPerforceConnection connection, string? streamPath, int maxResults, string? filter, bool bUnloaded, CancellationToken cancellationToken = default)
+		public static async Task<PerforceResponseList<StreamsRecord>> TryGetStreamsAsync(this IPerforceConnection connection, string? streamPath, int maxResults, string? filter, bool unloaded, CancellationToken cancellationToken = default)
 		{
 			// Build the command line
 			List<string> arguments = new List<string>();
-			if (bUnloaded)
+			if (unloaded)
 			{
 				arguments.Add("-U");
 			}
@@ -3758,10 +3758,10 @@ namespace EpicGames.Perforce
 		/// <param name="min">Minimum number of files in a parallel sync</param>
 		/// <param name="minSize">Minimum number of bytes in a parallel sync</param>
 		/// <param name="fileSpecs">Files to sync</param>
-		/// <param name="bQuiet">Whether to use quiet output</param>
+		/// <param name="quiet">Whether to use quiet output</param>
 		/// <param name="cancellationToken">Token used to cancel the operation</param>
 		/// <returns>Arguments for the command</returns>
-		private static IAsyncEnumerable<PerforceResponse<T>> SyncInternalAsync<T>(this IPerforceConnection connection, SyncOptions options, int maxFiles, int numThreads, int batch, int batchSize, int min, int minSize, FileSpecList fileSpecs, bool bQuiet, CancellationToken cancellationToken = default) where T : class
+		private static IAsyncEnumerable<PerforceResponse<T>> SyncInternalAsync<T>(this IPerforceConnection connection, SyncOptions options, int maxFiles, int numThreads, int batch, int batchSize, int min, int minSize, FileSpecList fileSpecs, bool quiet, CancellationToken cancellationToken = default) where T : class
 		{
 			List<string> arguments = new List<string>();
 			if ((options & SyncOptions.Force) != 0)
@@ -3788,7 +3788,7 @@ namespace EpicGames.Perforce
 			{
 				arguments.Add("-p");
 			}
-			if (bQuiet)
+			if (quiet)
 			{
 				arguments.Add("-q");
 			}
diff --git a/Engine/Source/Programs/Shared/EpicGames.Perforce/PerforceSettings.cs b/Engine/Source/Programs/Shared/EpicGames.Perforce/PerforceSettings.cs
--- a/Engine/Source/Programs/Shared/EpicGames.Perforce/PerforceSettings.cs
+++ b/Engine/Source/Programs/Shared/EpicGames.Perforce/PerforceSettings.cs
@@ -176,16 +176,16 @@ namespace EpicGames.Perforce
 		/// Gets the command line arguments to launch an external program, such as P4V or P4VC
 		/// </summary>
 		/// <param name="settings"></param>
-		/// <param name="bIncludeClient"></param>
+		/// <param name="includeClient"></param>
 		/// <returns></returns>
-		public static string GetArgumentsForExternalProgram(this IPerforceSettings settings, bool bIncludeClient)
+		public static string GetArgumentsForExternalProgram(this IPerforceSettings settings, bool includeClient)
 		{
 			StringBuilder basicCommandArgs = new StringBuilder();
 
 			basicCommandArgs.AppendFormat("-p \"{0}\"", settings.ServerAndPort);
 			basicCommandArgs.AppendFormat(" -u \"{0}\"", settings.UserName);
 
-			if (bIncludeClient && settings.ClientName != null)
+			if (includeClient && settings.ClientName != null)
 			{
 				basicCommandArgs.AppendFormat(" -c \"{0}\" ", settings.ClientName);
 			}
