EpicGames.Perforce.Native: Set the version string before every API call, per documentation.

diff --git a/Engine/Source/Programs/Shared/EpicGames.Perforce.Native/Client.cpp b/Engine/Source/Programs/Shared/EpicGames.Perforce.Native/Client.cpp
--- a/Engine/Source/Programs/Shared/EpicGames.Perforce.Native/Client.cpp
+++ b/Engine/Source/Programs/Shared/EpicGames.Perforce.Native/Client.cpp
@@ -667,10 +667,20 @@ class FClient
 public:
 	ClientApi ClientApi;
 	FClientUser User;
+	char* AppVersion;
 
 	FClient(FWriteBuffer* WriteBuffer, FOnBufferReadyFn* OnBufferReady)
 		: User(WriteBuffer, OnBufferReady)
 	{
+		AppVersion = nullptr;
+	}
+
+	~FClient()
+	{
+		if (AppVersion != nullptr)
+		{
+			free(AppVersion);
+		}
 	}
 };
 
@@ -706,7 +716,7 @@ extern "C" NATIVE_API FClient* Client_Create(const FSettings* Settings, FWriteBu
 		}
 		if (Settings->AppVersion != nullptr)
 		{
-			Client->ClientApi.SetVersion(Settings->AppVersion);
+			Client->AppVersion = strdup(Settings->AppVersion);
 		}
 	}
 	Client->ClientApi.SetProtocol("tag", "");
@@ -724,6 +734,11 @@ extern "C" NATIVE_API FClient* Client_Create(const FSettings* Settings, FWriteBu
 
 extern "C" NATIVE_API void Client_Command(FClient* Client, const char* Func, int ArgCount, const char** Args, const char* InputData, int InputLength, const char* promptResponse, bool InterceptIo)
 {
+	if (Client->AppVersion != nullptr)
+	{
+		Client->ClientApi.SetVersion(Client->AppVersion);
+	}
+
 	Client->User.InterceptIo = InterceptIo;
 	Client->User.Func = Func;
 	Client->User.SetInputBuffer(InputData, InputLength);
diff --git a/Engine/Source/Programs/Shared/EpicGames.Perforce.Native/DotNetPerforceLib.Build.cs b/Engine/Source/Programs/Shared/EpicGames.Perforce.Native/DotNetPerforceLib.Build.cs
--- a/Engine/Source/Programs/Shared/EpicGames.Perforce.Native/DotNetPerforceLib.Build.cs
+++ b/Engine/Source/Programs/Shared/EpicGames.Perforce.Native/DotNetPerforceLib.Build.cs
@@ -9,6 +9,7 @@ public class DotNetPerforceLib : ModuleRules
 		if (Target.Platform == UnrealTargetPlatform.Win64)
 		{
 			PrivateDefinitions.Add("OS_NT");
+			PrivateDefinitions.Add("_CRT_NONSTDC_NO_DEPRECATE");
 		}
 
 		PublicIncludePathModuleNames.Add("Core");
