EpicGames.Perforce: Use the assembly InformationalVersion attribute to determine the default version number for Perforce commands. Also allow overriding the default values via global settings.

diff --git a/Engine/Source/Programs/Shared/EpicGames.Perforce/PerforceSettings.cs b/Engine/Source/Programs/Shared/EpicGames.Perforce/PerforceSettings.cs
--- a/Engine/Source/Programs/Shared/EpicGames.Perforce/PerforceSettings.cs
+++ b/Engine/Source/Programs/Shared/EpicGames.Perforce/PerforceSettings.cs
@@ -58,11 +58,26 @@ namespace EpicGames.Perforce
 	/// </summary>
 	public class PerforceSettings : IPerforceSettings
 	{
+		static readonly AssemblyName? s_entryAssemblyName = Assembly.GetEntryAssembly()!.GetName();
+
+		string? _appName;
+		string? _appVersion;
+
 		/// <summary>
 		/// The default settings
 		/// </summary>
 		public static IPerforceSettings Default { get; } = new PerforceSettings(PerforceEnvironment.Default);
 
+		/// <summary>
+		/// Default app name to use for new settings objects
+		/// </summary>
+		public static string? DefaultAppName { get; set; } = s_entryAssemblyName?.Name;
+
+		/// <summary>
+		/// Default version string to use for new settings objects
+		/// </summary>
+		public static string? DefaultAppVersion { get; set; } = Assembly.GetExecutingAssembly().GetCustomAttribute<AssemblyInformationalVersionAttribute>()?.InformationalVersion ?? s_entryAssemblyName?.Version?.ToString();
+
 		/// <inheritdoc/>
 		public string ServerAndPort { get; set; }
 
@@ -79,10 +94,18 @@ namespace EpicGames.Perforce
 		public string? ClientName { get; set; }
 
 		/// <inheritdoc/>
-		public string? AppName { get; set; }
+		public string? AppName
+		{
+			get => _appName ?? DefaultAppName;
+			set => _appName = value;
+		}
 
 		/// <inheritdoc/>
-		public string? AppVersion { get; set; }
+		public string? AppVersion
+		{
+			get => _appVersion ?? DefaultAppVersion;
+			set => _appVersion = value;
+		}
 
 		/// <inheritdoc/>
 		public bool PreferNativeClient { get; set; }
