Fixed C# UHT to properly generate errors when TObjectPtr and TLazyObjectPtr are used in a function signature.

diff --git a/Engine/Source/Programs/Shared/EpicGames.UHT/Types/Properties/UhtClassPtrProperties.cs b/Engine/Source/Programs/Shared/EpicGames.UHT/Types/Properties/UhtClassPtrProperties.cs
--- a/Engine/Source/Programs/Shared/EpicGames.UHT/Types/Properties/UhtClassPtrProperties.cs
+++ b/Engine/Source/Programs/Shared/EpicGames.UHT/Types/Properties/UhtClassPtrProperties.cs
@@ -75,7 +75,7 @@ namespace EpicGames.UHT.Types
 
 			// UFunctions with a smart pointer as input parameter wont compile anyway, because of missing P_GET_... macro.
 			// UFunctions with a smart pointer as return type will crash when called via blueprint, because they are not supported in VM.
-			if (PropertyCategory == UhtPropertyCategory.RegularParameter || PropertyCategory == UhtPropertyCategory.ReplicatedParameter)
+			if (this == outermostProperty && PropertyCategory != UhtPropertyCategory.Member)
 			{
 				outerStruct.LogError("UFunctions cannot take a TObjectPtr as a parameter.");
 			}
diff --git a/Engine/Source/Programs/Shared/EpicGames.UHT/Types/Properties/UhtLazyObjectPtrProperty.cs b/Engine/Source/Programs/Shared/EpicGames.UHT/Types/Properties/UhtLazyObjectPtrProperty.cs
--- a/Engine/Source/Programs/Shared/EpicGames.UHT/Types/Properties/UhtLazyObjectPtrProperty.cs
+++ b/Engine/Source/Programs/Shared/EpicGames.UHT/Types/Properties/UhtLazyObjectPtrProperty.cs
@@ -99,7 +99,7 @@ namespace EpicGames.UHT.Types
 
 			// UFunctions with a smart pointer as input parameter wont compile anyway, because of missing P_GET_... macro.
 			// UFunctions with a smart pointer as return type will crash when called via blueprint, because they are not supported in VM.
-			if (PropertyCategory == UhtPropertyCategory.RegularParameter || PropertyCategory == UhtPropertyCategory.ReplicatedParameter)
+			if (this == outermostProperty && PropertyCategory != UhtPropertyCategory.Member)
 			{
 				outerStruct.LogError("UFunctions cannot take a lazy pointer as a parameter.");
 			}
diff --git a/Engine/Source/Programs/Shared/EpicGames.UHT/Types/Properties/UhtObjectPtrProperty.cs b/Engine/Source/Programs/Shared/EpicGames.UHT/Types/Properties/UhtObjectPtrProperty.cs
--- a/Engine/Source/Programs/Shared/EpicGames.UHT/Types/Properties/UhtObjectPtrProperty.cs
+++ b/Engine/Source/Programs/Shared/EpicGames.UHT/Types/Properties/UhtObjectPtrProperty.cs
@@ -94,7 +94,7 @@ namespace EpicGames.UHT.Types
 
 			// UFunctions with a smart pointer as input parameter wont compile anyway, because of missing P_GET_... macro.
 			// UFunctions with a smart pointer as return type will crash when called via blueprint, because they are not supported in VM.
-			if (PropertyCategory == UhtPropertyCategory.RegularParameter || PropertyCategory == UhtPropertyCategory.ReplicatedParameter)
+			if (this == outermostProperty && PropertyCategory != UhtPropertyCategory.Member)
 			{
 				outerStruct.LogError("UFunctions cannot take a TObjectPtr as a parameter.");
 			}
