- Minorly improved a printout

diff --git a/Engine/Source/Programs/Shared/EpicGames.Build/UEBuildPlatformSDK.cs b/Engine/Source/Programs/Shared/EpicGames.Build/UEBuildPlatformSDK.cs
--- a/Engine/Source/Programs/Shared/EpicGames.Build/UEBuildPlatformSDK.cs
+++ b/Engine/Source/Programs/Shared/EpicGames.Build/UEBuildPlatformSDK.cs
@@ -1,4 +1,4 @@
-// Copyright Epic Games, Inc. All Rights Reserved.
+﻿// Copyright Epic Games, Inc. All Rights Reserved.
 
 using System;
 using System.Collections.Generic;
@@ -760,7 +760,7 @@ namespace EpicGames.Core
 			UInt64 Result;
 			if (!TryConvertVersionToInt(StringValue, out Result, Hint))
 			{
-				throw new Exception($"Unable to convert {GetType()} version {StringValue} to an integer. Likely this version was supplied by code, and is expected to be valid.");
+				throw new Exception($"Unable to convert {GetType()} version '{StringValue}' to an integer. Likely this version was supplied by code, and is expected to be valid.");
 			}
 			return Result;
 		}
