Fix TObjectPtr checks to use IsPlugin.

#rnx

diff --git a/Engine/Source/Programs/Shared/EpicGames.UHT/Types/Properties/UhtObjectPropertyBase.cs b/Engine/Source/Programs/Shared/EpicGames.UHT/Types/Properties/UhtObjectPropertyBase.cs
--- a/Engine/Source/Programs/Shared/EpicGames.UHT/Types/Properties/UhtObjectPropertyBase.cs
+++ b/Engine/Source/Programs/Shared/EpicGames.UHT/Types/Properties/UhtObjectPropertyBase.cs
@@ -245,7 +245,7 @@ namespace EpicGames.UHT.Types
 			UhtPointerMemberBehavior behavior = nonEngineBehavior;
 			if (package.IsPartOfEngine)
 			{
-				if (package.Module.BaseDirectory.Contains("/Plugins/", StringComparison.InvariantCultureIgnoreCase) || package.Module.BaseDirectory.Contains("\\Plugins\\", StringComparison.InvariantCultureIgnoreCase))
+				if (package.IsPlugin)
 				{
 					behavior = enginePluginBehavior;
 				}
