Added check to verify the end of the macro text is as expected.

diff --git a/Engine/Source/Programs/Shared/EpicGames.UHT/Exporters/CodeGen/UhtMacroCreator.cs b/Engine/Source/Programs/Shared/EpicGames.UHT/Exporters/CodeGen/UhtMacroCreator.cs
--- a/Engine/Source/Programs/Shared/EpicGames.UHT/Exporters/CodeGen/UhtMacroCreator.cs
+++ b/Engine/Source/Programs/Shared/EpicGames.UHT/Exporters/CodeGen/UhtMacroCreator.cs
@@ -3,6 +3,7 @@
 using System;
 using System.Text;
 using EpicGames.UHT.Types;
+using EpicGames.UHT.Utils;
 
 namespace EpicGames.UHT.Exporters.CodeGen
 {
@@ -42,6 +43,15 @@ namespace EpicGames.UHT.Exporters.CodeGen
 		public void Dispose()
 		{
 			int finalLength = _builder.Length;
+			if (finalLength < 4 ||
+				_builder[finalLength - 4] != ' ' ||
+				_builder[finalLength - 3] != '\\' ||
+				_builder[finalLength - 2] != '\r' ||
+				_builder[finalLength - 1] != '\n')
+			{
+				throw new UhtException("Macro line must end in ' \\\\\\r\\n'");
+			}
+
 			_builder.Length -= 4;
 			if (finalLength == _startingLength)
 			{
