EpicGames.Core: Add helper extensions for ILogger for more convenient attaching of key/value properties to log scope

diff --git a/Engine/Source/Programs/Shared/EpicGames.Core.Tests/ILoggerPropertyTests.cs b/Engine/Source/Programs/Shared/EpicGames.Core.Tests/ILoggerPropertyTests.cs
new file mode 100644
--- /dev/null
+++ b/Engine/Source/Programs/Shared/EpicGames.Core.Tests/ILoggerPropertyTests.cs
@@ -0,0 +1,87 @@
+// Copyright Epic Games, Inc. All Rights Reserved.
+
+using System;
+using System.Collections.Generic;
+using Microsoft.Extensions.Logging;
+using Microsoft.VisualStudio.TestTools.UnitTesting;
+
+namespace EpicGames.Core.Tests;
+
+public sealed class StubLogScope : IDisposable
+{
+	public object? State { get;}
+	public bool IsDisposed { get; private set; }
+
+	public StubLogScope(object? state)
+	{
+		State = state;
+	}
+
+	public void Dispose()
+	{
+		IsDisposed = true;
+	}
+}
+
+public class SpyLogger : ILogger
+{
+	public List<object?> ReceivedLogStates { get; } = new();
+	public List<StubLogScope> ScopesStarted { get; } = new();
+	
+	public void Log<TState>(LogLevel logLevel, EventId eventId, TState? state, Exception? exception, Func<TState, Exception?, string> formatter)
+	{
+		ReceivedLogStates.Add(state);
+	}
+
+	public bool IsEnabled(LogLevel logLevel)
+	{
+		return true;
+	}
+
+	public IDisposable BeginScope<TState>(TState state)
+	{
+		StubLogScope scope = new (state);
+		ScopesStarted.Add(scope);
+		return scope;
+	}
+}
+
+public sealed class FakeLoggerProvider : ILoggerProvider
+{
+	private ILogger _logger;
+
+	public FakeLoggerProvider(ILogger logger)
+	{
+		_logger = logger;
+	}
+
+	public ILogger CreateLogger(string categoryName)
+	{
+		return _logger;
+	}
+	
+	public void Dispose()
+	{
+	}
+}
+
+[TestClass]
+public class ILoggerPropertyTests
+{
+	[TestMethod]
+	public void AttachPropertiesToScopeTest()
+	{
+		SpyLogger logger = new ();
+		using LoggerFactory loggerFactory = new ();
+		using ILoggerProvider provider = new FakeLoggerProvider(logger);
+		loggerFactory.AddProvider(provider);
+
+		using (logger.WithProperty("user", "foobar").BeginScope())
+		{
+			logger.LogInformation("Some log info");
+		}
+
+		Dictionary<string, object> scopeState = (Dictionary<string, object>)logger.ScopesStarted[0].State!;
+		Assert.AreEqual("foobar", scopeState["user"]);
+	}
+}
diff --git a/Engine/Source/Programs/Shared/EpicGames.Core/ILoggerProperty.cs b/Engine/Source/Programs/Shared/EpicGames.Core/ILoggerProperty.cs
new file mode 100644
--- /dev/null
+++ b/Engine/Source/Programs/Shared/EpicGames.Core/ILoggerProperty.cs
@@ -0,0 +1,77 @@
+// Copyright Epic Games, Inc. All Rights Reserved.
+
+using System;
+using System.Collections.Generic;
+using Microsoft.Extensions.Logging;
+
+namespace EpicGames.Core;
+
+/// <summary>
+/// Extension methods for more convenient attaching of properties to log messages
+/// </summary>
+public static class LoggerPropertyExtensions
+{
+	/// <summary>
+	/// Builder struct for attaching properties to log messages
+	/// Wraps and disposes log scope.
+	/// </summary>
+	public readonly struct LogScopedPropertyBuilder
+	{
+		private readonly Dictionary<string, object> _properties = new();
+		private readonly ILogger _logger;
+		
+		/// <summary>
+		/// Constructor
+		/// </summary>
+		/// <param name="logger">Logger to wrap</param>
+		public LogScopedPropertyBuilder(ILogger logger)
+		{
+			_logger = logger;
+		}
+
+		/// <summary>
+		/// Set a key/value property for current log scope
+		/// </summary>
+		/// <param name="key">Arbitrary key</param>
+		/// <param name="value">Arbitrary value</param>
+		/// <returns></returns>
+		public LogScopedPropertyBuilder WithProperty(string key, object value)
+		{
+			_properties[key] = value;
+			return this;
+		}
+
+		/// <summary>
+		/// Starts new scope on underlying logger, passing previously set properties
+		/// </summary>
+		/// <returns>Original IDisposable from ILogger.BeginScope()</returns>
+		public IDisposable BeginScope()
+		{
+			return _logger.BeginScope(_properties);
+		}
+	}
+	
+	/// <summary>
+	/// Creates a new log scope and attaches a key/value property
+	/// Additional properties can be added with <see cref="LogScopedPropertyBuilder.WithProperty" />.
+	/// </summary>
+	/// <param name="logger">Logger to use</param>
+	/// <param name="key">Arbitrary key</param>
+	/// <param name="value">Arbitrary value</param>
+	/// <returns>A disposable log property builder that wraps the original ILogger scope</returns>
+	/// <example>
+	/// Simple example of attaching two properties
+	/// <code>
+	/// using (_logger.WithProperty("stream", "ue5").WithProperty("os", "windows").BuildScope())
+	/// {
+	///     _logger.LogInformation("A log message with properties attached");
+	/// }
+	/// </code>
+	/// </example>
+	public static LogScopedPropertyBuilder WithProperty(this ILogger logger, string key, object value)
+	{
+		LogScopedPropertyBuilder builder = new (logger);
+		builder.WithProperty(key, value);
+		return builder;
+	}
+}
\ No newline at end of file
