AutomationTool: ScriptModules not in the Engine directory should still be compiled when using an installed engine

Requires marking any engine ScriptModules and their dependencies as valid and up to date so those are ignored when processing project ScriptModule dependencies.

#rnx

diff --git a/Engine/Source/Programs/Shared/EpicGames.Build/System/CompileScriptModules.cs b/Engine/Source/Programs/Shared/EpicGames.Build/System/CompileScriptModules.cs
--- a/Engine/Source/Programs/Shared/EpicGames.Build/System/CompileScriptModules.cs
+++ b/Engine/Source/Programs/Shared/EpicGames.Build/System/CompileScriptModules.cs
@@ -172,16 +172,62 @@ namespace UnrealBuildBase
 			{
 				BuildFlags |= BuildFlags.NoCompile;
 			}
+			if (bUseBuildRecords)
+			{
+				BuildFlags |= BuildFlags.UseBuildRecords;
+			}
+
+			Dictionary<FileReference, CsProjBuildRecordEntry> BuildResults;
 			if (Unreal.IsEngineInstalled())
 			{
-				BuildFlags |= BuildFlags.NoCompile | BuildFlags.ErrorOnMissingTarget;
+				// Warn if not using build records or force compiling
+				if (BuildFlags.HasFlag(BuildFlags.ForceCompile))
+				{
+					Logger.LogWarning("ForceCompile not supported if Unreal.IsEngineInstalled() == true");
+				}
+				if (!BuildFlags.HasFlag(BuildFlags.UseBuildRecords))
+				{
+					Logger.LogWarning("Disabling UseBuildRecords not supported if Unreal.IsEngineInstalled() == true");
+				}
+				BuildFlags EngineBuildFlags = BuildFlags | BuildFlags.UseBuildRecords | BuildFlags.NoCompile | BuildFlags.ErrorOnMissingTarget;
+
+				bool bEngineBuildSuccess = false;
+				Dictionary<FileReference, CsProjBuildRecordEntry> EngineBuildResults = Build(
+					RulesFileType,
+					FoundProjects.Where(x => x.IsUnderDirectory(Unreal.EngineDirectory)).ToHashSet(),
+					BaseDirectories.Where(x => x.IsUnderDirectory(Unreal.EngineDirectory)).ToList(),
+					null,
+					EngineBuildFlags,
+					out bEngineBuildSuccess,
+					OnBuildingProjects,
+					Logger);
+
+				bool bProjectBuildSuccess = false;
+				BuildFlags ProjectBuildFlags = BuildFlags | BuildFlags.UseBuildRecords;
+				// ForceCompile not supported for installed builds
+				if (ProjectBuildFlags.HasFlag(BuildFlags.ForceCompile))
+				{
+					ProjectBuildFlags &= ~BuildFlags.ForceCompile;
+				}
+				Dictionary<FileReference, CsProjBuildRecordEntry> ProjectBuildResults = Build(
+					RulesFileType,
+					FoundProjects.Where(x => !x.IsUnderDirectory(Unreal.EngineDirectory)).ToHashSet(),
+					BaseDirectories,
+					null,
+					ProjectBuildFlags,
+					out bProjectBuildSuccess,
+					OnBuildingProjects,
+					Logger);
+
+				bBuildSuccess = bEngineBuildSuccess && bProjectBuildSuccess;
+				BuildResults = EngineBuildResults.Union(ProjectBuildResults).Distinct().ToDictionary(x => x.Key, x => x.Value);
 			}
-			if (bUseBuildRecords)
+			else
 			{
-				BuildFlags |= BuildFlags.UseBuildRecords;
+				BuildResults = Build(RulesFileType, FoundProjects, BaseDirectories, null, BuildFlags, out bBuildSuccess, OnBuildingProjects, Logger);
 			}
 
-			return GetTargetPaths(Build(RulesFileType, FoundProjects, BaseDirectories, null, BuildFlags, out bBuildSuccess, OnBuildingProjects, Logger));
+			return GetTargetPaths(BuildResults);
 		}
 
 		/// <summary>
@@ -517,6 +563,16 @@ namespace UnrealBuildBase
 				return;
 			}
 
+
+			// If the engine is installed, and the project being validated is under the engine directory, treat it as valid
+			// so it is not built as a dependency when building an external project.
+			// It is assumed these projects exist and have already been verified.
+			if (Unreal.IsEngineInstalled() && ProjectPath.IsUnderDirectory(Unreal.EngineDirectory))
+			{
+				Entry.Status = CsProjBuildRecordStatus.Valid;
+				return;
+			}
+
 			// Is this particular build record valid?
 			if (!ValidateBuildRecord(Entry.BuildRecord, ProjectPath.Directory, out string ValidationFailureMessage, Hook))
 			{
