Horde: Suppress errors when inviting users to a channel that are already there.

diff --git a/Engine/Source/Programs/Shared/EpicGames.Slack/SlackClient.cs b/Engine/Source/Programs/Shared/EpicGames.Slack/SlackClient.cs
--- a/Engine/Source/Programs/Shared/EpicGames.Slack/SlackClient.cs
+++ b/Engine/Source/Programs/Shared/EpicGames.Slack/SlackClient.cs
@@ -385,7 +385,10 @@ namespace EpicGames.Slack
 			InviteMessage message = new InviteMessage();
 			message.Channel = channel;
 			message.Users = String.Join(",", userIds);
-			await SendRequestAsync<SlackResponse>(ConversationsInviteUrl, message);
+
+			static bool ShouldLogError(SlackResponse response) => !response.Ok && !String.Equals(response.Error, "already_in_channel", StringComparison.Ordinal);
+
+			await SendRequestAsync<SlackResponse>(ConversationsInviteUrl, message, ShouldLogError);
 		}
 
 		#endregion
