Horde: Use "PUT" over "POST" when uploading a redirected blob.

diff --git a/Engine/Source/Programs/Shared/EpicGames.Horde/Storage/Backends/HttpStorageClient.cs b/Engine/Source/Programs/Shared/EpicGames.Horde/Storage/Backends/HttpStorageClient.cs
--- a/Engine/Source/Programs/Shared/EpicGames.Horde/Storage/Backends/HttpStorageClient.cs
+++ b/Engine/Source/Programs/Shared/EpicGames.Horde/Storage/Backends/HttpStorageClient.cs
@@ -85,7 +85,7 @@ namespace EpicGames.Horde.Storage.Backends
 				if(redirectResponse.UploadUrl != null)
 				{
 					_logger.LogDebug("Using upload redirect for {Locator}", redirectResponse.Locator);
-					using (HttpResponseMessage uploadResponse = await _httpClient.PostAsync(redirectResponse.UploadUrl, streamContent, cancellationToken))
+					using (HttpResponseMessage uploadResponse = await _httpClient.PutAsync(redirectResponse.UploadUrl, streamContent, cancellationToken))
 					{
 						uploadResponse.EnsureSuccessStatusCode();
 					}
