Fixed C# UHT to properly detect system include files

diff --git a/Engine/Source/Programs/Shared/EpicGames.UHT/Parsers/UhtHeaderFileParser.cs b/Engine/Source/Programs/Shared/EpicGames.UHT/Parsers/UhtHeaderFileParser.cs
--- a/Engine/Source/Programs/Shared/EpicGames.UHT/Parsers/UhtHeaderFileParser.cs
+++ b/Engine/Source/Programs/Shared/EpicGames.UHT/Parsers/UhtHeaderFileParser.cs
@@ -660,14 +660,25 @@ namespace EpicGames.UHT.Parsers
 			{
 				if (isBeingIncluded)
 				{
+					if (SpottedAutogeneratedHeaderInclude)
+					{
+						HeaderFile.LogError("#include found after .generated.h file - the .generated.h file should always be the last #include in a header");
+					}
+
+					StringView includeNameString = new StringView();
+
 					UhtToken includeName = lineTokenReader.GetToken();
 					if (includeName.IsConstString())
 					{
-						StringView includeNameString = includeName.GetUnescapedString(HeaderFile);
-						if (SpottedAutogeneratedHeaderInclude)
-						{
-							HeaderFile.LogError("#include found after .generated.h file - the .generated.h file should always be the last #include in a header");
-						}
+						includeNameString = includeName.GetUnescapedString(HeaderFile);
+					}
+					else if (includeName.IsSymbol('<'))
+					{
+						includeNameString = new StringView(lineTokenReader.GetRawString('>', UhtRawStringOptions.DontConsumeTerminator).Memory.Trim());
+					}
+
+					if (includeNameString.Length > 0)
+					{ 
 						if (HeaderFile.GeneratedHeaderFileName.AsSpan().Equals(includeNameString.Span, StringComparison.OrdinalIgnoreCase))
 						{
 							SpottedAutogeneratedHeaderInclude = true;
