Fixed C# UHT to be compatible with C++ UHT where commented out text between the member property declaration and the assignment or semi-colon wouldn't be picked up as a comment.

diff --git a/Engine/Source/Programs/Shared/EpicGames.UHT/Parsers/UhtPropertyParser.cs b/Engine/Source/Programs/Shared/EpicGames.UHT/Parsers/UhtPropertyParser.cs
--- a/Engine/Source/Programs/Shared/EpicGames.UHT/Parsers/UhtPropertyParser.cs
+++ b/Engine/Source/Programs/Shared/EpicGames.UHT/Parsers/UhtPropertyParser.cs
@@ -1162,6 +1162,8 @@ namespace EpicGames.UHT.Parsers
 				UhtSpecifierParser specifiers = UhtSpecifierParser.GetThreadInstance(specifierContext, newProperty.SourceName, 
 					specifierContext.Type.Session.GetSpecifierTable(UhtTableNames.PropertyMember));
 				specifiers.ParseFieldMetaData();
+				tokenReader.SkipWhitespaceAndComments(); //TODO - old UHT compatibility.  Commented out initializers can cause comment/tooltip to be used as meta data.
+				tokenReader.CommitPendingComments(); //TODO - old UHT compatibility.  Commented out initializers can cause comment/tooltip to be used as meta data.
 				topScope.AddFormattedCommentsAsTooltipMetaData(newProperty);
 			}
 
diff --git a/Engine/Source/Programs/Shared/EpicGames.UHT/Tokenizer/UhtTokenBufferReader.cs b/Engine/Source/Programs/Shared/EpicGames.UHT/Tokenizer/UhtTokenBufferReader.cs
--- a/Engine/Source/Programs/Shared/EpicGames.UHT/Tokenizer/UhtTokenBufferReader.cs
+++ b/Engine/Source/Programs/Shared/EpicGames.UHT/Tokenizer/UhtTokenBufferReader.cs
@@ -122,8 +122,11 @@ namespace EpicGames.UHT.Tokenizer
 		/// <inheritdoc/>
 		public void SkipWhitespaceAndComments()
 		{
-			bool gotInlineComment = false;
-			SkipWhitespaceAndCommentsInternal(ref gotInlineComment, true);
+			if (!_hasToken)
+			{
+				bool gotInlineComment = false;
+				SkipWhitespaceAndCommentsInternal(ref gotInlineComment, true);
+			}
 		}
 
 		/// <inheritdoc/>
