Horde: Use a separate HTTP client for redirected uploads, to avoid sending invalid auth headers.

diff --git a/Engine/Source/Programs/Horde/Horde.Agent/Commands/Bundles/CreateCommand.cs b/Engine/Source/Programs/Horde/Horde.Agent/Commands/Bundles/CreateCommand.cs
--- a/Engine/Source/Programs/Horde/Horde.Agent/Commands/Bundles/CreateCommand.cs
+++ b/Engine/Source/Programs/Horde/Horde.Agent/Commands/Bundles/CreateCommand.cs
@@ -42,17 +42,20 @@ namespace Horde.Agent.Commands.Bundles
 		{
 			readonly IMemoryCache _cache;
 			readonly HttpClient _httpClient;
+			readonly HttpClient _redirectHttpClient;
 			public IStorageClient Store { get; }
 
 			public HttpStorageClientOwner(NamespaceId namespaceId, Uri baseUri, ILogger logger)
 			{
 				_cache = new MemoryCache(new MemoryCacheOptions());
 				_httpClient = new HttpClient { BaseAddress = baseUri };
-				Store = new HttpStorageClient(namespaceId, _httpClient, _cache, logger);
+				_redirectHttpClient = new HttpClient();
+				Store = new HttpStorageClient(namespaceId, _httpClient, _redirectHttpClient, _cache, logger);
 			}
 
 			public void Dispose()
 			{
+				_redirectHttpClient.Dispose();
 				_httpClient.Dispose();
 				_cache.Dispose();
 			}
diff --git a/Engine/Source/Programs/Horde/Horde.Agent/Execution/JobExecutor.cs b/Engine/Source/Programs/Horde/Horde.Agent/Execution/JobExecutor.cs
--- a/Engine/Source/Programs/Horde/Horde.Agent/Execution/JobExecutor.cs
+++ b/Engine/Source/Programs/Horde/Horde.Agent/Execution/JobExecutor.cs
@@ -660,7 +660,9 @@ namespace Horde.Agent.Execution
 			httpClient.BaseAddress = _session.ServerUrl;
 			httpClient.DefaultRequestHeaders.Authorization = new AuthenticationHeaderValue("Bearer", _session.Token);
 
-			HttpStorageClient storage = new HttpStorageClient(new NamespaceId("default"), httpClient, cache, logger);
+			using HttpClient redirectHttpClient = _httpClientFactory.CreateClient();
+
+			HttpStorageClient storage = new HttpStorageClient(new NamespaceId("default"), httpClient, redirectHttpClient, cache, logger);
 //			FileStorageClient storage = new FileStorageClient(DirectoryReference.Combine(sharedStorageDir, "bundles"), cache, logger);
 			logger.LogInformation("Using Horde-managed shared storage via {SharedStorageDir}", sharedStorageDir);
 
diff --git a/Engine/Source/Programs/Shared/EpicGames.Horde/Storage/Backends/HttpStorageClient.cs b/Engine/Source/Programs/Shared/EpicGames.Horde/Storage/Backends/HttpStorageClient.cs
--- a/Engine/Source/Programs/Shared/EpicGames.Horde/Storage/Backends/HttpStorageClient.cs
+++ b/Engine/Source/Programs/Shared/EpicGames.Horde/Storage/Backends/HttpStorageClient.cs
@@ -36,17 +36,19 @@ namespace EpicGames.Horde.Storage.Backends
 
 		readonly NamespaceId _namespaceId;
 		readonly HttpClient _httpClient;
+		readonly HttpClient _redirectHttpClient;
 		readonly ILogger _logger;
 		bool _supportsUploadRedirects = true;
 
 		/// <summary>
 		/// Constructor
 		/// </summary>
-		public HttpStorageClient(NamespaceId namespaceId, HttpClient httpClient, IMemoryCache cache, ILogger logger) 
+		public HttpStorageClient(NamespaceId namespaceId, HttpClient httpClient, HttpClient redirectHttpClient, IMemoryCache cache, ILogger logger) 
 			: base(cache, logger)
 		{
 			_namespaceId = namespaceId;
 			_httpClient = httpClient;
+			_redirectHttpClient = redirectHttpClient;
 			_logger = logger;
 		}
 
@@ -85,9 +87,13 @@ namespace EpicGames.Horde.Storage.Backends
 				if(redirectResponse.UploadUrl != null)
 				{
 					_logger.LogDebug("Using upload redirect for {Locator}", redirectResponse.Locator);
-					using (HttpResponseMessage uploadResponse = await _httpClient.PutAsync(redirectResponse.UploadUrl, streamContent, cancellationToken))
+					using (HttpResponseMessage uploadResponse = await _redirectHttpClient.PutAsync(redirectResponse.UploadUrl, streamContent, cancellationToken))
 					{
-						uploadResponse.EnsureSuccessStatusCode();
+						if (!uploadResponse.IsSuccessStatusCode)
+						{
+							string body = await uploadResponse.Content.ReadAsStringAsync(cancellationToken);
+							throw new StorageException($"Unable to upload data to redirected URL: {body}", null);
+						}
 					}
 					return redirectResponse.Locator;
 				}
