Fix duplicate keys being added to Dictionary when compiling script modules

#rnx

diff --git a/Engine/Source/Programs/Shared/EpicGames.Build/System/CompileScriptModules.cs b/Engine/Source/Programs/Shared/EpicGames.Build/System/CompileScriptModules.cs
--- a/Engine/Source/Programs/Shared/EpicGames.Build/System/CompileScriptModules.cs
+++ b/Engine/Source/Programs/Shared/EpicGames.Build/System/CompileScriptModules.cs
@@ -220,7 +220,14 @@ namespace UnrealBuildBase
 					Logger);
 
 				bBuildSuccess = bEngineBuildSuccess && bProjectBuildSuccess;
-				BuildResults = EngineBuildResults.Union(ProjectBuildResults).Distinct().ToDictionary(x => x.Key, x => x.Value);
+				BuildResults = new Dictionary<FileReference, CsProjBuildRecordEntry>(EngineBuildResults);
+				foreach (KeyValuePair<FileReference, CsProjBuildRecordEntry> Item in ProjectBuildResults)
+				{
+					if (!BuildResults.ContainsKey(Item.Key))
+					{
+						BuildResults.Add(Item.Key, Item.Value);
+					}
+				}
 			}
 			else
 			{
