Horde: Handle 404 responses correctly when reading refs through HTTP storage backend.

diff --git a/Engine/Source/Programs/Shared/EpicGames.Horde/Storage/Backends/HttpStorageClient.cs b/Engine/Source/Programs/Shared/EpicGames.Horde/Storage/Backends/HttpStorageClient.cs
--- a/Engine/Source/Programs/Shared/EpicGames.Horde/Storage/Backends/HttpStorageClient.cs
+++ b/Engine/Source/Programs/Shared/EpicGames.Horde/Storage/Backends/HttpStorageClient.cs
@@ -150,9 +150,16 @@ namespace EpicGames.Horde.Storage.Backends
 			{
 				using (HttpResponseMessage response = await _httpClient.SendAsync(request, cancellationToken))
 				{
-					response.EnsureSuccessStatusCode();
-					ReadRefResponse? data = await response.Content.ReadFromJsonAsync<ReadRefResponse>(cancellationToken: cancellationToken);
-					return new NodeLocator(data!.Blob, data!.ExportIdx);
+					if (response.StatusCode == HttpStatusCode.NotFound)
+					{
+						return default;
+					}
+					else
+					{
+						response.EnsureSuccessStatusCode();
+						ReadRefResponse? data = await response.Content.ReadFromJsonAsync<ReadRefResponse>(cancellationToken: cancellationToken);
+						return new NodeLocator(data!.Blob, data!.ExportIdx);
+					}
 				}
 			}
 		}
