Horde: Additional logging for HTTP storage client.

diff --git a/Engine/Source/Programs/Shared/EpicGames.Horde/Storage/Backends/HttpStorageClient.cs b/Engine/Source/Programs/Shared/EpicGames.Horde/Storage/Backends/HttpStorageClient.cs
--- a/Engine/Source/Programs/Shared/EpicGames.Horde/Storage/Backends/HttpStorageClient.cs
+++ b/Engine/Source/Programs/Shared/EpicGames.Horde/Storage/Backends/HttpStorageClient.cs
@@ -57,6 +57,7 @@ namespace EpicGames.Horde.Storage.Backends
 		/// <inheritdoc/>
 		public override async Task<Stream> ReadBlobAsync(BlobLocator locator, CancellationToken cancellationToken = default)
 		{
+			_logger.LogDebug("Reading {Locator}", locator);
 			using (HttpRequestMessage request = new HttpRequestMessage(HttpMethod.Get, $"api/v1/storage/{_namespaceId}/blobs/{locator}"))
 			{
 				HttpResponseMessage response = await _httpClient.SendAsync(request, cancellationToken);
@@ -68,6 +69,7 @@ namespace EpicGames.Horde.Storage.Backends
 		/// <inheritdoc/>
 		public override async Task<Stream> ReadBlobRangeAsync(BlobLocator locator, int offset, int length, CancellationToken cancellationToken = default)
 		{
+			_logger.LogDebug("Reading {Locator} ({Offset}+{Length})", locator, offset, length);
 			using (HttpRequestMessage request = new HttpRequestMessage(HttpMethod.Get, $"api/v1/storage/{_namespaceId}/blobs/{locator}?offset={offset}&length={length}"))
 			{
 				HttpResponseMessage response = await _httpClient.SendAsync(request, cancellationToken);
@@ -84,9 +86,8 @@ namespace EpicGames.Horde.Storage.Backends
 			if (_supportsUploadRedirects)
 			{
 				WriteBlobResponse redirectResponse = await SendWriteRequestAsync(null, prefix, cancellationToken);
-				if(redirectResponse.UploadUrl != null)
+				if (redirectResponse.UploadUrl != null)
 				{
-					_logger.LogDebug("Using upload redirect for {Locator}", redirectResponse.Locator);
 					using (HttpResponseMessage uploadResponse = await _redirectHttpClient.PutAsync(redirectResponse.UploadUrl, streamContent, cancellationToken))
 					{
 						if (!uploadResponse.IsSuccessStatusCode)
@@ -95,12 +96,14 @@ namespace EpicGames.Horde.Storage.Backends
 							throw new StorageException($"Unable to upload data to redirected URL: {body}", null);
 						}
 					}
+					_logger.LogDebug("Written {Locator} (using redirect)", redirectResponse.Locator);
 					return redirectResponse.Locator;
 				}
 			}
 
 			WriteBlobResponse response = await SendWriteRequestAsync(streamContent, prefix, cancellationToken);
 			_supportsUploadRedirects = response.SupportsRedirects ?? false;
+			_logger.LogDebug("Written {Locator} (direct)", response.Locator);
 			return response.Locator;
 		}
 
@@ -134,6 +137,7 @@ namespace EpicGames.Horde.Storage.Backends
 		/// <inheritdoc/>
 		public override async Task DeleteRefAsync(RefName name, CancellationToken cancellationToken)
 		{
+			_logger.LogDebug("Deleting ref {RefName}", name);
 			using (HttpRequestMessage request = new HttpRequestMessage(HttpMethod.Delete, $"api/v1/storage/{_namespaceId}/refs/{name}"))
 			{
 				using (HttpResponseMessage response = await _httpClient.SendAsync(request, cancellationToken))
@@ -152,12 +156,19 @@ namespace EpicGames.Horde.Storage.Backends
 				{
 					if (response.StatusCode == HttpStatusCode.NotFound)
 					{
+						_logger.LogDebug("Read ref {RefName} -> None", name);
 						return default;
 					}
+					else if (!response.IsSuccessStatusCode)
+					{
+						_logger.LogError("Unable to read ref {RefName} (status: {StatusCode}, body: {Body})", name, response.StatusCode, await response.Content.ReadAsStringAsync(cancellationToken));
+						throw new StorageException($"Unable to read ref '{name}'", null);
+					}
 					else
 					{
 						response.EnsureSuccessStatusCode();
 						ReadRefResponse? data = await response.Content.ReadFromJsonAsync<ReadRefResponse>(cancellationToken: cancellationToken);
+						_logger.LogDebug("Read ref {RefName} -> {Blob}#{ExportIdx}", name, data!.Blob, data!.ExportIdx);
 						return new NodeLocator(data!.Blob, data!.ExportIdx);
 					}
 				}
@@ -167,6 +178,7 @@ namespace EpicGames.Horde.Storage.Backends
 		/// <inheritdoc/>
 		public override async Task WriteRefTargetAsync(RefName name, NodeLocator target, RefOptions? options = null, CancellationToken cancellationToken = default)
 		{
+			_logger.LogDebug("Setting ref {RefName} -> {Blob}#{ExportIdx}", name, target.Blob, target.ExportIdx);
 			using (HttpResponseMessage response = await _httpClient.PutAsync($"api/v1/storage/{_namespaceId}/refs/{name}", new { locator = target.Blob, exportIdx = target.ExportIdx, options }, cancellationToken))
 			{
 				response.EnsureSuccessStatusCode();
