Horde: Remove bundle read/write methods from IStorageClient. We don't need overridable interfaces for this functionality (and polluting the cache with GC scrapes of header data is likely to be a anti-feature).

diff --git a/Engine/Source/Programs/Horde/Horde.Build/Storage/GcService.cs b/Engine/Source/Programs/Horde/Horde.Build/Storage/GcService.cs
--- a/Engine/Source/Programs/Horde/Horde.Build/Storage/GcService.cs
+++ b/Engine/Source/Programs/Horde/Horde.Build/Storage/GcService.cs
@@ -20,6 +20,7 @@ using MongoDB.Driver;
 using System.Runtime.CompilerServices;
 using System.IO;
 using System.Diagnostics;
+using System.Buffers;
 
 namespace Horde.Build.Storage
 {
@@ -275,7 +276,7 @@ namespace Horde.Build.Storage
 			List<WriteModel<GcNode>> updates = new List<WriteModel<GcNode>>();
 			node.Imports = new List<IoHash>();
 
-			BundleHeader? header = await _store.ReadBundleHeaderAsync(node.Locator, cancellationToken);
+			BundleHeader? header = await ReadHeaderAsync(_store, node.Locator, cancellationToken);
 			if (header != null)
 			{
 				foreach (BundleImport import in header.Imports)
@@ -295,6 +296,47 @@ namespace Horde.Build.Storage
 			await _nodeCollection.BulkWriteAsync(updates, new BulkWriteOptions { IsOrdered = false }, cancellationToken);
 		}
 
+		protected async Task<BundleHeader?> ReadHeaderAsync(IStorageClient store, BlobLocator locator, CancellationToken cancellationToken)
+		{
+			int fetchSize = 64 * 1024;
+			for (; ; )
+			{
+				using (IMemoryOwner<byte> owner = MemoryPool<byte>.Shared.Rent(fetchSize))
+				{
+					using (Stream? stream = await store.ReadBlobRangeAsync(locator, 0, fetchSize, cancellationToken))
+					{
+						if (stream == null)
+						{
+							_logger.LogError("Unable to read blob {Blob}", locator);
+							return null;
+						}
+
+						// Read the start of the blob
+						Memory<byte> memory = owner.Memory.Slice(0, fetchSize);
+
+						int length = await stream.ReadGreedyAsync(memory, cancellationToken);
+						if (length < BundleHeader.PreludeLength)
+						{
+							_logger.LogError("Blob {Blob} does not have a valid prelude", locator);
+							return null;
+						}
+
+						memory = memory.Slice(0, length);
+
+						// Make sure it's large enough to hold the header
+						int headerSize = BundleHeader.ReadPrelude(memory);
+						if (headerSize <= fetchSize)
+						{
+							return new BundleHeader(new MemoryReader(memory));
+						}
+
+						// Increase the fetch size and retry
+						fetchSize = headerSize;
+					}
+				}
+			}
+		}
+
 		WriteModel<GcNode> CreateGcNodeUpsert(IoHash hash, BlobLocator locator)
 		{
 			FilterDefinition<GcNode> filter = Builders<GcNode>.Filter.Eq(x => x.Id, hash);
diff --git a/Engine/Source/Programs/Shared/EpicGames.Horde/Storage/IStorageClient.cs b/Engine/Source/Programs/Shared/EpicGames.Horde/Storage/IStorageClient.cs
--- a/Engine/Source/Programs/Shared/EpicGames.Horde/Storage/IStorageClient.cs
+++ b/Engine/Source/Programs/Shared/EpicGames.Horde/Storage/IStorageClient.cs
@@ -143,35 +143,6 @@ namespace EpicGames.Horde.Storage
 
 		#endregion
 
-		#region Bundles
-
-		/// <summary>
-		/// Reads data for a bundle from the store
-		/// </summary>
-		/// <param name="locator">The blob identifier</param>
-		/// <param name="cancellationToken">Cancellation token for the operation</param>
-		/// <returns></returns>
-		Task<Bundle> ReadBundleAsync(BlobLocator locator, CancellationToken cancellationToken = default);
-
-		/// <summary>
-		/// Reads the header for a bundle from the store
-		/// </summary>
-		/// <param name="locator">The blob identifier</param>
-		/// <param name="cancellationToken">Cancellation token for the operation</param>
-		/// <returns></returns>
-		Task<BundleHeader> ReadBundleHeaderAsync(BlobLocator locator, CancellationToken cancellationToken = default);
-
-		/// <summary>
-		/// Writes a new bundle to the store
-		/// </summary>
-		/// <param name="bundle">Bundle data</param>
-		/// <param name="cancellationToken">Cancellation token for the operation</param>
-		/// <param name="prefix">Prefix for blob names. While the returned BlobId is guaranteed to be unique, this name can be used as a prefix to aid debugging.</param>
-		/// <returns>Unique identifier for the blob</returns>
-		Task<BlobLocator> WriteBundleAsync(Bundle bundle, Utf8String prefix = default, CancellationToken cancellationToken = default);
-
-		#endregion
-
 		#region Nodes
 
 		/// <summary>
@@ -267,19 +238,6 @@ namespace EpicGames.Horde.Storage
 
 			#endregion
 
-			#region Bundles
-
-			/// <inheritdoc/>
-			public Task<Bundle> ReadBundleAsync(BlobLocator locator, CancellationToken cancellationToken = default) => _inner.ReadBundleAsync(locator, cancellationToken);
-
-			/// <inheritdoc/>
-			public Task<BundleHeader> ReadBundleHeaderAsync(BlobLocator locator, CancellationToken cancellationToken = default) => _inner.ReadBundleHeaderAsync(locator, cancellationToken);
-
-			/// <inheritdoc/>
-			public Task<BlobLocator> WriteBundleAsync(Bundle bundle, Utf8String prefix = default, CancellationToken cancellationToken = default) => _inner.WriteBundleAsync(bundle, prefix, cancellationToken);
-
-			#endregion
-
 			#region Trees
 
 			/// <inheritdoc/>
@@ -314,6 +272,39 @@ namespace EpicGames.Horde.Storage
 		/// <returns></returns>
 		public static IStorageClient<T> ForType<T>(this IStorageClient blobStore) => new TypedStorageClient<T>(blobStore);
 
+		#region Bundles
+
+		/// <summary>
+		/// Reads a bundle from the given blob id, or retrieves it from the cache
+		/// </summary>
+		/// <param name="store">The store instance to read from</param>
+		/// <param name="locator"></param>
+		/// <param name="cancellationToken"></param>
+		/// <returns></returns>
+		public static async Task<Bundle> ReadBundleAsync(this IStorageClient store, BlobLocator locator, CancellationToken cancellationToken = default)
+		{
+			using (Stream stream = await store.ReadBlobAsync(locator, cancellationToken))
+			{
+				return await Bundle.FromStreamAsync(stream, cancellationToken);
+			}
+		}
+
+		/// <summary>
+		/// Writes a new bundle to the store
+		/// </summary>
+		/// <param name="store">The store instance to write to</param>
+		/// <param name="bundle">Bundle data</param>
+		/// <param name="cancellationToken">Cancellation token for the operation</param>
+		/// <param name="prefix">Prefix for blob names. While the returned BlobId is guaranteed to be unique, this name can be used as a prefix to aid debugging.</param>
+		/// <returns>Unique identifier for the blob</returns>
+		public static async Task<BlobLocator> WriteBundleAsync(this IStorageClient store, Bundle bundle, Utf8String prefix = default, CancellationToken cancellationToken = default)
+		{
+			using ReadOnlySequenceStream stream = new ReadOnlySequenceStream(bundle.AsSequence());
+			return await store.WriteBlobAsync(stream, prefix, cancellationToken);
+		}
+
+		#endregion
+
 		#region Nodes
 
 		/// <summary>
diff --git a/Engine/Source/Programs/Shared/EpicGames.Horde/Storage/StorageClientBase.cs b/Engine/Source/Programs/Shared/EpicGames.Horde/Storage/StorageClientBase.cs
--- a/Engine/Source/Programs/Shared/EpicGames.Horde/Storage/StorageClientBase.cs
+++ b/Engine/Source/Programs/Shared/EpicGames.Horde/Storage/StorageClientBase.cs
@@ -398,27 +398,13 @@ namespace EpicGames.Horde.Storage
 			}
 		}
 
-		/// <summary>
-		/// Reads a bundle from the given blob id, or retrieves it from the cache
-		/// </summary>
-		/// <param name="locator"></param>
-		/// <param name="cancellationToken"></param>
-		/// <returns></returns>
-		public async Task<Bundle> ReadBundleAsync(BlobLocator locator, CancellationToken cancellationToken = default)
-		{
-			using (Stream stream = await ReadBlobAsync(locator, cancellationToken))
-			{
-				return await Bundle.FromStreamAsync(stream, cancellationToken);
-			}
-		}
-
 		/// <summary>
 		/// Reads a bundle header from the given blob locator, or retrieves it from the cache
 		/// </summary>
 		/// <param name="locator"></param>
 		/// <param name="cancellationToken"></param>
 		/// <returns></returns>
-		public async Task<BundleHeader> ReadBundleHeaderAsync(BlobLocator locator, CancellationToken cancellationToken = default)
+		async Task<BundleHeader> ReadBundleHeaderAsync(BlobLocator locator, CancellationToken cancellationToken = default)
 		{
 			BundleInfo info = await GetBundleInfoAsync(locator, cancellationToken);
 			return info.Header;
@@ -554,13 +540,6 @@ namespace EpicGames.Horde.Storage
 			return await queuedPacket.CompletionSource.Task;
 		}
 
-		/// <inheritdoc/>
-		public async Task<BlobLocator> WriteBundleAsync(Bundle bundle, Utf8String prefix = default, CancellationToken cancellationToken = default)
-		{
-			using ReadOnlySequenceStream stream = new ReadOnlySequenceStream(bundle.AsSequence());
-			return await WriteBlobAsync(stream, prefix, cancellationToken);
-		}
-
 		#endregion
 
 		#region Nodes
@@ -627,7 +606,7 @@ namespace EpicGames.Horde.Storage
 		/// <inheritdoc/>
 		public virtual async Task<NodeLocator> WriteRefAsync(RefName name, Bundle bundle, int exportIdx = 0, Utf8String prefix = default, RefOptions? options = null, CancellationToken cancellationToken = default)
 		{
-			BlobLocator locator = await WriteBundleAsync(bundle, prefix, cancellationToken);
+			BlobLocator locator = await this.WriteBundleAsync(bundle, prefix, cancellationToken);
 			NodeLocator target = new NodeLocator(locator, exportIdx);
 			await WriteRefTargetAsync(name, target, options, cancellationToken);
 			return target;
