Fixed C# UHT issue where if a child class has the same sparse data as the parent, the accessor methods would not be generated in the child class.

diff --git a/Engine/Source/Programs/Shared/EpicGames.UHT/Exporters/CodeGen/UhtHeaderCodeGeneratorHFile.cs b/Engine/Source/Programs/Shared/EpicGames.UHT/Exporters/CodeGen/UhtHeaderCodeGeneratorHFile.cs
--- a/Engine/Source/Programs/Shared/EpicGames.UHT/Exporters/CodeGen/UhtHeaderCodeGeneratorHFile.cs
+++ b/Engine/Source/Programs/Shared/EpicGames.UHT/Exporters/CodeGen/UhtHeaderCodeGeneratorHFile.cs
@@ -549,53 +549,53 @@ namespace EpicGames.UHT.Exporters.CodeGen
 							builder.Append(api).Append("const F").Append(sparseDataType).Append("* Get").Append(sparseDataType).Append("(EGetSparseClassDataMethod GetMethod) const; \\\r\n");
 							builder.Append(api).Append("static UScriptStruct* StaticGet").Append(sparseDataType).Append("ScriptStruct(); \\\r\n");
 						}
+					}
 
-						foreach (string sparseDataType in sparseDataTypes)
+					foreach (string sparseDataType in sparseDataTypes)
+					{
+						UhtScriptStruct? sparseScriptStruct = classObj.FindType(UhtFindOptions.EngineName | UhtFindOptions.ScriptStruct | UhtFindOptions.NoSelf, sparseDataType) as UhtScriptStruct;
+						while (sparseScriptStruct != null)
 						{
-							UhtScriptStruct? sparseScriptStruct = classObj.FindType(UhtFindOptions.EngineName | UhtFindOptions.ScriptStruct | UhtFindOptions.NoSelf, sparseDataType) as UhtScriptStruct;
-							while (sparseScriptStruct != null)
+							foreach (UhtProperty sparseProperty in sparseScriptStruct.Properties)
 							{
-								foreach (UhtProperty sparseProperty in sparseScriptStruct.Properties)
+								if (!sparseProperty.MetaData.ContainsKey(UhtNames.NoGetter))
 								{
-									if (!sparseProperty.MetaData.ContainsKey(UhtNames.NoGetter))
+									string propertyName = sparseProperty.SourceName;
+									string cleanPropertyName = propertyName;
+									if (sparseProperty is UhtBoolProperty && propertyName.StartsWith("b", StringComparison.Ordinal))
 									{
-										string propertyName = sparseProperty.SourceName;
-										string cleanPropertyName = propertyName;
-										if (sparseProperty is UhtBoolProperty && propertyName.StartsWith("b", StringComparison.Ordinal))
-										{
-											cleanPropertyName = propertyName[1..];
-										}
-
-										bool getByRef = sparseProperty.MetaData.ContainsKey(UhtNames.GetByRef);
-
-										if (getByRef)
-										{
-											builder.Append("const ").AppendSparse(sparseProperty).Append("& Get").Append(cleanPropertyName).Append("() \\\r\n");
-										}
-										else
-										{
-											builder.AppendSparse(sparseProperty).Append(" Get").Append(cleanPropertyName).Append("() \\\r\n");
-										}
-										builder.Append("{ \\\r\n");
-										builder.Append("\treturn Get").Append(sparseDataType).Append("()->").Append(propertyName).Append("; \\\r\n");
-										builder.Append("} \\\r\n");
-
-										if (getByRef)
-										{
-											builder.Append("const ").AppendSparse(sparseProperty).Append("& Get").Append(cleanPropertyName).Append("() const \\\r\n");
-										}
-										else
-										{
-											builder.AppendSparse(sparseProperty).Append(" Get").Append(cleanPropertyName).Append("() const \\\r\n");
-										}
-										builder.Append("{ \\\r\n");
-										builder.Append("\treturn Get").Append(sparseDataType).Append("()->").Append(propertyName).Append("; \\\r\n");
-										builder.Append("} \\\r\n");
+										cleanPropertyName = propertyName[1..];
 									}
-								}
 
-								sparseScriptStruct = sparseScriptStruct.SuperScriptStruct;
+									bool getByRef = sparseProperty.MetaData.ContainsKey(UhtNames.GetByRef);
+
+									if (getByRef)
+									{
+										builder.Append("const ").AppendSparse(sparseProperty).Append("& Get").Append(cleanPropertyName).Append("() \\\r\n");
+									}
+									else
+									{
+										builder.AppendSparse(sparseProperty).Append(" Get").Append(cleanPropertyName).Append("() \\\r\n");
+									}
+									builder.Append("{ \\\r\n");
+									builder.Append("\treturn Get").Append(sparseDataType).Append("()->").Append(propertyName).Append("; \\\r\n");
+									builder.Append("} \\\r\n");
+
+									if (getByRef)
+									{
+										builder.Append("const ").AppendSparse(sparseProperty).Append("& Get").Append(cleanPropertyName).Append("() const \\\r\n");
+									}
+									else
+									{
+										builder.AppendSparse(sparseProperty).Append(" Get").Append(cleanPropertyName).Append("() const \\\r\n");
+									}
+									builder.Append("{ \\\r\n");
+									builder.Append("\treturn Get").Append(sparseDataType).Append("()->").Append(propertyName).Append("; \\\r\n");
+									builder.Append("} \\\r\n");
+								}
 							}
+
+							sparseScriptStruct = sparseScriptStruct.SuperScriptStruct;
 						}
 					}
 				}
