Fixed issue where C# UHT didn't require a trailing ';' at the end of a delegate declaration where C++ UHT does.

diff --git a/Engine/Source/Programs/Shared/EpicGames.UHT/Parsers/UhtFunctionParser.cs b/Engine/Source/Programs/Shared/EpicGames.UHT/Parsers/UhtFunctionParser.cs
--- a/Engine/Source/Programs/Shared/EpicGames.UHT/Parsers/UhtFunctionParser.cs
+++ b/Engine/Source/Programs/Shared/EpicGames.UHT/Parsers/UhtFunctionParser.cs
@@ -290,8 +290,8 @@ namespace EpicGames.UHT.Parsers
 
 					topScope.AddFormattedCommentsAsTooltipMetaData();
 
-					// Optionally consume a semicolon, it's not required for the delegate macro since it contains one internally
-					topScope.TokenReader.Optional(';');
+					// Consume a semicolon, it's not required for the delegate macro since it contains one internally
+					topScope.TokenReader.Require(';');
 				}
 				return UhtParseResult.Handled;
 			}
