Control Rig: Fix TEnumAsByte not casting correctly in the generated headers
#udn 00490176

diff --git a/Engine/Source/Developer/RigVMDeveloper/Private/RigVMDeveloperTypeUtils.cpp b/Engine/Source/Developer/RigVMDeveloper/Private/RigVMDeveloperTypeUtils.cpp
--- a/Engine/Source/Developer/RigVMDeveloper/Private/RigVMDeveloperTypeUtils.cpp
+++ b/Engine/Source/Developer/RigVMDeveloper/Private/RigVMDeveloperTypeUtils.cpp
@@ -37,6 +37,11 @@ FRigVMExternalVariable RigVMTypeUtils::ExternalVariableFromRigVMVariableDescript
 		ExternalVariable.TypeObject = InVariableDescription.CPPTypeObject;
 	}
 
+	if (UEnum* Enum = Cast<UEnum>(ExternalVariable.TypeObject))
+	{
+		ExternalVariable.TypeName = Enum->GetFName();
+	}	
+
 	ExternalVariable.bIsPublic = false;
 	ExternalVariable.bIsReadOnly = false;
 	ExternalVariable.Memory = nullptr;
diff --git a/Engine/Source/Programs/Shared/EpicGames.UHT/Exporters/CodeGen/UhtHeaderCodeGeneratorHFile.cs b/Engine/Source/Programs/Shared/EpicGames.UHT/Exporters/CodeGen/UhtHeaderCodeGeneratorHFile.cs
--- a/Engine/Source/Programs/Shared/EpicGames.UHT/Exporters/CodeGen/UhtHeaderCodeGeneratorHFile.cs
+++ b/Engine/Source/Programs/Shared/EpicGames.UHT/Exporters/CodeGen/UhtHeaderCodeGeneratorHFile.cs
@@ -216,6 +216,11 @@ namespace EpicGames.UHT.Exporters.CodeGen
 										string variableType = parameter.TypeVariableRef(true);
 										string extractedType = parameter.TypeOriginal();
 
+										if (parameter.IsEnumAsByte)
+										{
+											extractedType = parameter.TypeOriginal(true);
+										}
+
 										builder
 											.Append("\t\t")
 											.Append(variableType)
diff --git a/Engine/Source/Programs/Shared/EpicGames.UHT/Types/Properties/UhtByteProperty.cs b/Engine/Source/Programs/Shared/EpicGames.UHT/Types/Properties/UhtByteProperty.cs
--- a/Engine/Source/Programs/Shared/EpicGames.UHT/Types/Properties/UhtByteProperty.cs
+++ b/Engine/Source/Programs/Shared/EpicGames.UHT/Types/Properties/UhtByteProperty.cs
@@ -51,6 +51,7 @@ namespace EpicGames.UHT.Types
 			if (Enum != null)
 			{
 				PropertyCaps |= UhtPropertyCaps.IsRigVMEnum;
+				PropertyCaps |= UhtPropertyCaps.IsRigVMEnumAsByte;
 			}
 		}
 
diff --git a/Engine/Source/Programs/Shared/EpicGames.UHT/Types/UhtProperty.cs b/Engine/Source/Programs/Shared/EpicGames.UHT/Types/UhtProperty.cs
--- a/Engine/Source/Programs/Shared/EpicGames.UHT/Types/UhtProperty.cs
+++ b/Engine/Source/Programs/Shared/EpicGames.UHT/Types/UhtProperty.cs
@@ -400,6 +400,11 @@ namespace EpicGames.UHT.Types
 		/// True if the property should codegen as an array
 		/// </summary>
 		IsRigVMArray = 1 << 14,
+		
+		/// <summary>
+		/// True if the property should codegen as an array
+		/// </summary>
+		IsRigVMEnumAsByte = 1 << 15,
 	}
 
 	/// <summary>
diff --git a/Engine/Source/Programs/Shared/EpicGames.UHT/Types/UhtScriptStruct.cs b/Engine/Source/Programs/Shared/EpicGames.UHT/Types/UhtScriptStruct.cs
--- a/Engine/Source/Programs/Shared/EpicGames.UHT/Types/UhtScriptStruct.cs
+++ b/Engine/Source/Programs/Shared/EpicGames.UHT/Types/UhtScriptStruct.cs
@@ -50,6 +50,11 @@ namespace EpicGames.UHT.Types
 		/// Set if the property is an enum
 		/// </summary>
 		IsEnum = 0x00010000,
+		
+		/// <summary>
+		/// Set if the property is an enum as byte
+		/// </summary>
+		IsEnumAsByte = 0x00010004,
 
 		/// <summary>
 		/// Set if the property is an array
@@ -176,6 +181,12 @@ namespace EpicGames.UHT.Types
 		/// </summary>
 		[JsonIgnore]
 		public bool IsEnum => ParameterFlags.HasAnyFlags(UhtRigVMParameterFlags.IsEnum);
+		
+		/// <summary>
+		/// True if the parameter is an enum as byte
+		/// </summary>
+		[JsonIgnore]
+		public bool IsEnumAsByte => ParameterFlags.HasAnyFlags(UhtRigVMParameterFlags.IsEnumAsByte);
 
 		/// <summary>
 		/// True if the parameter is an array
@@ -217,6 +228,11 @@ namespace EpicGames.UHT.Types
 				{
 					ParameterFlags |= UhtRigVMParameterFlags.IsEnum;
 				}
+				if (property.PropertyCaps.HasAnyFlags(UhtPropertyCaps.IsRigVMEnumAsByte))
+				{
+					ParameterFlags |= UhtRigVMParameterFlags.IsEnumAsByte;
+					CastType = $"TEnumAsByte<{Type}>";
+				}
 				if (property.PropertyCaps.HasAnyFlags(UhtPropertyCaps.IsRigVMArray))
 				{
 					ParameterFlags |= UhtRigVMParameterFlags.IsArray;
diff --git a/Engine/Source/Programs/UnrealHeaderTool/Private/CodeGenerator.cpp b/Engine/Source/Programs/UnrealHeaderTool/Private/CodeGenerator.cpp
--- a/Engine/Source/Programs/UnrealHeaderTool/Private/CodeGenerator.cpp
+++ b/Engine/Source/Programs/UnrealHeaderTool/Private/CodeGenerator.cpp
@@ -3457,6 +3457,11 @@ void FNativeClassHeaderGenerator::ExportGeneratedStructBodyMacros(FOutputDevice&
 			{
 				FString VariableType = Parameter.TypeVariableRef(true);
 				FString ExtractedType = Parameter.TypeOriginal();
+
+				if (Parameter.bIsEnumAsByte)
+				{
+					ExtractedType = Parameter.TypeOriginal(true);
+				}
 				FString ParameterCast = FString::Printf(TEXT("*(%s*)"), *ExtractedType);
 
 				RigVMStubProlog.Add(FString::Printf(TEXT("%s %s = %sRigVMMemoryHandles[%d].GetData(false%s);"),
diff --git a/Engine/Source/Programs/UnrealHeaderTool/Private/HeaderParser.cpp b/Engine/Source/Programs/UnrealHeaderTool/Private/HeaderParser.cpp
--- a/Engine/Source/Programs/UnrealHeaderTool/Private/HeaderParser.cpp
+++ b/Engine/Source/Programs/UnrealHeaderTool/Private/HeaderParser.cpp
@@ -6651,6 +6651,7 @@ void FHeaderParser::ParseRigVMMethodParameters(FUnrealStructDefinitionInfo& Stru
 	// validate the property types for this struct
 	for (FUnrealPropertyDefinitionInfo* PropertyDef : TUHTFieldRange<FUnrealPropertyDefinitionInfo>(StructDef))
 	{
+		FRigVMParameter Parameter;
 		FString MemberCPPType;
 		FString ExtendedCPPType;
 		MemberCPPType = PropertyDef->GetCPPType(&ExtendedCPPType);
@@ -6658,9 +6659,13 @@ void FHeaderParser::ParseRigVMMethodParameters(FUnrealStructDefinitionInfo& Stru
 		if (ExtendedCPPType.IsEmpty() && MemberCPPType.StartsWith(TEnumAsByteText))
 		{
 			MemberCPPType = MemberCPPType.LeftChop(1).RightChop(12);
+			Parameter.bIsEnumAsByte = true;
+		}
+		else
+		{
+			Parameter.bIsEnumAsByte = false;
 		}
 
-		FRigVMParameter Parameter;
 		Parameter.Name = PropertyDef->GetName();
 		Parameter.Type = MemberCPPType + ExtendedCPPType;
 		Parameter.bConstant = PropertyDef->HasMetaData(NAME_ConstantText);
diff --git a/Engine/Source/Programs/UnrealHeaderTool/Private/ParserHelper.h b/Engine/Source/Programs/UnrealHeaderTool/Private/ParserHelper.h
--- a/Engine/Source/Programs/UnrealHeaderTool/Private/ParserHelper.h
+++ b/Engine/Source/Programs/UnrealHeaderTool/Private/ParserHelper.h
@@ -532,6 +532,7 @@ struct FRigVMParameter
 		, CastType()
 		, bEditorOnly(false)
 		, bIsEnum(false)
+		, bIsEnumAsByte(false)
 	{
 	}
 
@@ -546,6 +547,7 @@ struct FRigVMParameter
 	FString CastType;
 	bool bEditorOnly;
 	bool bIsEnum;
+	bool bIsEnumAsByte;
 
 	const FString& NameOriginal(bool bCastName = false) const
 	{
