[UHT]
* Fixed so delegate wrapper function is implemented in cpp file to be able to forward declare types used by delegate

diff --git a/Engine/Plugins/Runtime/AudioWidgets/Source/AudioWidgets/Public/AudioSlider.h b/Engine/Plugins/Runtime/AudioWidgets/Source/AudioWidgets/Public/AudioSlider.h
--- a/Engine/Plugins/Runtime/AudioWidgets/Source/AudioWidgets/Public/AudioSlider.h
+++ b/Engine/Plugins/Runtime/AudioWidgets/Source/AudioWidgets/Public/AudioSlider.h
@@ -4,6 +4,7 @@
 
 #include "AudioDefines.h"
 #include "AudioWidgetsSlateTypes.h"
+#include "Components/Slider.h"
 #include "Components/Widget.h"
 #include "Curves/CurveFloat.h"
 #include "Styling/StyleColors.h"
@@ -14,8 +15,6 @@
 
 class SAudioSliderBase;
 
-DECLARE_DYNAMIC_MULTICAST_DELEGATE_OneParam(FOnFloatValueChangedEvent, float, Value);
-
 /**
  * An audio slider widget. 
  */
diff --git a/Engine/Plugins/Runtime/Synthesis/Source/Synthesis/Public/UI/SynthKnob.h b/Engine/Plugins/Runtime/Synthesis/Source/Synthesis/Public/UI/SynthKnob.h
--- a/Engine/Plugins/Runtime/Synthesis/Source/Synthesis/Public/UI/SynthKnob.h
+++ b/Engine/Plugins/Runtime/Synthesis/Source/Synthesis/Public/UI/SynthKnob.h
@@ -10,19 +10,12 @@
 #include "UObject/ObjectMacros.h"
 #include "Styling/SlateTypes.h"
 #include "Widgets/SWidget.h"
+#include "Components/Slider.h"
 #include "Components/TextBlock.h"
 #include "Components/Widget.h"
 #include "Components/ContentWidget.h"
 #include "SynthKnob.generated.h"
 
-
-DECLARE_DYNAMIC_MULTICAST_DELEGATE(FOnMouseCaptureBeginEvent);
-DECLARE_DYNAMIC_MULTICAST_DELEGATE(FOnMouseCaptureEndEvent);
-DECLARE_DYNAMIC_MULTICAST_DELEGATE(FOnControllerCaptureBeginEvent);
-DECLARE_DYNAMIC_MULTICAST_DELEGATE(FOnControllerCaptureEndEvent);
-DECLARE_DYNAMIC_MULTICAST_DELEGATE_OneParam(FOnFloatValueChangedEvent, float, Value);
-
-
 /**
  * A simple widget that shows a sliding bar with a handle that allows you to control the value between 0..1.
  *
diff --git a/Engine/Source/Programs/Shared/EpicGames.UHT/Exporters/CodeGen/UhtHeaderCodeGenerator.cs b/Engine/Source/Programs/Shared/EpicGames.UHT/Exporters/CodeGen/UhtHeaderCodeGenerator.cs
--- a/Engine/Source/Programs/Shared/EpicGames.UHT/Exporters/CodeGen/UhtHeaderCodeGenerator.cs
+++ b/Engine/Source/Programs/Shared/EpicGames.UHT/Exporters/CodeGen/UhtHeaderCodeGenerator.cs
@@ -275,7 +275,7 @@ namespace EpicGames.UHT.Exporters.CodeGen
 			return builder;
 		}
 
-		protected static StringBuilder AppendEventFunctionPrologue(StringBuilder builder, UhtFunction function, string functionName, int tabs, string endl)
+		protected static StringBuilder AppendEventFunctionPrologue(StringBuilder builder, UhtFunction function, string functionName, int tabs, string endl, bool addEventParameterStruct)
 		{
 			builder.AppendTabs(tabs).Append('{').Append(endl);
 			if (function.Children.Count == 0)
@@ -285,6 +285,9 @@ namespace EpicGames.UHT.Exporters.CodeGen
 
 			++tabs;
 
+			if (addEventParameterStruct)
+				AppendEventParameter(builder, function, functionName, UhtPropertyTextType.EventParameterMember, true, tabs, "\r\n");
+
 			string eventParameterStructName = GetEventStructParametersName(function.Outer, functionName);
 
 			builder.AppendTabs(tabs).Append(eventParameterStructName).Append(" Parms;").Append(endl);
diff --git a/Engine/Source/Programs/Shared/EpicGames.UHT/Exporters/CodeGen/UhtHeaderCodeGeneratorCppFile.cs b/Engine/Source/Programs/Shared/EpicGames.UHT/Exporters/CodeGen/UhtHeaderCodeGeneratorCppFile.cs
--- a/Engine/Source/Programs/Shared/EpicGames.UHT/Exporters/CodeGen/UhtHeaderCodeGeneratorCppFile.cs
+++ b/Engine/Source/Programs/Shared/EpicGames.UHT/Exporters/CodeGen/UhtHeaderCodeGeneratorCppFile.cs
@@ -424,7 +424,7 @@ namespace EpicGames.UHT.Exporters.CodeGen
 
 				builder.Append("\t\tif (!").Append(registrationName).Append(".OuterSingleton)\r\n");
 				builder.Append("\t\t{\r\n");
-				builder.Append("\t\t\t").Append(registrationName).Append(".OuterSingleton = GetStaticEnum(").Append(singletonName).Append(", ")
+				builder.Append("\t\t\t").Append(registrationName).Append(".OuterSingleton = GetStaticEnum(").Append(singletonName).Append(", (UObject*)")
 					.Append(PackageSingletonName).Append("(), TEXT(\"").Append(enumObj.SourceName).Append("\"));\r\n");
 				builder.Append("\t\t}\r\n");
 				builder.Append("\t\treturn ").Append(registrationName).Append(".OuterSingleton;\r\n");
@@ -532,7 +532,7 @@ namespace EpicGames.UHT.Exporters.CodeGen
 					.Append(registrationName)
 					.Append(".OuterSingleton = GetStaticStruct(")
 					.Append(singletonName)
-					.Append(", ")
+					.Append(", (UObject*)")
 					.Append(PackageSingletonName)
 					.Append("(), TEXT(\"")
 					.Append(scriptStruct.EngineName)
@@ -751,6 +751,24 @@ namespace EpicGames.UHT.Exporters.CodeGen
 		private StringBuilder AppendDelegate(StringBuilder builder, UhtFunction function)
 		{
 			AppendFunction(builder, function, false);
+
+			int tabs = 0;
+			string strippedFunctionName = function.StrippedFunctionName;
+			string exportFunctionName = GetDelegateFunctionExportName(function);
+			string extraParameter = GetDelegateFunctionExtraParameter(function);
+
+			AppendNativeFunctionHeader(builder, function, UhtPropertyTextType.EventFunctionArgOrRetVal, false, exportFunctionName, extraParameter, UhtFunctionExportFlags.None, "\r\n");
+			AppendEventFunctionPrologue(builder, function, strippedFunctionName, tabs, "\r\n", true);
+			builder
+				.Append('\t')
+				.Append(strippedFunctionName)
+				.Append('.')
+				.Append(function.FunctionFlags.HasAnyFlags(EFunctionFlags.MulticastDelegate) ? "ProcessMulticastDelegate" : "ProcessDelegate")
+				.Append("<UObject>(")
+				.Append(function.Children.Count > 0 ? "&Parms" : "NULL")
+				.Append(");\r\n");
+			AppendEventFunctionEpilogue(builder, function, tabs, "\r\n");
+
 			return builder;
 		}
 
@@ -1769,7 +1787,7 @@ namespace EpicGames.UHT.Exporters.CodeGen
 						}
 						else
 						{
-							AppendEventFunctionPrologue(builder, function, function.EngineName, 1, "\r\n");
+							AppendEventFunctionPrologue(builder, function, function.EngineName, 1, "\r\n", false);
 
 							// Cast away const just in case, because ProcessEvent isn't const
 							builder.Append("\t\t");
diff --git a/Engine/Source/Programs/Shared/EpicGames.UHT/Exporters/CodeGen/UhtHeaderCodeGeneratorHFile.cs b/Engine/Source/Programs/Shared/EpicGames.UHT/Exporters/CodeGen/UhtHeaderCodeGeneratorHFile.cs
--- a/Engine/Source/Programs/Shared/EpicGames.UHT/Exporters/CodeGen/UhtHeaderCodeGeneratorHFile.cs
+++ b/Engine/Source/Programs/Shared/EpicGames.UHT/Exporters/CodeGen/UhtHeaderCodeGeneratorHFile.cs
@@ -317,24 +317,30 @@ namespace EpicGames.UHT.Exporters.CodeGen
 		{
 			using (UhtMacroCreator macro = new(builder, this, function, DelegateMacroSuffix))
 			{
-				int tabs = 0;
-				string strippedFunctionName = function.StrippedFunctionName;
 				string exportFunctionName = GetDelegateFunctionExportName(function);
 				string extraParameter = GetDelegateFunctionExtraParameter(function);
 
-				AppendEventParameter(builder, function, strippedFunctionName, UhtPropertyTextType.EventParameterMember, true, tabs, " \\\r\n");
-				builder.Append("static ");
-				AppendNativeFunctionHeader(builder, function, UhtPropertyTextType.EventFunctionArgOrRetVal, true, exportFunctionName, extraParameter, UhtFunctionExportFlags.Inline, " \\\r\n");
-				AppendEventFunctionPrologue(builder, function, strippedFunctionName, tabs, " \\\r\n");
-				builder
-					.Append('\t')
-					.Append(strippedFunctionName)
-					.Append('.')
-					.Append(function.FunctionFlags.HasAnyFlags(EFunctionFlags.MulticastDelegate) ? "ProcessMulticastDelegate" : "ProcessDelegate")
-					.Append("<UObject>(")
-					.Append(function.Children.Count > 0 ? "&Parms" : "NULL")
-					.Append("); \\\r\n");
-				AppendEventFunctionEpilogue(builder, function, tabs, " \\\r\n");
+				bool addAPI = true;
+
+				UhtClass? outerClass = function.Outer as UhtClass;
+				if (outerClass != null)
+				{
+					builder.Append("static ");
+					if (outerClass.ClassFlags.HasFlag(EClassFlags.RequiredAPI))
+					{
+						addAPI = false;
+					}
+				}
+
+				//if (!function.FunctionFlags.HasAnyFlags(EFunctionFlags.RequiredAPI)) // TODO: This requires too much fixup for now
+				//	addAPI = false
+
+				if (addAPI)
+				{
+					builder.Append(PackageApi);
+				}
+
+				AppendNativeFunctionHeader(builder, function, UhtPropertyTextType.EventFunctionArgOrRetVal, true, exportFunctionName, extraParameter, UhtFunctionExportFlags.None, "; \\\r\n");
 			}
 			return builder;
 		}
