Fixed C# UHT so that enumerations can not contain value names of TRUE or FALSE (case insensitive).

diff --git a/Engine/Source/Developer/AITestSuite/Classes/BehaviorTree/TestBTDecorator_Blueprint.h b/Engine/Source/Developer/AITestSuite/Classes/BehaviorTree/TestBTDecorator_Blueprint.h
--- a/Engine/Source/Developer/AITestSuite/Classes/BehaviorTree/TestBTDecorator_Blueprint.h
+++ b/Engine/Source/Developer/AITestSuite/Classes/BehaviorTree/TestBTDecorator_Blueprint.h
@@ -10,8 +10,8 @@ UENUM()
 enum class EBPConditionType
 {
 	NoCondition,
-	True,
-	False
+	TrueCondition,
+	FalseCondition
 };
 
 UCLASS(meta = (HiddenNode))
@@ -22,7 +22,7 @@ class UTestBTDecorator_Blueprint : public UBTDecorator_BlueprintBase
 public:
 
 	UPROPERTY()
-	EBPConditionType BPConditionType = EBPConditionType::True;
+	EBPConditionType BPConditionType = EBPConditionType::TrueCondition;
 	
 	UPROPERTY()
 	int32 LogIndexBecomeRelevant;
diff --git a/Engine/Source/Developer/AITestSuite/Private/BehaviorTree/TestBTDecorator_Blueprint.cpp b/Engine/Source/Developer/AITestSuite/Private/BehaviorTree/TestBTDecorator_Blueprint.cpp
--- a/Engine/Source/Developer/AITestSuite/Private/BehaviorTree/TestBTDecorator_Blueprint.cpp
+++ b/Engine/Source/Developer/AITestSuite/Private/BehaviorTree/TestBTDecorator_Blueprint.cpp
@@ -27,7 +27,7 @@ void UTestBTDecorator_Blueprint::OnCeaseRelevant(UBehaviorTreeComponent& OwnerCo
 bool UTestBTDecorator_Blueprint::CalculateRawConditionValue(UBehaviorTreeComponent& OwnerComp, uint8* NodeMemory) const
 {
 	LogExecution(LogIndexCalculate);
-	return BPConditionType == EBPConditionType::True ? true : false;
+	return BPConditionType == EBPConditionType::TrueCondition ? true : false;
 }
 
 void UTestBTDecorator_Blueprint::PostLoad()
diff --git a/Engine/Source/Developer/AITestSuite/Private/Tests/BTTest.cpp b/Engine/Source/Developer/AITestSuite/Private/Tests/BTTest.cpp
--- a/Engine/Source/Developer/AITestSuite/Private/Tests/BTTest.cpp
+++ b/Engine/Source/Developer/AITestSuite/Private/Tests/BTTest.cpp
@@ -170,7 +170,7 @@ struct FAITest_BTDecoratorBlueprint : public FAITest_SimpleBT
 		UBTCompositeNode& CompNode = FBTBuilder::AddSelector(*BTAsset);
 		{
 			UBTCompositeNode& CompNode1 = FBTBuilder::AddSequence(CompNode);
-			FBTBuilder::WithDecoratorBlueprint(CompNode, EBTFlowAbortMode::Both, EBPConditionType::True, DecoratorBlueprintBecomeRelevant, DecoratorBlueprintCeaseRelevant, DecoratorBlueprintCalculate);
+			FBTBuilder::WithDecoratorBlueprint(CompNode, EBTFlowAbortMode::Both, EBPConditionType::TrueCondition, DecoratorBlueprintBecomeRelevant, DecoratorBlueprintCeaseRelevant, DecoratorBlueprintCalculate);
 			{
 				// Task1
 				FBTBuilder::AddTask(CompNode1, Task1Execute, EBTNodeResult::Succeeded);
diff --git a/Engine/Source/Developer/AITestSuite/Public/BTBuilder.h b/Engine/Source/Developer/AITestSuite/Public/BTBuilder.h
--- a/Engine/Source/Developer/AITestSuite/Public/BTBuilder.h
+++ b/Engine/Source/Developer/AITestSuite/Public/BTBuilder.h
@@ -302,7 +302,7 @@ struct FBTBuilder
 		AbortDecorator.bOnlyOnce = bAbortOnlyOnce;
 	}
 
-	static void WithDecoratorBlueprint(UBTCompositeNode& ParentNode, EBTFlowAbortMode::Type Observer, EBPConditionType BPConditionType = EBPConditionType::True,
+	static void WithDecoratorBlueprint(UBTCompositeNode& ParentNode, EBTFlowAbortMode::Type Observer, EBPConditionType BPConditionType = EBPConditionType::TrueCondition,
 		int32 LogIndexBecomeRelevant = -1, int32 LogIndexCeaseRelevant = -1, int32 LogIndexCalculate = -1, FName ObservingKeyName = NAME_None )
 	{
 		UTestBTDecorator_Blueprint& BPDecorator = WithDecorator<UTestBTDecorator_Blueprint>(ParentNode);
diff --git a/Engine/Source/Programs/Shared/EpicGames.UHT/Parsers/UhtEnumParser.cs b/Engine/Source/Programs/Shared/EpicGames.UHT/Parsers/UhtEnumParser.cs
--- a/Engine/Source/Programs/Shared/EpicGames.UHT/Parsers/UhtEnumParser.cs
+++ b/Engine/Source/Programs/Shared/EpicGames.UHT/Parsers/UhtEnumParser.cs
@@ -1,6 +1,7 @@
 // Copyright Epic Games, Inc. All Rights Reserved.
 
 using System.Diagnostics.CodeAnalysis;
+using System.Runtime.CompilerServices;
 using EpicGames.Core;
 using EpicGames.UHT.Tables;
 using EpicGames.UHT.Tokenizer;
@@ -161,6 +162,11 @@ namespace EpicGames.UHT.Parsers
 					topScope.TokenReader.RequireList('}', ',', true, () =>
 					{
 						UhtToken tagToken = topScope.TokenReader.GetIdentifier();
+						if (tagToken.IsValue("true", true) || tagToken.IsValue("false", true))
+						{
+							// C++ UHT compatibility - TODO
+							topScope.TokenReader.LogError("Enumerations can't have a elements named 'true' or 'false' regardless of case");
+						}
 
 						StringView fullEnumName;
 						switch (enumObject.CppForm)
