Fixed C# UHT to properly handle a Java style ignore comment followed by a C++ line comment resulting in a mismatch between C++ UHT and C# UHT in the tooltip.

diff --git a/Engine/Source/Programs/Shared/EpicGames.UHT/Parsers/UhtParsingScope.cs b/Engine/Source/Programs/Shared/EpicGames.UHT/Parsers/UhtParsingScope.cs
--- a/Engine/Source/Programs/Shared/EpicGames.UHT/Parsers/UhtParsingScope.cs
+++ b/Engine/Source/Programs/Shared/EpicGames.UHT/Parsers/UhtParsingScope.cs
@@ -270,8 +270,8 @@ namespace EpicGames.UHT.Parsers
 			// 2) Block comment markers are removed first so that '///**' is process by removing '/**' first and then '//' second
 			// 3) We only remove block comments if we find the start of a comment.  This means that if there is a '*/' in a line comment, it won't be removed.
 			// 4) We must check to see if we have cppStyle prior to removing block style comments. 
-			ReadOnlySpan<char> span = comments.AsSpan(0, comments.Length);
 			int commentsLength = RemoveIgnoreComments(scratchChars, comments.Length);
+			ReadOnlySpan<char> span = scratchChars.AsSpan(0, commentsLength);
 			bool javaDocStyle = span.Contains("/**", StringComparison.Ordinal);
 			bool cStyle = javaDocStyle || span.Contains("/*", StringComparison.Ordinal);
 			bool cppStyle = span.StartsWith("//", StringComparison.Ordinal);
