Fixed C++ UHT to match recent changes to C# UHT

diff --git a/Engine/Source/Programs/Shared/EpicGames.UHT/Types/UhtProperty.cs b/Engine/Source/Programs/Shared/EpicGames.UHT/Types/UhtProperty.cs
--- a/Engine/Source/Programs/Shared/EpicGames.UHT/Types/UhtProperty.cs
+++ b/Engine/Source/Programs/Shared/EpicGames.UHT/Types/UhtProperty.cs
@@ -402,7 +402,7 @@ namespace EpicGames.UHT.Types
 		IsRigVMArray = 1 << 14,
 		
 		/// <summary>
-		/// True if the property should codegen as an array
+		/// True if the property should codegen as a byte enumeration
 		/// </summary>
 		IsRigVMEnumAsByte = 1 << 15,
 	}
diff --git a/Engine/Source/Programs/Shared/EpicGames.UHT/Types/UhtScriptStruct.cs b/Engine/Source/Programs/Shared/EpicGames.UHT/Types/UhtScriptStruct.cs
--- a/Engine/Source/Programs/Shared/EpicGames.UHT/Types/UhtScriptStruct.cs
+++ b/Engine/Source/Programs/Shared/EpicGames.UHT/Types/UhtScriptStruct.cs
@@ -52,14 +52,14 @@ namespace EpicGames.UHT.Types
 		IsEnum = 0x00010000,
 		
 		/// <summary>
-		/// Set if the property is an enum as byte
+		/// Set if the property is an array
 		/// </summary>
-		IsEnumAsByte = 0x00010004,
+		IsArray = 0x00020000,
 
 		/// <summary>
-		/// Set if the property is an array
+		/// Set if the property is an enum as byte
 		/// </summary>
-		IsArray = 0x00020000,
+		IsEnumAsByte = 0x00040000,
 	}
 
 	/// <summary>
diff --git a/Engine/Source/Programs/UnrealHeaderTool/Private/HeaderParser.cpp b/Engine/Source/Programs/UnrealHeaderTool/Private/HeaderParser.cpp
--- a/Engine/Source/Programs/UnrealHeaderTool/Private/HeaderParser.cpp
+++ b/Engine/Source/Programs/UnrealHeaderTool/Private/HeaderParser.cpp
@@ -6658,6 +6658,7 @@ void FHeaderParser::ParseRigVMMethodParameters(FUnrealStructDefinitionInfo& Stru
 
 		if (ExtendedCPPType.IsEmpty() && MemberCPPType.StartsWith(TEnumAsByteText))
 		{
+			Parameter.CastType = MemberCPPType;
 			MemberCPPType = MemberCPPType.LeftChop(1).RightChop(12);
 			Parameter.bIsEnumAsByte = true;
 		}
