Horde: Add alternative sync method for ManagedWorkspace not using client's have table

The new class DepotStreamTreeBuilder is simply code extracted from FindClientContentsAsync to avoid duplicating that. Ideally, that method would also use this class but want to minimize changes to existing code with this CL.

diff --git a/Engine/Source/Programs/Shared/EpicGames.Perforce.Managed.Tests/EpicGames.Perforce.Managed.Tests.csproj b/Engine/Source/Programs/Shared/EpicGames.Perforce.Managed.Tests/EpicGames.Perforce.Managed.Tests.csproj
new file mode 100644
--- /dev/null
+++ b/Engine/Source/Programs/Shared/EpicGames.Perforce.Managed.Tests/EpicGames.Perforce.Managed.Tests.csproj
@@ -0,0 +1,37 @@
+<Project Sdk="Microsoft.NET.Sdk">
+
+  <PropertyGroup>
+    <TargetFramework>net6.0</TargetFramework>
+    <Configuration Condition=" '$(Configuration)' == '' ">Development</Configuration>
+    <OutputType>Library</OutputType>
+    <Deterministic>true</Deterministic>
+    <RootNamespace>EpicGames.Perforce.Managed.Tests</RootNamespace>
+    <GenerateAssemblyCompanyAttribute>false</GenerateAssemblyCompanyAttribute>
+    <GenerateAssemblyInformationalVersionAttribute>false</GenerateAssemblyInformationalVersionAttribute>
+    <GenerateAssemblyProductAttribute>false</GenerateAssemblyProductAttribute>
+    <GenerateAssemblyVersionAttribute>false</GenerateAssemblyVersionAttribute>
+    <Configurations>Debug;Release;Development</Configurations>
+    <IsPackable>false</IsPackable>
+    <Nullable>enable</Nullable>
+  </PropertyGroup>
+  <ItemGroup>
+    <Compile Include="..\..\Shared\MetaData.cs">
+      <Link>Properties\MetaData.cs</Link>
+    </Compile>
+  </ItemGroup>
+
+  <ItemGroup>
+    <PackageReference Include="Microsoft.NET.Test.Sdk" Version="17.3.0" />
+    <PackageReference Include="MSTest.TestAdapter" Version="2.2.10" />
+    <PackageReference Include="MSTest.TestFramework" Version="2.2.10" />
+    <PackageReference Include="coverlet.collector" Version="3.1.2">
+      <PrivateAssets>all</PrivateAssets>
+      <IncludeAssets>runtime; build; native; contentfiles; analyzers; buildtransitive</IncludeAssets>
+    </PackageReference>
+  </ItemGroup>
+
+  <ItemGroup>
+    <ProjectReference Include="..\EpicGames.Perforce.Managed\EpicGames.Perforce.Managed.csproj" />
+  </ItemGroup>
+
+</Project>
diff --git a/Engine/Source/Programs/Shared/EpicGames.Perforce.Managed.Tests/StreamTreeBuilderTest.cs b/Engine/Source/Programs/Shared/EpicGames.Perforce.Managed.Tests/StreamTreeBuilderTest.cs
new file mode 100644
--- /dev/null
+++ b/Engine/Source/Programs/Shared/EpicGames.Perforce.Managed.Tests/StreamTreeBuilderTest.cs
@@ -0,0 +1,62 @@
+// Copyright Epic Games, Inc. All Rights Reserved.
+
+using System;
+using System.IO;
+using EpicGames.Core;
+using Microsoft.VisualStudio.TestTools.UnitTesting;
+
+namespace EpicGames.Perforce.Managed.Tests;
+
+[TestClass]
+public class DepotStreamTreeBuilderTest
+{
+	private readonly StreamFile _file1 = new ("//UE5/Main/main.cpp", 100, new FileContentId(Md5Hash.Zero, "text"), 10);
+	private readonly StreamFile _file2 = new ("//UE5/Main/Data/data1.bin", 200, new FileContentId(Md5Hash.Zero, "binary"), 20);
+	private readonly StreamFile _file3 = new ("//UE5/Main/Data/data2.bin", 300, new FileContentId(Md5Hash.Zero, "binary"), 30);
+	
+	[TestMethod]
+	public void Basic()
+	{
+		DepotStreamTreeBuilder builder = new ("//UE5/");
+		builder.AddDepotFile(_file1);
+		builder.AddDepotFile(_file2);
+		builder.AddDepotFile(_file3);
+		StreamSnapshotFromMemory snapshot = new (builder);
+
+		StreamTree ue5 = snapshot.Lookup(snapshot.Root);
+		{
+			Assert.AreEqual(0, ue5.NameToFile.Count);
+			Assert.AreEqual(1, ue5.NameToTree.Count);
+		}
+
+		StreamTree main = snapshot.Lookup(ue5.NameToTree["Main"]);
+		{
+			Assert.AreEqual(1, main.NameToFile.Count);
+			Assert.AreEqual(1, main.NameToTree.Count);
+			
+			Assert.AreEqual(100, main.NameToFile["main.cpp"].Length);
+		}
+		
+		StreamTree data = snapshot.Lookup(main.NameToTree["Data"]);
+		{
+			Assert.AreEqual(2, data.NameToFile.Count);
+			Assert.AreEqual(0, data.NameToTree.Count);
+			
+			Assert.AreEqual(200, data.NameToFile["data1.bin"].Length);
+			Assert.AreEqual(300, data.NameToFile["data2.bin"].Length);
+		}
+	}
+
+	[TestMethod]
+	public void InstantiateWithInvalidPrefix()
+	{
+		Assert.ThrowsException<ArgumentException>(() => new DepotStreamTreeBuilder("//UE5"));
+	}
+
+	[TestMethod]
+	public void AddFileWithInvalidPrefix()
+	{
+		DepotStreamTreeBuilder builder = new("//Foo/");
+		Assert.ThrowsException<InvalidDataException>(() => builder.AddDepotFile(_file1));
+	}
+}
\ No newline at end of file
diff --git a/Engine/Source/Programs/Shared/EpicGames.Perforce.Managed/ManagedWorkspace.cs b/Engine/Source/Programs/Shared/EpicGames.Perforce.Managed/ManagedWorkspace.cs
--- a/Engine/Source/Programs/Shared/EpicGames.Perforce.Managed/ManagedWorkspace.cs
+++ b/Engine/Source/Programs/Shared/EpicGames.Perforce.Managed/ManagedWorkspace.cs
@@ -826,7 +826,65 @@ namespace EpicGames.Perforce.Managed
 					contents = await TryLoadClientContentsAsync(cacheFile, streamName, cancellationToken);
 					if (contents == null)
 					{
-						contents = await FindAndSaveClientContentsAsync(perforce, streamName, changeNumber, cacheFile, cancellationToken);
+						contents = await FindAndSaveClientContentsAsync(perforce, streamName, changeNumber, cacheFile, true, cancellationToken);
+					}
+				}
+
+				// Sync all the appropriate files
+				await RemoveFilesFromWorkspaceAsync(contents, cancellationToken);
+				await AddFilesToWorkspaceAsync(perforce, contents, fakeSync, cancellationToken);
+			}
+
+			_logger.LogInformation("Completed in {ElapsedTime}s", $"{timer.Elapsed.TotalSeconds:0.0}");
+		}
+		
+		/// <summary>
+		/// Switches to the given stream
+		/// </summary>
+		/// <param name="perforce">The perforce connection</param>
+		/// <param name="streamName">Name of the stream to sync</param>
+		/// <param name="changeNumber">Changelist number to sync. -1 to sync to latest.</param>
+		/// <param name="view">View of the workspace</param>
+		/// <param name="removeUntracked">Whether to remove untracked files from the workspace</param>
+		/// <param name="fakeSync">Whether to simulate the syncing operation rather than actually getting files from the server</param>
+		/// <param name="cacheFile">If set, uses the given file to cache the contents of the workspace. This can improve sync times when multiple machines sync the same workspace.</param>
+		/// <param name="cancellationToken">Cancellation token</param>
+		public async Task SyncWithoutHaveTableAsync(IPerforceConnection perforce, string streamName, int changeNumber, IReadOnlyList<string> view, bool removeUntracked, bool fakeSync, FileReference? cacheFile, CancellationToken cancellationToken)
+		{
+			Stopwatch timer = Stopwatch.StartNew();
+			if (changeNumber == -1)
+			{
+				_logger.LogInformation("Syncing without have-table to {StreamName} at latest", streamName);
+			}
+			else
+			{
+				_logger.LogInformation("Syncing without have-table to {StreamName} at CL {CL}", streamName, changeNumber);
+			}
+
+			using (_logger.WithProperty("useHaveTable", true).BeginScope())
+			using (_logger.BeginIndentScope("  "))
+			{
+				// Update the client to the current stream
+				await UpdateClientAsync(perforce, streamName, cancellationToken);
+
+				// Get the latest change number
+				if (changeNumber == -1)
+				{
+					changeNumber = await GetLatestClientChangeAsync(perforce, cancellationToken);
+				}
+
+				// Update the state of the current stream, if necessary
+				StreamSnapshot? contents;
+				if (cacheFile == null)
+				{
+					contents = await FindClientContentsWithoutHaveTableAsync(perforce, changeNumber, cancellationToken);
+				}
+				else
+				{
+					contents = await TryLoadClientContentsAsync(cacheFile, streamName, cancellationToken);
+					if (contents == null)
+					{
+						contents = await FindAndSaveClientContentsAsync(perforce, streamName, changeNumber, cacheFile, false, cancellationToken);
 					}
 				}
 
@@ -1457,11 +1515,11 @@ namespace EpicGames.Perforce.Managed
 
 				// Create the workspace, and add records for all the files. Exclude deleted files with digest = null.
 				List<string> arguments = new List<string>();
-				arguments.Add("-Ol");
-				arguments.Add("-Op");
-				arguments.Add("-Os");
-				arguments.Add("-Rh");
-				arguments.Add("-T");
+				arguments.Add("-Ol"); // Output fileSize and digest field
+				arguments.Add("-Op"); // Output clientFile field in both server and local path syntax
+				arguments.Add("-Os"); // Shorten output by excluding client workspace data (for instance, the clientFile field).
+				arguments.Add("-Rh"); // Limit output to files on your have list;
+				arguments.Add("-T"); // Include only the fields listed below
 				arguments.Add(String.Join(",", FStatIndexedRecord.FieldNames));
 				arguments.Add($"//{perforceClient.Settings.ClientName}/...@{changeNumber}");
 				await perforceClient.RecordCommandAsync("fstat", arguments, null, HandleRecord, cancellationToken);
@@ -1472,6 +1530,73 @@ namespace EpicGames.Perforce.Managed
 
 			return new StreamSnapshotFromMemory(builder);
 		}
+		
+		class FStatRecordWithoutHaveTable
+		{
+			// Note: This enum is used for indexing an array of fields, and member names much match P4 field names (including case).
+			enum Field
+			{
+				code,
+				depotFile,
+				headType,
+				headRev,
+				fileSize,
+				digest
+			}
+
+			public static readonly string[] FieldNames = Enum.GetNames(typeof(Field));
+			public static readonly Utf8String[] Utf8FieldNames = Array.ConvertAll(FieldNames, x => new Utf8String(x));
+
+			public PerforceValue[] Values { get; } = new PerforceValue[FieldNames.Length];
+
+			public Utf8String DepotFile => Values[(int)Field.depotFile].GetString();
+
+			public Utf8String HeadType => Values[(int)Field.headType].GetString();
+
+			public int HeadRev => Values[(int)Field.headRev].AsInteger();
+
+			public long FileSize => Values[(int)Field.fileSize].AsLong();
+
+			public Utf8String Digest => Values[(int)Field.digest].GetString();
+		}
+		
+		/// <summary>
+		/// Get the contents of the client without using the have table
+		/// </summary>
+		/// <param name="perforceClient">The client connection</param>
+		/// <param name="changeNumber">The change number being synced. This must be specified in order to get the digest at the correct revision.</param>
+		/// <param name="cancellationToken">Cancellation token</param>
+		public async Task<StreamSnapshotFromMemory> FindClientContentsWithoutHaveTableAsync(IPerforceConnection perforceClient, int changeNumber, CancellationToken cancellationToken)
+		{
+			DepotStreamTreeBuilder builder = new ($"//{perforceClient.Settings.ClientName}/");
+			
+			// Re-use a single class instance as there can be millions of records
+			FStatRecordWithoutHaveTable record = new ();
+			void HandleRecord(PerforceRecord rawRecord)
+			{
+				// Copy into the values array
+				rawRecord.CopyInto(FStatRecordWithoutHaveTable.Utf8FieldNames, record.Values);
+				if (record.Digest.IsEmpty) { return; }
+				
+				Md5Hash md5Hash = Md5Hash.Parse(record.Digest);
+				FileContentId fileContentId = new (md5Hash, record.HeadType);
+				builder.AddDepotFile(new StreamFile(record.DepotFile, record.FileSize, fileContentId, record.HeadRev));
+			}
+
+			string fileSpec = $"//{perforceClient.Settings.ClientName}/...@{changeNumber}";
+			List<string> arguments = new ();
+			arguments.Add("-Ol"); // Output fileSize and digest field
+			arguments.Add("-Os"); // Shorten output by excluding client workspace data (for instance, the clientFile field).
+			arguments.Add("-F"); // Filter any files not existing at current revision (filter below)
+			arguments.Add("^headAction=delete&^headAction=move/delete&^headAction=purge");
+			arguments.Add("-T"); // Include only the fields listed below
+			arguments.Add(String.Join(",", FStatRecordWithoutHaveTable.FieldNames));
+			arguments.Add(fileSpec);
+
+			await perforceClient.RecordCommandAsync("fstat", arguments, null, HandleRecord, cancellationToken);
+			
+			return new StreamSnapshotFromMemory(builder);
+		}
 
 		/// <summary>
 		/// Loads the contents of a client from disk
@@ -1503,11 +1628,14 @@ namespace EpicGames.Perforce.Managed
 		/// <param name="basePath">Base path for the stream</param>
 		/// <param name="changeNumber">The change number being synced. This must be specified in order to get the digest at the correct revision.</param>
 		/// <param name="cacheFile">Location of the file to save the cached contents</param>
+		/// <param name="useHaveTable">Use client's have table when resolving contents</param>
 		/// <param name="cancellationToken">Cancellation token</param>
 		/// <returns>Contents of the workspace</returns>
-		private async Task<StreamSnapshotFromMemory> FindAndSaveClientContentsAsync(IPerforceConnection perforceClient, Utf8String basePath, int changeNumber, FileReference cacheFile, CancellationToken cancellationToken)
+		private async Task<StreamSnapshotFromMemory> FindAndSaveClientContentsAsync(IPerforceConnection perforceClient, Utf8String basePath, int changeNumber, FileReference cacheFile, bool useHaveTable, CancellationToken cancellationToken)
 		{
-			StreamSnapshotFromMemory contents = await FindClientContentsAsync(perforceClient, changeNumber, cancellationToken);
+			StreamSnapshotFromMemory contents = useHaveTable
+				? await FindClientContentsAsync(perforceClient, changeNumber, cancellationToken)
+				: await FindClientContentsWithoutHaveTableAsync(perforceClient, changeNumber, cancellationToken);
 
 			using (Trace("WriteMetadata"))
 			using (ILoggerProgress scope = _logger.BeginProgressScope($"Saving metadata to {cacheFile}..."))
diff --git a/Engine/Source/Programs/Shared/EpicGames.Perforce.Managed/StreamTree.cs b/Engine/Source/Programs/Shared/EpicGames.Perforce.Managed/StreamTree.cs
--- a/Engine/Source/Programs/Shared/EpicGames.Perforce.Managed/StreamTree.cs
+++ b/Engine/Source/Programs/Shared/EpicGames.Perforce.Managed/StreamTree.cs
@@ -82,6 +82,12 @@ namespace EpicGames.Perforce.Managed
 			writer.WriteUtf8String(s_typeField, ContentId.Type);
 			writer.WriteInteger(s_revisionField, Revision);
 		}
+
+		/// <inheritdoc/>
+		public override string ToString()
+		{
+			return $"StreamFile({Path}#{Revision} Len={Length} Digest={ContentId.Digest} Type={ContentId.Type})";
+		}
 	}
 
 	/// <summary>
diff --git a/Engine/Source/Programs/Shared/EpicGames.Perforce.Managed/StreamTreeBuilder.cs b/Engine/Source/Programs/Shared/EpicGames.Perforce.Managed/StreamTreeBuilder.cs
--- a/Engine/Source/Programs/Shared/EpicGames.Perforce.Managed/StreamTreeBuilder.cs
+++ b/Engine/Source/Programs/Shared/EpicGames.Perforce.Managed/StreamTreeBuilder.cs
@@ -2,6 +2,7 @@
 
 using System;
 using System.Collections.Generic;
+using System.IO;
 using EpicGames.Core;
 
 namespace EpicGames.Perforce.Managed
@@ -121,4 +122,97 @@ namespace EpicGames.Perforce.Managed
 			}
 		}
 	}
+
+	/// <summary>
+	/// Variant of StreamTreeBuilder which caches path fragments lookups for faster building
+	/// when there's lots of files.
+	/// </summary>
+	public class DepotStreamTreeBuilder : StreamTreeBuilder
+	{
+		/// <summary>
+		/// List of the last path fragments
+		/// Since file records that are returned are typically sorted by their position in the tree,
+		/// we can save quite a lot of processing by reusing as many fragments as possible
+		/// </summary>
+		private readonly List<(Utf8String, StreamTreeBuilder)> _fragments = new ();
+
+		private readonly Utf8String _expectedPathPrefix;
+
+		public DepotStreamTreeBuilder(string expectedPathPrefix)
+		{
+			if (expectedPathPrefix[^1] != '/')
+			{
+				throw new ArgumentException("Prefix must end with a slash");
+			}
+			_expectedPathPrefix = expectedPathPrefix;
+		}
+
+		public void AddDepotFile(StreamFile file)
+		{
+			if (!file.Path.StartsWith(_expectedPathPrefix))
+			{
+				throw new InvalidDataException($"File path ('{file.Path}') does not begin with prefix ('{_expectedPathPrefix}')");
+			}
+
+			ReadOnlySpan<byte> pathSpan = file.Path.Span;
+
+			// Parse out the data
+			StreamTreeBuilder lastStreamDirectory = this;
+
+			// Try to match up as many fragments from the last file.
+			int fragmentMinIdx = _expectedPathPrefix.Length;
+			for (int fragmentIdx = 0; ; fragmentIdx++)
+			{
+				// Find the next directory separator
+				int fragmentMaxIdx = fragmentMinIdx;
+				while (fragmentMaxIdx < pathSpan.Length && pathSpan[fragmentMaxIdx] != '/')
+				{
+					fragmentMaxIdx++;
+				}
+				if (fragmentMaxIdx == pathSpan.Length)
+				{
+					_fragments.RemoveRange(fragmentIdx, _fragments.Count - fragmentIdx);
+					break;
+				}
+
+				// Get the fragment text
+				Utf8String fragment = new (file.Path.Memory.Slice(fragmentMinIdx, fragmentMaxIdx - fragmentMinIdx));
+
+				// If this fragment matches the same fragment from the previous iteration, take the last stream directory straight away
+				if (fragmentIdx < _fragments.Count)
+				{
+					if (_fragments[fragmentIdx].Item1 == fragment)
+					{
+						lastStreamDirectory = _fragments[fragmentIdx].Item2;
+					}
+					else
+					{
+						_fragments.RemoveRange(fragmentIdx, _fragments.Count - fragmentIdx);
+					}
+				}
+
+				// Otherwise, find or add a directory for this fragment into the last directory
+				if (fragmentIdx >= _fragments.Count)
+				{
+					Utf8String unescapedFragment = PerforceUtils.UnescapePath(fragment);
+
+					StreamTreeBuilder? nextStreamDirectory;
+					if (!lastStreamDirectory.NameToTreeBuilder.TryGetValue(unescapedFragment, out nextStreamDirectory))
+					{
+						nextStreamDirectory = new StreamTreeBuilder();
+						lastStreamDirectory.NameToTreeBuilder.Add(unescapedFragment, nextStreamDirectory);
+					}
+					lastStreamDirectory = nextStreamDirectory;
+
+					_fragments.Add((fragment, lastStreamDirectory));
+				}
+
+				// Move to the next fragment
+				fragmentMinIdx = fragmentMaxIdx + 1;
+			}
+			
+			Utf8String fileName = PerforceUtils.UnescapePath(file.Path.Slice(fragmentMinIdx));
+			lastStreamDirectory.NameToFile.Add(fileName, file);
+		}
+	}
 }
