Horde: Respect client file paths when syncing without a have table

diff --git a/Engine/Source/Programs/Shared/EpicGames.Perforce.Managed.Tests/StreamTreeBuilderTest.cs b/Engine/Source/Programs/Shared/EpicGames.Perforce.Managed.Tests/StreamTreeBuilderTest.cs
--- a/Engine/Source/Programs/Shared/EpicGames.Perforce.Managed.Tests/StreamTreeBuilderTest.cs
+++ b/Engine/Source/Programs/Shared/EpicGames.Perforce.Managed.Tests/StreamTreeBuilderTest.cs
@@ -1,7 +1,5 @@
 // Copyright Epic Games, Inc. All Rights Reserved.
 
-using System;
-using System.IO;
 using EpicGames.Core;
 using Microsoft.VisualStudio.TestTools.UnitTesting;
 
@@ -11,52 +9,49 @@ namespace EpicGames.Perforce.Managed.Tests;
 public class DepotStreamTreeBuilderTest
 {
 	private readonly StreamFile _file1 = new ("//UE5/Main/main.cpp", 100, new FileContentId(Md5Hash.Zero, "text"), 10);
-	private readonly StreamFile _file2 = new ("//UE5/Main/Data/data1.bin", 200, new FileContentId(Md5Hash.Zero, "binary"), 20);
-	private readonly StreamFile _file3 = new ("//UE5/Main/Data/data2.bin", 300, new FileContentId(Md5Hash.Zero, "binary"), 30);
+	private readonly StreamFile _file2 = new ("//UE5/Main/Data/data2.bin", 200, new FileContentId(Md5Hash.Zero, "binary"), 20);
+	private readonly StreamFile _file3 = new ("//UE5/Main/Data/data3.bin", 300, new FileContentId(Md5Hash.Zero, "binary"), 30);
+	private readonly StreamFile _file4 = new ("//UE5/Main/Data/Audio/Samples/data4.bin", 400, new FileContentId(Md5Hash.Zero, "binary"), 40);
 	
 	[TestMethod]
 	public void Basic()
 	{
-		DepotStreamTreeBuilder builder = new ("//UE5/");
-		builder.AddDepotFile(_file1);
-		builder.AddDepotFile(_file2);
-		builder.AddDepotFile(_file3);
+		DepotStreamTreeBuilder builder = new ();
+		builder.AddFile("/main.cpp", _file1);
+		builder.AddFile("Data/data2.bin", _file2);
+		builder.AddFile("Data/data3.bin", _file3);
+		builder.AddFile("Data/Audio/Samples/data4.bin", _file4);
 		StreamSnapshotFromMemory snapshot = new (builder);
 
-		StreamTree ue5 = snapshot.Lookup(snapshot.Root);
+		StreamTree root = snapshot.Lookup(snapshot.Root);
 		{
-			Assert.AreEqual(0, ue5.NameToFile.Count);
-			Assert.AreEqual(1, ue5.NameToTree.Count);
+			Assert.AreEqual(1, root.NameToFile.Count);
+			Assert.AreEqual(100, root.NameToFile["main.cpp"].Length);
+			
+			Assert.AreEqual(1, root.NameToTree.Count);
 		}
-
-		StreamTree main = snapshot.Lookup(ue5.NameToTree["Main"]);
+		
+		StreamTree data = snapshot.Lookup(root.NameToTree["Data"]);
 		{
-			Assert.AreEqual(1, main.NameToFile.Count);
-			Assert.AreEqual(1, main.NameToTree.Count);
+			Assert.AreEqual(2, data.NameToFile.Count);
+			Assert.AreEqual(1, data.NameToTree.Count);
 			
-			Assert.AreEqual(100, main.NameToFile["main.cpp"].Length);
+			Assert.AreEqual(_file2.Length, data.NameToFile["data2.bin"].Length);
+			Assert.AreEqual(_file3.Length, data.NameToFile["data3.bin"].Length);
 		}
 		
-		StreamTree data = snapshot.Lookup(main.NameToTree["Data"]);
+		StreamTree audio = snapshot.Lookup(data.NameToTree["Audio"]);
 		{
-			Assert.AreEqual(2, data.NameToFile.Count);
-			Assert.AreEqual(0, data.NameToTree.Count);
+			Assert.AreEqual(0, audio.NameToFile.Count);
+			Assert.AreEqual(1, audio.NameToTree.Count);
+		}
+		
+		StreamTree samples = snapshot.Lookup(audio.NameToTree["Samples"]);
+		{
+			Assert.AreEqual(1, samples.NameToFile.Count);
+			Assert.AreEqual(0, samples.NameToTree.Count);
 			
-			Assert.AreEqual(200, data.NameToFile["data1.bin"].Length);
-			Assert.AreEqual(300, data.NameToFile["data2.bin"].Length);
+			Assert.AreEqual(_file4.Length, samples.NameToFile["data4.bin"].Length);
 		}
 	}
-
-	[TestMethod]
-	public void InstantiateWithInvalidPrefix()
-	{
-		Assert.ThrowsException<ArgumentException>(() => new DepotStreamTreeBuilder("//UE5"));
-	}
-
-	[TestMethod]
-	public void AddFileWithInvalidPrefix()
-	{
-		DepotStreamTreeBuilder builder = new("//Foo/");
-		Assert.ThrowsException<InvalidDataException>(() => builder.AddDepotFile(_file1));
-	}
 }
\ No newline at end of file
diff --git a/Engine/Source/Programs/Shared/EpicGames.Perforce.Managed/ManagedWorkspace.cs b/Engine/Source/Programs/Shared/EpicGames.Perforce.Managed/ManagedWorkspace.cs
--- a/Engine/Source/Programs/Shared/EpicGames.Perforce.Managed/ManagedWorkspace.cs
+++ b/Engine/Source/Programs/Shared/EpicGames.Perforce.Managed/ManagedWorkspace.cs
@@ -877,7 +877,7 @@ namespace EpicGames.Perforce.Managed
 				StreamSnapshot? contents;
 				if (cacheFile == null)
 				{
-					contents = await FindClientContentsWithoutHaveTableAsync(perforce, changeNumber, cancellationToken);
+					contents = await FindClientContentsWithoutHaveTableAsync(perforce, streamName, changeNumber, cancellationToken);
 				}
 				else
 				{
@@ -1564,11 +1564,15 @@ namespace EpicGames.Perforce.Managed
 		/// Get the contents of the client without using the have table
 		/// </summary>
 		/// <param name="perforceClient">The client connection</param>
+		/// <param name="streamName">Name of stream</param>
 		/// <param name="changeNumber">The change number being synced. This must be specified in order to get the digest at the correct revision.</param>
 		/// <param name="cancellationToken">Cancellation token</param>
-		public async Task<StreamSnapshotFromMemory> FindClientContentsWithoutHaveTableAsync(IPerforceConnection perforceClient, int changeNumber, CancellationToken cancellationToken)
+		public async Task<StreamSnapshotFromMemory> FindClientContentsWithoutHaveTableAsync(IPerforceConnection perforceClient, string streamName, int changeNumber, CancellationToken cancellationToken)
 		{
-			DepotStreamTreeBuilder builder = new ($"//{perforceClient.Settings.ClientName}/");
+			DepotStreamTreeBuilder builder = new ();
+			
+			StreamRecord streamRecord = await perforceClient.GetStreamAsync(streamName, true, cancellationToken);
+			PerforceViewMap viewMap = PerforceViewMap.Parse(streamRecord.View);
 			
 			// Re-use a single class instance as there can be millions of records
 			FStatRecordWithoutHaveTable record = new ();
@@ -1578,9 +1582,16 @@ namespace EpicGames.Perforce.Managed
 				rawRecord.CopyInto(FStatRecordWithoutHaveTable.Utf8FieldNames, record.Values);
 				if (record.Digest.IsEmpty) { return; }
 				
-				Md5Hash md5Hash = Md5Hash.Parse(record.Digest);
-				FileContentId fileContentId = new (md5Hash, record.HeadType);
-				builder.AddDepotFile(new StreamFile(record.DepotFile, record.FileSize, fileContentId, record.HeadRev));
+				if (viewMap.TryMapFile(record.DepotFile.ToString(), StringComparison.Ordinal, out string clientFile))
+				{
+					Md5Hash md5Hash = Md5Hash.Parse(record.Digest);
+					FileContentId fileContentId = new (md5Hash, record.HeadType);
+					builder.AddFile(clientFile, new StreamFile(record.DepotFile, record.FileSize, fileContentId, record.HeadRev));					
+				}
+				else
+				{
+					_logger.LogError("Failed to view map depot file {DepotFile}", record.DepotFile.ToString());
+				}
 			}
 
 			string fileSpec = $"//{perforceClient.Settings.ClientName}/...@{changeNumber}";
@@ -1635,7 +1646,7 @@ namespace EpicGames.Perforce.Managed
 		{
 			StreamSnapshotFromMemory contents = useHaveTable
 				? await FindClientContentsAsync(perforceClient, changeNumber, cancellationToken)
-				: await FindClientContentsWithoutHaveTableAsync(perforceClient, changeNumber, cancellationToken);
+				: await FindClientContentsWithoutHaveTableAsync(perforceClient, basePath.ToString(), changeNumber, cancellationToken);
 
 			using (Trace("WriteMetadata"))
 			using (ILoggerProgress scope = _logger.BeginProgressScope($"Saving metadata to {cacheFile}..."))
diff --git a/Engine/Source/Programs/Shared/EpicGames.Perforce.Managed/StreamTreeBuilder.cs b/Engine/Source/Programs/Shared/EpicGames.Perforce.Managed/StreamTreeBuilder.cs
--- a/Engine/Source/Programs/Shared/EpicGames.Perforce.Managed/StreamTreeBuilder.cs
+++ b/Engine/Source/Programs/Shared/EpicGames.Perforce.Managed/StreamTreeBuilder.cs
@@ -124,95 +124,38 @@ namespace EpicGames.Perforce.Managed
 	}
 
 	/// <summary>
-	/// Variant of StreamTreeBuilder which caches path fragments lookups for faster building
-	/// when there's lots of files.
+	/// Helper variant of StreamTreeBuilder capable of adding files to tree using both client and depot file paths
 	/// </summary>
 	public class DepotStreamTreeBuilder : StreamTreeBuilder
 	{
-		/// <summary>
-		/// List of the last path fragments
-		/// Since file records that are returned are typically sorted by their position in the tree,
-		/// we can save quite a lot of processing by reusing as many fragments as possible
-		/// </summary>
-		private readonly List<(Utf8String, StreamTreeBuilder)> _fragments = new ();
-
-		private readonly Utf8String _expectedPathPrefix;
-
-		public DepotStreamTreeBuilder(string expectedPathPrefix)
+		public void AddFile(string clientFile, StreamFile depotFile)
 		{
-			if (expectedPathPrefix[^1] != '/')
+			if (clientFile[0] == '/')
 			{
-				throw new ArgumentException("Prefix must end with a slash");
+				// Strip any slash at the start
+				clientFile = clientFile[1..];
 			}
-			_expectedPathPrefix = expectedPathPrefix;
-		}
+			
+			StreamTreeBuilder currentStreamDirectory = this;
+			string[] pathFragments = clientFile.Split('/');
 
-		public void AddDepotFile(StreamFile file)
-		{
-			if (!file.Path.StartsWith(_expectedPathPrefix))
+			// Find stream tree builder for deepest path fragment
+			// and skip last fragment as that is the filename
+			for (int i = 0; i < pathFragments.Length - 1; i++)
 			{
-				throw new InvalidDataException($"File path ('{file.Path}') does not begin with prefix ('{_expectedPathPrefix}')");
-			}
-
-			ReadOnlySpan<byte> pathSpan = file.Path.Span;
-
-			// Parse out the data
-			StreamTreeBuilder lastStreamDirectory = this;
+				string pathFragment = pathFragments[i];
+				Utf8String unescapedFragment = PerforceUtils.UnescapePath(pathFragment);
 
-			// Try to match up as many fragments from the last file.
-			int fragmentMinIdx = _expectedPathPrefix.Length;
-			for (int fragmentIdx = 0; ; fragmentIdx++)
-			{
-				// Find the next directory separator
-				int fragmentMaxIdx = fragmentMinIdx;
-				while (fragmentMaxIdx < pathSpan.Length && pathSpan[fragmentMaxIdx] != '/')
-				{
-					fragmentMaxIdx++;
-				}
-				if (fragmentMaxIdx == pathSpan.Length)
+				if (!currentStreamDirectory.NameToTreeBuilder.TryGetValue(unescapedFragment, out StreamTreeBuilder? nextStreamDirectory))
 				{
-					_fragments.RemoveRange(fragmentIdx, _fragments.Count - fragmentIdx);
-					break;
+					nextStreamDirectory = new StreamTreeBuilder();
+					currentStreamDirectory.NameToTreeBuilder.Add(unescapedFragment, nextStreamDirectory);
 				}
-
-				// Get the fragment text
-				Utf8String fragment = new (file.Path.Memory.Slice(fragmentMinIdx, fragmentMaxIdx - fragmentMinIdx));
-
-				// If this fragment matches the same fragment from the previous iteration, take the last stream directory straight away
-				if (fragmentIdx < _fragments.Count)
-				{
-					if (_fragments[fragmentIdx].Item1 == fragment)
-					{
-						lastStreamDirectory = _fragments[fragmentIdx].Item2;
-					}
-					else
-					{
-						_fragments.RemoveRange(fragmentIdx, _fragments.Count - fragmentIdx);
-					}
-				}
-
-				// Otherwise, find or add a directory for this fragment into the last directory
-				if (fragmentIdx >= _fragments.Count)
-				{
-					Utf8String unescapedFragment = PerforceUtils.UnescapePath(fragment);
-
-					StreamTreeBuilder? nextStreamDirectory;
-					if (!lastStreamDirectory.NameToTreeBuilder.TryGetValue(unescapedFragment, out nextStreamDirectory))
-					{
-						nextStreamDirectory = new StreamTreeBuilder();
-						lastStreamDirectory.NameToTreeBuilder.Add(unescapedFragment, nextStreamDirectory);
-					}
-					lastStreamDirectory = nextStreamDirectory;
-
-					_fragments.Add((fragment, lastStreamDirectory));
-				}
-
-				// Move to the next fragment
-				fragmentMinIdx = fragmentMaxIdx + 1;
+				currentStreamDirectory = nextStreamDirectory;
 			}
-			
-			Utf8String fileName = PerforceUtils.UnescapePath(file.Path.Slice(fragmentMinIdx));
-			lastStreamDirectory.NameToFile.Add(fileName, file);
+
+			string filename = PerforceUtils.UnescapePath(pathFragments[^1]); // Last fragment is filename
+			currentStreamDirectory.NameToFile[filename] = depotFile;
 		}
 	}
 }
