P4VUtils: Fix compilation errors due to using NET Core 3.1.

diff --git a/Engine/Extras/P4VUtils/P4VUtils.csproj b/Engine/Extras/P4VUtils/P4VUtils.csproj
--- a/Engine/Extras/P4VUtils/P4VUtils.csproj
+++ b/Engine/Extras/P4VUtils/P4VUtils.csproj
@@ -23,8 +23,8 @@
   </ItemGroup>
 
   <ItemGroup>
-    <ProjectReference Include="..\..\Source\Programs\Shared\EpicGames.Core\EpicGames.Core.csproj" />
-    <ProjectReference Include="..\..\Source\Programs\Shared\EpicGames.Perforce\EpicGames.Perforce.csproj" />
+    <ProjectReference Include="..\..\Source\Programs\UnrealGameSync\Shared\EpicGames.Core\EpicGames.Core.csproj" />
+    <ProjectReference Include="..\..\Source\Programs\UnrealGameSync\Shared\EpicGames.Perforce\EpicGames.Perforce.csproj" />
   </ItemGroup>
 
   <ItemGroup>
diff --git a/Engine/Source/Programs/Shared/EpicGames.Core/ILoggerProperty.cs b/Engine/Source/Programs/Shared/EpicGames.Core/ILoggerProperty.cs
--- a/Engine/Source/Programs/Shared/EpicGames.Core/ILoggerProperty.cs
+++ b/Engine/Source/Programs/Shared/EpicGames.Core/ILoggerProperty.cs
@@ -4,74 +4,76 @@ using System;
 using System.Collections.Generic;
 using Microsoft.Extensions.Logging;
 
-namespace EpicGames.Core;
-
-/// <summary>
-/// Extension methods for more convenient attaching of properties to log messages
-/// </summary>
-public static class LoggerPropertyExtensions
+namespace EpicGames.Core
 {
 	/// <summary>
-	/// Builder struct for attaching properties to log messages
-	/// Wraps and disposes log scope.
+	/// Extension methods for more convenient attaching of properties to log messages
 	/// </summary>
-	public readonly struct LogScopedPropertyBuilder
+	public static class LoggerPropertyExtensions
 	{
-		private readonly Dictionary<string, object> _properties = new();
-		private readonly ILogger _logger;
-		
 		/// <summary>
-		/// Constructor
+		/// Builder struct for attaching properties to log messages
+		/// Wraps and disposes log scope.
 		/// </summary>
-		/// <param name="logger">Logger to wrap</param>
-		public LogScopedPropertyBuilder(ILogger logger)
+		public readonly struct LogScopedPropertyBuilder
 		{
-			_logger = logger;
+			private readonly Dictionary<string, object> _properties;
+			private readonly ILogger _logger;
+
+			/// <summary>
+			/// Constructor
+			/// </summary>
+			/// <param name="logger">Logger to wrap</param>
+			public LogScopedPropertyBuilder(ILogger logger)
+			{
+				_properties = new Dictionary<string, object>();
+				_logger = logger;
+			}
+
+			/// <summary>
+			/// Set a key/value property for current log scope
+			/// </summary>
+			/// <param name="key">Arbitrary key</param>
+			/// <param name="value">Arbitrary value</param>
+			/// <returns></returns>
+			public LogScopedPropertyBuilder WithProperty(string key, object value)
+			{
+				_properties[key] = value;
+				return this;
+			}
+
+			/// <summary>
+			/// Starts new scope on underlying logger, passing previously set properties
+			/// </summary>
+			/// <returns>Original IDisposable from ILogger.BeginScope()</returns>
+			public IDisposable BeginScope()
+			{
+				return _logger.BeginScope(_properties);
+			}
 		}
 
 		/// <summary>
-		/// Set a key/value property for current log scope
+		/// Creates a new log scope and attaches a key/value property
+		/// Additional properties can be added with <see cref="LogScopedPropertyBuilder.WithProperty" />.
 		/// </summary>
+		/// <param name="logger">Logger to use</param>
 		/// <param name="key">Arbitrary key</param>
 		/// <param name="value">Arbitrary value</param>
-		/// <returns></returns>
-		public LogScopedPropertyBuilder WithProperty(string key, object value)
+		/// <returns>A disposable log property builder that wraps the original ILogger scope</returns>
+		/// <example>
+		/// Simple example of attaching two properties
+		/// <code>
+		/// using (_logger.WithProperty("stream", "ue5").WithProperty("os", "windows").BuildScope())
+		/// {
+		///     _logger.LogInformation("A log message with properties attached");
+		/// }
+		/// </code>
+		/// </example>
+		public static LogScopedPropertyBuilder WithProperty(this ILogger logger, string key, object value)
 		{
-			_properties[key] = value;
-			return this;
+			LogScopedPropertyBuilder builder = new LogScopedPropertyBuilder(logger);
+			builder.WithProperty(key, value);
+			return builder;
 		}
-
-		/// <summary>
-		/// Starts new scope on underlying logger, passing previously set properties
-		/// </summary>
-		/// <returns>Original IDisposable from ILogger.BeginScope()</returns>
-		public IDisposable BeginScope()
-		{
-			return _logger.BeginScope(_properties);
-		}
-	}
-	
-	/// <summary>
-	/// Creates a new log scope and attaches a key/value property
-	/// Additional properties can be added with <see cref="LogScopedPropertyBuilder.WithProperty" />.
-	/// </summary>
-	/// <param name="logger">Logger to use</param>
-	/// <param name="key">Arbitrary key</param>
-	/// <param name="value">Arbitrary value</param>
-	/// <returns>A disposable log property builder that wraps the original ILogger scope</returns>
-	/// <example>
-	/// Simple example of attaching two properties
-	/// <code>
-	/// using (_logger.WithProperty("stream", "ue5").WithProperty("os", "windows").BuildScope())
-	/// {
-	///     _logger.LogInformation("A log message with properties attached");
-	/// }
-	/// </code>
-	/// </example>
-	public static LogScopedPropertyBuilder WithProperty(this ILogger logger, string key, object value)
-	{
-		LogScopedPropertyBuilder builder = new (logger);
-		builder.WithProperty(key, value);
-		return builder;
 	}
-}
\ No newline at end of file
+}
diff --git a/Engine/Source/Programs/Shared/EpicGames.Core/MemoryWriter.cs b/Engine/Source/Programs/Shared/EpicGames.Core/MemoryWriter.cs
--- a/Engine/Source/Programs/Shared/EpicGames.Core/MemoryWriter.cs
+++ b/Engine/Source/Programs/Shared/EpicGames.Core/MemoryWriter.cs
@@ -423,7 +423,7 @@ namespace EpicGames.Core
 		/// <param name="dictionary">The dictionary to write</param>
 		/// <param name="writeKey">Delegate to write an individual key</param>
 		/// <param name="writeValue">Delegate to write an individual value</param>
-		public static void WriteDictionary<TKey, TValue>(this IMemoryWriter writer, IReadOnlyDictionary<TKey, TValue> dictionary, Action<TKey> writeKey, Action<TValue> writeValue)
+		public static void WriteDictionary<TKey, TValue>(this IMemoryWriter writer, IReadOnlyDictionary<TKey, TValue> dictionary, Action<TKey> writeKey, Action<TValue> writeValue) where TKey : notnull
 		{
 			writer.WriteUnsignedVarInt(dictionary.Count);
 			foreach (KeyValuePair<TKey, TValue> kvp in dictionary)
