Added more complete recursion include checking in C# UHT.
Fixed threading issue in C# UHT message output.

#rnx

diff --git a/Engine/Source/Programs/Shared/EpicGames.UHT/Utils/UhtSession.cs b/Engine/Source/Programs/Shared/EpicGames.UHT/Utils/UhtSession.cs
--- a/Engine/Source/Programs/Shared/EpicGames.UHT/Utils/UhtSession.cs
+++ b/Engine/Source/Programs/Shared/EpicGames.UHT/Utils/UhtSession.cs
@@ -1457,7 +1457,7 @@ namespace EpicGames.UHT.Utils
 
 				// If we aren't caching messages and this is the first message,
 				// start a task to flush the messages.
-				if (!CacheMessages && _messages.Count == 1)
+				if (!CacheMessages && _messageTask == null)
 				{
 					_messageTask = Task.Factory.StartNew(() => FlushMessages(), CancellationToken.None, TaskCreationOptions.None, TaskScheduler.Default);
 				}
@@ -1490,9 +1490,14 @@ namespace EpicGames.UHT.Utils
 		/// </summary>
 		public void LogMessages()
 		{
-			if (_messageTask != null)
+			Task? messageTask = null;
+			lock (_messages)
+			{
+				messageTask = _messageTask;
+			}
+			if (messageTask != null)
 			{
-				_messageTask.Wait();
+				messageTask.Wait();
 			}
 
 			foreach (UhtMessage message in FetchOrderedMessages())
@@ -1506,16 +1511,24 @@ namespace EpicGames.UHT.Utils
 		/// </summary>
 		private void FlushMessages()
 		{
-			UhtMessage[]? messageArray = null;
-			lock (_messages)
+			while (true)
 			{
-				messageArray = _messages.ToArray();
-				_messages.Clear();
-			}
+				UhtMessage[]? messageArray = null;
+				lock (_messages)
+				{
+					messageArray = _messages.ToArray();
+					_messages.Clear();
+					if (messageArray.Length == 0)
+					{
+						_messageTask = null;
+						return;
+					}
+				}
 
-			foreach (UhtMessage message in messageArray)
-			{
-				LogMessage(message);
+				foreach (UhtMessage message in messageArray)
+				{
+					LogMessage(message);
+				}
 			}
 		}
 
@@ -2130,16 +2143,20 @@ namespace EpicGames.UHT.Utils
 			Permanent,
 		}
 
-		private void TopologicalRecursion(List<TopologicalState> states, UhtHeaderFile first, UhtHeaderFile visit)
+		private void TopologicalRecursion(List<TopologicalState> states, UhtHeaderFile first, UhtHeaderFile visit, HashSet<UhtHeaderFile> signaledFiles)
 		{
 			foreach (UhtHeaderFile referenced in visit.ReferencedHeadersNoLock)
 			{
 				if (states[referenced.HeaderFileTypeIndex] == TopologicalState.Temporary)
 				{
+					if (!signaledFiles.Add(referenced))
+					{
+						break;
+					}
 					first.LogError($"'{visit.FilePath}' includes/requires '{referenced.FilePath}'");
 					if (first != referenced)
 					{
-						TopologicalRecursion(states, first, referenced);
+						TopologicalRecursion(states, first, referenced, signaledFiles);
 					}
 					break;
 				}
@@ -2200,8 +2217,8 @@ namespace EpicGames.UHT.Utils
 						if (recursion != null)
 						{
 							headerFile.LogError("Circular dependency detected:");
-							TopologicalRecursion(states, recursion, recursion);
-							return;
+							HashSet<UhtHeaderFile> signaledFiles = new();
+							TopologicalRecursion(states, recursion, recursion, signaledFiles);
 						}
 					}
 				}
