UGS: Fix PCBs not showing if built for changelists not visible for the current project.

diff --git a/Engine/Source/Programs/Shared/EpicGames.Core/ListExtensions.cs b/Engine/Source/Programs/Shared/EpicGames.Core/ListExtensions.cs
--- a/Engine/Source/Programs/Shared/EpicGames.Core/ListExtensions.cs
+++ b/Engine/Source/Programs/Shared/EpicGames.Core/ListExtensions.cs
@@ -10,6 +10,80 @@ namespace EpicGames.Core
 	/// </summary>s
 	public static class ListExtensions
 	{
+		/// <summary>
+		/// Performs a binary search on the given list
+		/// </summary>
+		/// <typeparam name="T">The element type for the list</typeparam>
+		/// <param name="list">The list type</param>
+		/// <param name="item">Item to search for</param>
+		/// <returns>As List.BinarySearch</returns>
+		public static int BinarySearch<T>(this IList<T> list, T item)
+		{
+			return BinarySearch(list, x => x, item, Comparer<T>.Default);
+		}
+
+		/// <summary>
+		/// Performs a binary search on the given list
+		/// </summary>
+		/// <typeparam name="T">The element type for the list</typeparam>
+		/// <param name="list">The list type</param>
+		/// <param name="item">Item to search for</param>
+		/// <param name="comparer">Comparer for elements in the list</param>
+		/// <returns>As List.BinarySearch</returns>
+		public static int BinarySearch<T>(this IList<T> list, T item, IComparer<T> comparer)
+		{
+			return BinarySearch(list, x => x, item, comparer);
+		}
+
+		/// <summary>
+		/// Binary searches a list based on a projection
+		/// </summary>
+		/// <typeparam name="TItem">The item in the list</typeparam>
+		/// <typeparam name="TField">The field to search on</typeparam>
+		/// <param name="list">The list to search</param>
+		/// <param name="projection">The projection to apply to each item in the list</param>
+		/// <param name="item">The item to find</param>
+		/// <returns>As <see cref="List{T}.BinarySearch(T)"/></returns>
+		public static int BinarySearch<TItem, TField>(this IList<TItem> list, Func<TItem, TField> projection, TField item)
+		{
+			return BinarySearch(list, projection, item, Comparer<TField>.Default);
+		}
+
+		/// <summary>
+		/// Binary searches a list based on a projection
+		/// </summary>
+		/// <typeparam name="TItem">The item in the list</typeparam>
+		/// <typeparam name="TField">The field to search on</typeparam>
+		/// <param name="list">The list to search</param>
+		/// <param name="projection">The projection to apply to each item in the list</param>
+		/// <param name="item">The item to find</param>
+		/// <param name="comparer">Comparer for field elements</param>
+		/// <returns>As <see cref="List{T}.BinarySearch(T)"/></returns>
+		public static int BinarySearch<TItem, TField>(this IList<TItem> list, Func<TItem, TField> projection, TField item, IComparer<TField> comparer)
+		{
+			int lowerBound = 0;
+			int upperBound = list.Count - 1;
+			while (lowerBound <= upperBound)
+			{
+				int idx = lowerBound + (upperBound - lowerBound) / 2;
+
+				int comparison = comparer.Compare(projection(list[idx]), item);
+				if (comparison == 0)
+				{
+					return idx;
+				}
+				else if (comparison < 0)
+				{
+					lowerBound = idx + 1;
+				}
+				else
+				{
+					upperBound = idx - 1;
+				}
+			}
+			return ~lowerBound;
+		}
+
 		/// <summary>
 		/// Sorts a list by a particular field
 		/// </summary>
diff --git a/Engine/Source/Programs/UnrealGameSync/UnrealGameSync/Controls/WorkspaceControl.cs b/Engine/Source/Programs/UnrealGameSync/UnrealGameSync/Controls/WorkspaceControl.cs
--- a/Engine/Source/Programs/UnrealGameSync/UnrealGameSync/Controls/WorkspaceControl.cs
+++ b/Engine/Source/Programs/UnrealGameSync/UnrealGameSync/Controls/WorkspaceControl.cs
@@ -1999,22 +1999,25 @@ namespace UnrealGameSync
 
 		private string? GetArchiveKeyForChangeNumber(IArchiveInfo archive, int changeNumber)
 		{
-			string? archivePath;
-
 			Dictionary<int, string?>? changeNumberToArchivePath;
 			if (!_archiveToChangeNumberToArchiveKey.TryGetValue(archive.Name, out changeNumberToArchivePath))
 			{
 				changeNumberToArchivePath = new Dictionary<int, string?>();
 				_archiveToChangeNumberToArchiveKey[archive.Name] = changeNumberToArchivePath;
 			}
+			return GetArchiveKeyForChangeNumber(archive, changeNumber, changeNumber, changeNumberToArchivePath);
+		}
 
+		private string? GetArchiveKeyForChangeNumber(IArchiveInfo archive, int changeNumber, int maxChangeNumber, Dictionary<int, string?> changeNumberToArchivePath)
+		{
+			string? archivePath;
 			if (!changeNumberToArchivePath.TryGetValue(changeNumber, out archivePath))
 			{
 				PerforceChangeDetails? details;
 				if (_perforceMonitor.TryGetChangeDetails(changeNumber, out details))
 				{
 					// Try to get the archive for this CL
-					if (!archive.TryGetArchiveKeyForChangeNumber(changeNumber, out archivePath) && !details.ContainsCode)
+					if (!archive.TryGetArchiveKeyForChangeNumber(changeNumber, maxChangeNumber, out archivePath) && !details.ContainsCode)
 					{
 						// Otherwise if it's a content-only change, find the previous build any use the archive path from that
 						int index = _sortedChangeNumbers.BinarySearch(changeNumber);
diff --git a/Engine/Source/Programs/UnrealGameSync/UnrealGameSync/Jupiter/JupiterMonitor.cs b/Engine/Source/Programs/UnrealGameSync/UnrealGameSync/Jupiter/JupiterMonitor.cs
--- a/Engine/Source/Programs/UnrealGameSync/UnrealGameSync/Jupiter/JupiterMonitor.cs
+++ b/Engine/Source/Programs/UnrealGameSync/UnrealGameSync/Jupiter/JupiterMonitor.cs
@@ -223,7 +223,7 @@ namespace UnrealGameSync
 				return _changeToKey.Count > 0;
 			}
 
-			public bool TryGetArchiveKeyForChangeNumber(int changeNumber, out string archiveKey)
+			public bool TryGetArchiveKeyForChangeNumber(int changeNumber, int maxChangeNumber, out string archiveKey)
 			{
 				return _changeToKey.TryGetValue(changeNumber, out archiveKey);
 			}
diff --git a/Engine/Source/Programs/UnrealGameSync/UnrealGameSyncShared/ArchiveInfo.cs b/Engine/Source/Programs/UnrealGameSync/UnrealGameSyncShared/ArchiveInfo.cs
--- a/Engine/Source/Programs/UnrealGameSync/UnrealGameSyncShared/ArchiveInfo.cs
+++ b/Engine/Source/Programs/UnrealGameSync/UnrealGameSyncShared/ArchiveInfo.cs
@@ -2,6 +2,7 @@
 
 using EpicGames.Core;
 using EpicGames.Perforce;
+using JetBrains.Annotations;
 using Microsoft.Extensions.Logging;
 using System;
 using System.Collections.Generic;
@@ -24,7 +25,7 @@ namespace UnrealGameSync
 		string? Target { get; }
 
 		bool Exists();
-		bool TryGetArchiveKeyForChangeNumber(int changeNumber, [NotNullWhen(true)] out string? archiveKey);
+		bool TryGetArchiveKeyForChangeNumber(int changeNumber, int maxChangeNumber, [NotNullWhen(true)] out string? archiveKey);
 		Task<bool> DownloadArchive(IPerforceConnection perforce, string archiveKey, DirectoryReference localRootPath, FileReference manifestFileName, ILogger logger, ProgressValue progress, CancellationToken cancellationToken);
 	}
 
@@ -93,9 +94,24 @@ namespace UnrealGameSync
 			return true;
 		}
 
-		public bool TryGetArchiveKeyForChangeNumber(int changeNumber, [NotNullWhen(true)] out string? archiveKey)
+		public bool TryGetArchiveKeyForChangeNumber(int changeNumber, int maxChangeNumber, [NotNullWhen(true)] out string? archiveKey)
 		{
-			return ChangeNumberToFileRevision.TryGetValue(changeNumber, out archiveKey);
+			int idx = ChangeNumberToFileRevision.Keys.BinarySearch(changeNumber);
+			if (idx >= 0)
+			{
+				archiveKey = ChangeNumberToFileRevision.Values[idx];
+				return true;
+			}
+
+			int nextIdx = ~idx;
+			if (nextIdx < ChangeNumberToFileRevision.Count && ChangeNumberToFileRevision.Keys[nextIdx] <= maxChangeNumber)
+			{
+				archiveKey = ChangeNumberToFileRevision.Values[nextIdx];
+				return true;
+			}
+
+			archiveKey = null;
+			return false;
 		}
 
 		public async Task<bool> DownloadArchive(IPerforceConnection perforce, string archiveKey, DirectoryReference localRootPath, FileReference manifestFileName, ILogger logger, ProgressValue progress, CancellationToken cancellationToken)
