Remove IList.BinarySearch extension methods due to IList not deriving from IReadOnlyList, and some classes implementing both.

diff --git a/Engine/Source/Programs/Shared/EpicGames.Core/ListExtensions.cs b/Engine/Source/Programs/Shared/EpicGames.Core/ListExtensions.cs
--- a/Engine/Source/Programs/Shared/EpicGames.Core/ListExtensions.cs
+++ b/Engine/Source/Programs/Shared/EpicGames.Core/ListExtensions.cs
@@ -1,6 +1,7 @@
 // Copyright Epic Games, Inc. All Rights Reserved.
 
 using System;
+using System.Collections;
 using System.Collections.Generic;
 
 namespace EpicGames.Core
@@ -11,78 +12,31 @@ namespace EpicGames.Core
 	public static class ListExtensions
 	{
 		/// <summary>
-		/// Performs a binary search on the given list
+		/// Wrapper around a list to implement <see cref="IReadOnlyList{T}"/>
 		/// </summary>
-		/// <typeparam name="T">The element type for the list</typeparam>
-		/// <param name="list">The list type</param>
-		/// <param name="item">Item to search for</param>
-		/// <returns>As List.BinarySearch</returns>
-		public static int BinarySearch<T>(this IList<T> list, T item)
+		/// <typeparam name="T"></typeparam>
+		class ReadOnlyList<T> : IReadOnlyList<T>
 		{
-			return BinarySearch(list, x => x, item, Comparer<T>.Default);
-		}
+			readonly IList<T> _inner;
 
-		/// <summary>
-		/// Performs a binary search on the given list
-		/// </summary>
-		/// <typeparam name="T">The element type for the list</typeparam>
-		/// <param name="list">The list type</param>
-		/// <param name="item">Item to search for</param>
-		/// <param name="comparer">Comparer for elements in the list</param>
-		/// <returns>As List.BinarySearch</returns>
-		public static int BinarySearch<T>(this IList<T> list, T item, IComparer<T> comparer)
-		{
-			return BinarySearch(list, x => x, item, comparer);
-		}
+			public ReadOnlyList(IList<T> inner) => _inner = inner;
 
-		/// <summary>
-		/// Binary searches a list based on a projection
-		/// </summary>
-		/// <typeparam name="TItem">The item in the list</typeparam>
-		/// <typeparam name="TField">The field to search on</typeparam>
-		/// <param name="list">The list to search</param>
-		/// <param name="projection">The projection to apply to each item in the list</param>
-		/// <param name="item">The item to find</param>
-		/// <returns>As <see cref="List{T}.BinarySearch(T)"/></returns>
-		public static int BinarySearch<TItem, TField>(this IList<TItem> list, Func<TItem, TField> projection, TField item)
-		{
-			return BinarySearch(list, projection, item, Comparer<TField>.Default);
+			public T this[int index] => _inner[index];
+
+			public int Count => _inner.Count;
+
+			public IEnumerator<T> GetEnumerator() => _inner.GetEnumerator();
+
+			IEnumerator IEnumerable.GetEnumerator() => _inner.GetEnumerator();
 		}
 
 		/// <summary>
-		/// Binary searches a list based on a projection
+		/// Create a read-only wrapper around a list
 		/// </summary>
-		/// <typeparam name="TItem">The item in the list</typeparam>
-		/// <typeparam name="TField">The field to search on</typeparam>
-		/// <param name="list">The list to search</param>
-		/// <param name="projection">The projection to apply to each item in the list</param>
-		/// <param name="item">The item to find</param>
-		/// <param name="comparer">Comparer for field elements</param>
-		/// <returns>As <see cref="List{T}.BinarySearch(T)"/></returns>
-		public static int BinarySearch<TItem, TField>(this IList<TItem> list, Func<TItem, TField> projection, TField item, IComparer<TField> comparer)
-		{
-			int lowerBound = 0;
-			int upperBound = list.Count - 1;
-			while (lowerBound <= upperBound)
-			{
-				int idx = lowerBound + (upperBound - lowerBound) / 2;
-
-				int comparison = comparer.Compare(projection(list[idx]), item);
-				if (comparison == 0)
-				{
-					return idx;
-				}
-				else if (comparison < 0)
-				{
-					lowerBound = idx + 1;
-				}
-				else
-				{
-					upperBound = idx - 1;
-				}
-			}
-			return ~lowerBound;
-		}
+		/// <typeparam name="T"></typeparam>
+		/// <param name="list"></param>
+		/// <returns></returns>
+		public static IReadOnlyList<T> AsReadOnly<T>(this IList<T> list) => new ReadOnlyList<T>(list);
 
 		/// <summary>
 		/// Sorts a list by a particular field
diff --git a/Engine/Source/Programs/UnrealGameSync/UnrealGameSyncShared/ArchiveInfo.cs b/Engine/Source/Programs/UnrealGameSync/UnrealGameSyncShared/ArchiveInfo.cs
--- a/Engine/Source/Programs/UnrealGameSync/UnrealGameSyncShared/ArchiveInfo.cs
+++ b/Engine/Source/Programs/UnrealGameSync/UnrealGameSyncShared/ArchiveInfo.cs
@@ -96,7 +96,7 @@ namespace UnrealGameSync
 
 		public bool TryGetArchiveKeyForChangeNumber(int changeNumber, int maxChangeNumber, [NotNullWhen(true)] out string? archiveKey)
 		{
-			int idx = ChangeNumberToFileRevision.Keys.BinarySearch(changeNumber);
+			int idx = ChangeNumberToFileRevision.Keys.AsReadOnly().BinarySearch(changeNumber);
 			if (idx >= 0)
 			{
 				archiveKey = ChangeNumberToFileRevision.Values[idx];
