Improved C# UHT circular include detection and logging.

diff --git a/Engine/Source/Programs/Shared/EpicGames.UHT/Utils/UhtSession.cs b/Engine/Source/Programs/Shared/EpicGames.UHT/Utils/UhtSession.cs
--- a/Engine/Source/Programs/Shared/EpicGames.UHT/Utils/UhtSession.cs
+++ b/Engine/Source/Programs/Shared/EpicGames.UHT/Utils/UhtSession.cs
@@ -2143,28 +2143,9 @@ namespace EpicGames.UHT.Utils
 			Permanent,
 		}
 
-		private void TopologicalRecursion(List<TopologicalState> states, UhtHeaderFile first, UhtHeaderFile visit, HashSet<UhtHeaderFile> signaledFiles)
-		{
-			foreach (UhtHeaderFile referenced in visit.ReferencedHeadersNoLock)
-			{
-				if (states[referenced.HeaderFileTypeIndex] == TopologicalState.Temporary)
-				{
-					if (!signaledFiles.Add(referenced))
-					{
-						break;
-					}
-					first.LogError($"'{visit.FilePath}' includes/requires '{referenced.FilePath}'");
-					if (first != referenced)
-					{
-						TopologicalRecursion(states, first, referenced, signaledFiles);
-					}
-					break;
-				}
-			}
-		}
-
-		private UhtHeaderFile? TopologicalVisit(List<TopologicalState> states, UhtHeaderFile visit)
+		private void TopologicalVisit(List<TopologicalState> states, UhtHeaderFile visit, List<UhtHeaderFile> headerStack)
 		{
+			headerStack.Add(visit);
 			switch (states[visit.HeaderFileTypeIndex])
 			{
 				case TopologicalState.Unmarked:
@@ -2173,26 +2154,39 @@ namespace EpicGames.UHT.Utils
 					{
 						if (visit != referenced)
 						{
-							UhtHeaderFile? recursion = TopologicalVisit(states, referenced);
-							if (recursion != null)
-							{
-								return recursion;
-							}
+							TopologicalVisit(states, referenced, headerStack);
 						}
 					}
 					states[visit.HeaderFileTypeIndex] = TopologicalState.Permanent;
 					_sortedHeaderFiles.Add(visit);
-					return null;
+					break;
 
 				case TopologicalState.Temporary:
-					return visit;
+					{
+						int index = headerStack.IndexOf(visit);
+						if (index == -1 || index == headerStack.Count - 1)
+						{
+							throw new UhtIceException("Error locating include file loop");
+						}
+						index++;
+						visit.LogError("Circular dependency detected:");
+						UhtHeaderFile previous = visit;
+						for (int loopIndex = index; loopIndex < headerStack.Count; loopIndex++)
+						{
+							UhtHeaderFile next = headerStack[loopIndex];
+							previous.LogError($"includes/requires '{next.FilePath}'");
+							previous = next;
+						}
+					}
+					break;
 
 				case TopologicalState.Permanent:
-					return null;
+					break;
 
 				default:
 					throw new UhtIceException("Unknown topological state");
 			}
+			headerStack.RemoveAt(headerStack.Count - 1);
 		}
 
 		private void TopologicalSortHeaderFiles()
@@ -2209,17 +2203,12 @@ namespace EpicGames.UHT.Utils
 					states.Add(TopologicalState.Unmarked);
 				}
 
+				List<UhtHeaderFile> headerStack = new(32); // arbitrary capacity
 				foreach (UhtHeaderFile headerFile in HeaderFiles)
 				{
 					if (states[headerFile.HeaderFileTypeIndex] == TopologicalState.Unmarked)
 					{
-						UhtHeaderFile? recursion = TopologicalVisit(states, headerFile);
-						if (recursion != null)
-						{
-							headerFile.LogError("Circular dependency detected:");
-							HashSet<UhtHeaderFile> signaledFiles = new();
-							TopologicalRecursion(states, recursion, recursion, signaledFiles);
-						}
+						TopologicalVisit(states, headerFile, headerStack);
 					}
 				}
 			});
