- Added CreateObjectList to StringUtils, this is helpful for some refactoring in other code, but it could be a useful utility function to have

diff --git a/Engine/Source/Programs/Shared/EpicGames.Core/StringUtils.cs b/Engine/Source/Programs/Shared/EpicGames.Core/StringUtils.cs
--- a/Engine/Source/Programs/Shared/EpicGames.Core/StringUtils.cs
+++ b/Engine/Source/Programs/Shared/EpicGames.Core/StringUtils.cs
@@ -292,6 +292,17 @@ namespace EpicGames.Core
 			}
 		}
 
+		/// <summary>
+		/// Generates a string suitable for debugging a list of objects using ToString(). Lists one per line with the prefix string on the first line.
+		/// </summary>
+		/// <param name="prefix">Prefix string to print along with the list of BuildInfos.</param>
+		/// <param name="buildInfos"></param>
+		/// <returns>the resulting debug string</returns>
+		public static string CreateObjectList<T>(this IEnumerable<T> objects, string prefix)
+		{
+			return objects.Aggregate(new StringBuilder(prefix), (sb, obj) => sb.AppendFormat($"\n    {obj}")).ToString();
+		}
+		
 		/// <summary>
 		/// Parses a hexadecimal digit
 		/// </summary>
