- Added missing param comment for a new function in EpicGames.Core

diff --git a/Engine/Source/Programs/Shared/EpicGames.Core/StringUtils.cs b/Engine/Source/Programs/Shared/EpicGames.Core/StringUtils.cs
--- a/Engine/Source/Programs/Shared/EpicGames.Core/StringUtils.cs
+++ b/Engine/Source/Programs/Shared/EpicGames.Core/StringUtils.cs
@@ -295,8 +295,8 @@ namespace EpicGames.Core
 		/// <summary>
 		/// Generates a string suitable for debugging a list of objects using ToString(). Lists one per line with the prefix string on the first line.
 		/// </summary>
-		/// <param name="prefix">Prefix string to print along with the list of BuildInfos.</param>
-		/// <param name="buildInfos"></param>
+		/// <param name="objects">The list of objects to inset into the output string</param>
+		/// <param name="prefix">Prefix string to print along with the list of objects</param>
 		/// <returns>the resulting debug string</returns>
 		public static string CreateObjectList<T>(this IEnumerable<T> objects, string prefix)
 		{
