Horde: Various changes to bundle utility classes.

* Bundles now include a unique identifier (GUID) for the type of each exported node, as well as a version number for serializer. Unique identifier and version number is configured through the TreeNodeAttribute class which must be attached to each node type.
* Known node types are registered explicitly through the TreeReaderOptions class passed to the TreeReader constructor. Types are discovered as needed when writing nodes to TreeWriter, via the type of supplied objects.
* Support for custom deserializers is removed. Nodes are always constructed via a constructor taking a ITreeNodeReader argument.
* Nodes are now read from bundles using an explicitly instantiated TreeReader class, which removes a lot of the complexity from StorageNodeBase.

diff --git a/Engine/Source/Programs/Horde/Horde.Agent/Commands/Bundles/CreateCommand.cs b/Engine/Source/Programs/Horde/Horde.Agent/Commands/Bundles/CreateCommand.cs
--- a/Engine/Source/Programs/Horde/Horde.Agent/Commands/Bundles/CreateCommand.cs
+++ b/Engine/Source/Programs/Horde/Horde.Agent/Commands/Bundles/CreateCommand.cs
@@ -18,46 +18,47 @@ namespace Horde.Agent.Commands.Bundles
 	{
 		protected interface IStorageClientOwner : IDisposable
 		{
+			public IMemoryCache Cache { get; }
 			public IStorageClient Store { get; }
 		}
 
 		class FileStorageClientOwner : IStorageClientOwner
 		{
-			readonly IMemoryCache _cache;
+			public IMemoryCache Cache { get; }
 			public IStorageClient Store { get; }
 
 			public FileStorageClientOwner(DirectoryReference storageDir, ILogger logger)
 			{
-				_cache = new MemoryCache(new MemoryCacheOptions());
-				Store = new FileStorageClient(storageDir, _cache, logger);
+				Cache = new MemoryCache(new MemoryCacheOptions());
+				Store = new FileStorageClient(storageDir, logger);
 			}
 
 			public void Dispose()
 			{
-				_cache.Dispose();
+				Cache.Dispose();
 			}
 		}
 
 		class HttpStorageClientOwner : IStorageClientOwner
 		{
-			readonly IMemoryCache _cache;
+			public IMemoryCache Cache { get; }
 			readonly HttpClient _httpClient;
 			readonly HttpClient _redirectHttpClient;
 			public IStorageClient Store { get; }
 
 			public HttpStorageClientOwner(NamespaceId namespaceId, Uri baseUri, ILogger logger)
 			{
-				_cache = new MemoryCache(new MemoryCacheOptions());
+				Cache = new MemoryCache(new MemoryCacheOptions());
 				_httpClient = new HttpClient { BaseAddress = baseUri };
 				_redirectHttpClient = new HttpClient();
-				Store = new HttpStorageClient(namespaceId, _httpClient, _redirectHttpClient, _cache, logger);
+				Store = new HttpStorageClient(namespaceId, _httpClient, _redirectHttpClient, logger);
 			}
 
 			public void Dispose()
 			{
 				_redirectHttpClient.Dispose();
 				_httpClient.Dispose();
-				_cache.Dispose();
+				Cache.Dispose();
 			}
 		}
 
diff --git a/Engine/Source/Programs/Horde/Horde.Agent/Commands/Bundles/ExtractCommand.cs b/Engine/Source/Programs/Horde/Horde.Agent/Commands/Bundles/ExtractCommand.cs
--- a/Engine/Source/Programs/Horde/Horde.Agent/Commands/Bundles/ExtractCommand.cs
+++ b/Engine/Source/Programs/Horde/Horde.Agent/Commands/Bundles/ExtractCommand.cs
@@ -21,13 +21,14 @@ namespace Horde.Agent.Commands.Bundles
 
 		public override async Task<int> ExecuteAsync(ILogger logger)
 		{
-			using IStorageClientOwner storeOwner = CreateStorageClient(logger);
-			IStorageClient store = storeOwner.Store;
+			using IStorageClientOwner owner = CreateStorageClient(logger);
+			IStorageClient store = owner.Store;
+			TreeReader reader = new TreeReader(owner.Store, owner.Cache, logger);
 
 			Stopwatch timer = Stopwatch.StartNew();
 
-			DirectoryNode node = await store.ReadNodeAsync<DirectoryNode>(RefName);
-			await node.CopyToDirectoryAsync(OutputDir.ToDirectoryInfo(), logger, CancellationToken.None);
+			DirectoryNode node = await reader.ReadNodeAsync<DirectoryNode>(RefName);
+			await node.CopyToDirectoryAsync(reader, OutputDir.ToDirectoryInfo(), logger, CancellationToken.None);
 
 			logger.LogInformation("Elapsed: {Time}s", timer.Elapsed.TotalSeconds);
 			return 0;
diff --git a/Engine/Source/Programs/Horde/Horde.Agent/Commands/Vcs/VcsCommand.cs b/Engine/Source/Programs/Horde/Horde.Agent/Commands/Vcs/VcsCommand.cs
--- a/Engine/Source/Programs/Horde/Horde.Agent/Commands/Vcs/VcsCommand.cs
+++ b/Engine/Source/Programs/Horde/Horde.Agent/Commands/Vcs/VcsCommand.cs
@@ -507,8 +507,9 @@ namespace Horde.Agent.Commands.Bundles
 
 			using IStorageClientOwner owner = CreateStorageClient(rootDir, logger);
 			IStorageClient store = owner.Store;
+			TreeReader reader = new TreeReader(owner.Store, owner.Cache, logger);
 
-			CommitNode tip = await store.ReadNodeAsync<CommitNode>(branchName);
+			CommitNode tip = await reader.ReadNodeAsync<CommitNode>(branchName);
 			if (Change != 0)
 			{
 				while (tip.Number != Change)
@@ -518,24 +519,24 @@ namespace Horde.Agent.Commands.Bundles
 						logger.LogError("Unable to find change {Change}", Change);
 						return 1;
 					}
-					tip = await tip.Parent.ExpandAsync();
+					tip = await tip.Parent.ExpandAsync(reader);
 				}
 			}
 
-			workspaceState.Tree = await RealizeAsync(tip.Contents, rootDir, newState, Clean, logger);
+			workspaceState.Tree = await RealizeAsync(reader, tip.Contents, rootDir, newState, Clean, logger);
 			await WriteStateAsync(rootDir, workspaceState);
 
 			logger.LogInformation("Updated workspace to change {Number}", tip.Number);
 			return 0;
 		}
 
-		static async Task<DirectoryState> RealizeAsync(TreeNodeRef<DirectoryNode> directoryRef, DirectoryReference dirPath, DirectoryState? directoryState, bool clean, ILogger logger)
+		static async Task<DirectoryState> RealizeAsync(TreeReader reader, TreeNodeRef<DirectoryNode> directoryRef, DirectoryReference dirPath, DirectoryState? directoryState, bool clean, ILogger logger)
 		{
 			DirectoryReference.CreateDirectory(dirPath);
 
 			DirectoryState newState = new DirectoryState();
 
-			DirectoryNode directoryNode = await directoryRef.ExpandAsync();
+			DirectoryNode directoryNode = await directoryRef.ExpandAsync(reader);
 			foreach ((Utf8String name, DirectoryEntry? subDirEntry, DirectoryState? subDirState) in Zip(directoryNode.NameToDirectory, directoryState?.Directories))
 			{
 				DirectoryReference subDirPath = DirectoryReference.Combine(dirPath, name.ToString());
@@ -548,7 +549,7 @@ namespace Horde.Agent.Commands.Bundles
 				}
 				else
 				{
-					newState.Directories[name] = await RealizeAsync(subDirEntry, subDirPath, subDirState, clean, logger);
+					newState.Directories[name] = await RealizeAsync(reader, subDirEntry, subDirPath, subDirState, clean, logger);
 				}
 			}
 
@@ -564,7 +565,7 @@ namespace Horde.Agent.Commands.Bundles
 				}
 				else if (fileState == null || fileState.Hash != fileEntry.Hash)
 				{
-					newState.Files[name] = await CheckoutFileAsync(fileEntry, filePath.ToFileInfo(), logger);
+					newState.Files[name] = await CheckoutFileAsync(reader, fileEntry, filePath.ToFileInfo(), logger);
 				}
 				else
 				{
@@ -576,11 +577,11 @@ namespace Horde.Agent.Commands.Bundles
 			return newState;
 		}
 
-		static async Task<FileState> CheckoutFileAsync(TreeNodeRef<FileNode> fileRef, FileInfo fileInfo, ILogger logger)
+		static async Task<FileState> CheckoutFileAsync(TreeReader reader, TreeNodeRef<FileNode> fileRef, FileInfo fileInfo, ILogger logger)
 		{
 			logger.LogInformation("Updating {File} to {Hash}", fileInfo, fileRef.Hash);
-			FileNode fileNode = await fileRef.ExpandAsync();
-			await fileNode.CopyToFileAsync(fileInfo, CancellationToken.None);
+			FileNode fileNode = await fileRef.ExpandAsync(reader);
+			await fileNode.CopyToFileAsync(reader, fileInfo, CancellationToken.None);
 			fileInfo.Refresh();
 			return new FileState(fileInfo, fileRef.Hash);
 		}
@@ -617,8 +618,9 @@ namespace Horde.Agent.Commands.Bundles
 
 			using IStorageClientOwner owner = CreateStorageClient(rootDir, logger);
 			IStorageClient store = owner.Store;
+			TreeReader reader = new TreeReader(owner.Store, owner.Cache, logger);
 
-			CommitNode? tip = await store.TryReadNodeAsync<CommitNode>(workspaceState.Branch);
+			CommitNode? tip = await reader.TryReadNodeAsync<CommitNode>(workspaceState.Branch);
 			TreeNodeRef<CommitNode>? tipRef = (tip == null) ? null : new TreeNodeRef<CommitNode>(tip);
 
 			DirectoryNode rootNode;
@@ -628,14 +630,14 @@ namespace Horde.Agent.Commands.Bundles
 			}
 			else
 			{
-				rootNode = await tip.Contents.ExpandCopyAsync();
+				rootNode = await tip.Contents.ExpandCopyAsync(reader);
 			}
 
 			TreeNodeRef<DirectoryNode> rootRef = new TreeNodeRef<DirectoryNode>(rootNode);
 
 			List<(DirectoryNode, FileInfo, FileState)> files = new List<(DirectoryNode, FileInfo, FileState)>();
 			List<(TreeNodeRef<DirectoryNode>, DirectoryState)> directories = new List<(TreeNodeRef<DirectoryNode>, DirectoryState)>();
-			await UpdateTreeAsync(rootRef, rootDir, oldState, newState, files, directories);
+			await UpdateTreeAsync(reader, rootRef, rootDir, oldState, newState, files, directories);
 
 			TreeWriter writer = new TreeWriter(store);
 			await DirectoryNode.CopyFromDirectoryAsync(files.ConvertAll(x => (x.Item1, x.Item2)), new ChunkingOptions(), writer, CancellationToken.None);
@@ -670,11 +672,11 @@ namespace Horde.Agent.Commands.Bundles
 			return 0;
 		}
 
-		private async Task UpdateTreeAsync(TreeNodeRef<DirectoryNode> rootRef, DirectoryReference rootDir, DirectoryState? oldState, DirectoryState newState, List<(DirectoryNode, FileInfo, FileState)> files, List<(TreeNodeRef<DirectoryNode>, DirectoryState)> directories)
+		private async Task UpdateTreeAsync(TreeReader reader, TreeNodeRef<DirectoryNode> rootRef, DirectoryReference rootDir, DirectoryState? oldState, DirectoryState newState, List<(DirectoryNode, FileInfo, FileState)> files, List<(TreeNodeRef<DirectoryNode>, DirectoryState)> directories)
 		{
 			directories.Add((rootRef, newState));
 
-			DirectoryNode root = await rootRef.ExpandAsync();
+			DirectoryNode root = await rootRef.ExpandAsync(reader);
 			foreach ((Utf8String name, DirectoryState? oldSubDirState, DirectoryState? newSubDirState) in Zip(oldState?.Directories, newState.Directories))
 			{
 				if (newSubDirState == null)
@@ -683,7 +685,7 @@ namespace Horde.Agent.Commands.Bundles
 				}
 				else if(oldSubDirState != newSubDirState)
 				{
-					await UpdateTreeAsync(root.FindOrAddDirectoryEntry(name), DirectoryReference.Combine(rootDir, name.ToString()), oldSubDirState, newSubDirState, files, directories);
+					await UpdateTreeAsync(reader, root.FindOrAddDirectoryEntry(name), DirectoryReference.Combine(rootDir, name.ToString()), oldSubDirState, newSubDirState, files, directories);
 				}
 			}
 
@@ -715,10 +717,11 @@ namespace Horde.Agent.Commands.Bundles
 
 			using IStorageClientOwner owner = CreateStorageClient(rootDir, logger);
 			IStorageClient store = owner.Store;
+			TreeReader reader = new TreeReader(owner.Store, owner.Cache, logger);
 
 			List<CommitNode> commits = new List<CommitNode>();
 
-			CommitNode? tip = await store.TryReadNodeAsync<CommitNode>(workspaceState.Branch);
+			CommitNode? tip = await reader.TryReadNodeAsync<CommitNode>(workspaceState.Branch);
 			if (tip != null)
 			{
 				for (int idx = 0; idx < Count; idx++)
@@ -730,7 +733,7 @@ namespace Horde.Agent.Commands.Bundles
 						break;
 					}
 
-					tip = await tip.Parent.ExpandAsync();
+					tip = await tip.Parent.ExpandAsync(reader);
 				}
 			}
 
diff --git a/Engine/Source/Programs/Horde/Horde.Agent/Execution/JobExecutor.cs b/Engine/Source/Programs/Horde/Horde.Agent/Execution/JobExecutor.cs
--- a/Engine/Source/Programs/Horde/Horde.Agent/Execution/JobExecutor.cs
+++ b/Engine/Source/Programs/Horde/Horde.Agent/Execution/JobExecutor.cs
@@ -665,8 +665,9 @@ namespace Horde.Agent.Execution
 
 			using HttpClient redirectHttpClient = _httpClientFactory.CreateClient(StorageHttpClientName);
 
-			HttpStorageClient storage = new HttpStorageClient(new NamespaceId("default"), httpClient, redirectHttpClient, cache, logger);
-//			FileStorageClient storage = new FileStorageClient(DirectoryReference.Combine(sharedStorageDir, "bundles"), cache, logger);
+			HttpStorageClient storage = new HttpStorageClient(new NamespaceId("default"), httpClient, redirectHttpClient, logger);
+			//			FileStorageClient storage = new FileStorageClient(DirectoryReference.Combine(sharedStorageDir, "bundles"), cache, logger);
+			TreeReader reader = new TreeReader(storage, cache, logger);
 			logger.LogInformation("Using Horde-managed shared storage via {SharedStorageDir}", sharedStorageDir);
 
 			// Create the mapping of tag names to file sets
@@ -686,7 +687,7 @@ namespace Horde.Agent.Execution
 				string nodeName = input.Substring(0, slashIdx);
 				string tagName = input.Substring(slashIdx + 1);
 
-				TempStorageTagManifest fileList = await TempStorage.RetrieveTagAsync(storage, RefPrefix, nodeName, tagName, manifestDir, logger, cancellationToken);
+				TempStorageTagManifest fileList = await TempStorage.RetrieveTagAsync(reader, RefPrefix, nodeName, tagName, manifestDir, logger, cancellationToken);
 				tagNameToFileSet[tagName] = fileList.ToFileSet(workspaceDir);
 				inputStorageBlocks.UnionWith(fileList.Blocks);
 			}
@@ -699,7 +700,7 @@ namespace Horde.Agent.Execution
 				scope.Span.SetTag("blocks", inputStorageBlocks.Count);
 				foreach (TempStorageBlockRef inputStorageBlock in inputStorageBlocks)
 				{
-					TempStorageBlockManifest manifest = await TempStorage.RetrieveBlockAsync(storage, RefPrefix, inputStorageBlock.NodeName, inputStorageBlock.OutputName, workspaceDir, manifestDir, logger, cancellationToken);
+					TempStorageBlockManifest manifest = await TempStorage.RetrieveBlockAsync(reader, RefPrefix, inputStorageBlock.NodeName, inputStorageBlock.OutputName, workspaceDir, manifestDir, logger, cancellationToken);
 					inputManifests[inputStorageBlock] = manifest;
 				}
 				scope.Span.SetTag("size", inputManifests.Sum(x => x.Value.GetTotalSize()));
diff --git a/Engine/Source/Programs/Horde/Horde.Agent/Utility/TempStorage.cs b/Engine/Source/Programs/Horde/Horde.Agent/Utility/TempStorage.cs
--- a/Engine/Source/Programs/Horde/Horde.Agent/Utility/TempStorage.cs
+++ b/Engine/Source/Programs/Horde/Horde.Agent/Utility/TempStorage.cs
@@ -682,7 +682,7 @@ namespace Horde.Storage.Utility
 		/// <summary>
 		/// Reads a set of tagged files from disk
 		/// </summary>
-		/// <param name="store"></param>
+		/// <param name="reader">Reader for node data</param>
 		/// <param name="refPrefix">Prefix for ref names</param>
 		/// <param name="nodeName">Name of the node which produced the tag set</param>
 		/// <param name="tagName">Name of the tag, with a '#' prefix</param>
@@ -690,7 +690,7 @@ namespace Horde.Storage.Utility
 		/// <param name="logger">Logger for output</param>
 		/// <param name="cancellationToken"></param>
 		/// <returns>The set of files</returns>
-		public static async Task<TempStorageTagManifest> RetrieveTagAsync(IStorageClient store, string refPrefix, string nodeName, string tagName, DirectoryReference manifestDir, ILogger logger, CancellationToken cancellationToken)
+		public static async Task<TempStorageTagManifest> RetrieveTagAsync(TreeReader reader, string refPrefix, string nodeName, string tagName, DirectoryReference manifestDir, ILogger logger, CancellationToken cancellationToken)
 		{
 			TempStorageTagManifest fileList;
 
@@ -706,7 +706,7 @@ namespace Horde.Storage.Utility
 				RefName refName = GetRefNameForNode(refPrefix, nodeName);
 				logger.LogInformation("Reading node \"{NodeName}\" tag \"{TagName}\" from temp storage (ref: {RefName}, localFile: {LocalFile})", nodeName, tagName, refName, localFileListLocation);
 
-				TempStorageNode node = await store.ReadNodeAsync<TempStorageNode>(refName, cancellationToken: cancellationToken);
+				TempStorageNode node = await reader.ReadNodeAsync<TempStorageNode>(refName, cancellationToken: cancellationToken);
 
 				TreeNodeRef<TempStorageTagNode>? tagNodeRef;
 				if (!node.Tags.TryGetValue(tagName, out tagNodeRef))
@@ -714,7 +714,7 @@ namespace Horde.Storage.Utility
 					throw new TempStorageException($"Missing tag {tagName} from node {nodeName}");
 				}
 
-				TempStorageTagNode tagNode = await tagNodeRef.ExpandAsync(cancellationToken);
+				TempStorageTagNode tagNode = await tagNodeRef.ExpandAsync(reader, cancellationToken);
 				fileList = tagNode.FileList;
 
 				// Save the manifest locally
@@ -807,7 +807,7 @@ namespace Horde.Storage.Utility
 		/// <summary>
 		/// Retrieve an output of the given node. Fetches and decompresses the files from shared storage if necessary, or validates the local files.
 		/// </summary>
-		/// <param name="store">Store to read data from</param>
+		/// <param name="reader">Store to read data from</param>
 		/// <param name="refPrefix">Prefix for ref names</param>
 		/// <param name="nodeName">The node which created the storage block</param>
 		/// <param name="blockName">Name of the block to retrieve.</param>
@@ -816,7 +816,7 @@ namespace Horde.Storage.Utility
 		/// <param name="logger">Logger for output</param>
 		/// <param name="cancellationToken"></param>
 		/// <returns>Manifest of the files retrieved</returns>
-		public static async Task<TempStorageBlockManifest> RetrieveBlockAsync(IStorageClient store, string refPrefix, string nodeName, string blockName, DirectoryReference rootDir, DirectoryReference manifestDir, ILogger logger, CancellationToken cancellationToken)
+		public static async Task<TempStorageBlockManifest> RetrieveBlockAsync(TreeReader reader, string refPrefix, string nodeName, string blockName, DirectoryReference rootDir, DirectoryReference manifestDir, ILogger logger, CancellationToken cancellationToken)
 		{
 			// Get the path to the local manifest
 			FileReference localManifestFile = GetBlockManifestLocation(manifestDir, nodeName, blockName);
@@ -835,7 +835,7 @@ namespace Horde.Storage.Utility
 				RefName refName = GetRefNameForNode(refPrefix, nodeName);
 				logger.LogInformation("Reading node \"{NodeName}\" block \"{BlockName}\" from temp storage (ref: {RefName}, local: {LocalFile})", nodeName, blockName, refName, localManifestFile);
 
-				TempStorageNode node = await store.ReadNodeAsync<TempStorageNode>(refName, cancellationToken: cancellationToken);
+				TempStorageNode node = await reader.ReadNodeAsync<TempStorageNode>(refName, cancellationToken: cancellationToken);
 				
 				TreeNodeRef<TempStorageBlockNode>? blockNodeRef;
 				if (!node.Blocks.TryGetValue(blockName, out blockNodeRef))
@@ -843,7 +843,7 @@ namespace Horde.Storage.Utility
 					throw new TempStorageException($"Missing block \"{blockName}\" from node \"{nodeName}\"");
 				}
 
-				TempStorageBlockNode blockNode = await blockNodeRef.ExpandAsync(cancellationToken);
+				TempStorageBlockNode blockNode = await blockNodeRef.ExpandAsync(reader, cancellationToken);
 				manifest = blockNode.Manifest;
 
 				// Delete all the existing files. They may be read-only.
@@ -854,8 +854,8 @@ namespace Horde.Storage.Utility
 				}
 
 				// Add all the files and flush the ref
-				DirectoryNode rootDirNode = await blockNode.Contents.ExpandAsync(cancellationToken);
-				await rootDirNode.CopyToDirectoryAsync(rootDir.ToDirectoryInfo(), logger, cancellationToken);
+				DirectoryNode rootDirNode = await blockNode.Contents.ExpandAsync(reader, cancellationToken);
+				await rootDirNode.CopyToDirectoryAsync(reader, rootDir.ToDirectoryInfo(), logger, cancellationToken);
 
 				// Update the timestamps to match the manifest.
 				foreach (TempStorageFile ManifestFile in manifest.Files)
diff --git a/Engine/Source/Programs/Horde/Horde.Build.Tests/BlobStoreTests.cs b/Engine/Source/Programs/Horde/Horde.Build.Tests/BlobStoreTests.cs
--- a/Engine/Source/Programs/Horde/Horde.Build.Tests/BlobStoreTests.cs
+++ b/Engine/Source/Programs/Horde/Horde.Build.Tests/BlobStoreTests.cs
@@ -47,6 +47,9 @@ namespace Horde.Build.Tests
 
 		static Bundle CreateTestBundle(ReadOnlyMemory<byte> data, IReadOnlyList<BlobLocator> refs)
 		{
+			List<BundleType> types = new List<BundleType>();
+			types.Add(new BundleType(Guid.Parse("{AFDF76A7-5333-4DEE-B837-B5F5CA511245}"), 1));
+
 			List<BundleImport> imports = new List<BundleImport>();
 
 			Dictionary<BlobLocator, int> locatorToIndex = new Dictionary<BlobLocator, int>();
@@ -62,12 +65,12 @@ namespace Horde.Build.Tests
 			}
 
 			List<BundleExport> exports = new List<BundleExport>();
-			exports.Add(new BundleExport(IoHash.Compute(data.Span), (int)data.Length, refs.Select(x => locatorToIndex[x]).ToArray()));
+			exports.Add(new BundleExport(0, IoHash.Compute(data.Span), (int)data.Length, refs.Select(x => locatorToIndex[x]).ToArray()));
 
 			List<BundlePacket> packets = new List<BundlePacket>();
 			packets.Add(new BundlePacket((int)data.Length, (int)data.Length));
 
-			BundleHeader header = new BundleHeader(BundleCompressionFormat.None, imports, exports, packets);
+			BundleHeader header = new BundleHeader(BundleCompressionFormat.None, types, imports, exports, packets);
 			return new Bundle(header, new[] { data });
 		}
 
diff --git a/Engine/Source/Programs/Horde/Horde.Build.Tests/GcServiceTests.cs b/Engine/Source/Programs/Horde/Horde.Build.Tests/GcServiceTests.cs
--- a/Engine/Source/Programs/Horde/Horde.Build.Tests/GcServiceTests.cs
+++ b/Engine/Source/Programs/Horde/Horde.Build.Tests/GcServiceTests.cs
@@ -105,8 +105,9 @@ namespace Horde.Build.Tests
 			BlobLocator[] locators = new BlobLocator[children.Length];
 			for (int idx = numNodes - 1; idx >= 0; idx--)
 			{
+				List<BundleType> types = new List<BundleType> { new BundleType(Guid.Parse("{AFDF76A7-5333-4DEE-B837-B5F5CA511245}"), 0) };
 				List<BundleImport> imports = children[idx].ConvertAll(x => new BundleImport(locators[x], Array.Empty<(int, IoHash)>()));
-				BundleHeader header = new BundleHeader(BundleCompressionFormat.None, imports, Array.Empty<BundleExport>(), Array.Empty<BundlePacket>());
+				BundleHeader header = new BundleHeader(BundleCompressionFormat.None, types, imports, Array.Empty<BundleExport>(), Array.Empty<BundlePacket>());
 				Bundle bundle = new Bundle(header, Array.Empty<ReadOnlyMemory<byte>>());
 				locators[idx] = await store.WriteBundleAsync(bundle, prefix: "gctest");
 			}
diff --git a/Engine/Source/Programs/Horde/Horde.Build/Commands/Bundles/ExtractCommand.cs b/Engine/Source/Programs/Horde/Horde.Build/Commands/Bundles/ExtractCommand.cs
--- a/Engine/Source/Programs/Horde/Horde.Build/Commands/Bundles/ExtractCommand.cs
+++ b/Engine/Source/Programs/Horde/Horde.Build/Commands/Bundles/ExtractCommand.cs
@@ -36,10 +36,11 @@ namespace Horde.Build.Commands.Bundles
 			using ServiceProvider serviceProvider = Startup.CreateServiceProvider(_configuration, _loggerProvider);
 
 			IStorageClient store = serviceProvider.GetRequiredService<IStorageClient<ReplicationService>>();
+			TreeReader reader = new TreeReader(store, serviceProvider.GetRequiredService<IMemoryCache>(), serviceProvider.GetRequiredService<ILogger<ExtractCommand>>()); 
 
-			ReplicationNode node = await store.ReadNodeAsync<ReplicationNode>(RefName);
-			DirectoryNode root = await node.Contents.ExpandAsync();
-			await root.CopyToDirectoryAsync(OutputDir.ToDirectoryInfo(), logger, CancellationToken.None);
+			ReplicationNode node = await reader.ReadNodeAsync<ReplicationNode>(RefName);
+			DirectoryNode root = await node.Contents.ExpandAsync(reader);
+			await root.CopyToDirectoryAsync(reader, OutputDir.ToDirectoryInfo(), logger, CancellationToken.None);
 
 			return 0;
 		}
diff --git a/Engine/Source/Programs/Horde/Horde.Build/Commands/Bundles/PerforceCommand.cs b/Engine/Source/Programs/Horde/Horde.Build/Commands/Bundles/PerforceCommand.cs
--- a/Engine/Source/Programs/Horde/Horde.Build/Commands/Bundles/PerforceCommand.cs
+++ b/Engine/Source/Programs/Horde/Horde.Build/Commands/Bundles/PerforceCommand.cs
@@ -12,6 +12,7 @@ using Horde.Build.Perforce;
 using Horde.Build.Storage;
 using Horde.Build.Streams;
 using Horde.Build.Utilities;
+using Microsoft.Extensions.Caching.Memory;
 using Microsoft.Extensions.Configuration;
 using Microsoft.Extensions.DependencyInjection;
 using Microsoft.Extensions.Logging;
@@ -69,6 +70,7 @@ namespace Horde.Build.Commands.Bundles
 			StorageService storageService = serviceProvider.GetRequiredService<StorageService>();
 
 			IStorageClient storage = await storageService.GetClientAsync(Namespace.Perforce, CancellationToken.None);
+			TreeReader reader = new TreeReader(storage, serviceProvider.GetRequiredService<IMemoryCache>(), serviceProvider.GetRequiredService<ILogger<PerforceCommand>>());
 
 			IStream? stream = await streamCollection.GetAsync(new StreamId(StreamId));
 			if (stream == null)
@@ -86,7 +88,7 @@ namespace Horde.Build.Commands.Bundles
 			}
 			else
 			{
-				baseContents = await ReplicationService.ReadCommitTreeAsync(storage, stream, BaseChange, Filter, RevisionsOnly, CancellationToken.None);
+				baseContents = await ReplicationService.ReadCommitTreeAsync(reader, stream, BaseChange, Filter, RevisionsOnly, CancellationToken.None);
 			}
 
 			await foreach (ICommit commit in commitService.GetCollection(stream).FindAsync(Change, null, Count))
@@ -98,7 +100,7 @@ namespace Horde.Build.Commands.Bundles
 
 				if (Content)
 				{
-					baseContents = await replicationService.WriteCommitTreeAsync(storage, stream, commit.Number, BaseChange, baseContents, Filter, RevisionsOnly, CancellationToken.None);
+					baseContents = await replicationService.WriteCommitTreeAsync(storage, reader, stream, commit.Number, BaseChange, baseContents, Filter, RevisionsOnly, CancellationToken.None);
 					BaseChange = commit.Number;
 				}
 			}
diff --git a/Engine/Source/Programs/Horde/Horde.Build/Perforce/ReplicationService.cs b/Engine/Source/Programs/Horde/Horde.Build/Perforce/ReplicationService.cs
--- a/Engine/Source/Programs/Horde/Horde.Build/Perforce/ReplicationService.cs
+++ b/Engine/Source/Programs/Horde/Horde.Build/Perforce/ReplicationService.cs
@@ -20,6 +20,7 @@ using Horde.Build.Server;
 using Horde.Build.Storage;
 using Horde.Build.Streams;
 using Horde.Build.Utilities;
+using Microsoft.Extensions.Caching.Memory;
 using Microsoft.Extensions.Hosting;
 using Microsoft.Extensions.Logging;
 using Microsoft.Extensions.Options;
@@ -42,6 +43,7 @@ namespace Horde.Build.Perforce
 	/// <summary>
 	/// Root node for a commit snapshot
 	/// </summary>
+	[TreeNode("{8C874966-2E89-4273-A4AB-AC9F5491C86D}")]
 	public class ReplicationNode : TreeNode
 	{
 		/// <summary>
@@ -127,6 +129,7 @@ namespace Horde.Build.Perforce
 		readonly IStreamCollection _streamCollection;
 		readonly IPerforceService _perforceService;
 		readonly StorageService _storageService;
+		readonly IMemoryCache _cache;
 		readonly ILogger _logger;
 
 		const int MaxBackgroundTasks = 2;
@@ -146,7 +149,7 @@ namespace Horde.Build.Perforce
 		/// <summary>
 		/// Constructor
 		/// </summary>
-		public ReplicationService(RedisService redisService, /* CommitService commitService, */IStreamCollection streamCollection, IPerforceService perforceService, StorageService storageService, IOptions<ServerSettings> settings, ILogger<ReplicationService> logger)
+		public ReplicationService(RedisService redisService, /* CommitService commitService, */IStreamCollection streamCollection, IPerforceService perforceService, StorageService storageService, IMemoryCache cache, IOptions<ServerSettings> settings, ILogger<ReplicationService> logger)
 		{
 			_settings = settings.Value;
 
@@ -160,6 +163,7 @@ namespace Horde.Build.Perforce
 			_streamCollection = streamCollection;
 			_perforceService = perforceService;
 			_storageService = storageService;
+			_cache = cache;
 			_logger = logger;
 		}
 
@@ -327,9 +331,10 @@ namespace Horde.Build.Perforce
 				else
 				{
 					IStorageClient storage = await _storageService.GetClientAsync(Namespace.Perforce, cancellationToken);
+					TreeReader reader = new TreeReader(storage, _cache, _logger);
 
 					// Update the stream, updating the reservation every 30 seconds
-					Task internalTask = Task.Run(() => UpdateStreamContentInternalAsync(storage, stream, streamChanges, cancellationToken), cancellationToken);
+					Task internalTask = Task.Run(() => UpdateStreamContentInternalAsync(storage, reader, stream, streamChanges, cancellationToken), cancellationToken);
 					while (!internalTask.IsCompleted)
 					{
 						Task delayTask = Task.Delay(TimeSpan.FromSeconds(ExtendReservationAfterSeconds), cancellationToken);
@@ -364,7 +369,7 @@ namespace Horde.Build.Perforce
 			await _redisReservations.RemoveAsync(streamId);
 		}
 
-		async Task UpdateStreamContentInternalAsync(IStorageClient storage, IStream stream, RedisList<int> changes, CancellationToken cancellationToken)
+		async Task UpdateStreamContentInternalAsync(IStorageClient storage, TreeReader reader, IStream stream, RedisList<int> changes, CancellationToken cancellationToken)
 		{
 			int prevChange = 0;
 			ReplicationNode prevContents = new ReplicationNode();
@@ -382,11 +387,11 @@ namespace Horde.Build.Perforce
 				{
 					RefName prevRefName = GetRefName(stream, values[0]);
 
-					ReplicationNode? contents = await storage.TryReadNodeAsync<ReplicationNode>(prevRefName, cancellationToken: cancellationToken);
+					ReplicationNode? contents = await reader.TryReadNodeAsync<ReplicationNode>(prevRefName, cancellationToken: cancellationToken);
 					if (contents == null)
 					{
 						_logger.LogInformation("No content for CL {Change}; creating full snapshot", values[0]);
-						prevContents = await WriteCommitTreeAsync(storage, stream, values[0], 0, new ReplicationNode(), cancellationToken);
+						prevContents = await WriteCommitTreeAsync(storage, reader, stream, values[0], 0, new ReplicationNode(), cancellationToken);
 					}
 					else
 					{
@@ -399,7 +404,7 @@ namespace Horde.Build.Perforce
 				else if (values.Length == 2)
 				{
 					// Perform a snapshot of the new change, then remove it from the list
-					prevContents = await WriteCommitTreeAsync(storage, stream, values[1], prevChange, prevContents, cancellationToken);
+					prevContents = await WriteCommitTreeAsync(storage, reader, stream, values[1], prevChange, prevContents, cancellationToken);
 					prevChange = values[1];
 
 					// Remove the first item from the list
@@ -453,33 +458,34 @@ namespace Horde.Build.Perforce
 		/// <summary>
 		/// Reads a tree from the given stream
 		/// </summary>
-		/// <param name="storage">Storage client instance</param>
+		/// <param name="reader">Storage client instance</param>
 		/// <param name="stream"></param>
 		/// <param name="change"></param>
 		/// <param name="filter"></param>
 		/// <param name="revisionsOnly"></param>
 		/// <param name="cancellationToken"></param>
 		/// <returns></returns>
-		public static async Task<ReplicationNode> ReadCommitTreeAsync(IStorageClient storage, IStream stream, int change, string? filter, bool revisionsOnly, CancellationToken cancellationToken)
+		public static async Task<ReplicationNode> ReadCommitTreeAsync(TreeReader reader, IStream stream, int change, string? filter, bool revisionsOnly, CancellationToken cancellationToken)
 		{
 			RefName name = GetRefName(stream.Id, change, filter, revisionsOnly);
-			return await storage.ReadNodeAsync<ReplicationNode>(name, cancellationToken: cancellationToken);
+			return await reader.ReadNodeAsync<ReplicationNode>(name, cancellationToken: cancellationToken);
 		}
 
 		/// <summary>
 		/// Replicates the contents of a stream to Horde storage, optionally using the given change as a starting point
 		/// </summary>
 		/// <param name="storage">Storage client instance</param>
+		/// <param name="reader"></param>
 		/// <param name="stream">The stream to replicate</param>
 		/// <param name="change">Commit to store the tree ref</param>
 		/// <param name="baseChange">The base change to update from</param>
 		/// <param name="baseContents">Initial contents of the tree at baseChange</param>
 		/// <param name="cancellationToken">Cancellation token for the operation</param>
 		/// <returns>Root tree object</returns>
-		public async Task<ReplicationNode> WriteCommitTreeAsync(IStorageClient storage, IStream stream, int change, int baseChange, ReplicationNode baseContents, CancellationToken cancellationToken)
+		public async Task<ReplicationNode> WriteCommitTreeAsync(IStorageClient storage, TreeReader reader, IStream stream, int change, int baseChange, ReplicationNode baseContents, CancellationToken cancellationToken)
 		{
 			bool revisionsOnly = stream.Config.ReplicationMode == ContentReplicationMode.RevisionsOnly;
-			return await WriteCommitTreeAsync(storage, stream, change, baseChange, baseContents, stream.Config.ReplicationFilter, revisionsOnly, cancellationToken);
+			return await WriteCommitTreeAsync(storage, reader, stream, change, baseChange, baseContents, stream.Config.ReplicationFilter, revisionsOnly, cancellationToken);
 		}
 
 		[DebuggerDisplay("{_path}")]
@@ -496,47 +502,27 @@ namespace Horde.Build.Perforce
 			}
 		}
 
-		class FileWriter : IDisposable
+		class FileWriter
 		{
-			readonly FileEntry _entry;
-			readonly IncrementalHash _hash;
+			readonly Utf8String _path;
+			FileNode _node;
 			readonly long _size;
 
-			public FileWriter(FileEntry entry, long size)
+			public Utf8String Path => _path;
+			public FileNode Node => _node;
+			public long Size => _size;
+
+			public FileWriter(Utf8String path, long size)
 			{
-				_entry = entry;
-				_hash = IncrementalHash.CreateHash(HashAlgorithmName.SHA1);
+				_path = path;
+				_node = new LeafFileNode();
 				_size = size;
-
-				GitObject.WriteHeader(GitObjectType.Blob, size, _hash);
 			}
 
-			public void Dispose() => _hash.Dispose();
-
 			public async Task AppendAsync(ReadOnlyMemory<byte> data, ChunkingOptions options, TreeWriter writer, CancellationToken cancellationToken)
 			{
 				_hash.AppendData(data.Span);
-				await _entry.AppendAsync(data, options, writer, cancellationToken);
-			}
-
-			public async Task FinishAsync(TreeWriter writer, CancellationToken cancellationToken)
-			{
-				if (_entry.Length != _size)
-				{
-					throw new ReplicationException($"Invalid size for replicated file '{_entry.Name}'. Expected {_size}, got {_entry.Length}.");
-				}
-				UpdateHash();
-				if (_entry.Target != null)
-				{
-					await writer.WriteAsync(_entry, cancellationToken);
-				}
-			}
-
-			void UpdateHash()
-			{
-				Span<byte> buffer = stackalloc byte[20];
-				_hash.GetHashAndReset(buffer);
-				_entry.GitHash = new Sha1Hash(buffer);
+				_node = await _node.AppendAsync(data, options, writer, cancellationToken);
 			}
 		}
 
@@ -544,6 +530,7 @@ namespace Horde.Build.Perforce
 		/// Replicates the contents of a stream to Horde storage, optionally using the given change as a starting point
 		/// </summary>
 		/// <param name="storage">Storage client</param>
+		/// <param name="reader">Reader for nodes</param>
 		/// <param name="stream">The stream to replicate</param>
 		/// <param name="change">Commit to store the tree ref</param>
 		/// <param name="baseChange">The base change to update from</param>
@@ -552,7 +539,7 @@ namespace Horde.Build.Perforce
 		/// <param name="revisionsOnly"></param>
 		/// <param name="cancellationToken">Cancellation token for the operation</param>
 		/// <returns>Root tree object</returns>
-		public async Task<ReplicationNode> WriteCommitTreeAsync(IStorageClient storage, IStream stream, int change, int? baseChange, ReplicationNode baseContents, string? filter, bool revisionsOnly, CancellationToken cancellationToken)
+		public async Task<ReplicationNode> WriteCommitTreeAsync(IStorageClient storage, TreeReader reader, IStream stream, int change, int? baseChange, ReplicationNode baseContents, string? filter, bool revisionsOnly, CancellationToken cancellationToken)
 		{
 			// Create a client to replicate from this stream
 			ReplicationClient clientInfo = await FindOrAddReplicationClientAsync(stream);
@@ -574,16 +561,16 @@ namespace Horde.Build.Perforce
 			TreeWriter writer = new TreeWriter(storage, prefix: refName.Text);
 
 			// Get the current sync state for this change
-			ReplicationNode? syncNode = await storage.TryReadNodeAsync<ReplicationNode>(incRefName, cancellationToken: cancellationToken);
+			ReplicationNode? syncNode = await reader.TryReadNodeAsync<ReplicationNode>(incRefName, cancellationToken: cancellationToken);
 			bool deleteIncRef = syncNode != null;
 			if (syncNode == null)
 			{
-				DirectoryNode contents = await baseContents.Contents.ExpandAsync(cancellationToken);
+				DirectoryNode contents = await baseContents.Contents.ExpandAsync(reader, cancellationToken);
 				syncNode = new ReplicationNode(contents);
 			}
 
 			// Get the contents of the current tree
-			DirectoryNode root = await syncNode.Contents.ExpandAsync(cancellationToken);
+			DirectoryNode root = await syncNode.Contents.ExpandAsync(reader, cancellationToken);
 
 			if (revisionsOnly)
 			{
@@ -602,12 +589,10 @@ namespace Horde.Build.Perforce
 					}
 
 					Utf8String path = syncRecord.Path.Substring(clientRoot.Length);
-					byte[] data = Encoding.UTF8.GetBytes($"{syncRecord.DepotFile}#{syncRecord.Revision}");
 
-					using ReadOnlyMemoryStream dataStream = new ReadOnlyMemoryStream(data);
-
-					FileEntry entry = await root.AddFileByPathAsync(path, FileEntryFlags.PerforceDepotPathAndRevision, cancellationToken);
-					await entry.AppendAsync(dataStream, _settings.Commits.Chunking, writer, cancellationToken);
+					byte[] data = Encoding.UTF8.GetBytes($"{syncRecord.DepotFile}#{syncRecord.Revision}");
+					FileNode node = await FileNode.CreateAsync(data, _settings.Commits.Chunking, writer, cancellationToken);
+					await root.AddFileByPathAsync(reader, path, FileEntryFlags.PerforceDepotPathAndRevision, node, cancellationToken);
 				}
 			}
 			else
@@ -756,8 +741,7 @@ namespace Horde.Build.Perforce
 										throw new ReplicationException($"Unable to find file entry for {file}");
 									}
 
-									FileEntry entry = await root.AddFileByPathAsync(file, FileEntryFlags.None, cancellationToken);
-									handles[io.File] = new FileWriter(entry, fileSize);
+									handles[io.File] = new FileWriter(file, fileSize);
 								}
 								else if (io.Command == PerforceIoCommand.Write)
 								{
@@ -771,7 +755,14 @@ namespace Horde.Build.Perforce
 									{
 										if (handles.Remove(io.File, out file))
 										{
-											await file.FinishAsync(writer, cancellationToken);
+											FileNode node = file.Node;
+											if (node.Length != file.Size)
+											{
+												throw new ReplicationException($"Invalid size for replicated file '{file.Path}'. Expected {file.Size}, got {node.Length}.");
+											}
+
+											FileEntry entry = await root.AddFileByPathAsync(reader, file.Path, FileEntryFlags.None, node, cancellationToken);
+											await writer.WriteAsync(entry, cancellationToken);
 										}
 									}
 									finally
@@ -782,7 +773,7 @@ namespace Horde.Build.Perforce
 								else if (io.Command == PerforceIoCommand.Unlink)
 								{
 									Utf8String file = GetClientRelativePath(io.Payload, clientInfo.Client.Root);
-									await root.DeleteFileByPathAsync(file, cancellationToken);
+									await root.DeleteFileByPathAsync(reader, file, cancellationToken);
 								}
 								else
 								{
diff --git a/Engine/Source/Programs/Horde/Horde.Build/Storage/StorageService.cs b/Engine/Source/Programs/Horde/Horde.Build/Storage/StorageService.cs
--- a/Engine/Source/Programs/Horde/Horde.Build/Storage/StorageService.cs
+++ b/Engine/Source/Programs/Horde/Horde.Build/Storage/StorageService.cs
@@ -103,8 +103,7 @@ namespace Horde.Build.Storage
 
 			public bool SupportsRedirects { get; }
 
-			public StorageClient(StorageService outer, NamespaceConfig config, IStorageBackend backend, IMemoryCache cache, ILogger logger)
-				: base(cache, logger)
+			public StorageClient(StorageService outer, NamespaceConfig config, IStorageBackend backend)
 			{
 				_outer = outer;
 				Config = config;
@@ -545,7 +544,7 @@ namespace Horde.Build.Storage
 						IStorageBackend backend = CreateStorageBackend(backendConfig);
 
 #pragma warning disable CA2000 // Dispose objects before losing scope
-						StorageClient client = new StorageClient(this, namespaceConfig, backend, _cache, _logger);
+						StorageClient client = new StorageClient(this, namespaceConfig, backend);
 						nextState.Namespaces.Add(namespaceConfig.Id, new NamespaceInfo(namespaceConfig, client));
 #pragma warning restore CA2000 // Dispose objects before losing scope
 					}
diff --git a/Engine/Source/Programs/Shared/EpicGames.Horde.Tests/BundleStoreTests.cs b/Engine/Source/Programs/Shared/EpicGames.Horde.Tests/BundleStoreTests.cs
--- a/Engine/Source/Programs/Shared/EpicGames.Horde.Tests/BundleStoreTests.cs
+++ b/Engine/Source/Programs/Shared/EpicGames.Horde.Tests/BundleStoreTests.cs
@@ -22,12 +22,14 @@ namespace EpicGames.Horde.Tests
 	[TestClass]
 	public class BundleStoreTests
 	{
+		static readonly TreeReaderOptions s_readOptions = new TreeReaderOptions(typeof(SimpleNode));
+
 		[TestMethod]
 		public async Task TestTreeAsync()
 		{
 			using IMemoryCache cache = new MemoryCache(new MemoryCacheOptions());
 
-			MemoryStorageClient blobStore = new MemoryStorageClient(cache, NullLogger.Instance);
+			MemoryStorageClient blobStore = new MemoryStorageClient();
 			await TestTreeAsync(blobStore, new TreeOptions { MaxBlobSize = 1024 * 1024 });
 
 			Assert.AreEqual(1, blobStore.Blobs.Count);
@@ -39,13 +41,14 @@ namespace EpicGames.Horde.Tests
 		{
 			using IMemoryCache cache = new MemoryCache(new MemoryCacheOptions());
 
-			MemoryStorageClient blobStore = new MemoryStorageClient(cache, NullLogger.Instance);
+			MemoryStorageClient blobStore = new MemoryStorageClient();
 			await TestTreeAsync(blobStore, new TreeOptions { MaxBlobSize = 1 });
 
 			Assert.AreEqual(5, blobStore.Blobs.Count);
 			Assert.AreEqual(1, blobStore.Refs.Count);
 		}
 
+		[TreeNode("{F63606D4-5DBB-4061-A655-6F444F65229F}")]
 		class SimpleNode : TreeNode
 		{
 			public ReadOnlySequence<byte> Data { get; }
@@ -87,17 +90,19 @@ namespace EpicGames.Horde.Tests
 
 				await writer.WriteRefAsync(new RefName("test"), root);
 
-				await CheckTree(root);
+				TreeReader reader = new TreeReader(store, null, s_readOptions, NullLogger.Instance);
+				await CheckTree(reader, root);
 			}
 
 			// Check we can read it back in
 			{
-				SimpleNode root = await store.ReadNodeAsync<SimpleNode>(new RefName("test"));
-				await CheckTree(root!);
+				TreeReader reader = new TreeReader(store, null, s_readOptions, NullLogger.Instance);
+				SimpleNode root = await reader.ReadNodeAsync<SimpleNode>(new RefName("test"));
+				await CheckTree(reader, root);
 			}
 		}
 
-		static async Task CheckTree(SimpleNode root)
+		static async Task CheckTree(TreeReader reader, SimpleNode root)
 		{
 			SimpleNode node5 = root;
 			byte[] data5 = node5.Data.ToArray();
@@ -105,25 +110,25 @@ namespace EpicGames.Horde.Tests
 			IReadOnlyList<TreeNodeRef<SimpleNode>> refs5 = node5.Refs;
 			Assert.AreEqual(2, refs5.Count);
 
-			SimpleNode node4 = await refs5[0].ExpandAsync();
+			SimpleNode node4 = await refs5[0].ExpandAsync(reader);
 			byte[] data4 = node4.Data.ToArray();
 			Assert.IsTrue(data4.SequenceEqual(new byte[] { 4 }));
 			IReadOnlyList<TreeNodeRef<SimpleNode>> refs4 = node4.Refs;
 			Assert.AreEqual(0, refs4.Count);
 
-			SimpleNode node3 = await refs5[1].ExpandAsync();
+			SimpleNode node3 = await refs5[1].ExpandAsync(reader);
 			byte[] data3 = node3.Data.ToArray();
 			Assert.IsTrue(data3.SequenceEqual(new byte[] { 3 }));
 			IReadOnlyList<TreeNodeRef<SimpleNode>> refs3 = node3.Refs;
 			Assert.AreEqual(1, refs3.Count);
 
-			SimpleNode node2 = await refs3[0].ExpandAsync();
+			SimpleNode node2 = await refs3[0].ExpandAsync(reader);
 			byte[] data2 = node2.Data.ToArray();
 			Assert.IsTrue(data2.SequenceEqual(new byte[] { 2 }));
 			IReadOnlyList<TreeNodeRef<SimpleNode>> refs2 = node2.Refs;
 			Assert.AreEqual(1, refs2.Count);
 
-			SimpleNode node1 = await refs2[0].ExpandAsync();
+			SimpleNode node1 = await refs2[0].ExpandAsync(reader);
 			byte[] data1 = node1.Data.ToArray();
 			Assert.IsTrue(data1.SequenceEqual(new byte[] { 1 }));
 			IReadOnlyList<TreeNodeRef<SimpleNode>> refs1 = node1.Refs;
@@ -133,8 +138,8 @@ namespace EpicGames.Horde.Tests
 		[TestMethod]
 		public async Task DirectoryNodesAsync()
 		{
-			using IMemoryCache cache = new MemoryCache(new MemoryCacheOptions());
-			MemoryStorageClient store = new MemoryStorageClient(cache, NullLogger.Instance);
+			MemoryStorageClient store = new MemoryStorageClient();
+			TreeReader reader = new TreeReader(store, null, s_readOptions, NullLogger.Instance);
 
 			// Generate a tree
 			{
@@ -143,26 +148,26 @@ namespace EpicGames.Horde.Tests
 				DirectoryNode world = hello.AddDirectory("world");
 				await store.WriteNodeAsync(new RefName("test"), root);
 
-				await CheckDirectoryTreeAsync(root);
+				await CheckDirectoryTreeAsync(reader, root);
 			}
 
 			// Check we can read it back in
 			{
-				DirectoryNode root = await store.ReadNodeAsync<DirectoryNode>(new RefName("test"));
-				await CheckDirectoryTreeAsync(root);
+				DirectoryNode root = await reader.ReadNodeAsync<DirectoryNode>(new RefName("test"));
+				await CheckDirectoryTreeAsync(reader, root);
 			}
 		}
 
-		static async Task CheckDirectoryTreeAsync(DirectoryNode root)
+		static async Task CheckDirectoryTreeAsync(TreeReader reader, DirectoryNode root)
 		{
 			Assert.AreEqual(1, root.Directories.Count);
 			Assert.AreEqual("hello", root.Directories.First().Name);
 
-			DirectoryNode hello = await root.Directories.First().ExpandAsync(CancellationToken.None);
+			DirectoryNode hello = await root.Directories.First().ExpandAsync(reader, CancellationToken.None);
 			Assert.AreEqual(1, hello.Directories.Count);
 			Assert.AreEqual("world", hello.Directories.First().Name);
 
-			DirectoryNode world = await hello.Directories.First().ExpandAsync(CancellationToken.None);
+			DirectoryNode world = await hello.Directories.First().ExpandAsync(reader, CancellationToken.None);
 			Assert.AreEqual(0, world.Directories.Count);
 		}
 
@@ -170,7 +175,8 @@ namespace EpicGames.Horde.Tests
 		public async Task FileNodesAsync()
 		{
 			using IMemoryCache cache = new MemoryCache(new MemoryCacheOptions());
-			MemoryStorageClient store = new MemoryStorageClient(cache, NullLogger.Instance);
+			MemoryStorageClient store = new MemoryStorageClient();
+			TreeReader reader = new TreeReader(store, null, NullLogger.Instance);
 
 			// Generate a tree
 			{
@@ -179,34 +185,34 @@ namespace EpicGames.Horde.Tests
 				DirectoryNode root = new DirectoryNode(DirectoryFlags.None);
 				DirectoryNode hello = root.AddDirectory("hello");
 
-				FileEntry world = hello.AddFile("world", FileEntryFlags.None);
-				await world.AppendAsync(Encoding.UTF8.GetBytes("world"), new ChunkingOptions(), writer, CancellationToken.None);
+				FileNode world = await FileNode.CreateAsync(Encoding.UTF8.GetBytes("world"), new ChunkingOptions(), writer, CancellationToken.None);
+				hello.AddFile("world", FileEntryFlags.None, world);
 
 				await writer.WriteRefAsync(new RefName("test"), root);
 
-				await CheckFileTreeAsync(root);
+				await CheckFileTreeAsync(reader, root);
 			}
 
 			// Check we can read it back in
 			{
-				DirectoryNode root = await store.ReadNodeAsync<DirectoryNode>(new RefName("test"));
-				await CheckFileTreeAsync(root);
+				DirectoryNode root = await reader.ReadNodeAsync<DirectoryNode>(new RefName("test"));
+				await CheckFileTreeAsync(reader, root);
 			}
 		}
 
-		static async Task CheckFileTreeAsync(DirectoryNode root)
+		static async Task CheckFileTreeAsync(TreeReader reader, DirectoryNode root)
 		{
 			Assert.AreEqual(1, root.Directories.Count);
 			Assert.AreEqual("hello", root.Directories.First().Name);
 
-			DirectoryNode hello = await root.Directories.First().ExpandAsync(CancellationToken.None);
+			DirectoryNode hello = await root.Directories.First().ExpandAsync(reader);
 			Assert.AreEqual(0, hello.Directories.Count);
 			Assert.AreEqual(1, hello.Files.Count);
 			Assert.AreEqual("world", hello.Files.First().Name);
 
-			FileNode world = await hello.Files.First().ExpandAsync(CancellationToken.None);
+			FileNode world = await hello.Files.First().ExpandAsync(reader);
 
-			byte[] worldData = await GetFileDataAsync(world);
+			byte[] worldData = await GetFileDataAsync(reader, world);
 			Assert.IsTrue(worldData.SequenceEqual(Encoding.UTF8.GetBytes("world")));
 		}
 
@@ -215,7 +221,7 @@ namespace EpicGames.Horde.Tests
 		{
 			using IMemoryCache cache = new MemoryCache(new MemoryCacheOptions());
 
-			MemoryStorageClient store = new MemoryStorageClient(cache, NullLogger.Instance);
+			MemoryStorageClient store = new MemoryStorageClient();
 			TreeWriter writer = new TreeWriter(store, new TreeOptions());
 
 			byte[] chunk = RandomNumberGenerator.GetBytes(256);
@@ -291,7 +297,7 @@ namespace EpicGames.Horde.Tests
 		public async Task LargeFileTestAsync()
 		{
 			using IMemoryCache cache = new MemoryCache(new MemoryCacheOptions());
-			MemoryStorageClient store = new MemoryStorageClient(cache, NullLogger.Instance);
+			MemoryStorageClient store = new MemoryStorageClient();
 
 			const int length = 1024;
 			const int copies = 4096;
@@ -317,19 +323,22 @@ namespace EpicGames.Horde.Tests
 				options.LeafOptions.TargetSize = 128;
 				options.LeafOptions.MaxSize = 64 * 1024;
 
-				FileEntry file = root.AddFile("test", FileEntryFlags.None);
-				await file.AppendAsync(data, options, writer, CancellationToken.None);
+				FileNode file = await FileNode.CreateAsync(data, options, writer, CancellationToken.None);
+				root.AddFile("test", FileEntryFlags.None, file);
 
 				await writer.WriteRefAsync(new RefName("test"), root);
 
-				await CheckLargeFileTreeAsync(root, data);
+				TreeReader reader = new TreeReader(store, null, NullLogger.Instance);
+				await CheckLargeFileTreeAsync(reader, root, data);
 			}
 
 			// Check we can read it back in
 			{
-				DirectoryNode newRoot = await store.ReadNodeAsync<DirectoryNode>(new RefName("test"));
-				await CompareTrees(root, newRoot);
-				await CheckLargeFileTreeAsync(root, data);
+				TreeReader reader = new TreeReader(store, null, NullLogger.Instance);
+
+				DirectoryNode newRoot = await reader.ReadNodeAsync<DirectoryNode>(new RefName("test"));
+				await CompareTrees(reader, root, newRoot);
+				await CheckLargeFileTreeAsync(reader, root, data);
 
 				TreeNodeRef<FileNode> file = root.GetFileEntry("test");
 
@@ -338,7 +347,7 @@ namespace EpicGames.Horde.Tests
 			}
 		}
 
-		static async Task CompareTrees(DirectoryNode oldNode, DirectoryNode newNode)
+		static async Task CompareTrees(TreeReader reader, DirectoryNode oldNode, DirectoryNode newNode)
 		{
 			Assert.AreEqual(oldNode.Length, newNode.Length);
 			Assert.AreEqual(oldNode.Files.Count, newNode.Files.Count);
@@ -346,13 +355,13 @@ namespace EpicGames.Horde.Tests
 
 			foreach ((FileEntry oldFileEntry, FileEntry newFileEntry) in oldNode.Files.Zip(newNode.Files))
 			{
-				FileNode oldFile = await oldFileEntry.ExpandAsync();
-				FileNode newFile = await newFileEntry.ExpandAsync();
-				await CompareTrees(oldFile, newFile);
+				FileNode oldFile = await oldFileEntry.ExpandAsync(reader);
+				FileNode newFile = await newFileEntry.ExpandAsync(reader);
+				await CompareTrees(reader, oldFile, newFile);
 			}
 		}
 
-		static async Task CompareTrees(FileNode oldNode, FileNode newNode)
+		static async Task CompareTrees(TreeReader reader, FileNode oldNode, FileNode newNode)
 		{
 			if (oldNode is InteriorFileNode oldInteriorNode)
 			{
@@ -365,9 +374,9 @@ namespace EpicGames.Horde.Tests
 				int index = 0;
 				foreach ((TreeNodeRef<FileNode> oldFileRef, TreeNodeRef<FileNode> newFileRef) in oldInteriorNode.Children.Zip(newInteriorNode.Children))
 				{
-					FileNode oldFile = await oldFileRef.ExpandAsync();
-					FileNode newFile = await newFileRef.ExpandAsync();
-					await CompareTrees(oldFile, newFile);
+					FileNode oldFile = await oldFileRef.ExpandAsync(reader);
+					FileNode newFile = await newFileRef.ExpandAsync(reader);
+					await CompareTrees(reader, oldFile, newFile);
 					index++;
 				}
 			}
@@ -383,23 +392,23 @@ namespace EpicGames.Horde.Tests
 			}
 		}
 
-		static async Task CheckLargeFileTreeAsync(DirectoryNode root, byte[] data)
+		static async Task CheckLargeFileTreeAsync(TreeReader reader, DirectoryNode root, byte[] data)
 		{
 			Assert.AreEqual(0, root.Directories.Count);
 			Assert.AreEqual(1, root.Files.Count);
 
-			FileNode world = await root.Files.First().ExpandAsync(CancellationToken.None);
+			FileNode world = await root.Files.First().ExpandAsync(reader, CancellationToken.None);
 
-			byte[] worldData = await GetFileDataAsync(world);
+			byte[] worldData = await GetFileDataAsync(reader, world);
 			Assert.AreEqual(data.Length, worldData.Length);
 			Assert.IsTrue(worldData.SequenceEqual(data));
 		}
 
-		static async Task<byte[]> GetFileDataAsync(FileNode fileNode)
+		static async Task<byte[]> GetFileDataAsync(TreeReader reader, FileNode fileNode)
 		{
 			using (MemoryStream stream = new MemoryStream())
 			{
-				await fileNode.CopyToStreamAsync(stream, CancellationToken.None);
+				await fileNode.CopyToStreamAsync(reader, stream, CancellationToken.None);
 				return stream.ToArray();
 			}
 		}
diff --git a/Engine/Source/Programs/Shared/EpicGames.Horde.Tests/BundleTests.cs b/Engine/Source/Programs/Shared/EpicGames.Horde.Tests/BundleTests.cs
--- a/Engine/Source/Programs/Shared/EpicGames.Horde.Tests/BundleTests.cs
+++ b/Engine/Source/Programs/Shared/EpicGames.Horde.Tests/BundleTests.cs
@@ -13,6 +13,7 @@ using System.Linq;
 using System.Text;
 using System.Threading;
 using System.Threading.Tasks;
+using static System.Formats.Asn1.AsnWriter;
 
 namespace EpicGames.Horde.Tests
 {
@@ -25,7 +26,7 @@ namespace EpicGames.Horde.Tests
 		public BundleTests()
 		{
 			_cache = new MemoryCache(new MemoryCacheOptions());
-			_storage = new MemoryStorageClient(_cache, NullLogger.Instance);
+			_storage = new MemoryStorageClient();
 		}
 
 		public void Dispose()
@@ -66,6 +67,7 @@ namespace EpicGames.Horde.Tests
 		public async Task BasicChunkingTests()
 		{
 			RefName refName = new RefName("test");
+			TreeReader reader = new TreeReader(_storage, null, NullLogger.Instance);
 			TreeWriter writer = new TreeWriter(_storage, new TreeOptions(), refName.Text);
 
 			ChunkingOptions options = new ChunkingOptions();
@@ -86,12 +88,12 @@ namespace EpicGames.Horde.Tests
 			Assert.IsTrue(node is InteriorFileNode);
 			Assert.AreEqual(2, ((InteriorFileNode)node).Children.Count);
 
-			FileNode? childNode1 = await ((TreeNodeRef<FileNode>)((InteriorFileNode)node).Children[0]).ExpandAsync();
+			FileNode? childNode1 = await ((TreeNodeRef<FileNode>)((InteriorFileNode)node).Children[0]).ExpandAsync(reader);
 			Assert.IsNotNull(childNode1);
 			Assert.IsTrue(childNode1 is LeafFileNode);
 			Assert.AreEqual(8, ((LeafFileNode)childNode1!).Data.Length);
 
-			FileNode? childNode2 = await ((TreeNodeRef<FileNode>)((InteriorFileNode)node).Children[1]).ExpandAsync();
+			FileNode? childNode2 = await ((TreeNodeRef<FileNode>)((InteriorFileNode)node).Children[1]).ExpandAsync(reader);
 			Assert.IsNotNull(childNode2);
 			Assert.IsTrue(childNode2 is LeafFileNode);
 			Assert.AreEqual(1, ((LeafFileNode)childNode2!).Data.Length);
@@ -102,19 +104,22 @@ namespace EpicGames.Horde.Tests
 		{
 			BundleHeader oldHeader;
 			{
+				List<BundleType> types = new List<BundleType>();
+				types.Add(new BundleType(Guid.NewGuid(), 0));
+
 				List<BundleImport> imports = new List<BundleImport>();
 				imports.Add(new BundleImport(new BlobLocator("import1"), new (int, IoHash)[] { (5, IoHash.Compute(Encoding.UTF8.GetBytes("blob1"))) }));
 				imports.Add(new BundleImport(new BlobLocator("import2"), new (int, IoHash)[] { (6, IoHash.Compute(Encoding.UTF8.GetBytes("blob2"))) }));
 
 				List<BundleExport> exports = new List<BundleExport>();
-				exports.Add(new BundleExport(IoHash.Compute(Encoding.UTF8.GetBytes("export1")), 2, new int[] { 1, 2 }));
-				exports.Add(new BundleExport(IoHash.Compute(Encoding.UTF8.GetBytes("export2")), 3, new int[] { 3 }));
+				exports.Add(new BundleExport(0, IoHash.Compute(Encoding.UTF8.GetBytes("export1")), 2, new int[] { 1, 2 }));
+				exports.Add(new BundleExport(0, IoHash.Compute(Encoding.UTF8.GetBytes("export2")), 3, new int[] { 3 }));
 
 				List<BundlePacket> packets = new List<BundlePacket>();
 				packets.Add(new BundlePacket(20, 40));
 				packets.Add(new BundlePacket(10, 20));
 
-				oldHeader = new BundleHeader(BundleCompressionFormat.LZ4, imports, exports, packets);
+				oldHeader = new BundleHeader(BundleCompressionFormat.LZ4, types, imports, exports, packets);
 			}
 
 			ByteArrayBuilder writer = new ByteArrayBuilder();
@@ -167,6 +172,7 @@ namespace EpicGames.Horde.Tests
 		public async Task BasicTestDirectory()
 		{
 			MemoryStorageClient store = _storage;
+			TreeReader reader = new TreeReader(store, null, NullLogger.Instance);
 
 			DirectoryNode root = new DirectoryNode(DirectoryFlags.None);
 			DirectoryNode node = root.AddDirectory("hello");
@@ -186,16 +192,16 @@ namespace EpicGames.Horde.Tests
 			Assert.AreEqual(3, bundle.Header.Exports.Count);
 
 			// Create a new bundle and read it back in again
-			DirectoryNode newRoot = await store.ReadNodeAsync<DirectoryNode>(refName);
+			DirectoryNode newRoot = await reader.ReadNodeAsync<DirectoryNode>(refName);
 
 			Assert.AreEqual(0, newRoot.Files.Count);
 			Assert.AreEqual(1, newRoot.Directories.Count);
-			DirectoryNode? outputNode = await newRoot.FindDirectoryAsync("hello", CancellationToken.None);
+			DirectoryNode? outputNode = await newRoot.FindDirectoryAsync(reader, "hello", CancellationToken.None);
 			Assert.IsNotNull(outputNode);
 
 			Assert.AreEqual(0, outputNode!.Files.Count);
 			Assert.AreEqual(1, outputNode!.Directories.Count);
-			DirectoryNode? outputNode2 = await outputNode.FindDirectoryAsync("world", CancellationToken.None);
+			DirectoryNode? outputNode2 = await outputNode.FindDirectoryAsync(reader, "world", CancellationToken.None);
 			Assert.IsNotNull(outputNode2);
 
 			Assert.AreEqual(0, outputNode2!.Files.Count);
@@ -246,18 +252,20 @@ namespace EpicGames.Horde.Tests
 			}
 
 			{
-				DirectoryNode root = await _storage.ReadNodeAsync<DirectoryNode>(refName);
+				TreeReader reader = new TreeReader(_storage, null, NullLogger.Instance);
+
+				DirectoryNode root = await reader.ReadNodeAsync<DirectoryNode>(refName);
 
-				DirectoryNode? newNode1 = await root.FindDirectoryAsync("node1", CancellationToken.None);
+				DirectoryNode? newNode1 = await root.FindDirectoryAsync(reader, "node1", CancellationToken.None);
 				Assert.IsNotNull(newNode1);
 
-				DirectoryNode? newNode2 = await newNode1!.FindDirectoryAsync("node2", CancellationToken.None);
+				DirectoryNode? newNode2 = await newNode1!.FindDirectoryAsync(reader, "node2", CancellationToken.None);
 				Assert.IsNotNull(newNode2);
 
-				DirectoryNode? newNode3 = await newNode2!.FindDirectoryAsync("node3", CancellationToken.None);
+				DirectoryNode? newNode3 = await newNode2!.FindDirectoryAsync(reader, "node3", CancellationToken.None);
 				Assert.IsNotNull(newNode3);
 
-				DirectoryNode? newNode4 = await newNode3!.FindDirectoryAsync("node4", CancellationToken.None);
+				DirectoryNode? newNode4 = await newNode3!.FindDirectoryAsync(reader, "node4", CancellationToken.None);
 				Assert.IsNotNull(newNode4);
 			}
 		}
@@ -334,6 +342,7 @@ namespace EpicGames.Horde.Tests
 		[TestMethod]
 		public async Task CoreAppendTest()
 		{
+			TreeReader reader = new TreeReader(_storage, null, NullLogger.Instance);
 			TreeWriter writer = new TreeWriter(_storage);
 
 			byte[] data = new byte[4096];
@@ -347,7 +356,7 @@ namespace EpicGames.Horde.Tests
 			{
 				node = await node.AppendAsync(data.AsMemory(idx, 1), options, writer, CancellationToken.None);
 
-				byte[] outputData = await node.ToByteArrayAsync(CancellationToken.None);
+				byte[] outputData = await node.ToByteArrayAsync(reader, CancellationToken.None);
 				Assert.IsTrue(data.AsMemory(0, idx + 1).Span.SequenceEqual(outputData.AsSpan(0, idx + 1)));
 			}
 		}
@@ -392,7 +401,8 @@ namespace EpicGames.Horde.Tests
 				root = await root.AppendAsync(data, options, writer, CancellationToken.None);
 			}
 
-			byte[] result = await root.ToByteArrayAsync(CancellationToken.None);
+			TreeReader reader = new TreeReader(_storage, null, NullLogger.Instance);
+			byte[] result = await root.ToByteArrayAsync(reader, CancellationToken.None);
 			Assert.AreEqual(NumIterations * data.Length, result.Length);
 
 			for (int idx = 0; idx < NumIterations; idx++)
@@ -401,10 +411,10 @@ namespace EpicGames.Horde.Tests
 				Assert.IsTrue(spanData.Span.SequenceEqual(data));
 			}
 
-			await CheckSizes(root, options, true);
+			await CheckSizes(reader, root, options, true);
 		}
 
-		async Task CheckSizes(FileNode node, ChunkingOptions options, bool rightmost)
+		async Task CheckSizes(TreeReader reader, FileNode node, ChunkingOptions options, bool rightmost)
 		{
 			if (node is LeafFileNode leafNode)
 			{
@@ -421,8 +431,8 @@ namespace EpicGames.Horde.Tests
 				int childCount = interiorNode.Children.Count;
 				for (int idx = 0; idx < childCount; idx++)
 				{
-					FileNode childNode = await interiorNode.Children[idx].ExpandAsync(CancellationToken.None);
-					await CheckSizes(childNode, options, idx == childCount - 1);
+					FileNode childNode = await interiorNode.Children[idx].ExpandAsync(reader, CancellationToken.None);
+					await CheckSizes(reader, childNode, options, idx == childCount - 1);
 				}
 			}
 		}
diff --git a/Engine/Source/Programs/Shared/EpicGames.Horde.Tests/ChunkingTests.cs b/Engine/Source/Programs/Shared/EpicGames.Horde.Tests/ChunkingTests.cs
--- a/Engine/Source/Programs/Shared/EpicGames.Horde.Tests/ChunkingTests.cs
+++ b/Engine/Source/Programs/Shared/EpicGames.Horde.Tests/ChunkingTests.cs
@@ -52,7 +52,8 @@ namespace EpicGames.Horde.Tests
 		{
 			using MemoryCache cache = new MemoryCache(new MemoryCacheOptions());
 
-			MemoryStorageClient store = new MemoryStorageClient(cache, NullLogger.Instance);
+			MemoryStorageClient store = new MemoryStorageClient();
+			TreeReader reader = new TreeReader(store, cache, NullLogger.Instance);
 			TreeWriter writer = new TreeWriter(store, new TreeOptions(), "test");
 
 			ChunkingOptions options = new ChunkingOptions();
@@ -62,13 +63,13 @@ namespace EpicGames.Horde.Tests
 			node = await node.AppendAsync(new byte[7], options, writer, CancellationToken.None);
 			Assert.AreEqual(7, node.Length);
 			Assert.IsTrue(node is LeafFileNode);
-			Assert.AreEqual(7, (await node.ToByteArrayAsync(CancellationToken.None)).Length);
+			Assert.AreEqual(7, (await node.ToByteArrayAsync(reader, CancellationToken.None)).Length);
 
 			node = new LeafFileNode();
 			node = await node.AppendAsync(new byte[8], options, writer, CancellationToken.None);
 			Assert.AreEqual(8, node.Length);
 			Assert.IsTrue(node is LeafFileNode);
-			Assert.AreEqual(8, (await node.ToByteArrayAsync(CancellationToken.None)).Length);
+			Assert.AreEqual(8, (await node.ToByteArrayAsync(reader, CancellationToken.None)).Length);
 
 			node = new LeafFileNode();
 			node = await node.AppendAsync(new byte[9], options, writer, CancellationToken.None);
@@ -76,15 +77,15 @@ namespace EpicGames.Horde.Tests
 			Assert.IsTrue(node is InteriorFileNode);
 			Assert.AreEqual(2, ((InteriorFileNode)node).Children.Count);
 
-			FileNode? childNode1 = await ((TreeNodeRef<FileNode>)((InteriorFileNode)node).Children[0]).ExpandAsync();
+			FileNode? childNode1 = await ((TreeNodeRef<FileNode>)((InteriorFileNode)node).Children[0]).ExpandAsync(reader);
 			Assert.IsNotNull(childNode1);
 			Assert.IsTrue(childNode1 is LeafFileNode);
-			Assert.AreEqual(8, (await childNode1!.ToByteArrayAsync(CancellationToken.None)).Length);
+			Assert.AreEqual(8, (await childNode1!.ToByteArrayAsync(reader, CancellationToken.None)).Length);
 
-			FileNode? childNode2 = await ((TreeNodeRef<FileNode>)((InteriorFileNode)node).Children[1]).ExpandAsync();
+			FileNode? childNode2 = await ((TreeNodeRef<FileNode>)((InteriorFileNode)node).Children[1]).ExpandAsync(reader);
 			Assert.IsNotNull(childNode2);
 			Assert.IsTrue(childNode2 is LeafFileNode);
-			Assert.AreEqual(1, (await childNode2!.ToByteArrayAsync(CancellationToken.None)).Length);
+			Assert.AreEqual(1, (await childNode2!.ToByteArrayAsync(reader, CancellationToken.None)).Length);
 		}
 
 		[TestMethod]
@@ -111,7 +112,8 @@ namespace EpicGames.Horde.Tests
 		{
 			using MemoryCache cache = new MemoryCache(new MemoryCacheOptions());
 
-			MemoryStorageClient store = new MemoryStorageClient(cache, NullLogger.Instance);
+			MemoryStorageClient store = new MemoryStorageClient();
+			TreeReader reader = new TreeReader(store, cache, NullLogger.Instance);
 			TreeWriter writer = new TreeWriter(store, new TreeOptions());
 
 			byte[] data = new byte[4096];
@@ -133,7 +135,7 @@ namespace EpicGames.Horde.Tests
 			byte[] result;
 			using (MemoryStream stream = new MemoryStream())
 			{
-				await root.CopyToStreamAsync(stream, CancellationToken.None);
+				await root.CopyToStreamAsync(reader, stream, CancellationToken.None);
 				result = stream.ToArray();
 			}
 
diff --git a/Engine/Source/Programs/Shared/EpicGames.Horde/Storage/Backends/FileStorageClient.cs b/Engine/Source/Programs/Shared/EpicGames.Horde/Storage/Backends/FileStorageClient.cs
--- a/Engine/Source/Programs/Shared/EpicGames.Horde/Storage/Backends/FileStorageClient.cs
+++ b/Engine/Source/Programs/Shared/EpicGames.Horde/Storage/Backends/FileStorageClient.cs
@@ -23,10 +23,8 @@ namespace EpicGames.Horde.Storage.Backends
 		/// Constructor
 		/// </summary>
 		/// <param name="rootDir">Root directory for storing blobs</param>
-		/// <param name="cache">Cache for blob data</param>
 		/// <param name="logger">Logger interface</param>
-		public FileStorageClient(DirectoryReference rootDir, IMemoryCache cache, ILogger logger)
-			: base(cache, logger)
+		public FileStorageClient(DirectoryReference rootDir, ILogger logger)
 		{
 			_rootDir = rootDir;
 			_logger = logger;
diff --git a/Engine/Source/Programs/Shared/EpicGames.Horde/Storage/Backends/HttpStorageClient.cs b/Engine/Source/Programs/Shared/EpicGames.Horde/Storage/Backends/HttpStorageClient.cs
--- a/Engine/Source/Programs/Shared/EpicGames.Horde/Storage/Backends/HttpStorageClient.cs
+++ b/Engine/Source/Programs/Shared/EpicGames.Horde/Storage/Backends/HttpStorageClient.cs
@@ -43,8 +43,7 @@ namespace EpicGames.Horde.Storage.Backends
 		/// <summary>
 		/// Constructor
 		/// </summary>
-		public HttpStorageClient(NamespaceId namespaceId, HttpClient httpClient, HttpClient redirectHttpClient, IMemoryCache cache, ILogger logger) 
-			: base(cache, logger)
+		public HttpStorageClient(NamespaceId namespaceId, HttpClient httpClient, HttpClient redirectHttpClient, ILogger logger) 
 		{
 			_namespaceId = namespaceId;
 			_httpClient = httpClient;
diff --git a/Engine/Source/Programs/Shared/EpicGames.Horde/Storage/Backends/MemoryStorageClient.cs b/Engine/Source/Programs/Shared/EpicGames.Horde/Storage/Backends/MemoryStorageClient.cs
--- a/Engine/Source/Programs/Shared/EpicGames.Horde/Storage/Backends/MemoryStorageClient.cs
+++ b/Engine/Source/Programs/Shared/EpicGames.Horde/Storage/Backends/MemoryStorageClient.cs
@@ -37,8 +37,7 @@ namespace EpicGames.Horde.Storage.Backends
 		/// <summary>
 		/// Constructor
 		/// </summary>
-		public MemoryStorageClient(IMemoryCache cache, ILogger logger) 
-			: base(cache, logger)
+		public MemoryStorageClient() 
 		{
 		}
 
diff --git a/Engine/Source/Programs/Shared/EpicGames.Horde/Storage/Bundle.cs b/Engine/Source/Programs/Shared/EpicGames.Horde/Storage/Bundle.cs
--- a/Engine/Source/Programs/Shared/EpicGames.Horde/Storage/Bundle.cs
+++ b/Engine/Source/Programs/Shared/EpicGames.Horde/Storage/Bundle.cs
@@ -6,9 +6,11 @@ using System;
 using System.Buffers;
 using System.Buffers.Binary;
 using System.Collections.Generic;
+using System.Diagnostics;
 using System.IO;
 using System.IO.Compression;
 using System.Linq;
+using System.Reflection;
 using System.Text;
 using System.Threading;
 using System.Threading.Tasks;
@@ -171,6 +173,11 @@ namespace EpicGames.Horde.Storage
 		/// </summary>
 		public BundleCompressionFormat CompressionFormat { get; }
 
+		/// <summary>
+		/// Types for exports within this bundle
+		/// </summary>
+		public IReadOnlyList<BundleType> Types { get; }
+
 		/// <summary>
 		/// References to nodes in other bundles
 		/// </summary>
@@ -190,12 +197,14 @@ namespace EpicGames.Horde.Storage
 		/// Constructs a new bundle header
 		/// </summary>
 		/// <param name="compressionFormat">Compression format for bundle packets</param>
+		/// <param name="types">Type array indexed by each export</param>
 		/// <param name="imports">Imports from other bundles</param>
 		/// <param name="exports">Exports for nodes</param>
 		/// <param name="packets">Compression packets within the bundle</param>
-		public BundleHeader(BundleCompressionFormat compressionFormat, IReadOnlyList<BundleImport> imports, IReadOnlyList<BundleExport> exports, IReadOnlyList<BundlePacket> packets)
+		public BundleHeader(BundleCompressionFormat compressionFormat, IReadOnlyList<BundleType> types, IReadOnlyList<BundleImport> imports, IReadOnlyList<BundleExport> exports, IReadOnlyList<BundlePacket> packets)
 		{
 			CompressionFormat = compressionFormat;
+			Types = types;
 			Imports = imports;
 			Exports = exports;
 			Packets = packets;
@@ -219,6 +228,7 @@ namespace EpicGames.Horde.Storage
 			}
 
 			CompressionFormat = (BundleCompressionFormat)reader.ReadUnsignedVarInt();
+			Types = reader.ReadVariableLengthArray(() => new BundleType(reader));
 			Imports = reader.ReadVariableLengthArray(() => new BundleImport(reader));
 			Exports = reader.ReadVariableLengthArray(() => new BundleExport(reader));
 
@@ -295,6 +305,7 @@ namespace EpicGames.Horde.Storage
 			writer.WriteUnsignedVarInt((ulong)BundleVersion.Latest);
 
 			writer.WriteUnsignedVarInt((ulong)CompressionFormat);
+			writer.WriteVariableLengthArray(Types, x => x.Write(writer));
 			writer.WriteVariableLengthArray(Imports, x => x.Write(writer));
 			writer.WriteVariableLengthArray(Exports, x => x.Write(writer));
 
@@ -308,6 +319,69 @@ namespace EpicGames.Horde.Storage
 		}
 	}
 
+	/// <summary>
+	/// Information about a type within a bundle
+	/// </summary>
+	[DebuggerDisplay("{Tag}#{Version}")]
+	public class BundleType : IEquatable<BundleType>
+	{
+		/// <summary>
+		/// Nominal identifier for the type
+		/// </summary>
+		public Guid Guid { get; }
+
+		/// <summary>
+		/// Version number for the serializer
+		/// </summary>
+		public int Version { get; }
+
+		/// <summary>
+		/// Constructor
+		/// </summary>
+		/// <param name="guid">Nominal identifier for the type</param>
+		/// <param name="version">Version number for the serializer</param>
+		public BundleType(Guid guid, int version)
+		{
+			Guid = guid;
+			Version = version;
+		}
+
+		/// <summary>
+		/// Deserialize a bundle type
+		/// </summary>
+		/// <param name="reader">Reader to deserialize from</param>
+		public BundleType(IMemoryReader reader)
+		{
+			Guid = reader.ReadGuid();
+			Version = (int)reader.ReadUnsignedVarInt();
+		}
+
+		/// <summary>
+		/// Serializes a type to storage
+		/// </summary>
+		/// <param name="writer">Writer for serialization</param>
+		public void Write(IMemoryWriter writer)
+		{
+			writer.WriteGuid(Guid);
+			writer.WriteUnsignedVarInt(Version);
+		}
+
+		/// <inheritdoc/>
+		public override bool Equals(object? obj) => obj is BundleType type && Equals(type);
+
+		/// <inheritdoc/>
+		public override int GetHashCode() => HashCode.Combine(Guid, Version);
+
+		/// <inheritdoc/>
+		public bool Equals(BundleType? type) => type is object && Guid == type.Guid && Version == type.Version;
+
+		/// <inheritdoc/>
+		public static bool operator ==(BundleType lhs, BundleType rhs) => lhs.Equals(rhs);
+
+		/// <inheritdoc/>
+		public static bool operator !=(BundleType lhs, BundleType rhs) => !lhs.Equals(rhs);
+	}
+
 	/// <summary>
 	/// Reference to another tree pack object
 	/// </summary>
@@ -422,6 +496,11 @@ namespace EpicGames.Horde.Storage
 	/// </summary>
 	public class BundleExport
 	{
+		/// <summary>
+		/// Index of the type within the type list
+		/// </summary>
+		public int TypeIdx { get; }
+
 		/// <summary>
 		/// Hash of the node
 		/// </summary>
@@ -442,8 +521,9 @@ namespace EpicGames.Horde.Storage
 		/// <summary>
 		/// Constructor
 		/// </summary>
-		public BundleExport(IoHash hash, int length, IReadOnlyList<int> references)
+		public BundleExport(int typeIdx, IoHash hash, int length, IReadOnlyList<int> references)
 		{
+			TypeIdx = typeIdx;
 			Hash = hash;
 			Length = length;
 			References = references;
@@ -455,6 +535,7 @@ namespace EpicGames.Horde.Storage
 		/// <param name="reader">Reader to deserialize from</param>
 		public BundleExport(IMemoryReader reader)
 		{
+			TypeIdx = (int)reader.ReadUnsignedVarInt();
 			Hash = reader.ReadIoHash();
 			Length = (int)reader.ReadUnsignedVarInt();
 
@@ -475,6 +556,7 @@ namespace EpicGames.Horde.Storage
 		/// <param name="writer">Writer to serialize to</param>
 		public void Write(IMemoryWriter writer)
 		{
+			writer.WriteUnsignedVarInt(TypeIdx);
 			writer.WriteIoHash(Hash);
 			writer.WriteUnsignedVarInt(Length);
 
diff --git a/Engine/Source/Programs/Shared/EpicGames.Horde/Storage/IStorageClient.cs b/Engine/Source/Programs/Shared/EpicGames.Horde/Storage/IStorageClient.cs
--- a/Engine/Source/Programs/Shared/EpicGames.Horde/Storage/IStorageClient.cs
+++ b/Engine/Source/Programs/Shared/EpicGames.Horde/Storage/IStorageClient.cs
@@ -159,27 +159,6 @@ namespace EpicGames.Horde.Storage
 
 		#endregion
 
-		#region Nodes
-
-		/// <summary>
-		/// Reads a node from a bundle
-		/// </summary>
-		/// <param name="locator">Locator for the node</param>
-		/// <param name="cancellationToken">Cancellation token for the operation</param>
-		/// <returns>Node data read from the given bundle</returns>
-		ValueTask<TNode> ReadNodeAsync<TNode>(NodeLocator locator, CancellationToken cancellationToken = default) where TNode : TreeNode;
-
-		/// <summary>
-		/// Reads data for a ref from the store, along with the node's contents.
-		/// </summary>
-		/// <param name="name">The ref name</param>
-		/// <param name="cacheTime">Minimum coherency for any cached value to be returned</param>
-		/// <param name="cancellationToken">Cancellation token for the operation</param>
-		/// <returns>Node for the given ref, or null if it does not exist</returns>
-		Task<TNode?> TryReadNodeAsync<TNode>(RefName name, DateTime cacheTime = default, CancellationToken cancellationToken = default) where TNode : TreeNode;
-
-		#endregion
-
 		#region Refs
 
 		/// <summary>
@@ -254,16 +233,6 @@ namespace EpicGames.Horde.Storage
 
 			#endregion
 
-			#region Trees
-
-			/// <inheritdoc/>
-			public ValueTask<TNode> ReadNodeAsync<TNode>(NodeLocator locator, CancellationToken cancellationToken = default) where TNode : TreeNode => _inner.ReadNodeAsync<TNode>(locator, cancellationToken);
-
-			/// <inheritdoc/>
-			public Task<TNode?> TryReadNodeAsync<TNode>(RefName name, DateTime cacheTime = default, CancellationToken cancellationToken = default) where TNode : TreeNode => _inner.TryReadNodeAsync<TNode>(name, cacheTime, cancellationToken);
-
-			#endregion
-
 			#region Refs
 
 			/// <inheritdoc/>
@@ -321,67 +290,33 @@ namespace EpicGames.Horde.Storage
 
 		#endregion
 
-		#region Nodes
-
-		/// <summary>
-		/// Attempts to reads a ref from the store
-		/// </summary>
-		/// <param name="store">The store instance to read from</param>
-		/// <param name="name">Name of the ref</param>
-		/// <param name="maxAge">Maximum age of any cached ref</param>
-		/// <param name="cancellationToken">Cancellation token for the operation</param>
-		/// <returns>The ref value</returns>
-		public static Task<TNode?> TryReadNodeAsync<TNode>(this IStorageClient store, RefName name, TimeSpan maxAge, CancellationToken cancellationToken = default) where TNode : TreeNode
-		{
-			return store.TryReadNodeAsync<TNode>(name, DateTime.UtcNow - maxAge, cancellationToken);
-		}
+		#region Blobs
 
 		/// <summary>
-		/// Reads a ref from the store, throwing an exception if it does not exist
+		/// Utility method to read a blob into a buffer
 		/// </summary>
-		/// <param name="store">The store instance to read from</param>
-		/// <param name="name">Id for the ref</param>
-		/// <param name="cacheTime">Minimum coherency of any cached result</param>
+		/// <param name="store">Store to read from</param>
+		/// <param name="locator">Blob location</param>
+		/// <param name="offset">Offset within the blob</param>
+		/// <param name="memory">Buffer to read into</param>
 		/// <param name="cancellationToken">Cancellation token for the operation</param>
-		/// <returns>The blob instance</returns>
-		public static async Task<TNode> ReadNodeAsync<TNode>(this IStorageClient store, RefName name, DateTime cacheTime = default, CancellationToken cancellationToken = default) where TNode : TreeNode
+		/// <returns>The data that was read</returns>
+		public static async Task<Memory<byte>> ReadBlobRangeAsync(this IStorageClient store, BlobLocator locator, int offset, Memory<byte> memory, CancellationToken cancellationToken = default)
 		{
-			TNode? refValue = await store.TryReadNodeAsync<TNode>(name, cacheTime, cancellationToken);
-			if (refValue == null)
+			using (Stream stream = await store.ReadBlobRangeAsync(locator, offset, memory.Length, cancellationToken))
 			{
-				throw new RefNameNotFoundException(name);
+				int length = 0;
+				while (length < memory.Length)
+				{
+					int readBytes = await stream.ReadAsync(memory.Slice(length), cancellationToken);
+					if (readBytes == 0)
+					{
+						break;
+					}
+					length += readBytes;
+				}
+				return memory.Slice(0, length);
 			}
-			return refValue;
-		}
-
-		/// <summary>
-		/// Reads a ref from the store, throwing an exception if it does not exist
-		/// </summary>
-		/// <param name="store">The store instance to read from</param>
-		/// <param name="name">Id for the ref</param>
-		/// <param name="maxAge">Maximum age for any cached result</param>
-		/// <param name="cancellationToken">Cancellation token for the operation</param>
-		/// <returns>The ref value</returns>
-		public static Task<TNode> ReadNodeAsync<TNode>(this IStorageClient store, RefName name, TimeSpan maxAge, CancellationToken cancellationToken = default) where TNode : TreeNode
-		{
-			return ReadNodeAsync<TNode>(store, name, DateTime.UtcNow - maxAge, cancellationToken);
-		}
-
-		/// <summary>
-		/// Writes a node to storage
-		/// </summary>
-		/// <param name="store">Store instance to write to</param>
-		/// <param name="name">Name of the ref containing this node</param>
-		/// <param name="node">Node to be written</param>
-		/// <param name="options">Options for the node writer</param>
-		/// <param name="prefix">Prefix for uploaded blobs</param>
-		/// <param name="refOptions">Options for the ref</param>
-		/// <param name="cancellationToken">Cancellation token for the operation</param>
-		/// <returns>Location of node targetted by the ref</returns>
-		public static async Task<NodeLocator> WriteNodeAsync(this IStorageClient store, RefName name, TreeNode node, TreeOptions? options = null, Utf8String prefix = default, RefOptions? refOptions = null, CancellationToken cancellationToken = default)
-		{
-			TreeWriter writer = new TreeWriter(store, options, prefix.IsEmpty ? name.Text : prefix);
-			return await writer.WriteRefAsync(name, node, refOptions, cancellationToken);
 		}
 
 		#endregion
diff --git a/Engine/Source/Programs/Shared/EpicGames.Horde/Storage/Nodes/DirectoryNode.cs b/Engine/Source/Programs/Shared/EpicGames.Horde/Storage/Nodes/DirectoryNode.cs
--- a/Engine/Source/Programs/Shared/EpicGames.Horde/Storage/Nodes/DirectoryNode.cs
+++ b/Engine/Source/Programs/Shared/EpicGames.Horde/Storage/Nodes/DirectoryNode.cs
@@ -1,20 +1,14 @@
 // Copyright Epic Games, Inc. All Rights Reserved.
 
 using EpicGames.Core;
-using EpicGames.Serialization;
 using Microsoft.Extensions.Logging;
-using Microsoft.Extensions.Options;
 using System;
-using System.Buffers;
 using System.Collections.Generic;
-using System.Diagnostics;
 using System.Diagnostics.CodeAnalysis;
 using System.IO;
 using System.Linq;
-using System.Reflection;
 using System.Threading;
 using System.Threading.Tasks;
-using System.Xml.Linq;
 
 namespace EpicGames.Horde.Storage.Nodes
 {
@@ -80,14 +74,6 @@ namespace EpicGames.Horde.Storage.Nodes
 		/// </summary>
 		readonly long _cachedLength;
 
-		/// <summary>
-		/// Constructor
-		/// </summary>
-		public FileEntry(Utf8String name, FileEntryFlags flags)
-			: this(name, flags, new LeafFileNode())
-		{
-		}
-
 		/// <summary>
 		/// Constructor
 		/// </summary>
@@ -124,53 +110,29 @@ namespace EpicGames.Horde.Storage.Nodes
 			writer.WriteUnsignedVarInt((ulong)Length);
 		}
 
-		/// <summary>
-		/// Appends data to this file
-		/// </summary>
-		/// <param name="data">Data to append to the file</param>
-		/// <param name="options">Options for chunking the data</param>
-		/// <param name="writer">Writer for new node data</param>
-		/// <param name="cancellationToken">Cancellation token for the operation</param>
-		public async Task AppendAsync(ReadOnlyMemory<byte> data, ChunkingOptions options, TreeWriter writer, CancellationToken cancellationToken)
-		{
-			FileNode node = await ExpandAsync(cancellationToken);
-			Target = await node.AppendAsync(data, options, writer, cancellationToken);
-		}
-
-		/// <summary>
-		/// Appends data to this file
-		/// </summary>
-		/// <param name="stream">Data to append to the file</param>
-		/// <param name="options">Options for chunking the data</param>
-		/// <param name="writer">Writer for new node data</param>
-		/// <param name="cancellationToken">Cancellation token for the operation</param>
-		public async Task AppendAsync(Stream stream, ChunkingOptions options, TreeWriter writer, CancellationToken cancellationToken)
-		{
-			FileNode node = await ExpandAsync(cancellationToken);
-			Target = await node.AppendAsync(stream, options, writer, cancellationToken);
-		}
-
 		/// <summary>
 		/// Copies the contents of this node and its children to the given output stream
 		/// </summary>
+		/// <param name="reader">Reader for retrieving existing node data</param>
 		/// <param name="outputStream">The output stream to receive the data</param>
 		/// <param name="cancellationToken">Cancellation token for the operation</param>
-		public async Task CopyToStreamAsync(Stream outputStream, CancellationToken cancellationToken)
+		public async Task CopyToStreamAsync(TreeReader reader, Stream outputStream, CancellationToken cancellationToken)
 		{
-			FileNode node = await ExpandAsync(cancellationToken);
-			await node.CopyToStreamAsync(outputStream, cancellationToken);
+			FileNode node = await ExpandAsync(reader, cancellationToken);
+			await node.CopyToStreamAsync(reader, outputStream, cancellationToken);
 		}
 
 		/// <summary>
 		/// Extracts the contents of this node to a file
 		/// </summary>
+		/// <param name="reader">Reader for retrieving existing node data</param>
 		/// <param name="file">File to write with the contents of this node</param>
 		/// <param name="cancellationToken">Cancellation token for the operation</param>
 		/// <returns></returns>
-		public async Task CopyToFileAsync(FileInfo file, CancellationToken cancellationToken)
+		public async Task CopyToFileAsync(TreeReader reader, FileInfo file, CancellationToken cancellationToken)
 		{
-			FileNode node = await ExpandAsync(cancellationToken);
-			await node.CopyToFileAsync(file, cancellationToken);
+			FileNode node = await ExpandAsync(reader, cancellationToken);
+			await node.CopyToFileAsync(reader, file, cancellationToken);
 		}
 
 		/// <inheritdoc/>
@@ -255,10 +217,9 @@ namespace EpicGames.Horde.Storage.Nodes
 	/// <summary>
 	/// A directory node
 	/// </summary>
+	[TreeNode("{0714EC11-291A-4D07-867F-E78AD6809979}", 1)]
 	public class DirectoryNode : TreeNode
 	{
-		internal const byte TypeId = (byte)'d';
-
 		readonly Dictionary<Utf8String, FileEntry> _nameToFileEntry = new Dictionary<Utf8String, FileEntry>();
 		readonly Dictionary<Utf8String, DirectoryEntry> _nameToDirectoryEntry = new Dictionary<Utf8String, DirectoryEntry>();
 
@@ -307,12 +268,6 @@ namespace EpicGames.Horde.Storage.Nodes
 		/// <param name="reader">Reader to deserialize from</param>
 		public DirectoryNode(ITreeNodeReader reader)
 		{
-			byte typeId = reader.ReadUInt8();
-			if (typeId != TypeId)
-			{
-				throw new InvalidOperationException("Invalid signature byte for directory");
-			}
-
 			Flags = (DirectoryFlags)reader.ReadUnsignedVarInt();
 
 			int fileCount = (int)reader.ReadUnsignedVarInt();
@@ -337,7 +292,6 @@ namespace EpicGames.Horde.Storage.Nodes
 		/// <inheritdoc/>
 		public override void Serialize(ITreeNodeWriter writer)
 		{
-			writer.WriteUInt8(TypeId);
 			writer.WriteUnsignedVarInt((ulong)Flags);
 
 			writer.WriteUnsignedVarInt(Files.Count);
@@ -400,39 +354,25 @@ namespace EpicGames.Horde.Storage.Nodes
 		/// </summary>
 		/// <param name="name">Name of the new directory</param>
 		/// <param name="flags">Flags for the new file</param>
+		/// <param name="node">The file data</param>
 		/// <returns>The new directory object</returns>
-		public FileEntry AddFile(Utf8String name, FileEntryFlags flags)
+		public FileEntry AddFile(Utf8String name, FileEntryFlags flags, FileNode node)
 		{
-			FileEntry entry = new FileEntry(name, flags);
+			FileEntry entry = new FileEntry(name, flags, node);
 			AddFile(entry);
 			return entry;
 		}
 
-		/// <summary>
-		/// Adds a new directory with the given name
-		/// </summary>
-		/// <param name="name">Name of the new directory</param>
-		/// <param name="flags">Flags for the new file</param>
-		/// <param name="stream">The stream to read from</param>
-		/// <param name="options">Options for chunking the data</param>
-		/// <param name="writer">Writer for new node data</param>
-		/// <param name="cancellationToken">Cancellation token for the operation</param>
-		/// <returns>The new directory object</returns>
-		public async Task<FileEntry> AddFileAsync(Utf8String name, FileEntryFlags flags, Stream stream, ChunkingOptions options, TreeWriter writer, CancellationToken cancellationToken = default)
-		{
-			FileEntry entry = AddFile(name, flags);
-			await entry.AppendAsync(stream, options, writer, cancellationToken);
-			return entry;
-		}
-
 		/// <summary>
 		/// Finds or adds a file with the given path
 		/// </summary>
+		/// <param name="reader">Reader for node data</param>
 		/// <param name="path">Path to the file</param>
 		/// <param name="flags">Flags for the new file</param>
+		/// <param name="node">THe file node</param>
 		/// <param name="cancellationToken">Cancellation token for the operation</param>
 		/// <returns>The new directory object</returns>
-		public async ValueTask<FileEntry> AddFileByPathAsync(Utf8String path, FileEntryFlags flags, CancellationToken cancellationToken = default)
+		public async ValueTask<FileEntry> AddFileByPathAsync(TreeReader reader, Utf8String path, FileEntryFlags flags, FileNode node, CancellationToken cancellationToken = default)
 		{
 			DirectoryNode directory = this;
 
@@ -449,7 +389,7 @@ namespace EpicGames.Horde.Storage.Nodes
 				{
 					if (length == remainingPath.Length)
 					{
-						return directory.AddFile(remainingPath, flags);
+						return directory.AddFile(remainingPath, flags, node);
 					}
 
 					byte character = remainingPath[length];
@@ -461,7 +401,7 @@ namespace EpicGames.Horde.Storage.Nodes
 
 				if (length > 0)
 				{
-					directory = await directory.FindOrAddDirectoryAsync(remainingPath.Slice(0, length), cancellationToken);
+					directory = await directory.FindOrAddDirectoryAsync(reader, remainingPath.Slice(0, length), cancellationToken);
 				}
 				remainingPath = remainingPath.Slice(length + 1);
 			}
@@ -500,10 +440,11 @@ namespace EpicGames.Horde.Storage.Nodes
 		/// <summary>
 		/// Deletes a file with the given path
 		/// </summary>
+		/// <param name="reader">Reader for existing node data</param>
 		/// <param name="path"></param>
 		/// <param name="cancellationToken"></param>
 		/// <returns></returns>
-		public async ValueTask<bool> DeleteFileByPathAsync(Utf8String path, CancellationToken cancellationToken)
+		public async ValueTask<bool> DeleteFileByPathAsync(TreeReader reader, Utf8String path, CancellationToken cancellationToken)
 		{
 			Utf8String remainingPath = path;
 			for (DirectoryNode? directory = this; directory != null;)
@@ -515,7 +456,7 @@ namespace EpicGames.Horde.Storage.Nodes
 				}
 				if (length > 0)
 				{
-					directory = await directory.FindDirectoryAsync(remainingPath.Slice(0, length), cancellationToken);
+					directory = await directory.FindDirectoryAsync(reader, remainingPath.Slice(0, length), cancellationToken);
 				}
 				remainingPath = remainingPath.Slice(length + 1);
 			}
@@ -582,14 +523,15 @@ namespace EpicGames.Horde.Storage.Nodes
 		/// <summary>
 		/// Tries to get a directory with the given name
 		/// </summary>
+		/// <param name="reader">Reader for node data</param>
 		/// <param name="name">Name of the new directory</param>
 		/// <param name="cancellationToken">Cancellation token for the operation</param>
 		/// <returns>The new directory object</returns>
-		public async ValueTask<DirectoryNode?> FindDirectoryAsync(Utf8String name, CancellationToken cancellationToken)
+		public async ValueTask<DirectoryNode?> FindDirectoryAsync(TreeReader reader, Utf8String name, CancellationToken cancellationToken)
 		{
 			if (TryGetDirectoryEntry(name, out DirectoryEntry? entry))
 			{
-				return await entry.ExpandAsync(cancellationToken);
+				return await entry.ExpandAsync(reader, cancellationToken);
 			}
 			else
 			{
@@ -600,12 +542,13 @@ namespace EpicGames.Horde.Storage.Nodes
 		/// <summary>
 		/// Tries to get a directory with the given name
 		/// </summary>
+		/// <param name="reader">Reader for node data</param>
 		/// <param name="name">Name of the new directory</param>
 		/// <param name="cancellationToken">Cancellation token for the operation</param>
 		/// <returns>The new directory object</returns>
-		public async ValueTask<DirectoryNode> FindOrAddDirectoryAsync(Utf8String name, CancellationToken cancellationToken)
+		public async ValueTask<DirectoryNode> FindOrAddDirectoryAsync(TreeReader reader, Utf8String name, CancellationToken cancellationToken)
 		{
-			DirectoryNode? directory = await FindDirectoryAsync(name, cancellationToken);
+			DirectoryNode? directory = await FindDirectoryAsync(reader, name, cancellationToken);
 			if (directory == null)
 			{
 				directory = AddDirectory(name);
@@ -715,13 +658,8 @@ namespace EpicGames.Horde.Storage.Nodes
 			for(int idx = minIdx; idx < maxIdx; idx++)
 			{
 				FileInfo fileInfo = files[idx].FileInfo;
-
-				FileEntry fileEntry = new FileEntry(fileInfo.Name, FileEntryFlags.None);
-				using (Stream stream = fileInfo.OpenRead())
-				{
-					await fileEntry.AppendAsync(stream, options, writer, cancellationToken);
-				}
-				entries[idx] = fileEntry;
+				FileNode fileNode = await FileNode.CreateAsync(fileInfo, options, writer, cancellationToken);
+				entries[idx] = new FileEntry(fileInfo.Name, FileEntryFlags.None, fileNode);
 			}
 			await writer.FlushAsync(cancellationToken);
 		}
@@ -729,11 +667,12 @@ namespace EpicGames.Horde.Storage.Nodes
 		/// <summary>
 		/// Utility function to allow extracting a packed directory to disk
 		/// </summary>
+		/// <param name="reader">Reader to retrieve data from</param>
 		/// <param name="directoryInfo"></param>
 		/// <param name="logger"></param>
 		/// <param name="cancellationToken"></param>
 		/// <returns></returns>
-		public async Task CopyToDirectoryAsync(DirectoryInfo directoryInfo, ILogger logger, CancellationToken cancellationToken)
+		public async Task CopyToDirectoryAsync(TreeReader reader, DirectoryInfo directoryInfo, ILogger logger, CancellationToken cancellationToken)
 		{
 			directoryInfo.Create();
 
@@ -741,29 +680,20 @@ namespace EpicGames.Horde.Storage.Nodes
 			foreach (FileEntry fileEntry in _nameToFileEntry.Values)
 			{
 				FileInfo fileInfo = new FileInfo(Path.Combine(directoryInfo.FullName, fileEntry.Name.ToString()));
-				FileNode fileNode = await fileEntry.ExpandAsync(cancellationToken);
-				tasks.Add(Task.Run(() => fileNode.CopyToFileAsync(fileInfo, cancellationToken), cancellationToken));
+				FileNode fileNode = await fileEntry.ExpandAsync(reader, cancellationToken);
+				tasks.Add(Task.Run(() => fileNode.CopyToFileAsync(reader, fileInfo, cancellationToken), cancellationToken));
 			}
 			foreach (DirectoryEntry directoryEntry in _nameToDirectoryEntry.Values)
 			{
 				DirectoryInfo subDirectoryInfo = directoryInfo.CreateSubdirectory(directoryEntry.Name.ToString());
-				DirectoryNode subDirectoryNode = await directoryEntry.ExpandAsync(cancellationToken);
-				tasks.Add(Task.Run(() => subDirectoryNode.CopyToDirectoryAsync(subDirectoryInfo, logger, cancellationToken), cancellationToken));
+				DirectoryNode subDirectoryNode = await directoryEntry.ExpandAsync(reader, cancellationToken);
+				tasks.Add(Task.Run(() => subDirectoryNode.CopyToDirectoryAsync(reader, subDirectoryInfo, logger, cancellationToken), cancellationToken));
 			}
 
 			await Task.WhenAll(tasks);
 		}
 	}
 
-	/// <summary>
-	/// Factory for creating and serializing <see cref="DirectoryNode"/> objects
-	/// </summary>
-	public class DirectoryNodeSerializer : TreeNodeSerializer<DirectoryNode>
-	{
-		/// <inheritdoc/>
-		public override DirectoryNode Deserialize(ITreeNodeReader reader) => new DirectoryNode(reader);
-	}
-
 	/// <summary>
 	/// Extension methods for <see cref="DirectoryNode"/>
 	/// </summary>
diff --git a/Engine/Source/Programs/Shared/EpicGames.Horde/Storage/Nodes/FileNode.cs b/Engine/Source/Programs/Shared/EpicGames.Horde/Storage/Nodes/FileNode.cs
--- a/Engine/Source/Programs/Shared/EpicGames.Horde/Storage/Nodes/FileNode.cs
+++ b/Engine/Source/Programs/Shared/EpicGames.Horde/Storage/Nodes/FileNode.cs
@@ -1,7 +1,9 @@
 // Copyright Epic Games, Inc. All Rights Reserved.
 
 using EpicGames.Core;
+using EpicGames.Horde.Storage.Nodes;
 using EpicGames.Serialization;
+using Microsoft.Extensions.Options;
 using System;
 using System.Buffers;
 using System.Buffers.Binary;
@@ -9,8 +11,10 @@ using System.Collections.Generic;
 using System.Diagnostics;
 using System.IO;
 using System.Linq;
+using System.Reflection;
 using System.Threading;
 using System.Threading.Tasks;
+using System.Xml.Linq;
 
 namespace EpicGames.Horde.Storage.Nodes
 {
@@ -75,7 +79,6 @@ namespace EpicGames.Horde.Storage.Nodes
 	/// Chunks are pushed into a tree hierarchy as data is appended to the root, with nodes of the tree also split along content-aware boundaries with <see cref="IoHash.NumBytes"/> granularity.
 	/// Once a chunk has been written to storage, it is treated as immutable.
 	/// </summary>
-	[TreeSerializer(typeof(FileNodeSerializer))]
 	public abstract class FileNode : TreeNode
 	{
 		/// <summary>
@@ -94,6 +97,36 @@ namespace EpicGames.Horde.Storage.Nodes
 		/// <returns>True if the node is read only</returns>
 		public abstract bool IsReadOnly();
 
+		/// <summary>
+		/// Creates a file node from a block of memory
+		/// </summary>
+		/// <param name="memory">The memory to read from</param>
+		/// <param name="options">Options for chunking the data</param>
+		/// <param name="writer">Writer for new tree nodes</param>
+		/// <param name="cancellationToken">Cancellation token for the operation</param>
+		/// <returns>New file node</returns>
+		public static Task<FileNode> CreateAsync(ReadOnlyMemory<byte> memory, ChunkingOptions options, TreeWriter writer, CancellationToken cancellationToken)
+		{
+			using ReadOnlyMemoryStream stream = new ReadOnlyMemoryStream(memory);
+			return CreateAsync(stream, options, writer, cancellationToken);
+		}
+
+		/// <summary>
+		/// Creates a file node from a file
+		/// </summary>
+		/// <param name="fileInfo">The stream to read from</param>
+		/// <param name="options">Options for chunking the data</param>
+		/// <param name="writer">Writer for new tree nodes</param>
+		/// <param name="cancellationToken">Cancellation token for the operation</param>
+		/// <returns>New file node</returns>
+		public static async Task<FileNode> CreateAsync(FileInfo fileInfo, ChunkingOptions options, TreeWriter writer, CancellationToken cancellationToken)
+		{
+			using (Stream stream = fileInfo.OpenRead())
+			{
+				return await FileNode.CreateAsync(stream, options, writer, cancellationToken);
+			}
+		}
+
 		/// <summary>
 		/// Creates a file node from a stream
 		/// </summary>
@@ -231,34 +264,37 @@ namespace EpicGames.Horde.Storage.Nodes
 		/// <summary>
 		/// Copies the contents of this node and its children to the given output stream
 		/// </summary>
+		/// <param name="reader">Reader for nodes in the tree</param>
 		/// <param name="outputStream">The output stream to receive the data</param>
 		/// <param name="cancellationToken">Cancellation token for the operation</param>
-		public abstract Task CopyToStreamAsync(Stream outputStream, CancellationToken cancellationToken);
+		public abstract Task CopyToStreamAsync(TreeReader reader, Stream outputStream, CancellationToken cancellationToken);
 
 		/// <summary>
 		/// Extracts the contents of this node to a file
 		/// </summary>
+		/// <param name="reader">Reader for nodes in the tree</param>
 		/// <param name="file">File to write with the contents of this node</param>
 		/// <param name="cancellationToken">Cancellation token for the operation</param>
 		/// <returns></returns>
-		public async Task CopyToFileAsync(FileInfo file, CancellationToken cancellationToken)
+		public async Task CopyToFileAsync(TreeReader reader, FileInfo file, CancellationToken cancellationToken)
 		{
 			using (FileStream stream = file.OpenWrite())
 			{
-				await CopyToStreamAsync(stream, cancellationToken);
+				await CopyToStreamAsync(reader, stream, cancellationToken);
 			}
 		}
 
 		/// <summary>
 		/// Serialize this node and its children into a byte array
 		/// </summary>
+		/// <param name="reader">Reader for nodes in the tree</param>
 		/// <param name="cancellationToken"></param>
 		/// <returns>Array of data stored by the tree</returns>
-		public async Task<byte[]> ToByteArrayAsync(CancellationToken cancellationToken)
+		public async Task<byte[]> ToByteArrayAsync(TreeReader reader, CancellationToken cancellationToken)
 		{
 			using (MemoryStream stream = new MemoryStream())
 			{
-				await CopyToStreamAsync(stream, cancellationToken);
+				await CopyToStreamAsync(reader, stream, cancellationToken);
 				return stream.ToArray();
 			}
 		}
@@ -267,6 +303,7 @@ namespace EpicGames.Horde.Storage.Nodes
 	/// <summary>
 	/// File node that contains a chunk of data
 	/// </summary>
+	[TreeNode("{B27AFB68-9E20-4A4B-A4D8-788A4098D439}", 1)]
 	public sealed class LeafFileNode : FileNode
 	{
 		class DataSegment : ReadOnlySequenceSegment<byte>
@@ -283,11 +320,6 @@ namespace EpicGames.Horde.Storage.Nodes
 			}
 		}
 
-		/// <summary>
-		/// First byte in the serialized data for this class, indicating its type.
-		/// </summary>
-		public const byte TypeId = (byte)'l';
-
 		bool _isReadOnly;
 		uint _rollingHash;
 		DataSegment _firstSegment;
@@ -310,12 +342,6 @@ namespace EpicGames.Horde.Storage.Nodes
 		/// </summary>
 		public LeafFileNode(ITreeNodeReader reader)
 		{
-			byte typeId = reader.ReadUInt8();
-			if (typeId != TypeId)
-			{
-				throw new InvalidDataException($"Invalid type id for {nameof(LeafFileNode)}");
-			}
-
 			_isReadOnly = true;
 			_rollingHash = reader.ReadUInt32();
 
@@ -328,7 +354,6 @@ namespace EpicGames.Horde.Storage.Nodes
 		/// <inheritdoc/>
 		public override void Serialize(ITreeNodeWriter writer)
 		{
-			writer.WriteUInt8(TypeId);
 			writer.WriteUInt32(_rollingHash);
 			writer.WriteVariableLengthBytes(_writtenSequence);
 		}
@@ -471,7 +496,7 @@ namespace EpicGames.Horde.Storage.Nodes
 		}
 
 		/// <inheritdoc/>
-		public override async Task CopyToStreamAsync(Stream outputStream, CancellationToken cancellationToken)
+		public override async Task CopyToStreamAsync(TreeReader reader, Stream outputStream, CancellationToken cancellationToken)
 		{
 			foreach (ReadOnlyMemory<byte> segment in _writtenSequence)
 			{
@@ -483,13 +508,9 @@ namespace EpicGames.Horde.Storage.Nodes
 	/// <summary>
 	/// An interior file node
 	/// </summary>
+	[TreeNode("{F4DEDDBC-70CB-4C7A-8347-F011AFCCCDB9}", 1)]
 	public class InteriorFileNode : FileNode
 	{
-		/// <summary>
-		/// Type identifier for interior nodes. First byte in the serialized stream.
-		/// </summary>
-		public const byte TypeId = (byte)'i';
-
 		bool _isReadOnly;
 		uint _rollingHash;
 		long _length;
@@ -522,12 +543,6 @@ namespace EpicGames.Horde.Storage.Nodes
 		/// </summary>
 		public InteriorFileNode(ITreeNodeReader reader)
 		{
-			byte typeId = reader.ReadUInt8();
-			if (typeId != TypeId)
-			{
-				throw new InvalidDataException($"Invalid type id for {nameof(InteriorFileNode)}");
-			}
-
 			_rollingHash = reader.ReadUInt32();
 			_length = (long)reader.ReadUnsignedVarInt();
 			_isReadOnly = reader.ReadBoolean();
@@ -540,8 +555,6 @@ namespace EpicGames.Horde.Storage.Nodes
 		/// <inheritdoc/>
 		public override void Serialize(ITreeNodeWriter writer)
 		{
-			writer.WriteUInt8(TypeId);
-
 			writer.WriteUInt32(_rollingHash);
 			writer.WriteUnsignedVarInt((ulong)_length);
 			writer.WriteBoolean(_isReadOnly);
@@ -617,34 +630,13 @@ namespace EpicGames.Horde.Storage.Nodes
 		}
 
 		/// <inheritdoc/>
-		public override async Task CopyToStreamAsync(Stream outputStream, CancellationToken cancellationToken)
+		public override async Task CopyToStreamAsync(TreeReader reader, Stream outputStream, CancellationToken cancellationToken)
 		{
 			foreach (TreeNodeRef<FileNode> childNodeRef in _children)
 			{
-				FileNode childNode = await childNodeRef.ExpandAsync(cancellationToken);
-				await childNode.CopyToStreamAsync(outputStream, cancellationToken);
+				FileNode childNode = await childNodeRef.ExpandAsync(reader, cancellationToken);
+				await childNode.CopyToStreamAsync(reader, outputStream, cancellationToken);
 			}
 		}
 	}
-
-	/// <summary>
-	/// Factory class for file nodes
-	/// </summary>
-	public class FileNodeSerializer : TreeNodeSerializer<FileNode>
-	{
-		/// <inheritdoc/>
-		public override FileNode Deserialize(ITreeNodeReader reader)
-		{
-			ReadOnlySpan<byte> data = reader.GetSpan(1);
-			switch (data[0])
-			{
-				case LeafFileNode.TypeId:
-					return new LeafFileNode(reader);
-				case InteriorFileNode.TypeId:
-					return new InteriorFileNode(reader);
-				default:
-					throw new InvalidDataException("Unknown type id while deserializing file node");
-			};
-		}
-	}
 }
diff --git a/Engine/Source/Programs/Shared/EpicGames.Horde/Storage/StorageClientBase.cs b/Engine/Source/Programs/Shared/EpicGames.Horde/Storage/StorageClientBase.cs
--- a/Engine/Source/Programs/Shared/EpicGames.Horde/Storage/StorageClientBase.cs
+++ b/Engine/Source/Programs/Shared/EpicGames.Horde/Storage/StorageClientBase.cs
@@ -1,16 +1,10 @@
 // Copyright Epic Games, Inc. All Rights Reserved.
 
 using System;
-using System.Buffers;
-using System.Collections.Generic;
 using System.IO;
-using System.Linq;
 using System.Threading;
 using System.Threading.Tasks;
 using EpicGames.Core;
-using Microsoft.CodeAnalysis;
-using Microsoft.Extensions.Caching.Memory;
-using Microsoft.Extensions.Logging;
 
 namespace EpicGames.Horde.Storage
 {
@@ -19,121 +13,11 @@ namespace EpicGames.Horde.Storage
 	/// </summary>
 	public abstract class StorageClientBase : IStorageClient
 	{
-		/// <summary>
-		/// Describes a bundle export
-		/// </summary>
-		/// <param name="PacketIdx">Index of the packet within a bundle</param>
-		/// <param name="Offset">Offset within the packet</param>
-		record struct ExportInfo(int PacketIdx, int Offset);
-
-		/// <summary>
-		/// Computed information about a bundle
-		/// </summary>
-		class BundleInfo
-		{
-			public readonly BlobLocator Locator;
-			public readonly BundleHeader Header;
-			public readonly int[] PacketOffsets;
-			public readonly ExportInfo[] Exports;
-			public readonly (IoHash, NodeLocator)[] References;
-
-			public BundleInfo(BlobLocator locator, BundleHeader header, int headerLength)
-			{
-				Locator = locator;
-				Header = header;
-
-				PacketOffsets = new int[header.Packets.Count];
-				Exports = new ExportInfo[header.Exports.Count];
-
-				int exportIdx = 0;
-				int packetOffset = headerLength;
-				for (int packetIdx = 0; packetIdx < header.Packets.Count; packetIdx++)
-				{
-					BundlePacket packet = header.Packets[packetIdx];
-					PacketOffsets[packetIdx] = packetOffset;
-
-					int nodeOffset = 0;
-					for (; exportIdx < header.Exports.Count && nodeOffset + header.Exports[exportIdx].Length <= packet.DecodedLength; exportIdx++)
-					{
-						Exports[exportIdx] = new ExportInfo(packetIdx, nodeOffset);
-						nodeOffset += header.Exports[exportIdx].Length;
-					}
-
-					packetOffset += packet.EncodedLength;
-				}
-
-				References = new (IoHash, NodeLocator)[header.Imports.Sum(x => x.Exports.Count) + header.Exports.Count];
-
-				int referenceIdx = 0;
-				foreach (BundleImport import in header.Imports)
-				{
-					foreach ((int importExportIdx, IoHash importExportHash) in import.Exports)
-					{
-						NodeLocator importLocator = new NodeLocator(import.Locator, importExportIdx);
-						References[referenceIdx++] = (importExportHash, importLocator);
-					}
-				}
-				for(int idx = 0; idx < header.Exports.Count; idx++)
-				{
-					References[referenceIdx++] = (header.Exports[idx].Hash, new NodeLocator(locator, idx));
-				}
-			}
-		}
-
-		/// <summary>
-		/// Bundle header queued to be read
-		/// </summary>
-		class QueuedHeader
-		{
-			public readonly BlobLocator Blob;
-			public readonly TaskCompletionSource<BundleInfo> CompletionSource = new TaskCompletionSource<BundleInfo>();
-
-			public BlobId BlobId => Blob.BlobId;
-
-			public QueuedHeader(BlobLocator blob)
-			{
-				Blob = blob;
-			}
-		}
-
-		/// <summary>
-		/// Encoded bundle packet queued to be read
-		/// </summary>
-		class QueuedPacket
-		{
-			public readonly BundleInfo Bundle;
-			public readonly int PacketIdx;
-			public readonly TaskCompletionSource<ReadOnlyMemory<byte>> CompletionSource = new TaskCompletionSource<ReadOnlyMemory<byte>>();
-
-			public BlobId BlobId => Bundle.Locator.BlobId;
-
-			public QueuedPacket(BundleInfo bundle, int packetIdx)
-			{
-				Bundle = bundle;
-				PacketIdx = packetIdx;
-			}
-		}
-
-		// Size of data to fetch by default. This is larger than the minimum request size to reduce number of reads.
-		const int DefaultFetchSize = 15 * 1024 * 1024;
-
-		readonly object _queueLock = new object();
-		readonly List<QueuedHeader> _queuedHeaders = new List<QueuedHeader>();
-		readonly List<QueuedPacket> _queuedPackets = new List<QueuedPacket>();
-		readonly Dictionary<string, Task<ReadOnlyMemory<byte>>> _decodeTasks = new Dictionary<string, Task<ReadOnlyMemory<byte>>>(StringComparer.Ordinal);
-		Task? _readTask;
-		readonly IMemoryCache _cache;
-		readonly ILogger _logger;
-
 		/// <summary>
 		/// Constructor
 		/// </summary>
-		/// <param name="cache">Cache for read data</param>
-		/// <param name="logger">Logger for output</param>
-		protected StorageClientBase(IMemoryCache cache, ILogger logger)
+		protected StorageClientBase()
 		{
-			_cache = cache;
-			_logger = logger;
 		}
 
 		#region Blobs
@@ -144,462 +28,16 @@ namespace EpicGames.Horde.Storage
 		/// <inheritdoc/>
 		public abstract Task<Stream> ReadBlobRangeAsync(BlobLocator locator, int offset, int length, CancellationToken cancellationToken = default);
 
-		/// <summary>
-		/// Utility method to read a blob into a buffer
-		/// </summary>
-		/// <param name="locator">Blob location</param>
-		/// <param name="offset">Offset within the blob</param>
-		/// <param name="memory">Buffer to read into</param>
-		/// <param name="cancellationToken">Cancellation token for the operation</param>
-		/// <returns>The data that was read</returns>
-		async Task<Memory<byte>> ReadBlobRangeAsync(BlobLocator locator, int offset, Memory<byte> memory, CancellationToken cancellationToken = default)
-		{
-			using (Stream stream = await ReadBlobRangeAsync(locator, offset, memory.Length, cancellationToken))
-			{
-				int length = 0;
-				while (length < memory.Length)
-				{
-					int readBytes = await stream.ReadAsync(memory.Slice(length), cancellationToken);
-					if (readBytes == 0)
-					{
-						break;
-					}
-					length += readBytes;
-				}
-				return memory.Slice(0, length);
-			}
-		}
-
 		/// <inheritdoc/>
 		public abstract Task<BlobLocator> WriteBlobAsync(Stream stream, Utf8String prefix = default, CancellationToken cancellationToken = default);
 
 		#endregion
 
-		#region Bundles
-
-		static string GetBundleInfoCacheKey(BlobId blobId) => $"bundle:{blobId}";
-		static string GetEncodedPacketCacheKey(BlobId blobId, int packetIdx) => $"encoded-packet:{blobId}#{packetIdx}";
-		static string GetDecodedPacketCacheKey(BlobId blobId, int packetIdx) => $"decoded-packet:{blobId}#{packetIdx}";
-
-		/// <summary>
-		/// Adds an object to the storage cache
-		/// </summary>
-		/// <param name="cacheKey">Key for the item</param>
-		/// <param name="value">Value to add</param>
-		/// <param name="size">Size of the value</param>
-		void AddToCache(string cacheKey, object value, int size)
-		{
-			using (ICacheEntry entry = _cache.CreateEntry(cacheKey))
-			{
-				entry.SetValue(value);
-				entry.SetSize(size);
-			}
-		}
-
-		/// <summary>
-		/// Starts the background task for reading data from the store
-		/// </summary>
-		void StartReadTask()
-		{
-			if (_readTask == null)
-			{
-				_readTask = Task.Run(() => ServiceReadQueueAsync(CancellationToken.None));
-			}
-		}
-
-		/// <summary>
-		/// Dispatches requests in the read queue
-		/// </summary>
-		/// <param name="cancellationToken">Cancellation token for the background task</param>
-		async Task ServiceReadQueueAsync(CancellationToken cancellationToken)
-		{
-			const int MaxConcurrentReads = 4;
-
-			List<(BlobId BlobId, Task Task)> currentTasks = new List<(BlobId, Task)>();
-			for (; ; )
-			{
-				// Start any new reads
-				lock (_queueLock)
-				{
-					while (currentTasks.Count < MaxConcurrentReads)
-					{
-						HashSet<BlobId> currentBlobIds = new HashSet<BlobId>(currentTasks.Select(x => x.BlobId));
-
-						// Try to start another header read
-						QueuedHeader? queuedHeader = _queuedHeaders.FirstOrDefault(x => !currentBlobIds.Contains(x.BlobId));
-						if (queuedHeader != null)
-						{
-							Task task = Task.Run(() => PerformHeaderReadAsync(queuedHeader, cancellationToken), cancellationToken);
-							currentTasks.Add((queuedHeader.BlobId, task));
-							continue;
-						}
-
-						// Try to start another packet read
-						QueuedPacket? queuedPacket = _queuedPackets.FirstOrDefault(x => !currentBlobIds.Contains(x.BlobId));
-						if (queuedPacket != null)
-						{
-							Task task = Task.Run(() => PerformPacketReadAsync(queuedPacket, cancellationToken), cancellationToken);
-							currentTasks.Add((queuedPacket.BlobId, task));
-							continue;
-						}
-
-						// If we're not waiting for anything else and there are no more requests, end the task thread.
-						if (currentTasks.Count == 0)
-						{
-							_readTask = null;
-							return;
-						}
-
-						// Break out of the loop
-						break;
-					}
-				}
-
-				// Wait for any read task to complete
-				await Task.WhenAny(currentTasks.Select(x => x.Task));
-
-				// Remove any tasks which are complete
-				for (int idx = 0; idx < currentTasks.Count; idx++)
-				{
-					Task task = currentTasks[idx].Task;
-					if (task.IsCompleted)
-					{
-						if (task.Exception != null)
-						{
-							_logger.LogError(task.Exception, "Exception while reading from blob {BlobId}.", currentTasks[idx].BlobId);
-						}
-						currentTasks.RemoveAt(idx--);
-					}
-				}
-			}
-		}
-
-		/// <summary>
-		/// Reads a bundle header from the queue
-		/// </summary>
-		/// <param name="queuedHeader">The header to read</param>
-		/// <param name="cancellationToken">Cancellation token for the operation</param>
-		async Task PerformHeaderReadAsync(QueuedHeader queuedHeader, CancellationToken cancellationToken)
-		{
-			int prefetchSize = DefaultFetchSize;
-			for (; ; )
-			{
-				using (IMemoryOwner<byte> owner = MemoryPool<byte>.Shared.Rent(prefetchSize))
-				{
-					// Read the prefetch size from the blob
-					Memory<byte> memory = owner.Memory.Slice(0, prefetchSize);
-					memory = await ReadBlobRangeAsync(queuedHeader.Blob, 0, memory, cancellationToken);
-
-					// Make sure it's large enough to hold the header
-					int headerSize = BundleHeader.ReadPrelude(memory);
-					if (headerSize > prefetchSize)
-					{
-						prefetchSize = headerSize;
-						continue;
-					}
-
-					// Parse the header and construct the bundle info from it
-					BundleHeader header = new BundleHeader(new MemoryReader(memory));
-					BundleInfo bundleInfo = new BundleInfo(queuedHeader.Blob, header, headerSize);
-
-					// Also add any encoded packets we prefetched
-					List<ReadOnlyMemory<byte>> packets = new List<ReadOnlyMemory<byte>>();
-					for (int packetOffset = headerSize; packets.Count < header.Packets.Count;)
-					{
-						int packetIdx = packets.Count;
-
-						int packetLength = header.Packets[packetIdx].EncodedLength;
-						if (packetOffset + packetLength > memory.Length)
-						{
-							break;
-						}
-
-						ReadOnlyMemory<byte> packetData = memory.Slice(packetOffset, packetLength).ToArray();
-						AddToCache(GetEncodedPacketCacheKey(queuedHeader.Blob.BlobId, packetIdx), packetData, packetData.Length);
-						packets.Add(packetData);
-						packetOffset += packetLength;
-					}
-
-					// Add the info to the cache
-					string cacheKey = GetBundleInfoCacheKey(queuedHeader.Blob.BlobId);
-					AddToCache(cacheKey, bundleInfo, headerSize);
-
-					// Update the task
-					queuedHeader.CompletionSource.SetResult(bundleInfo);
-
-					// Remove it from the queue
-					lock (_queueLock)
-					{
-						_queuedHeaders.Remove(queuedHeader);
-					}
-					break;
-				}
-			}
-		}
-
-		/// <summary>
-		/// Reads a packet from storage
-		/// </summary>
-		/// <param name="queuedPacket">The packet to read</param>
-		/// <param name="cancellationToken">Cancellation token for the operation</param>
-		async Task PerformPacketReadAsync(QueuedPacket queuedPacket, CancellationToken cancellationToken)
-		{
-			BundleInfo bundleInfo = queuedPacket.Bundle;
-			int minPacketIdx = queuedPacket.PacketIdx;
-
-			BundlePacket packet = bundleInfo.Header.Packets[minPacketIdx];
-			int readLength = packet.EncodedLength;
-
-			int maxPacketIdx = minPacketIdx + 1;
-			for (; maxPacketIdx < bundleInfo.Header.Packets.Count; maxPacketIdx++)
-			{
-				int nextReadLength = readLength + bundleInfo.Header.Packets[maxPacketIdx].EncodedLength;
-				if (nextReadLength > DefaultFetchSize)
-				{
-					break;
-				}
-				readLength = nextReadLength;
-			}
-
-			using (IMemoryOwner<byte> owner = MemoryPool<byte>.Shared.Rent(readLength))
-			{
-				Memory<byte> buffer = owner.Memory.Slice(0, readLength);
-				buffer = await ReadBlobRangeAsync(bundleInfo.Locator, bundleInfo.PacketOffsets[minPacketIdx], buffer, cancellationToken);
-
-				// Copy all the packets that have been read into separate buffers, so we can cache them individually.
-				List<ReadOnlyMemory<byte>> packets = new List<ReadOnlyMemory<byte>>();
-				for (int idx = minPacketIdx; idx < maxPacketIdx; idx++)
-				{
-					string cacheKey = GetEncodedPacketCacheKey(bundleInfo.Locator.BlobId, idx);
-					ReadOnlyMemory<byte> data = buffer.Slice(bundleInfo.PacketOffsets[idx] - bundleInfo.PacketOffsets[minPacketIdx], bundleInfo.Header.Packets[idx].EncodedLength).ToArray();
-					packets.Add(data);
-					AddToCache(cacheKey, data, data.Length);
-				}
-
-				// Find all the packets we can mark as complete
-				List<QueuedPacket> updatePackets = new List<QueuedPacket>();
-				lock (_queueLock)
-				{
-					updatePackets.AddRange(_queuedPackets.Where(x => x.BlobId == bundleInfo.Locator.BlobId && (x.PacketIdx >= minPacketIdx && x.PacketIdx < maxPacketIdx)));
-				}
-
-				// Mark them all as complete
-				foreach (QueuedPacket updatePacket in updatePackets)
-				{
-					ReadOnlyMemory<byte> data = packets[updatePacket.PacketIdx - minPacketIdx];
-					updatePacket.CompletionSource.SetResult(data);
-				}
-
-				// Remove all the completed packets from the queue
-				lock (_queueLock)
-				{
-					_queuedPackets.RemoveAll(x => updatePackets.Contains(x));
-				}
-			}
-		}
-
-		/// <summary>
-		/// Reads a bundle header from the given blob locator, or retrieves it from the cache
-		/// </summary>
-		/// <param name="locator"></param>
-		/// <param name="cancellationToken"></param>
-		/// <returns></returns>
-		async Task<BundleHeader> ReadBundleHeaderAsync(BlobLocator locator, CancellationToken cancellationToken = default)
-		{
-			BundleInfo info = await GetBundleInfoAsync(locator, cancellationToken);
-			return info.Header;
-		}
-
-		/// <summary>
-		/// Reads the header and structural metadata about the bundle
-		/// </summary>
-		/// <param name="locator">The bundle location</param>
-		/// <param name="cancellationToken">Cancellation token for the operation</param>
-		/// <returns>Information about the bundle</returns>
-		async ValueTask<BundleInfo> GetBundleInfoAsync(BlobLocator locator, CancellationToken cancellationToken = default)
-		{
-			string cacheKey = GetBundleInfoCacheKey(locator.BlobId);
-			if (_cache.TryGetValue(cacheKey, out BundleInfo bundleInfo))
-			{
-				return bundleInfo;
-			}
-
-			QueuedHeader? queuedHeader;
-			lock (_queueLock)
-			{
-				// Check the cache again inside lock scope to avoid races
-				if (_cache.TryGetValue(cacheKey, out bundleInfo))
-				{
-					return bundleInfo;
-				}
-
-				// Find or start the read
-				queuedHeader = _queuedHeaders.FirstOrDefault(x => x.Blob.BlobId == locator.BlobId);
-				if (queuedHeader == null)
-				{
-					queuedHeader = new QueuedHeader(locator);
-					_queuedHeaders.Add(queuedHeader);
-					StartReadTask();
-				}
-			}
-			return await queuedHeader.CompletionSource.Task.AbandonOnCancel(cancellationToken);
-		}
-
-		/// <summary>
-		/// Gets a decoded block from the store
-		/// </summary>
-		/// <param name="bundleInfo">Information about the bundle</param>
-		/// <param name="packetIdx">Index of the packet</param>
-		/// <param name="cancellationToken">Cancellation token for the operation</param>
-		/// <returns>The decoded data</returns>
-		async ValueTask<ReadOnlyMemory<byte>> ReadBundlePacketAsync(BundleInfo bundleInfo, int packetIdx, CancellationToken cancellationToken)
-		{
-			string decodedCacheKey = GetDecodedPacketCacheKey(bundleInfo.Locator.BlobId, packetIdx);
-			if (_cache.TryGetValue(decodedCacheKey, out ReadOnlyMemory<byte> decodedPacket))
-			{
-				return decodedPacket;
-			}
-
-			Task<ReadOnlyMemory<byte>>? decodeTask;
-			lock (_queueLock)
-			{
-				// Query the cache again, to eliminate races between cache checks and decode tasks finishing.
-				if (_cache.TryGetValue(decodedCacheKey, out decodedPacket))
-				{
-					return decodedPacket;
-				}
-
-				// Create an async task to read the data
-				if (!_decodeTasks.TryGetValue(decodedCacheKey, out decodeTask))
-				{
-					decodeTask = Task.Run(() => ReadAndDecodePacketAsync(bundleInfo, packetIdx), CancellationToken.None);
-					_decodeTasks.Add(decodedCacheKey, decodeTask);
-				}
-			}
-			return await decodeTask.AbandonOnCancel(cancellationToken);
-		}
-
-		/// <summary>
-		/// Reads and decodes a packet from a bundle
-		/// </summary>
-		/// <param name="bundleInfo">Bundle to read from</param>
-		/// <param name="packetIdx">Index of the packet to return</param>
-		/// <returns>The decoded packet data</returns>
-		async Task<ReadOnlyMemory<byte>> ReadAndDecodePacketAsync(BundleInfo bundleInfo, int packetIdx)
-		{
-			ReadOnlyMemory<byte> encodedPacket = await ReadEncodedPacketAsync(bundleInfo, packetIdx);
-
-			BundlePacket packet = bundleInfo.Header.Packets[packetIdx];
-			byte[] decodedPacket = new byte[packet.DecodedLength];
-
-			BundleData.Decompress(bundleInfo.Header.CompressionFormat, encodedPacket, decodedPacket);
-
-			string decodedCacheKey = GetDecodedPacketCacheKey(bundleInfo.Locator.BlobId, packetIdx);
-			AddToCache(decodedCacheKey, (ReadOnlyMemory<byte>)decodedPacket, decodedPacket.Length);
-
-			lock (_queueLock)
-			{
-				_decodeTasks.Remove(decodedCacheKey);
-			}
-
-			return decodedPacket;
-		}
-
-		/// <summary>
-		/// Reads an encoded packet from a bundle
-		/// </summary>
-		/// <param name="bundleInfo">Bundle to read from</param>
-		/// <param name="packetIdx">Index of the packet to return</param>
-		/// <returns>The encoded packet data</returns>
-		async ValueTask<ReadOnlyMemory<byte>> ReadEncodedPacketAsync(BundleInfo bundleInfo, int packetIdx)
-		{
-			string encodedCacheKey = GetEncodedPacketCacheKey(bundleInfo.Locator.BlobId, packetIdx);
-			if (_cache.TryGetValue(encodedCacheKey, out ReadOnlyMemory<byte> encodedPacket))
-			{
-				return encodedPacket;
-			}
-
-			QueuedPacket? queuedPacket;
-			lock (_queueLock)
-			{
-				// Query the cache again, to eliminate races between cache checks and decode tasks finishing.
-				if (_cache.TryGetValue(encodedCacheKey, out encodedPacket))
-				{
-					return encodedPacket;
-				}
-
-				// Add a read to the queue
-				queuedPacket = _queuedPackets.FirstOrDefault(x => x.Bundle.Locator.BlobId == bundleInfo.Locator.BlobId && x.PacketIdx == packetIdx);
-				if (queuedPacket == null)
-				{
-					queuedPacket = new QueuedPacket(bundleInfo, packetIdx);
-					_queuedPackets.Add(queuedPacket);
-					StartReadTask();
-				}
-			}
-			return await queuedPacket.CompletionSource.Task;
-		}
-
-		#endregion
-
-		#region Nodes
-
-		class NodeReader : MemoryReader, ITreeNodeReader
-		{
-			readonly IStorageClient _store;
-			readonly IReadOnlyList<(IoHash, NodeLocator)> _refs;
-
-			public NodeReader(IStorageClient store, ReadOnlyMemory<byte> data, IReadOnlyList<(IoHash, NodeLocator)> refs)
-				: base(data)
-			{
-				_store = store;
-				_refs = refs;
-			}
-
-			public TreeNodeRefData ReadRef()
-			{
-				(IoHash hash, NodeLocator locator) = _refs[(int)this.ReadUnsignedVarInt()];
-				return new TreeNodeRefData(_store, hash, locator);
-			}
-		}
-
-		/// <inheritdoc/>
-		public async ValueTask<TNode> ReadNodeAsync<TNode>(NodeLocator locator, CancellationToken cancellationToken = default) where TNode : TreeNode
-		{
-			BundleInfo bundleInfo = await GetBundleInfoAsync(locator.Blob, cancellationToken);
-			BundleExport export = bundleInfo.Header.Exports[locator.ExportIdx];
-
-			ExportInfo exportInfo = bundleInfo.Exports[locator.ExportIdx];
-			ReadOnlyMemory<byte> packetData = await ReadBundlePacketAsync(bundleInfo, exportInfo.PacketIdx, cancellationToken);
-
-			(IoHash, NodeLocator)[] refs = new (IoHash, NodeLocator)[export.References.Count];
-			for (int idx = 0; idx < export.References.Count; idx++)
-			{
-				refs[idx] = bundleInfo.References[export.References[idx]];
-			}
-
-			ReadOnlyMemory<byte> nodeData = packetData.Slice(exportInfo.Offset, export.Length);
-			return TreeNode.Deserialize<TNode>(new NodeReader(this, nodeData, refs));
-		}
-
-		#endregion
-
 		#region Refs
 
 		/// <inheritdoc/>
 		public abstract Task DeleteRefAsync(RefName name, CancellationToken cancellationToken = default);
 
-		/// <inheritdoc/>
-		public virtual async Task<TNode?> TryReadNodeAsync<TNode>(RefName name, DateTime cacheTime = default, CancellationToken cancellationToken = default) where TNode : TreeNode
-		{
-			NodeLocator refTarget = await TryReadRefTargetAsync(name, cacheTime, cancellationToken);
-			if (!refTarget.IsValid())
-			{
-				return null;
-			}
-			return await ReadNodeAsync<TNode>(refTarget, cancellationToken);
-		}
-
 		/// <inheritdoc/>
 		public abstract Task<NodeLocator> TryReadRefTargetAsync(RefName name, DateTime cacheTime = default, CancellationToken cancellationToken = default);
 
diff --git a/Engine/Source/Programs/Shared/EpicGames.Horde/Storage/TreeNode.cs b/Engine/Source/Programs/Shared/EpicGames.Horde/Storage/TreeNode.cs
--- a/Engine/Source/Programs/Shared/EpicGames.Horde/Storage/TreeNode.cs
+++ b/Engine/Source/Programs/Shared/EpicGames.Horde/Storage/TreeNode.cs
@@ -1,10 +1,12 @@
 // Copyright Epic Games, Inc. All Rights Reserved.
 
 using System;
+using System.Collections.Concurrent;
 using System.Collections.Generic;
-using System.Diagnostics;
 using System.Reflection;
 using System.Reflection.Emit;
+using System.Threading;
+using System.Threading.Tasks;
 using EpicGames.Core;
 
 namespace EpicGames.Horde.Storage
@@ -57,55 +59,6 @@ namespace EpicGames.Horde.Storage
 		/// </summary>
 		/// <returns>References to other nodes</returns>
 		public abstract IEnumerable<TreeNodeRef> EnumerateRefs();
-
-		/// <summary>
-		/// Static instance of the serializer for a particular <see cref="TreeNode"/> type.
-		/// </summary>
-		static class SerializerInstance<T> where T : TreeNode
-		{
-			public static readonly TreeNodeSerializer<T> Serializer = CreateSerializer();
-
-			static TreeNodeSerializer<T> CreateSerializer()
-			{
-				TreeSerializerAttribute? _attribute = typeof(T).GetCustomAttribute<TreeSerializerAttribute>()!;
-				if (_attribute == null)
-				{
-					return new DefaultTreeNodeSerializer<T>();
-				}
-				else
-				{
-					return (TreeNodeSerializer<T>)Activator.CreateInstance(_attribute.Type)!;
-				}
-			}
-		}
-
-		/// <summary>
-		/// Deserialize a node from data
-		/// </summary>
-		/// <typeparam name="T"></typeparam>
-		/// <param name="reader">Reader to deserialize data from</param>
-		/// <returns></returns>
-		public static T Deserialize<T>(ITreeNodeReader reader) where T : TreeNode
-		{
-			T node = SerializerInstance<T>.Serializer.Deserialize(reader);
-			foreach (TreeNodeRef nodeRef in node.EnumerateRefs())
-			{
-				Debug.Assert(nodeRef._owner == null);
-				nodeRef._owner = node;
-			}
-			return node;
-		}
-	}
-
-	/// <summary>
-	/// Reader for tree nodes
-	/// </summary>
-	public interface ITreeNodeReader : IMemoryReader
-	{
-		/// <summary>
-		/// Reads a reference to another node
-		/// </summary>
-		TreeNodeRefData ReadRef();
 	}
 
 	/// <summary>
@@ -121,71 +74,29 @@ namespace EpicGames.Horde.Storage
 	}
 
 	/// <summary>
-	/// Factory class for deserializing node types
-	/// </summary>
-	/// <typeparam name="T">The type of node returned</typeparam>
-	public abstract class TreeNodeSerializer<T> where T : TreeNode
-	{
-		/// <summary>
-		/// Deserializes data from the given data
-		/// </summary>
-		/// <param name="reader">Reader to deserialize data from</param>
-		/// <returns>New node parsed from the data</returns>
-		public abstract T Deserialize(ITreeNodeReader reader);
-	}
-
-	/// <summary>
-	/// Default implementation of <see cref="TreeNodeSerializer{T}"/> which calls a constructor taking an <see cref="ITreeNodeReader"/> argument.
+	/// Attribute used to define a factory for a particular node type
 	/// </summary>
-	/// <typeparam name="T"></typeparam>
-	public class DefaultTreeNodeSerializer<T> : TreeNodeSerializer<T> where T : TreeNode
+	[AttributeUsage(AttributeTargets.Class)]
+	public sealed class TreeNodeAttribute : Attribute
 	{
-		static readonly Type[] s_signature = new[] { typeof(ITreeNodeReader) };
-
-		readonly Func<ITreeNodeReader, T> _constructor;
-
 		/// <summary>
-		/// Constructor
+		/// Name of the type to store in the bundle header
 		/// </summary>
-		public DefaultTreeNodeSerializer()
-		{
-			Type type = typeof(T);
-
-			ConstructorInfo? constructorInfo = type.GetConstructor(s_signature);
-			if (constructorInfo == null)
-			{
-				throw new InvalidOperationException($"Type {type.Name} does not have a constructor taking an {typeof(ITreeNodeReader).Name} as parameter.");
-			}
-
-			DynamicMethod method = new DynamicMethod($"Create_{type.Name}", type, s_signature, true);
-
-			ILGenerator generator = method.GetILGenerator();
-			generator.Emit(OpCodes.Ldarg_0);
-			generator.Emit(OpCodes.Newobj, constructorInfo);
-			generator.Emit(OpCodes.Ret);
-
-			_constructor = (Func<ITreeNodeReader, T>)method.CreateDelegate(typeof(Func<ITreeNodeReader, T>));
-		}
-
-		/// <inheritdoc/>
-		public override T Deserialize(ITreeNodeReader reader) => (T)_constructor(reader);
-	}
+		public string Guid { get; }
 
-	/// <summary>
-	/// Attribute used to define a factory for a particular node type
-	/// </summary>
-	[AttributeUsage(AttributeTargets.Class)]
-	public sealed class TreeSerializerAttribute : Attribute
-	{
 		/// <summary>
-		/// The factory type. Should be derived from <see cref="TreeNodeSerializer{T}"/>
+		/// Version number of the serializer
 		/// </summary>
-		public Type Type { get; }
+		public int Version { get; }
 
 		/// <summary>
 		/// Constructor
 		/// </summary>
-		public TreeSerializerAttribute(Type type) => Type = type;
+		public TreeNodeAttribute(string guid, int version = 1)
+		{
+			Guid = guid;
+			Version = version;
+		}
 	}
 
 	/// <summary>
@@ -203,5 +114,54 @@ namespace EpicGames.Horde.Storage
 		{
 			return new TreeNodeRef<T>(reader);
 		}
+
+		/// <summary>
+		/// Writes a node to storage
+		/// </summary>
+		/// <param name="store">Store instance to write to</param>
+		/// <param name="name">Name of the ref containing this node</param>
+		/// <param name="node">Node to be written</param>
+		/// <param name="options">Options for the node writer</param>
+		/// <param name="prefix">Prefix for uploaded blobs</param>
+		/// <param name="refOptions">Options for the ref</param>
+		/// <param name="cancellationToken">Cancellation token for the operation</param>
+		/// <returns>Location of node targetted by the ref</returns>
+		public static async Task<NodeLocator> WriteNodeAsync(this IStorageClient store, RefName name, TreeNode node, TreeOptions? options = null, Utf8String prefix = default, RefOptions? refOptions = null, CancellationToken cancellationToken = default)
+		{
+			TreeWriter writer = new TreeWriter(store, options, prefix.IsEmpty ? name.Text : prefix);
+			return await writer.WriteRefAsync(name, node, refOptions, cancellationToken);
+		}
+
+		/// <summary>
+		/// Cache of constructed <see cref="BundleType"/> instances.
+		/// </summary>
+		static readonly ConcurrentDictionary<Type, BundleType> s_typeToBundleType = new ConcurrentDictionary<Type, BundleType>();
+
+		/// <summary>
+		/// Gets the bundle type object for a particular node
+		/// </summary>
+		/// <param name="node"></param>
+		/// <returns></returns>
+		public static BundleType GetBundleType(this TreeNode node) => GetBundleType(node.GetType());
+
+		/// <summary>
+		/// Gets the <see cref="BundleType"/> instance for a particular node type
+		/// </summary>
+		/// <param name="type"></param>
+		/// <returns>Bundle</returns>
+		public static BundleType GetBundleType(Type type)
+		{
+			BundleType? bundleType;
+			if (!s_typeToBundleType.TryGetValue(type, out bundleType))
+			{
+				TreeNodeAttribute? attribute = type.GetCustomAttribute<TreeNodeAttribute>();
+				if (attribute == null)
+				{
+					throw new InvalidOperationException($"Missing {nameof(TreeNodeAttribute)} from type {type.Name}");
+				}
+				bundleType = s_typeToBundleType.GetOrAdd(type, new BundleType(Guid.Parse(attribute.Guid), attribute.Version));
+			}
+			return bundleType;
+		}
 	}
 }
diff --git a/Engine/Source/Programs/Shared/EpicGames.Horde/Storage/TreeNodeRef.cs b/Engine/Source/Programs/Shared/EpicGames.Horde/Storage/TreeNodeRef.cs
--- a/Engine/Source/Programs/Shared/EpicGames.Horde/Storage/TreeNodeRef.cs
+++ b/Engine/Source/Programs/Shared/EpicGames.Horde/Storage/TreeNodeRef.cs
@@ -18,11 +18,6 @@ namespace EpicGames.Horde.Storage
 		/// </summary>
 		internal TreeNode? _owner;
 
-		/// <summary>
-		/// Store containing the node data. May be null for nodes in memory.
-		/// </summary>
-		public IStorageClient? Store { get; private set; }
-
 		/// <summary>
 		/// Hash of the referenced node. Invalid for nodes in memory.
 		/// </summary>
@@ -89,15 +84,12 @@ namespace EpicGames.Horde.Storage
 		/// Creates a reference to a node with the given hash
 		/// </summary>
 		/// <param name="owner">Node which owns the ref</param>
-		/// <param name="store">Store to fetch this node from</param>
 		/// <param name="hash">Hash of the referenced node</param>
 		/// <param name="locator">Locator for the node</param>
-		internal TreeNodeRef(TreeNode owner, IStorageClient store, IoHash hash, NodeLocator locator)
+		internal TreeNodeRef(TreeNode owner, IoHash hash, NodeLocator locator)
 		{
-			Debug.Assert(store != null);
 			_owner = owner;
 
-			Store = store;
 			Hash = hash;
 			Locator = locator;
 
@@ -111,7 +103,6 @@ namespace EpicGames.Horde.Storage
 		public TreeNodeRef(ITreeNodeReader reader)
 		{
 			TreeNodeRefData data = reader.ReadRef();
-			Store = data.Store;
 			Hash = data.Hash;
 			Locator = data.Locator;
 
@@ -132,7 +123,6 @@ namespace EpicGames.Horde.Storage
 		{
 			Debug.Assert(Target != null);
 
-			Store = null;
 			Hash = default;
 			Locator = default;
 
@@ -149,14 +139,12 @@ namespace EpicGames.Horde.Storage
 		/// <summary>
 		/// Update the reference to refer to a location in storage.
 		/// </summary>
-		/// <param name="store">The storage client</param>
 		/// <param name="hash">Hash of the node</param>
 		/// <param name="locator">Location of the node</param>
-		internal void MarkAsWritten(IStorageClient store, IoHash hash, NodeLocator locator)
+		internal void MarkAsWritten(IoHash hash, NodeLocator locator)
 		{
 			if (hash == Hash)
 			{
-				Store = store;
 				Locator = locator;
 				_dirty = false;
 			}
@@ -207,11 +195,10 @@ namespace EpicGames.Horde.Storage
 		/// Constructor
 		/// </summary>
 		/// <param name="owner">Node which owns the ref</param>
-		/// <param name="store">Storage client containing the node data</param>
 		/// <param name="hash">Hash of the referenced node</param>
 		/// <param name="locator">Locator for the node</param>
-		internal TreeNodeRef(TreeNode owner, IStorageClient store, IoHash hash, NodeLocator locator) 
-			: base(owner, store, hash, locator)
+		internal TreeNodeRef(TreeNode owner, IoHash hash, NodeLocator locator) 
+			: base(owner, hash, locator)
 		{
 		}
 
@@ -227,13 +214,14 @@ namespace EpicGames.Horde.Storage
 		/// <summary>
 		/// Resolve this reference to a concrete node
 		/// </summary>
+		/// <param name="reader">Reader to use for expanding this ref</param>
 		/// <param name="cancellationToken">Cancellation token for the operation</param>
 		/// <returns></returns>
-		public async ValueTask<T> ExpandAsync(CancellationToken cancellationToken = default)
+		public async ValueTask<T> ExpandAsync(TreeReader reader, CancellationToken cancellationToken = default)
 		{
 			if (base.Target == null)
 			{
-				base.Target = await Store!.ReadNodeAsync<T>(Locator, cancellationToken);
+				base.Target = await reader.ReadNodeAsync<T>(Locator, cancellationToken);
 				Target!.IncomingRef = this;
 			}
 			return Target!;
@@ -242,16 +230,17 @@ namespace EpicGames.Horde.Storage
 		/// <summary>
 		/// Resolve this reference to a concrete node
 		/// </summary>
+		/// <param name="reader">Reader to use for expanding this ref</param>
 		/// <param name="cancellationToken">Cancellation token for the operation</param>
 		/// <returns></returns>
-		public async ValueTask<T> ExpandCopyAsync(CancellationToken cancellationToken = default)
+		public async ValueTask<T> ExpandCopyAsync(TreeReader reader, CancellationToken cancellationToken = default)
 		{
-			return await Store!.ReadNodeAsync<T>(Locator, cancellationToken);
+			return await reader.ReadNodeAsync<T>(Locator, cancellationToken);
 		}
 	}
 
 	/// <summary>
 	/// Deserialized ref data
 	/// </summary>
-	public record struct TreeNodeRefData(IStorageClient Store, IoHash Hash, NodeLocator Locator);
+	public record struct TreeNodeRefData(IoHash Hash, NodeLocator Locator);
 }
diff --git a/Engine/Source/Programs/Shared/EpicGames.Horde/Storage/TreeWriter.cs b/Engine/Source/Programs/Shared/EpicGames.Horde/Storage/TreeWriter.cs
--- a/Engine/Source/Programs/Shared/EpicGames.Horde/Storage/TreeWriter.cs
+++ b/Engine/Source/Programs/Shared/EpicGames.Horde/Storage/TreeWriter.cs
@@ -2,12 +2,14 @@
 
 using System;
 using System.Buffers;
+using System.Collections.Concurrent;
 using System.Collections.Generic;
 using System.Diagnostics;
 using System.Linq;
 using System.Threading;
 using System.Threading.Tasks;
 using EpicGames.Core;
+using EpicGames.Horde.Compute.Impl;
 
 namespace EpicGames.Horde.Storage
 {
@@ -39,13 +41,15 @@ namespace EpicGames.Horde.Storage
 	{
 		class NodeInfo
 		{
+			public readonly BundleType Type;
 			public readonly IoHash Hash;
 			public readonly int Length;
 			public readonly IReadOnlyList<IoHash> RefHashes;
 			public readonly TreeNodeRef NodeRef;
 
-			public NodeInfo(IoHash hash, int length, IReadOnlyList<IoHash> refs, TreeNodeRef nodeRef)
+			public NodeInfo(BundleType type, IoHash hash, int length, IReadOnlyList<IoHash> refs, TreeNodeRef nodeRef)
 			{
+				Type = type;
 				Hash = hash;
 				Length = length;
 				RefHashes = refs;
@@ -205,14 +209,14 @@ namespace EpicGames.Horde.Storage
 			nodeRef.MarkAsPendingWrite(hash);
 
 			// Check if we're already tracking a node with the same hash
-			NodeInfo nodeInfo = new NodeInfo(hash, (int)data.Length, nextHashes, nodeRef);
+			NodeInfo nodeInfo = new NodeInfo(nodeRef.Target.GetBundleType(), hash, (int)data.Length, nextHashes, nodeRef);
 			if (_hashToNode.TryGetValue(hash, out NodeLocator locator))
 			{
 				// If the node is in the lookup but not currently valid, it's already queued for writing. Add this ref to the list of refs that needs fixing up,
 				// so we can update it after flushing.
 				if (locator.IsValid())
 				{
-					nodeRef.MarkAsWritten(_store, hash, locator);
+					nodeRef.MarkAsWritten(hash, locator);
 				}
 				else
 				{
@@ -343,18 +347,30 @@ namespace EpicGames.Horde.Storage
 					imports.Add(new BundleImport(blobLocator, importedNodes));
 				}
 
+				// List of types in the bundle
+				List<BundleType> types = new List<BundleType>();
+				Dictionary<BundleType, int> typeToIndex = new Dictionary<BundleType, int>();
+
 				// Create the export list
 				List<BundleExport> exports = new List<BundleExport>(_queue.Count);
 				foreach (NodeInfo nodeInfo in _queue)
 				{
+					int typeIdx;
+					if (!typeToIndex.TryGetValue(nodeInfo.Type, out typeIdx))
+					{
+						typeIdx = types.Count;
+						typeToIndex.Add(nodeInfo.Type, typeIdx);
+						types.Add(nodeInfo.Type);
+					}
+
 					int[] references = nodeInfo.RefHashes.Select(x => nodeToIndex[x]).ToArray();
-					BundleExport export = new BundleExport(nodeInfo.Hash, nodeInfo.Length, references);
+					BundleExport export = new BundleExport(typeIdx, nodeInfo.Hash, nodeInfo.Length, references);
 					nodeToIndex.Add(nodeInfo.Hash, nodeToIndex.Count);
 					exports.Add(export);
 				}
 
 				// Create the bundle
-				BundleHeader header = new BundleHeader(_options.CompressionFormat, imports, exports, _packets.ToArray());
+				BundleHeader header = new BundleHeader(_options.CompressionFormat, types, imports, exports, _packets.ToArray());
 				Bundle bundle = new Bundle(header, _encodedPackets.ToArray());
 
 				// Write the bundle
@@ -364,7 +380,7 @@ namespace EpicGames.Horde.Storage
 					NodeLocator nodeLocator = new NodeLocator(locator, idx);
 
 					NodeInfo nodeInfo = _queue[idx];
-					nodeInfo.NodeRef.MarkAsWritten(_store, nodeInfo.Hash, nodeLocator);
+					nodeInfo.NodeRef.MarkAsWritten(nodeInfo.Hash, nodeLocator);
 
 					_hashToNode[nodeInfo.Hash] = nodeLocator;
 				}
@@ -376,7 +392,7 @@ namespace EpicGames.Horde.Storage
 					NodeInfo nodeInfo = _secondaryQueue[refIdx];
 					if (_hashToNode.TryGetValue(nodeInfo.Hash, out NodeLocator refLocator))
 					{
-						nodeInfo.NodeRef.MarkAsWritten(_store, nodeInfo.Hash, refLocator);
+						nodeInfo.NodeRef.MarkAsWritten(nodeInfo.Hash, refLocator);
 					}
 				}
 				_secondaryQueue.RemoveRange(0, refIdx);
