Horde: Missing file.

diff --git a/Engine/Source/Programs/Shared/EpicGames.Horde/Storage/TreeReader.cs b/Engine/Source/Programs/Shared/EpicGames.Horde/Storage/TreeReader.cs
new file mode 100644
--- /dev/null
+++ b/Engine/Source/Programs/Shared/EpicGames.Horde/Storage/TreeReader.cs
@@ -0,0 +1,724 @@
+// Copyright Epic Games, Inc. All Rights Reserved.
+
+using System;
+using System.Buffers;
+using System.Collections.Generic;
+using System.Diagnostics;
+using System.Linq;
+using System.Reflection.Emit;
+using System.Reflection;
+using System.Threading;
+using System.Threading.Tasks;
+using EpicGames.Core;
+using Microsoft.Extensions.Caching.Memory;
+using Microsoft.Extensions.Logging;
+using System.Collections.Concurrent;
+using Microsoft.CodeAnalysis;
+using System.Reflection.PortableExecutable;
+
+namespace EpicGames.Horde.Storage
+{
+	/// <summary>
+	/// Reader for tree nodes
+	/// </summary>
+	public interface ITreeNodeReader : IMemoryReader
+	{
+		/// <summary>
+		/// Version of the current node, as specified via <see cref="TreeNodeAttribute"/>
+		/// </summary>
+		int Version { get; }
+
+		/// <summary>
+		/// Reads a reference to another node
+		/// </summary>
+		TreeNodeRefData ReadRef();
+	}
+
+	/// <summary>
+	/// Options for configuring a bundle serializer
+	/// </summary>
+	public class TreeReaderOptions
+	{
+		/// <summary>
+		/// Known node types. Each node type should have a <see cref="TreeNodeAttribute"/> indicating the guid and latest supported version number.
+		/// </summary>
+		public List<Type> Types { get; } = new List<Type>
+		{
+			typeof(Nodes.DirectoryNode),
+			typeof(Nodes.LeafFileNode),
+			typeof(Nodes.InteriorFileNode)
+		};
+
+		/// <summary>
+		/// Constructor
+		/// </summary>
+		/// <param name="types"></param>
+		public TreeReaderOptions(params Type[] types)
+		{
+			Types.AddRange(types);
+		}
+	}
+
+	/// <summary>
+	/// Writes nodes from bundles in an <see cref="IStorageClient"/> instance.
+	/// </summary>
+	public class TreeReader
+	{
+		/// <summary>
+		/// Describes a bundle export
+		/// </summary>
+		/// <param name="PacketIdx">Index of the packet within a bundle</param>
+		/// <param name="Offset">Offset within the packet</param>
+		record struct ExportInfo(int PacketIdx, int Offset);
+
+		/// <summary>
+		/// Information about a known type that can be deserialized from a bundle
+		/// </summary>
+		class TypeInfo
+		{
+			static readonly ConcurrentDictionary<Type, TypeInfo> s_cachedTypeInfo = new ConcurrentDictionary<Type, TypeInfo>();
+
+			public Type Type { get; }
+			public BundleType BundleType { get; }
+			public Func<ITreeNodeReader, TreeNode> Deserialize { get; }
+
+			private TypeInfo(Type type, BundleType bundleType, Func<ITreeNodeReader, TreeNode> deserialize)
+			{
+				Type = type;
+				BundleType = bundleType;
+				Deserialize = deserialize;
+			}
+
+			public static TypeInfo Create(Type type)
+			{
+				TypeInfo? typeInfo;
+				if (!s_cachedTypeInfo.TryGetValue(type, out typeInfo))
+				{
+					BundleType bundleType = TreeNodeExtensions.GetBundleType(type);
+
+					Type[] signature = new[] { typeof(ITreeNodeReader) };
+
+					ConstructorInfo? constructorInfo = type.GetConstructor(signature);
+					if (constructorInfo == null)
+					{
+						throw new InvalidOperationException($"Type {type.Name} does not have a constructor taking an {typeof(ITreeNodeReader).Name} as parameter.");
+					}
+
+					DynamicMethod method = new DynamicMethod($"Create_{type.Name}", type, signature, true);
+
+					ILGenerator generator = method.GetILGenerator();
+					generator.Emit(OpCodes.Ldarg_0);
+					generator.Emit(OpCodes.Newobj, constructorInfo);
+					generator.Emit(OpCodes.Ret);
+
+					Func<ITreeNodeReader, TreeNode> deserialize = (Func<ITreeNodeReader, TreeNode>)method.CreateDelegate(typeof(Func<ITreeNodeReader, TreeNode>));
+
+					typeInfo = s_cachedTypeInfo.GetOrAdd(type, new TypeInfo(type, bundleType, deserialize));
+				}
+				return typeInfo;
+			}
+		}
+
+		/// <summary>
+		/// Computed information about a bundle
+		/// </summary>
+		class BundleInfo
+		{
+			public readonly BlobLocator Locator;
+			public readonly BundleHeader Header;
+			public readonly int[] PacketOffsets;
+			public readonly ExportInfo[] Exports;
+			public readonly (IoHash, NodeLocator)[] References;
+			public readonly TypeInfo?[] Types;
+
+			public BundleInfo(BlobLocator locator, BundleHeader header, int headerLength, TypeInfo?[] types)
+			{
+				Locator = locator;
+				Header = header;
+
+				PacketOffsets = new int[header.Packets.Count];
+				Exports = new ExportInfo[header.Exports.Count];
+
+				int exportIdx = 0;
+				int packetOffset = headerLength;
+				for (int packetIdx = 0; packetIdx < header.Packets.Count; packetIdx++)
+				{
+					BundlePacket packet = header.Packets[packetIdx];
+					PacketOffsets[packetIdx] = packetOffset;
+
+					int nodeOffset = 0;
+					for (; exportIdx < header.Exports.Count && nodeOffset + header.Exports[exportIdx].Length <= packet.DecodedLength; exportIdx++)
+					{
+						Exports[exportIdx] = new ExportInfo(packetIdx, nodeOffset);
+						nodeOffset += header.Exports[exportIdx].Length;
+					}
+
+					packetOffset += packet.EncodedLength;
+				}
+
+				References = new (IoHash, NodeLocator)[header.Imports.Sum(x => x.Exports.Count) + header.Exports.Count];
+
+				int referenceIdx = 0;
+				foreach (BundleImport import in header.Imports)
+				{
+					foreach ((int importExportIdx, IoHash importExportHash) in import.Exports)
+					{
+						NodeLocator importLocator = new NodeLocator(import.Locator, importExportIdx);
+						References[referenceIdx++] = (importExportHash, importLocator);
+					}
+				}
+				for (int idx = 0; idx < header.Exports.Count; idx++)
+				{
+					References[referenceIdx++] = (header.Exports[idx].Hash, new NodeLocator(locator, idx));
+				}
+
+				Types = types;
+			}
+		}
+
+		/// <summary>
+		/// Bundle header queued to be read
+		/// </summary>
+		class QueuedHeader
+		{
+			public readonly BlobLocator Blob;
+			public readonly TaskCompletionSource<BundleInfo> CompletionSource = new TaskCompletionSource<BundleInfo>();
+
+			public BlobId BlobId => Blob.BlobId;
+
+			public QueuedHeader(BlobLocator blob)
+			{
+				Blob = blob;
+			}
+		}
+
+		/// <summary>
+		/// Encoded bundle packet queued to be read
+		/// </summary>
+		class QueuedPacket
+		{
+			public readonly BundleInfo Bundle;
+			public readonly int PacketIdx;
+			public readonly TaskCompletionSource<ReadOnlyMemory<byte>> CompletionSource = new TaskCompletionSource<ReadOnlyMemory<byte>>();
+
+			public BlobId BlobId => Bundle.Locator.BlobId;
+
+			public QueuedPacket(BundleInfo bundle, int packetIdx)
+			{
+				Bundle = bundle;
+				PacketIdx = packetIdx;
+			}
+		}
+
+		/// <summary>
+		/// Implementation of <see cref="ITreeNodeReader"/>
+		/// </summary>
+		class NodeReader : MemoryReader, ITreeNodeReader
+		{
+			readonly IReadOnlyList<(IoHash, NodeLocator)> _refs;
+			readonly BundleType _type;
+
+			public NodeReader(ReadOnlyMemory<byte> data, IReadOnlyList<(IoHash, NodeLocator)> refs, BundleType type)
+				: base(data)
+			{
+				_refs = refs;
+				_type = type;
+			}
+
+			public int Version => _type.Version;
+
+			public TreeNodeRefData ReadRef()
+			{
+				(IoHash hash, NodeLocator locator) = _refs[(int)this.ReadUnsignedVarInt()];
+				return new TreeNodeRefData(hash, locator);
+			}
+		}
+
+		// Size of data to fetch by default. This is larger than the minimum request size to reduce number of reads.
+		const int DefaultFetchSize = 15 * 1024 * 1024;
+
+		readonly IStorageClient _store;
+		readonly IMemoryCache? _cache;
+		readonly Dictionary<Guid, TypeInfo> _types;
+		readonly ILogger _logger;
+
+		readonly object _queueLock = new object();
+		readonly List<QueuedHeader> _queuedHeaders = new List<QueuedHeader>();
+		readonly List<QueuedPacket> _queuedPackets = new List<QueuedPacket>();
+		readonly Dictionary<string, Task<ReadOnlyMemory<byte>>> _decodeTasks = new Dictionary<string, Task<ReadOnlyMemory<byte>>>(StringComparer.Ordinal);
+		Task? _readTask;
+
+		// Default options object, if unspecified
+		static readonly TreeReaderOptions s_defaultOptions = new TreeReaderOptions();
+
+		/// <summary>
+		/// Constructor
+		/// </summary>
+		/// <param name="store"></param>
+		/// <param name="cache">Cache for data</param>
+		/// <param name="logger">Logger for output</param>
+		public TreeReader(IStorageClient store, IMemoryCache? cache, ILogger logger)
+			: this(store, cache, s_defaultOptions, logger)
+		{
+		}
+
+		/// <summary>
+		/// Constructor
+		/// </summary>
+		/// <param name="store"></param>
+		/// <param name="cache">Cache for data</param>
+		/// <param name="options"></param>
+		/// <param name="logger">Logger for output</param>
+		public TreeReader(IStorageClient store, IMemoryCache? cache, TreeReaderOptions options, ILogger logger)
+		{
+			_store = store;
+			_cache = cache;
+			_types = options.Types.Select(x => TypeInfo.Create(x)).ToDictionary(x => x.BundleType.Guid, x => x);
+			_logger = logger;
+		}
+
+		#region Bundles
+
+		static string GetBundleInfoCacheKey(BlobId blobId) => $"bundle:{blobId}";
+		static string GetEncodedPacketCacheKey(BlobId blobId, int packetIdx) => $"encoded-packet:{blobId}#{packetIdx}";
+		static string GetDecodedPacketCacheKey(BlobId blobId, int packetIdx) => $"decoded-packet:{blobId}#{packetIdx}";
+
+		/// <summary>
+		/// Adds an object to the storage cache
+		/// </summary>
+		/// <param name="cacheKey">Key for the item</param>
+		/// <param name="value">Value to add</param>
+		/// <param name="size">Size of the value</param>
+		void AddToCache(string cacheKey, object value, int size)
+		{
+			if (_cache != null)
+			{
+				using (ICacheEntry entry = _cache.CreateEntry(cacheKey))
+				{
+					entry.SetValue(value);
+					entry.SetSize(size);
+				}
+			}
+		}
+
+		/// <summary>
+		/// Starts the background task for reading data from the store
+		/// </summary>
+		void StartReadTask()
+		{
+			if (_readTask == null)
+			{
+				_readTask = Task.Run(() => ServiceReadQueueAsync(CancellationToken.None));
+			}
+		}
+
+		/// <summary>
+		/// Dispatches requests in the read queue
+		/// </summary>
+		/// <param name="cancellationToken">Cancellation token for the background task</param>
+		async Task ServiceReadQueueAsync(CancellationToken cancellationToken)
+		{
+			const int MaxConcurrentReads = 4;
+
+			List<(BlobId BlobId, Task Task)> currentTasks = new List<(BlobId, Task)>();
+			for (; ; )
+			{
+				// Start any new reads
+				lock (_queueLock)
+				{
+					while (currentTasks.Count < MaxConcurrentReads)
+					{
+						HashSet<BlobId> currentBlobIds = new HashSet<BlobId>(currentTasks.Select(x => x.BlobId));
+
+						// Try to start another header read
+						QueuedHeader? queuedHeader = _queuedHeaders.FirstOrDefault(x => !currentBlobIds.Contains(x.BlobId));
+						if (queuedHeader != null)
+						{
+							Task task = Task.Run(() => PerformHeaderReadAsync(queuedHeader, cancellationToken), cancellationToken);
+							currentTasks.Add((queuedHeader.BlobId, task));
+							continue;
+						}
+
+						// Try to start another packet read
+						QueuedPacket? queuedPacket = _queuedPackets.FirstOrDefault(x => !currentBlobIds.Contains(x.BlobId));
+						if (queuedPacket != null)
+						{
+							Task task = Task.Run(() => PerformPacketReadAsync(queuedPacket, cancellationToken), cancellationToken);
+							currentTasks.Add((queuedPacket.BlobId, task));
+							continue;
+						}
+
+						// If we're not waiting for anything else and there are no more requests, end the task thread.
+						if (currentTasks.Count == 0)
+						{
+							_readTask = null;
+							return;
+						}
+
+						// Break out of the loop
+						break;
+					}
+				}
+
+				// Wait for any read task to complete
+				await Task.WhenAny(currentTasks.Select(x => x.Task));
+
+				// Remove any tasks which are complete
+				for (int idx = 0; idx < currentTasks.Count; idx++)
+				{
+					Task task = currentTasks[idx].Task;
+					if (task.IsCompleted)
+					{
+						if (task.Exception != null)
+						{
+							_logger.LogError(task.Exception, "Exception while reading from blob {BlobId}.", currentTasks[idx].BlobId);
+						}
+						currentTasks.RemoveAt(idx--);
+					}
+				}
+			}
+		}
+
+		/// <summary>
+		/// Reads a bundle header from the queue
+		/// </summary>
+		/// <param name="queuedHeader">The header to read</param>
+		/// <param name="cancellationToken">Cancellation token for the operation</param>
+		async Task PerformHeaderReadAsync(QueuedHeader queuedHeader, CancellationToken cancellationToken)
+		{
+			int prefetchSize = DefaultFetchSize;
+			for (; ; )
+			{
+				using (IMemoryOwner<byte> owner = MemoryPool<byte>.Shared.Rent(prefetchSize))
+				{
+					// Read the prefetch size from the blob
+					Memory<byte> memory = owner.Memory.Slice(0, prefetchSize);
+					memory = await _store.ReadBlobRangeAsync(queuedHeader.Blob, 0, memory, cancellationToken);
+
+					// Make sure it's large enough to hold the header
+					int headerSize = BundleHeader.ReadPrelude(memory);
+					if (headerSize > prefetchSize)
+					{
+						prefetchSize = headerSize;
+						continue;
+					}
+
+					// Parse the header and construct the bundle info from it
+					BundleHeader header = new BundleHeader(new MemoryReader(memory));
+
+					// Get the types within this bundle
+					TypeInfo?[] types = new TypeInfo?[header.Types.Count];
+					for (int idx = 0; idx < header.Types.Count; idx++)
+					{
+						_types.TryGetValue(header.Types[idx].Guid, out types[idx]);
+					}
+
+					// Construct the bundle info
+					BundleInfo bundleInfo = new BundleInfo(queuedHeader.Blob, header, headerSize, types);
+
+					// Also add any encoded packets we prefetched
+					List<ReadOnlyMemory<byte>> packets = new List<ReadOnlyMemory<byte>>();
+					for (int packetOffset = headerSize; packets.Count < header.Packets.Count;)
+					{
+						int packetIdx = packets.Count;
+
+						int packetLength = header.Packets[packetIdx].EncodedLength;
+						if (packetOffset + packetLength > memory.Length)
+						{
+							break;
+						}
+
+						ReadOnlyMemory<byte> packetData = memory.Slice(packetOffset, packetLength).ToArray();
+						AddToCache(GetEncodedPacketCacheKey(queuedHeader.Blob.BlobId, packetIdx), packetData, packetData.Length);
+						packets.Add(packetData);
+						packetOffset += packetLength;
+					}
+
+					// Add the info to the cache
+					string cacheKey = GetBundleInfoCacheKey(queuedHeader.Blob.BlobId);
+					AddToCache(cacheKey, bundleInfo, headerSize);
+
+					// Update the task
+					queuedHeader.CompletionSource.SetResult(bundleInfo);
+
+					// Remove it from the queue
+					lock (_queueLock)
+					{
+						_queuedHeaders.Remove(queuedHeader);
+					}
+					break;
+				}
+			}
+		}
+
+		/// <summary>
+		/// Reads a packet from storage
+		/// </summary>
+		/// <param name="queuedPacket">The packet to read</param>
+		/// <param name="cancellationToken">Cancellation token for the operation</param>
+		async Task PerformPacketReadAsync(QueuedPacket queuedPacket, CancellationToken cancellationToken)
+		{
+			BundleInfo bundleInfo = queuedPacket.Bundle;
+			int minPacketIdx = queuedPacket.PacketIdx;
+
+			BundlePacket packet = bundleInfo.Header.Packets[minPacketIdx];
+			int readLength = packet.EncodedLength;
+
+			int maxPacketIdx = minPacketIdx + 1;
+			for (; maxPacketIdx < bundleInfo.Header.Packets.Count; maxPacketIdx++)
+			{
+				int nextReadLength = readLength + bundleInfo.Header.Packets[maxPacketIdx].EncodedLength;
+				if (nextReadLength > DefaultFetchSize)
+				{
+					break;
+				}
+				readLength = nextReadLength;
+			}
+
+			using (IMemoryOwner<byte> owner = MemoryPool<byte>.Shared.Rent(readLength))
+			{
+				Memory<byte> buffer = owner.Memory.Slice(0, readLength);
+				buffer = await _store.ReadBlobRangeAsync(bundleInfo.Locator, bundleInfo.PacketOffsets[minPacketIdx], buffer, cancellationToken);
+
+				// Copy all the packets that have been read into separate buffers, so we can cache them individually.
+				List<ReadOnlyMemory<byte>> packets = new List<ReadOnlyMemory<byte>>();
+				for (int idx = minPacketIdx; idx < maxPacketIdx; idx++)
+				{
+					string cacheKey = GetEncodedPacketCacheKey(bundleInfo.Locator.BlobId, idx);
+					ReadOnlyMemory<byte> data = buffer.Slice(bundleInfo.PacketOffsets[idx] - bundleInfo.PacketOffsets[minPacketIdx], bundleInfo.Header.Packets[idx].EncodedLength).ToArray();
+					packets.Add(data);
+					AddToCache(cacheKey, data, data.Length);
+				}
+
+				// Find all the packets we can mark as complete
+				List<QueuedPacket> updatePackets = new List<QueuedPacket>();
+				lock (_queueLock)
+				{
+					updatePackets.AddRange(_queuedPackets.Where(x => x.BlobId == bundleInfo.Locator.BlobId && (x.PacketIdx >= minPacketIdx && x.PacketIdx < maxPacketIdx)));
+				}
+
+				// Mark them all as complete
+				foreach (QueuedPacket updatePacket in updatePackets)
+				{
+					ReadOnlyMemory<byte> data = packets[updatePacket.PacketIdx - minPacketIdx];
+					updatePacket.CompletionSource.SetResult(data);
+				}
+
+				// Remove all the completed packets from the queue
+				lock (_queueLock)
+				{
+					_queuedPackets.RemoveAll(x => updatePackets.Contains(x));
+				}
+			}
+		}
+
+		/// <summary>
+		/// Reads a bundle header from the given blob locator, or retrieves it from the cache
+		/// </summary>
+		/// <param name="locator"></param>
+		/// <param name="cancellationToken"></param>
+		/// <returns></returns>
+		async Task<BundleHeader> ReadBundleHeaderAsync(BlobLocator locator, CancellationToken cancellationToken = default)
+		{
+			BundleInfo info = await GetBundleInfoAsync(locator, cancellationToken);
+			return info.Header;
+		}
+
+		/// <summary>
+		/// Reads the header and structural metadata about the bundle
+		/// </summary>
+		/// <param name="locator">The bundle location</param>
+		/// <param name="cancellationToken">Cancellation token for the operation</param>
+		/// <returns>Information about the bundle</returns>
+		async ValueTask<BundleInfo> GetBundleInfoAsync(BlobLocator locator, CancellationToken cancellationToken = default)
+		{
+			string cacheKey = GetBundleInfoCacheKey(locator.BlobId);
+			if (_cache != null && _cache.TryGetValue(cacheKey, out BundleInfo bundleInfo))
+			{
+				return bundleInfo;
+			}
+
+			QueuedHeader? queuedHeader;
+			lock (_queueLock)
+			{
+				// Check the cache again inside lock scope to avoid races
+				if (_cache != null && _cache.TryGetValue(cacheKey, out bundleInfo))
+				{
+					return bundleInfo;
+				}
+
+				// Find or start the read
+				queuedHeader = _queuedHeaders.FirstOrDefault(x => x.Blob.BlobId == locator.BlobId);
+				if (queuedHeader == null)
+				{
+					queuedHeader = new QueuedHeader(locator);
+					_queuedHeaders.Add(queuedHeader);
+					StartReadTask();
+				}
+			}
+			return await queuedHeader.CompletionSource.Task.AbandonOnCancel(cancellationToken);
+		}
+
+		/// <summary>
+		/// Gets a decoded block from the store
+		/// </summary>
+		/// <param name="bundleInfo">Information about the bundle</param>
+		/// <param name="packetIdx">Index of the packet</param>
+		/// <param name="cancellationToken">Cancellation token for the operation</param>
+		/// <returns>The decoded data</returns>
+		async ValueTask<ReadOnlyMemory<byte>> ReadBundlePacketAsync(BundleInfo bundleInfo, int packetIdx, CancellationToken cancellationToken)
+		{
+			string decodedCacheKey = GetDecodedPacketCacheKey(bundleInfo.Locator.BlobId, packetIdx);
+			if (_cache != null && _cache.TryGetValue(decodedCacheKey, out ReadOnlyMemory<byte> decodedPacket))
+			{
+				return decodedPacket;
+			}
+
+			Task<ReadOnlyMemory<byte>>? decodeTask;
+			lock (_queueLock)
+			{
+				// Query the cache again, to eliminate races between cache checks and decode tasks finishing.
+				if (_cache != null && _cache.TryGetValue(decodedCacheKey, out decodedPacket))
+				{
+					return decodedPacket;
+				}
+
+				// Create an async task to read the data
+				if (!_decodeTasks.TryGetValue(decodedCacheKey, out decodeTask))
+				{
+					decodeTask = Task.Run(() => ReadAndDecodePacketAsync(bundleInfo, packetIdx), CancellationToken.None);
+					_decodeTasks.Add(decodedCacheKey, decodeTask);
+				}
+			}
+			return await decodeTask.AbandonOnCancel(cancellationToken);
+		}
+
+		/// <summary>
+		/// Reads and decodes a packet from a bundle
+		/// </summary>
+		/// <param name="bundleInfo">Bundle to read from</param>
+		/// <param name="packetIdx">Index of the packet to return</param>
+		/// <returns>The decoded packet data</returns>
+		async Task<ReadOnlyMemory<byte>> ReadAndDecodePacketAsync(BundleInfo bundleInfo, int packetIdx)
+		{
+			ReadOnlyMemory<byte> encodedPacket = await ReadEncodedPacketAsync(bundleInfo, packetIdx);
+
+			BundlePacket packet = bundleInfo.Header.Packets[packetIdx];
+			byte[] decodedPacket = new byte[packet.DecodedLength];
+
+			BundleData.Decompress(bundleInfo.Header.CompressionFormat, encodedPacket, decodedPacket);
+
+			string decodedCacheKey = GetDecodedPacketCacheKey(bundleInfo.Locator.BlobId, packetIdx);
+			AddToCache(decodedCacheKey, (ReadOnlyMemory<byte>)decodedPacket, decodedPacket.Length);
+
+			lock (_queueLock)
+			{
+				_decodeTasks.Remove(decodedCacheKey);
+			}
+
+			return decodedPacket;
+		}
+
+		/// <summary>
+		/// Reads an encoded packet from a bundle
+		/// </summary>
+		/// <param name="bundleInfo">Bundle to read from</param>
+		/// <param name="packetIdx">Index of the packet to return</param>
+		/// <returns>The encoded packet data</returns>
+		async ValueTask<ReadOnlyMemory<byte>> ReadEncodedPacketAsync(BundleInfo bundleInfo, int packetIdx)
+		{
+			string encodedCacheKey = GetEncodedPacketCacheKey(bundleInfo.Locator.BlobId, packetIdx);
+			if (_cache != null && _cache.TryGetValue(encodedCacheKey, out ReadOnlyMemory<byte> encodedPacket))
+			{
+				return encodedPacket;
+			}
+
+			QueuedPacket? queuedPacket;
+			lock (_queueLock)
+			{
+				// Query the cache again, to eliminate races between cache checks and decode tasks finishing.
+				if (_cache != null && _cache.TryGetValue(encodedCacheKey, out encodedPacket))
+				{
+					return encodedPacket;
+				}
+
+				// Add a read to the queue
+				queuedPacket = _queuedPackets.FirstOrDefault(x => x.Bundle.Locator.BlobId == bundleInfo.Locator.BlobId && x.PacketIdx == packetIdx);
+				if (queuedPacket == null)
+				{
+					queuedPacket = new QueuedPacket(bundleInfo, packetIdx);
+					_queuedPackets.Add(queuedPacket);
+					StartReadTask();
+				}
+			}
+			return await queuedPacket.CompletionSource.Task;
+		}
+
+		#endregion
+
+		/// <summary>
+		/// Reads a node from a bundle
+		/// </summary>
+		/// <param name="locator">Locator for the node</param>
+		/// <param name="cancellationToken">Cancellation token for the operation</param>
+		/// <returns>Node data read from the given bundle</returns>
+		public async ValueTask<TNode> ReadNodeAsync<TNode>(NodeLocator locator, CancellationToken cancellationToken = default) where TNode : TreeNode
+		{
+			BundleInfo bundleInfo = await GetBundleInfoAsync(locator.Blob, cancellationToken);
+			BundleExport export = bundleInfo.Header.Exports[locator.ExportIdx];
+
+			ExportInfo exportInfo = bundleInfo.Exports[locator.ExportIdx];
+			ReadOnlyMemory<byte> packetData = await ReadBundlePacketAsync(bundleInfo, exportInfo.PacketIdx, cancellationToken);
+
+			(IoHash, NodeLocator)[] refs = new (IoHash, NodeLocator)[export.References.Count];
+			for (int idx = 0; idx < export.References.Count; idx++)
+			{
+				refs[idx] = bundleInfo.References[export.References[idx]];
+			}
+
+			ReadOnlyMemory<byte> nodeData = packetData.Slice(exportInfo.Offset, export.Length);
+
+			TypeInfo? typeInfo = bundleInfo.Types[export.TypeIdx];
+			if (typeInfo == null)
+			{
+				throw new InvalidOperationException($"No registered serializer for type {bundleInfo.Header.Types[export.TypeIdx].Guid}");
+			}
+
+			return (TNode)typeInfo.Deserialize(new NodeReader(nodeData, refs, typeInfo.BundleType));
+		}
+
+		/// <summary>
+		/// Reads data for a ref from the store, along with the node's contents.
+		/// </summary>
+		/// <param name="name">The ref name</param>
+		/// <param name="cacheTime">Minimum coherency for any cached value to be returned</param>
+		/// <param name="cancellationToken">Cancellation token for the operation</param>
+		/// <returns>Node for the given ref, or null if it does not exist</returns>
+		public async Task<TNode?> TryReadNodeAsync<TNode>(RefName name, DateTime cacheTime = default, CancellationToken cancellationToken = default) where TNode : TreeNode
+		{
+			NodeLocator refTarget = await _store.TryReadRefTargetAsync(name, cacheTime, cancellationToken);
+			if (!refTarget.IsValid())
+			{
+				return null;
+			}
+			return await ReadNodeAsync<TNode>(refTarget, cancellationToken);
+		}
+
+		/// <summary>
+		/// Reads a ref from the store, throwing an exception if it does not exist
+		/// </summary>
+		/// <param name="name">Id for the ref</param>
+		/// <param name="cacheTime">Minimum coherency of any cached result</param>
+		/// <param name="cancellationToken">Cancellation token for the operation</param>
+		/// <returns>The blob instance</returns>
+		public async Task<TNode> ReadNodeAsync<TNode>(RefName name, DateTime cacheTime = default, CancellationToken cancellationToken = default) where TNode : TreeNode
+		{
+			TNode? refValue = await TryReadNodeAsync<TNode>(name, cacheTime, cancellationToken);
+			if (refValue == null)
+			{
+				throw new RefNameNotFoundException(name);
+			}
+			return refValue;
+		}
+	}
+}
