Horde: Add a node type for storing opaque compact-binary objects without transcoding them into some other form.

diff --git a/Engine/Source/Programs/Shared/EpicGames.Horde/Storage/Nodes/CbNode.cs b/Engine/Source/Programs/Shared/EpicGames.Horde/Storage/Nodes/CbNode.cs
new file mode 100644
--- /dev/null
+++ b/Engine/Source/Programs/Shared/EpicGames.Horde/Storage/Nodes/CbNode.cs
@@ -0,0 +1,59 @@
+// Copyright Epic Games, Inc. All Rights Reserved.
+
+using EpicGames.Core;
+using EpicGames.Serialization;
+using System.Collections.Generic;
+using System.Linq;
+
+namespace EpicGames.Horde.Storage.Nodes
+{
+	/// <summary>
+	/// A node containing arbitrary compact binary data
+	/// </summary>
+	[TreeNode("{34A0793F-8364-42F4-8632-98A71C843229}", 1)]
+	public class CbNode : TreeNode
+	{
+		/// <summary>
+		/// The compact binary object
+		/// </summary>
+		public CbObject Object { get; set; }
+
+		/// <summary>
+		/// Imported nodes
+		/// </summary>
+		public Dictionary<IoHash, NodeLocator> References { get; }
+
+		/// <summary>
+		/// Constructor
+		/// </summary>
+		/// <param name="obj">The compact binary object</param>
+		/// <param name="references">Map of attachment hash to node locator</param>
+		public CbNode(CbObject obj, Dictionary<IoHash, NodeLocator> references)
+		{
+			Object = obj;
+			References = references;
+		}
+
+		/// <summary>
+		/// Deserialization constructor
+		/// </summary>
+		/// <param name="reader">Reader to deserialize from</param>
+		public CbNode(ITreeNodeReader reader)
+		{
+			Object = new CbObject(reader.ReadFixedLengthBytes(reader.Length));
+			References = reader.References.ToDictionary(x => x.Hash, x => x.Locator);
+		}
+
+		/// <inheritdoc/>
+		public override void Serialize(ITreeNodeWriter writer)
+		{
+			writer.WriteFixedLengthBytes(Object.GetView().Span);
+		}
+
+		/// <inheritdoc/>
+		public override IEnumerable<TreeNodeRef> EnumerateRefs()
+		{
+			return References.Select(x => new TreeNodeRef(this, x.Key, x.Value));
+		}
+	}
+}
diff --git a/Engine/Source/Programs/Shared/EpicGames.Horde/Storage/TreeReader.cs b/Engine/Source/Programs/Shared/EpicGames.Horde/Storage/TreeReader.cs
--- a/Engine/Source/Programs/Shared/EpicGames.Horde/Storage/TreeReader.cs
+++ b/Engine/Source/Programs/Shared/EpicGames.Horde/Storage/TreeReader.cs
@@ -3,7 +3,6 @@
 using System;
 using System.Buffers;
 using System.Collections.Generic;
-using System.Diagnostics;
 using System.Linq;
 using System.Reflection.Emit;
 using System.Reflection;
@@ -14,7 +13,6 @@ using Microsoft.Extensions.Caching.Memory;
 using Microsoft.Extensions.Logging;
 using System.Collections.Concurrent;
 using Microsoft.CodeAnalysis;
-using System.Reflection.PortableExecutable;
 
 namespace EpicGames.Horde.Storage
 {
@@ -28,6 +26,16 @@ namespace EpicGames.Horde.Storage
 		/// </summary>
 		int Version { get; }
 
+		/// <summary>
+		/// Total length of the data in this node
+		/// </summary>
+		int Length { get; }
+
+		/// <summary>
+		/// References to other nodes
+		/// </summary>
+		IReadOnlyList<TreeNodeRefData> References { get; }
+
 		/// <summary>
 		/// Reads a reference to another node
 		/// </summary>
@@ -128,7 +136,7 @@ namespace EpicGames.Horde.Storage
 			public readonly BundleHeader Header;
 			public readonly int[] PacketOffsets;
 			public readonly ExportInfo[] Exports;
-			public readonly (IoHash, NodeLocator)[] References;
+			public readonly TreeNodeRefData[] References;
 			public readonly TypeInfo?[] Types;
 
 			public BundleInfo(BlobLocator locator, BundleHeader header, int headerLength, TypeInfo?[] types)
@@ -156,7 +164,7 @@ namespace EpicGames.Horde.Storage
 					packetOffset += packet.EncodedLength;
 				}
 
-				References = new (IoHash, NodeLocator)[header.Imports.Sum(x => x.Exports.Count) + header.Exports.Count];
+				References = new TreeNodeRefData[header.Imports.Sum(x => x.Exports.Count) + header.Exports.Count];
 
 				int referenceIdx = 0;
 				foreach (BundleImport import in header.Imports)
@@ -164,12 +172,12 @@ namespace EpicGames.Horde.Storage
 					foreach ((int importExportIdx, IoHash importExportHash) in import.Exports)
 					{
 						NodeLocator importLocator = new NodeLocator(import.Locator, importExportIdx);
-						References[referenceIdx++] = (importExportHash, importLocator);
+						References[referenceIdx++] = new TreeNodeRefData(importExportHash, importLocator);
 					}
 				}
 				for (int idx = 0; idx < header.Exports.Count; idx++)
 				{
-					References[referenceIdx++] = (header.Exports[idx].Hash, new NodeLocator(locator, idx));
+					References[referenceIdx++] = new TreeNodeRefData(header.Exports[idx].Hash, new NodeLocator(locator, idx));
 				}
 
 				Types = types;
@@ -215,23 +223,25 @@ namespace EpicGames.Horde.Storage
 		/// </summary>
 		class NodeReader : MemoryReader, ITreeNodeReader
 		{
-			readonly IReadOnlyList<(IoHash, NodeLocator)> _refs;
+			readonly IReadOnlyList<TreeNodeRefData> _refs;
 			readonly BundleType _type;
+			readonly int _length;
 
-			public NodeReader(ReadOnlyMemory<byte> data, IReadOnlyList<(IoHash, NodeLocator)> refs, BundleType type)
+			public NodeReader(ReadOnlyMemory<byte> data, IReadOnlyList<TreeNodeRefData> refs, BundleType type)
 				: base(data)
 			{
 				_refs = refs;
 				_type = type;
+				_length = data.Length;
 			}
 
 			public int Version => _type.Version;
 
-			public TreeNodeRefData ReadRef()
-			{
-				(IoHash hash, NodeLocator locator) = _refs[(int)this.ReadUnsignedVarInt()];
-				return new TreeNodeRefData(hash, locator);
-			}
+			public int Length => _length;
+
+			public IReadOnlyList<TreeNodeRefData> References => _refs;
+
+			public TreeNodeRefData ReadRef() => _refs[(int)this.ReadUnsignedVarInt()];
 		}
 
 		// Size of data to fetch by default. This is larger than the minimum request size to reduce number of reads.
@@ -670,7 +680,7 @@ namespace EpicGames.Horde.Storage
 			ExportInfo exportInfo = bundleInfo.Exports[locator.ExportIdx];
 			ReadOnlyMemory<byte> packetData = await ReadBundlePacketAsync(bundleInfo, exportInfo.PacketIdx, cancellationToken);
 
-			(IoHash, NodeLocator)[] refs = new (IoHash, NodeLocator)[export.References.Count];
+			TreeNodeRefData[] refs = new TreeNodeRefData[export.References.Count];
 			for (int idx = 0; idx < export.References.Count; idx++)
 			{
 				refs[idx] = bundleInfo.References[export.References[idx]];
