Horde: Ensure computed hashes for serialized nodes are identical to serialized data (ie. not including a separate array of hashes).

diff --git a/Engine/Source/Programs/Shared/EpicGames.Horde.Tests/BundleStoreTests.cs b/Engine/Source/Programs/Shared/EpicGames.Horde.Tests/BundleStoreTests.cs
--- a/Engine/Source/Programs/Shared/EpicGames.Horde.Tests/BundleStoreTests.cs
+++ b/Engine/Source/Programs/Shared/EpicGames.Horde.Tests/BundleStoreTests.cs
@@ -334,7 +334,7 @@ namespace EpicGames.Horde.Tests
 
 			// Check we can read it back in
 			{
-				TreeReader reader = new TreeReader(store, null, NullLogger.Instance);
+				TreeReader reader = new TreeReader(store, cache, NullLogger.Instance);
 
 				DirectoryNode newRoot = await reader.ReadNodeAsync<DirectoryNode>(new RefName("test"));
 				await CompareTrees(reader, root, newRoot);
diff --git a/Engine/Source/Programs/Shared/EpicGames.Horde/Storage/Nodes/CbNode.cs b/Engine/Source/Programs/Shared/EpicGames.Horde/Storage/Nodes/CbNode.cs
--- a/Engine/Source/Programs/Shared/EpicGames.Horde/Storage/Nodes/CbNode.cs
+++ b/Engine/Source/Programs/Shared/EpicGames.Horde/Storage/Nodes/CbNode.cs
@@ -21,14 +21,14 @@ namespace EpicGames.Horde.Storage.Nodes
 		/// <summary>
 		/// Imported nodes
 		/// </summary>
-		public Dictionary<IoHash, NodeLocator> References { get; }
+		public IReadOnlyDictionary<IoHash, NodeLocator> References { get; }
 
 		/// <summary>
 		/// Constructor
 		/// </summary>
 		/// <param name="obj">The compact binary object</param>
 		/// <param name="references">Map of attachment hash to node locator</param>
-		public CbNode(CbObject obj, Dictionary<IoHash, NodeLocator> references)
+		public CbNode(CbObject obj, IReadOnlyDictionary<IoHash, NodeLocator> references)
 		{
 			Object = obj;
 			References = references;
@@ -41,7 +41,7 @@ namespace EpicGames.Horde.Storage.Nodes
 		public CbNode(ITreeNodeReader reader)
 		{
 			Object = new CbObject(reader.ReadFixedLengthBytes(reader.Length));
-			References = reader.References.ToDictionary(x => x.Hash, x => x.Locator);
+			References = reader.References;
 		}
 
 		/// <inheritdoc/>
diff --git a/Engine/Source/Programs/Shared/EpicGames.Horde/Storage/TreeReader.cs b/Engine/Source/Programs/Shared/EpicGames.Horde/Storage/TreeReader.cs
--- a/Engine/Source/Programs/Shared/EpicGames.Horde/Storage/TreeReader.cs
+++ b/Engine/Source/Programs/Shared/EpicGames.Horde/Storage/TreeReader.cs
@@ -34,7 +34,7 @@ namespace EpicGames.Horde.Storage
 		/// <summary>
 		/// References to other nodes
 		/// </summary>
-		IReadOnlyList<TreeNodeRefData> References { get; }
+		IReadOnlyDictionary<IoHash, NodeLocator> References { get; }
 
 		/// <summary>
 		/// Reads a reference to another node
@@ -223,11 +223,11 @@ namespace EpicGames.Horde.Storage
 		/// </summary>
 		class NodeReader : MemoryReader, ITreeNodeReader
 		{
-			readonly IReadOnlyList<TreeNodeRefData> _refs;
+			readonly IReadOnlyDictionary<IoHash, NodeLocator> _refs;
 			readonly BundleType _type;
 			readonly int _length;
 
-			public NodeReader(ReadOnlyMemory<byte> data, IReadOnlyList<TreeNodeRefData> refs, BundleType type)
+			public NodeReader(ReadOnlyMemory<byte> data, IReadOnlyDictionary<IoHash, NodeLocator> refs, BundleType type)
 				: base(data)
 			{
 				_refs = refs;
@@ -239,9 +239,13 @@ namespace EpicGames.Horde.Storage
 
 			public int Length => _length;
 
-			public IReadOnlyList<TreeNodeRefData> References => _refs;
+			public IReadOnlyDictionary<IoHash, NodeLocator> References => _refs;
 
-			public TreeNodeRefData ReadRef() => _refs[(int)this.ReadUnsignedVarInt()];
+			public TreeNodeRefData ReadRef()
+			{
+				IoHash hash = this.ReadIoHash();
+				return new TreeNodeRefData(hash, _refs[hash]);
+			}
 		}
 
 		// Size of data to fetch by default. This is larger than the minimum request size to reduce number of reads.
@@ -680,10 +684,11 @@ namespace EpicGames.Horde.Storage
 			ExportInfo exportInfo = bundleInfo.Exports[locator.ExportIdx];
 			ReadOnlyMemory<byte> packetData = await ReadBundlePacketAsync(bundleInfo, exportInfo.PacketIdx, cancellationToken);
 
-			TreeNodeRefData[] refs = new TreeNodeRefData[export.References.Count];
+			Dictionary<IoHash, NodeLocator> refs = new Dictionary<IoHash, NodeLocator>(export.References.Count);
 			for (int idx = 0; idx < export.References.Count; idx++)
 			{
-				refs[idx] = bundleInfo.References[export.References[idx]];
+				TreeNodeRefData reference = bundleInfo.References[export.References[idx]];
+				refs.Add(reference.Hash, reference.Locator);
 			}
 
 			ReadOnlyMemory<byte> nodeData = packetData.Slice(exportInfo.Offset, export.Length);
diff --git a/Engine/Source/Programs/Shared/EpicGames.Horde/Storage/TreeWriter.cs b/Engine/Source/Programs/Shared/EpicGames.Horde/Storage/TreeWriter.cs
--- a/Engine/Source/Programs/Shared/EpicGames.Horde/Storage/TreeWriter.cs
+++ b/Engine/Source/Programs/Shared/EpicGames.Horde/Storage/TreeWriter.cs
@@ -61,22 +61,23 @@ namespace EpicGames.Horde.Storage
 		{
 			public ByteArrayBuilder Data = new ByteArrayBuilder();
 
-			readonly Dictionary<TreeNodeRef, int> _refToIndex;
+			readonly IReadOnlyDictionary<TreeNodeRef, IoHash> _refToHash;
 
 			public int Length => Data.Length;
 
-			public NodeWriter(List<TreeNodeRef> nextRefs)
+			public NodeWriter(IReadOnlyDictionary<TreeNodeRef, IoHash> refToHash)
 			{
-				_refToIndex = new Dictionary<TreeNodeRef, int>(nextRefs.Count);
-				for (int idx = 0; idx < nextRefs.Count; idx++)
-				{
-					_refToIndex.TryAdd(nextRefs[idx], idx);
-				}
+				_refToHash = refToHash;
 			}
 
 			public void WriteRef(TreeNodeRef target)
 			{
-				Data.WriteUnsignedVarInt(_refToIndex[target]);
+				IoHash hash;
+				if (!_refToHash.TryGetValue(target, out hash))
+				{
+					throw new InvalidOperationException($"Cannot serialize value; was not returned by owner's {nameof(target.Target.EnumerateRefs)} method.");
+				}
+				Data.WriteIoHash(hash);
 			}
 
 			public Memory<byte> GetMemory(int minSize = 1) => Data.GetMemory(minSize);
@@ -188,28 +189,26 @@ namespace EpicGames.Horde.Storage
 				throw new InvalidOperationException("Node cannot have multiple incoming references");
 			}
 
-			// Enumerate all the outward references
-			List<TreeNodeRef> nextRefs = new List<TreeNodeRef>(node.EnumerateRefs());
-
 			// Write all the referenced nodes first
-			IoHash[] nextHashes = new IoHash[nextRefs.Count];
-			for (int idx = 0; idx < nextRefs.Count; idx++)
+			Dictionary<TreeNodeRef, IoHash> nextRefToHash = new Dictionary<TreeNodeRef, IoHash>();
+			foreach (TreeNodeRef nextRef in node.EnumerateRefs())
 			{
-				nextHashes[idx] = await WriteAsync(nextRefs[idx], cancellationToken);
-				Debug.Assert(nextHashes[idx] != IoHash.Zero);
+				IoHash nextHash = await WriteAsync(nextRef, cancellationToken);
+				Debug.Assert(nextHash != IoHash.Zero);
+				nextRefToHash.Add(nextRef, nextHash);
 			}
 
 			// Serialize the node
-			NodeWriter nodeWriter = new NodeWriter(nextRefs);
+			NodeWriter nodeWriter = new NodeWriter(nextRefToHash);
 			nodeRef.Target!.Serialize(nodeWriter);
 
 			// Get the hash for the new blob
 			ReadOnlySequence<byte> data = nodeWriter.Data.AsSequence();
-			IoHash hash = ComputeHash(data, nextHashes);
+			IoHash hash = IoHash.Compute(data);
 			nodeRef.MarkAsPendingWrite(hash);
 
 			// Check if we're already tracking a node with the same hash
-			NodeInfo nodeInfo = new NodeInfo(nodeRef.Target.GetBundleType(), hash, (int)data.Length, nextHashes, nodeRef);
+			NodeInfo nodeInfo = new NodeInfo(nodeRef.Target.GetBundleType(), hash, (int)data.Length, nextRefToHash.Values.Distinct().ToArray(), nodeRef);
 			if (_hashToNode.TryGetValue(hash, out NodeLocator locator))
 			{
 				// If the node is in the lookup but not currently valid, it's already queued for writing. Add this ref to the list of refs that needs fixing up,
@@ -404,21 +403,5 @@ namespace EpicGames.Horde.Storage
 				_currentBundleLength = 0;
 			}
 		}
-
-		static IoHash ComputeHash(ReadOnlySequence<byte> data, IoHash[] references)
-		{
-			byte[] buffer = new byte[IoHash.NumBytes * (references.Length + 1)];
-
-			Span<byte> span = buffer;
-			for (int idx = 0; idx < references.Length; idx++)
-			{
-				Debug.Assert(references[idx] != IoHash.Zero);
-				references[idx].CopyTo(span);
-				span = span[IoHash.NumBytes..];
-			}
-			IoHash.Compute(data).CopyTo(span);
-
-			return IoHash.Compute(buffer);
-		}
 	}
 }
