Horde: Move CommitNode into EpicGames.Horde, and allow reading anonymous TreeNodeRef types. Now that type information is embedded in the bundle, we no longer need to have overridden read methods for each type.

diff --git a/Engine/Source/Programs/Horde/Horde.Agent/Commands/Vcs/VcsCommand.cs b/Engine/Source/Programs/Horde/Horde.Agent/Commands/Vcs/VcsCommand.cs
--- a/Engine/Source/Programs/Horde/Horde.Agent/Commands/Vcs/VcsCommand.cs
+++ b/Engine/Source/Programs/Horde/Horde.Agent/Commands/Vcs/VcsCommand.cs
@@ -20,54 +20,6 @@ namespace Horde.Agent.Commands.Bundles
 		[CommandLine("-BaseDir=")]
 		public DirectoryReference? BaseDir { get; set; }
 
-		protected class CommitNode : TreeNode
-		{
-			public int Number { get; set; }
-			public TreeNodeRef<CommitNode>? Parent { get; set; }
-			public string Message { get; set; }
-			public TreeNodeRef<DirectoryNode> Contents { get; set; }
-
-			public CommitNode(ITreeNodeReader reader)
-			{
-				Number = reader.ReadInt32();
-				if (reader.ReadBoolean())
-				{
-					Parent = reader.ReadRef<CommitNode>();
-				}
-				Message = reader.ReadString();
-				Contents = reader.ReadRef<DirectoryNode>();
-			}
-
-			public CommitNode(int number, TreeNodeRef<CommitNode>? parent, string message, TreeNodeRef<DirectoryNode> contents)
-			{
-				Number = number;
-				Parent = parent;
-				Message = message;
-				Contents = contents;
-			}
-
-			public override void Serialize(ITreeNodeWriter writer)
-			{
-				writer.WriteInt32(Number);
-				writer.WriteBoolean(Parent != null);
-				if (Parent != null)
-				{
-					writer.WriteRef(Parent);
-				}
-				writer.WriteString(Message);
-				writer.WriteRef(Contents);
-			}
-
-			public override IEnumerable<TreeNodeRef> EnumerateRefs()
-			{
-				if (Parent != null)
-				{
-					yield return Parent;
-				}
-				yield return Contents;
-			}
-		}
-
 		protected class FileState
 		{
 			public IoHash Hash { get; set; }
@@ -645,11 +597,11 @@ namespace Horde.Agent.Commands.Bundles
 			CommitNode newTip;
 			if (tip == null)
 			{
-				newTip = new CommitNode(1, null, Message, rootRef);
+				newTip = new CommitNode(1, null, Message, DateTime.UtcNow, rootRef);
 			}
 			else
 			{
-				newTip = new CommitNode(tip.Number + 1, tipRef, Message, rootRef);
+				newTip = new CommitNode(tip.Number + 1, tipRef, Message, DateTime.UtcNow, rootRef);
 			}
 			await store.WriteNodeAsync(workspaceState.Branch, newTip);
 
diff --git a/Engine/Source/Programs/Shared/EpicGames.Horde/Storage/Nodes/CommitNode.cs b/Engine/Source/Programs/Shared/EpicGames.Horde/Storage/Nodes/CommitNode.cs
new file mode 100644
--- /dev/null
+++ b/Engine/Source/Programs/Shared/EpicGames.Horde/Storage/Nodes/CommitNode.cs
@@ -0,0 +1,104 @@
+// Copyright Epic Games, Inc. All Rights Reserved.
+
+using EpicGames.Core;
+using System;
+using System.Collections.Generic;
+
+namespace EpicGames.Horde.Storage.Nodes
+{
+	/// <summary>
+	/// A node representing commit metadata
+	/// </summary>
+	[TreeNode("{64D50724-6B22-41C0-A890-B51CD6241817}", 1)]
+	public class CommitNode : TreeNode
+	{
+		/// <summary>
+		/// The commit number
+		/// </summary>
+		public int Number { get; set; }
+
+		/// <summary>
+		/// Reference to the parent commit
+		/// </summary>
+		public TreeNodeRef<CommitNode>? Parent { get; set; }
+
+		/// <summary>
+		/// Message for this commit
+		/// </summary>
+		public string Message { get; set; }
+
+		/// <summary>
+		/// Time that this commit was created
+		/// </summary>
+		public DateTime Time { get; set; }
+
+		/// <summary>
+		/// Contents of the tree at this commit
+		/// </summary>
+		public TreeNodeRef<DirectoryNode> Contents { get; set; }
+
+		/// <summary>
+		/// Metadata for this commit, keyed by arbitrary GUID
+		/// </summary>
+		public Dictionary<Guid, TreeNodeRef> Metadata { get; } = new Dictionary<Guid, TreeNodeRef>();
+
+		/// <summary>
+		/// Constructor
+		/// </summary>
+		/// <param name="number">Commit number</param>
+		/// <param name="parent">The parent commit</param>
+		/// <param name="message">Message for the commit</param>
+		/// <param name="time">The commit time</param>
+		/// <param name="contents">Contents of the tree at this commit</param>
+		public CommitNode(int number, TreeNodeRef<CommitNode>? parent, string message, DateTime time, TreeNodeRef<DirectoryNode> contents)
+		{
+			Number = number;
+			Parent = parent;
+			Message = message;
+			Time = time;
+			Contents = contents;
+		}
+
+		/// <summary>
+		/// Deserializing constructor
+		/// </summary>
+		/// <param name="reader"></param>
+		public CommitNode(ITreeNodeReader reader)
+		{
+			Number = (int)reader.ReadUnsignedVarInt();
+			if (reader.ReadBoolean())
+			{
+				Parent = reader.ReadRef<CommitNode>();
+			}
+			Message = reader.ReadString();
+			Time = reader.ReadDateTime();
+			Contents = reader.ReadRef<DirectoryNode>();
+			Metadata = reader.ReadDictionary(() => reader.ReadGuid(), () => reader.ReadRef());
+		}
+
+		/// <inheritdoc/>
+		public override void Serialize(ITreeNodeWriter writer)
+		{
+			writer.WriteUnsignedVarInt(Number);
+			writer.WriteBoolean(Parent != null);
+			if (Parent != null)
+			{
+				writer.WriteRef(Parent);
+			}
+			writer.WriteString(Message);
+			writer.WriteDateTime(Time);
+			writer.WriteRef(Contents);
+			writer.WriteDictionary(Metadata, key => writer.WriteGuid(key), value => writer.WriteRef(value));
+		}
+
+		/// <inheritdoc/>
+		public override IEnumerable<TreeNodeRef> EnumerateRefs()
+		{
+			if (Parent != null)
+			{
+				yield return Parent;
+			}
+			yield return Contents;
+		}
+	}
+}
diff --git a/Engine/Source/Programs/Shared/EpicGames.Horde/Storage/TreeNode.cs b/Engine/Source/Programs/Shared/EpicGames.Horde/Storage/TreeNode.cs
--- a/Engine/Source/Programs/Shared/EpicGames.Horde/Storage/TreeNode.cs
+++ b/Engine/Source/Programs/Shared/EpicGames.Horde/Storage/TreeNode.cs
@@ -105,10 +105,20 @@ namespace EpicGames.Horde.Storage
 	public static class TreeNodeExtensions
 	{
 		/// <summary>
-		/// 
+		/// Read an untyped ref from the reader
 		/// </summary>
-		/// <typeparam name="T"></typeparam>
-		/// <param name="reader"></param>
+		/// <param name="reader">Reader to deserialize from</param>
+		/// <returns>New untyped ref</returns>
+		public static TreeNodeRef ReadRef(this ITreeNodeReader reader)
+		{
+			return new TreeNodeRef(reader);
+		}
+
+		/// <summary>
+		/// Read a strongly typed ref from the reader
+		/// </summary>
+		/// <typeparam name="T">Type of the referenced node</typeparam>
+		/// <param name="reader">Reader to deserialize from</param>
 		/// <returns></returns>
 		public static TreeNodeRef<T> ReadRef<T>(this ITreeNodeReader reader) where T : TreeNode
 		{
diff --git a/Engine/Source/Programs/Shared/EpicGames.Horde/Storage/TreeNodeRef.cs b/Engine/Source/Programs/Shared/EpicGames.Horde/Storage/TreeNodeRef.cs
--- a/Engine/Source/Programs/Shared/EpicGames.Horde/Storage/TreeNodeRef.cs
+++ b/Engine/Source/Programs/Shared/EpicGames.Horde/Storage/TreeNodeRef.cs
@@ -102,7 +102,7 @@ namespace EpicGames.Horde.Storage
 		/// <param name="reader"></param>
 		public TreeNodeRef(ITreeNodeReader reader)
 		{
-			TreeNodeRefData data = reader.ReadRef();
+			TreeNodeRefData data = reader.ReadRefData();
 			Hash = data.Hash;
 			Locator = data.Locator;
 
@@ -166,6 +166,22 @@ namespace EpicGames.Horde.Storage
 		/// </summary>
 		/// <param name="writer"></param>
 		public void Serialize(ITreeNodeWriter writer) => writer.WriteRef(this);
+
+		/// <summary>
+		/// Resolve this reference to a concrete node
+		/// </summary>
+		/// <param name="reader">Reader to use for expanding this ref</param>
+		/// <param name="cancellationToken">Cancellation token for the operation</param>
+		/// <returns></returns>
+		public async ValueTask<TreeNode> ExpandAsync(TreeReader reader, CancellationToken cancellationToken = default)
+		{
+			if (Target == null)
+			{
+				Target = await reader.ReadNodeAsync(Locator, cancellationToken);
+				Target!.IncomingRef = this;
+			}
+			return Target!;
+		}
 	}
 
 	/// <summary>
@@ -217,14 +233,9 @@ namespace EpicGames.Horde.Storage
 		/// <param name="reader">Reader to use for expanding this ref</param>
 		/// <param name="cancellationToken">Cancellation token for the operation</param>
 		/// <returns></returns>
-		public async ValueTask<T> ExpandAsync(TreeReader reader, CancellationToken cancellationToken = default)
+		public new async ValueTask<T> ExpandAsync(TreeReader reader, CancellationToken cancellationToken = default)
 		{
-			if (base.Target == null)
-			{
-				base.Target = await reader.ReadNodeAsync<T>(Locator, cancellationToken);
-				Target!.IncomingRef = this;
-			}
-			return Target!;
+			return (T)await base.ExpandAsync(reader, cancellationToken);
 		}
 
 		/// <summary>
diff --git a/Engine/Source/Programs/Shared/EpicGames.Horde/Storage/TreeReader.cs b/Engine/Source/Programs/Shared/EpicGames.Horde/Storage/TreeReader.cs
--- a/Engine/Source/Programs/Shared/EpicGames.Horde/Storage/TreeReader.cs
+++ b/Engine/Source/Programs/Shared/EpicGames.Horde/Storage/TreeReader.cs
@@ -39,7 +39,7 @@ namespace EpicGames.Horde.Storage
 		/// <summary>
 		/// Reads a reference to another node
 		/// </summary>
-		TreeNodeRefData ReadRef();
+		TreeNodeRefData ReadRefData();
 	}
 
 	/// <summary>
@@ -241,7 +241,7 @@ namespace EpicGames.Horde.Storage
 
 			public IReadOnlyDictionary<IoHash, NodeLocator> References => _refs;
 
-			public TreeNodeRefData ReadRef()
+			public TreeNodeRefData ReadRefData()
 			{
 				IoHash hash = this.ReadIoHash();
 				return new TreeNodeRefData(hash, _refs[hash]);
@@ -676,7 +676,7 @@ namespace EpicGames.Horde.Storage
 		/// <param name="locator">Locator for the node</param>
 		/// <param name="cancellationToken">Cancellation token for the operation</param>
 		/// <returns>Node data read from the given bundle</returns>
-		public async ValueTask<TNode> ReadNodeAsync<TNode>(NodeLocator locator, CancellationToken cancellationToken = default) where TNode : TreeNode
+		public async ValueTask<TreeNode> ReadNodeAsync(NodeLocator locator, CancellationToken cancellationToken = default)
 		{
 			BundleInfo bundleInfo = await GetBundleInfoAsync(locator.Blob, cancellationToken);
 			BundleExport export = bundleInfo.Header.Exports[locator.ExportIdx];
@@ -699,7 +699,18 @@ namespace EpicGames.Horde.Storage
 				throw new InvalidOperationException($"No registered serializer for type {bundleInfo.Header.Types[export.TypeIdx].Guid}");
 			}
 
-			return (TNode)typeInfo.Deserialize(new NodeReader(nodeData, refs, typeInfo.BundleType));
+			return typeInfo.Deserialize(new NodeReader(nodeData, refs, typeInfo.BundleType));
+		}
+
+		/// <summary>
+		/// Reads a node from a bundle
+		/// </summary>
+		/// <param name="locator">Locator for the node</param>
+		/// <param name="cancellationToken">Cancellation token for the operation</param>
+		/// <returns>Node data read from the given bundle</returns>
+		public async ValueTask<TNode> ReadNodeAsync<TNode>(NodeLocator locator, CancellationToken cancellationToken = default) where TNode : TreeNode
+		{
+			return (TNode)await ReadNodeAsync(locator, cancellationToken);
 		}
 
 		/// <summary>
