Horde: Include additional metadata for the author and committer of each change.

diff --git a/Engine/Source/Programs/Horde/Horde.Agent/Commands/Vcs/VcsCommand.cs b/Engine/Source/Programs/Horde/Horde.Agent/Commands/Vcs/VcsCommand.cs
--- a/Engine/Source/Programs/Horde/Horde.Agent/Commands/Vcs/VcsCommand.cs
+++ b/Engine/Source/Programs/Horde/Horde.Agent/Commands/Vcs/VcsCommand.cs
@@ -542,6 +542,9 @@ namespace Horde.Agent.Commands.Bundles
 	[Command("vcs", "commit", "Commits data to the VCS store")]
 	class VcsCommitCommand : VcsCommand
 	{
+		[CommandLine("-User=")]
+		public string UserName { get; set; } = Environment.UserName;
+
 		[CommandLine("-Message=")]
 		public string? Message { get; set; }
 
@@ -597,11 +600,11 @@ namespace Horde.Agent.Commands.Bundles
 			CommitNode newTip;
 			if (tip == null)
 			{
-				newTip = new CommitNode(1, null, Message, DateTime.UtcNow, rootRef);
+				newTip = new CommitNode(1, null, UserName, Message, DateTime.UtcNow, rootRef);
 			}
 			else
 			{
-				newTip = new CommitNode(tip.Number + 1, tipRef, Message, DateTime.UtcNow, rootRef);
+				newTip = new CommitNode(tip.Number + 1, tipRef, UserName, Message, DateTime.UtcNow, rootRef);
 			}
 			await store.WriteNodeAsync(workspaceState.Branch, newTip);
 
diff --git a/Engine/Source/Programs/Shared/EpicGames.Core/MemoryReader.cs b/Engine/Source/Programs/Shared/EpicGames.Core/MemoryReader.cs
--- a/Engine/Source/Programs/Shared/EpicGames.Core/MemoryReader.cs
+++ b/Engine/Source/Programs/Shared/EpicGames.Core/MemoryReader.cs
@@ -3,6 +3,7 @@
 using System;
 using System.Buffers.Binary;
 using System.Collections.Generic;
+using System.Text;
 
 namespace EpicGames.Core
 {
@@ -334,5 +335,61 @@ namespace EpicGames.Core
 			ReadDictionary(reader, dictionary, readKey, readValue);
 			return dictionary;
 		}
+
+		/// <summary>
+		/// Read a string from memory
+		/// </summary>
+		/// <param name="reader">Reader to deserialize from</param>
+		/// <returns>The string that was read</returns>
+		public static string ReadString(this IMemoryReader reader) => ReadString(reader, Encoding.UTF8);
+
+		/// <summary>
+		/// Read a string from memory
+		/// </summary>
+		/// <param name="reader">Reader to deserialize from</param>
+		/// <param name="encoding">Encoding to use for the string</param>
+		/// <returns>The string that was read</returns>
+		public static string ReadString(this IMemoryReader reader, Encoding encoding)
+		{
+			int length = (int)reader.ReadUnsignedVarInt();
+
+			ReadOnlySpan<byte> span = reader.GetSpan(length).Slice(0, length);
+			string str = encoding.GetString(span);
+			reader.Advance(length);
+
+			return str;
+		}
+
+		/// <summary>
+		/// Read a string from memory
+		/// </summary>
+		/// <param name="reader">Reader to deserialize from</param>
+		/// <returns>The string that was read</returns>
+		public static string? ReadOptionalString(this IMemoryReader reader) => ReadOptionalString(reader, Encoding.UTF8);
+
+		/// <summary>
+		/// Read a string from memory
+		/// </summary>
+		/// <param name="reader">Reader to deserialize from</param>
+		/// <param name="encoding">Encoding to use for the string</param>
+		/// <returns>The string that was read</returns>
+		public static string? ReadOptionalString(this IMemoryReader reader, Encoding encoding)
+		{
+			int length = (int)reader.ReadUnsignedVarInt();
+			if (length == 0)
+			{
+				return null;
+			}
+			else
+			{
+				length--;
+
+				ReadOnlySpan<byte> span = reader.GetSpan(length).Slice(0, length);
+				string str = encoding.GetString(span);
+				reader.Advance(length);
+
+				return str;
+			}
+		}
 	}
 }
diff --git a/Engine/Source/Programs/Shared/EpicGames.Core/MemoryWriter.cs b/Engine/Source/Programs/Shared/EpicGames.Core/MemoryWriter.cs
--- a/Engine/Source/Programs/Shared/EpicGames.Core/MemoryWriter.cs
+++ b/Engine/Source/Programs/Shared/EpicGames.Core/MemoryWriter.cs
@@ -434,52 +434,61 @@ namespace EpicGames.Core
 		}
 
 		/// <summary>
-		/// Read a string from memory
+		/// Write a string to memory
 		/// </summary>
-		/// <param name="reader">Reader to deserialize from</param>
-		/// <returns>The string that was read</returns>
-		public static string ReadString(this IMemoryReader reader) => ReadString(reader, Encoding.UTF8);
+		/// <param name="writer">Reader to deserialize from</param>
+		/// <param name="str">The string to be written</param>
+		public static void WriteString(this IMemoryWriter writer, string str) => WriteString(writer, str, Encoding.UTF8);
 
 		/// <summary>
-		/// Read a string from memory
+		/// Write a string to memory
 		/// </summary>
-		/// <param name="reader">Reader to deserialize from</param>
+		/// <param name="writer">Reader to deserialize from</param>
+		/// <param name="str">The string to be written</param>
 		/// <param name="encoding">Encoding to use for the string</param>
-		/// <returns>The string that was read</returns>
-		public static string ReadString(this IMemoryReader reader, Encoding encoding)
+		public static void WriteString(this IMemoryWriter writer, string str, Encoding encoding)
 		{
-			int length = (int)reader.ReadUnsignedVarInt();
+			int stringBytes = encoding.GetByteCount(str);
+			int lengthBytes = VarInt.MeasureUnsigned(stringBytes);
 
-			ReadOnlySpan<byte> span = reader.GetSpan(length).Slice(0, length);
-			string str = encoding.GetString(span);
-			reader.Advance(length);
+			Span<byte> span = writer.GetSpan(lengthBytes + stringBytes);
+			VarInt.WriteUnsigned(span, stringBytes);
+			encoding.GetBytes(str, span[lengthBytes..]);
 
-			return str;
+			writer.Advance(lengthBytes + stringBytes);
 		}
 
 		/// <summary>
-		/// Write a string to memory
+		/// Write a nullable string to memory
 		/// </summary>
 		/// <param name="writer">Reader to deserialize from</param>
 		/// <param name="str">The string to be written</param>
-		public static void WriteString(this IMemoryWriter writer, string str) => WriteString(writer, str, Encoding.UTF8);
+		public static void WriteOptionalString(this IMemoryWriter writer, string? str) => WriteOptionalString(writer, str, Encoding.UTF8);
 
 		/// <summary>
-		/// Write a string to memory
+		/// Write a nullable string to memory
 		/// </summary>
 		/// <param name="writer">Reader to deserialize from</param>
 		/// <param name="str">The string to be written</param>
 		/// <param name="encoding">Encoding to use for the string</param>
-		public static void WriteString(this IMemoryWriter writer, string str, Encoding encoding)
+		public static void WriteOptionalString(this IMemoryWriter writer, string? str, Encoding encoding)
 		{
-			int stringBytes = encoding.GetByteCount(str);
-			int lengthBytes = VarInt.MeasureUnsigned(stringBytes);
+			if (str == null)
+			{
+				Span<byte> span = writer.GetSpan(1);
+				span[0] = 0;
+			}
+			else
+			{
+				int stringBytes = encoding.GetByteCount(str);
+				int lengthBytes = VarInt.MeasureUnsigned(stringBytes + 1);
 
-			Span<byte> span = writer.GetSpan(lengthBytes + stringBytes);
-			VarInt.WriteUnsigned(span, stringBytes);
-			encoding.GetBytes(str, span[lengthBytes..]);
+				Span<byte> span = writer.GetSpan(lengthBytes + stringBytes);
+				VarInt.WriteUnsigned(span, stringBytes + 1);
+				encoding.GetBytes(str, span[lengthBytes..]);
 
-			writer.Advance(lengthBytes + stringBytes);
+				writer.Advance(lengthBytes + stringBytes);
+			}
 		}
 	}
 }
diff --git a/Engine/Source/Programs/Shared/EpicGames.Horde/Storage/Nodes/CommitNode.cs b/Engine/Source/Programs/Shared/EpicGames.Horde/Storage/Nodes/CommitNode.cs
--- a/Engine/Source/Programs/Shared/EpicGames.Horde/Storage/Nodes/CommitNode.cs
+++ b/Engine/Source/Programs/Shared/EpicGames.Horde/Storage/Nodes/CommitNode.cs
@@ -22,6 +22,26 @@ namespace EpicGames.Horde.Storage.Nodes
 		/// </summary>
 		public TreeNodeRef<CommitNode>? Parent { get; set; }
 
+		/// <summary>
+		/// Human readable name of the author of this change
+		/// </summary>
+		public string Author { get; set; }
+
+		/// <summary>
+		/// Optional unique identifier for the author. May be an email address, user id, etc...
+		/// </summary>
+		public string? AuthorId { get; set; }
+
+		/// <summary>
+		/// Human readable name of the committer of this change
+		/// </summary>
+		public string? Committer { get; set; }
+
+		/// <summary>
+		/// Optional unique identifier for the committer. May be an email address, user id, etc...
+		/// </summary>
+		public string? CommitterId { get; set; }
+
 		/// <summary>
 		/// Message for this commit
 		/// </summary>
@@ -47,13 +67,15 @@ namespace EpicGames.Horde.Storage.Nodes
 		/// </summary>
 		/// <param name="number">Commit number</param>
 		/// <param name="parent">The parent commit</param>
+		/// <param name="author">Author of this commit</param>
 		/// <param name="message">Message for the commit</param>
 		/// <param name="time">The commit time</param>
 		/// <param name="contents">Contents of the tree at this commit</param>
-		public CommitNode(int number, TreeNodeRef<CommitNode>? parent, string message, DateTime time, TreeNodeRef<DirectoryNode> contents)
+		public CommitNode(int number, TreeNodeRef<CommitNode>? parent, string author, string message, DateTime time, TreeNodeRef<DirectoryNode> contents)
 		{
 			Number = number;
 			Parent = parent;
+			Author = author;
 			Message = message;
 			Time = time;
 			Contents = contents;
@@ -66,10 +88,11 @@ namespace EpicGames.Horde.Storage.Nodes
 		public CommitNode(ITreeNodeReader reader)
 		{
 			Number = (int)reader.ReadUnsignedVarInt();
-			if (reader.ReadBoolean())
-			{
-				Parent = reader.ReadRef<CommitNode>();
-			}
+			Parent = reader.ReadOptionalRef<CommitNode>();
+			Author = reader.ReadString();
+			AuthorId = reader.ReadOptionalString();
+			Committer = reader.ReadOptionalString();
+			CommitterId = reader.ReadOptionalString();
 			Message = reader.ReadString();
 			Time = reader.ReadDateTime();
 			Contents = reader.ReadRef<DirectoryNode>();
@@ -80,11 +103,11 @@ namespace EpicGames.Horde.Storage.Nodes
 		public override void Serialize(ITreeNodeWriter writer)
 		{
 			writer.WriteUnsignedVarInt(Number);
-			writer.WriteBoolean(Parent != null);
-			if (Parent != null)
-			{
-				writer.WriteRef(Parent);
-			}
+			writer.WriteOptionalRef(Parent);
+			writer.WriteString(Author);
+			writer.WriteOptionalString(AuthorId);
+			writer.WriteOptionalString(Committer);
+			writer.WriteOptionalString(CommitterId);
 			writer.WriteString(Message);
 			writer.WriteDateTime(Time);
 			writer.WriteRef(Contents);
@@ -98,7 +121,13 @@ namespace EpicGames.Horde.Storage.Nodes
 			{
 				yield return Parent;
 			}
+
 			yield return Contents;
+
+			foreach (TreeNodeRef metadataRef in Metadata.Values)
+			{
+				yield return metadataRef;
+			}
 		}
 	}
 }
diff --git a/Engine/Source/Programs/Shared/EpicGames.Horde/Storage/TreeNode.cs b/Engine/Source/Programs/Shared/EpicGames.Horde/Storage/TreeNode.cs
--- a/Engine/Source/Programs/Shared/EpicGames.Horde/Storage/TreeNode.cs
+++ b/Engine/Source/Programs/Shared/EpicGames.Horde/Storage/TreeNode.cs
@@ -119,12 +119,64 @@ namespace EpicGames.Horde.Storage
 		/// </summary>
 		/// <typeparam name="T">Type of the referenced node</typeparam>
 		/// <param name="reader">Reader to deserialize from</param>
-		/// <returns></returns>
+		/// <returns>New strongly typed ref</returns>
 		public static TreeNodeRef<T> ReadRef<T>(this ITreeNodeReader reader) where T : TreeNode
 		{
 			return new TreeNodeRef<T>(reader);
 		}
 
+		/// <summary>
+		/// Read an optional untyped ref from the reader
+		/// </summary>
+		/// <param name="reader">Reader to deserialize from</param>
+		/// <returns>New untyped ref</returns>
+		public static TreeNodeRef? ReadOptionalRef(this ITreeNodeReader reader)
+		{
+			if (reader.ReadBoolean())
+			{
+				return reader.ReadRef();
+			}
+			else
+			{
+				return null;
+			}
+		}
+
+		/// <summary>
+		/// Read an optional strongly typed ref from the reader
+		/// </summary>
+		/// <param name="reader">Reader to deserialize from</param>
+		/// <returns>New strongly typed ref</returns>
+		public static TreeNodeRef<T>? ReadOptionalRef<T>(this ITreeNodeReader reader) where T : TreeNode
+		{
+			if (reader.ReadBoolean())
+			{
+				return reader.ReadRef<T>();
+			}
+			else
+			{
+				return null;
+			}
+		}
+
+		/// <summary>
+		/// Writes an optional ref value to storage
+		/// </summary>
+		/// <param name="writer">Writer to serialize to</param>
+		/// <param name="value">Value to write</param>
+		public static void WriteOptionalRef(this ITreeNodeWriter writer, TreeNodeRef? value)
+		{
+			if (value == null)
+			{
+				writer.WriteBoolean(false);
+			}
+			else
+			{
+				writer.WriteBoolean(true);
+				writer.WriteRef(value);
+			}
+		}
+
 		/// <summary>
 		/// Writes a node to storage
 		/// </summary>
