EpicGames.Perforce: Clarify parameter use in TryLoadAsync

diff --git a/Engine/Source/Programs/Shared/EpicGames.Perforce.Managed/StreamSnapshotFromMemory.cs b/Engine/Source/Programs/Shared/EpicGames.Perforce.Managed/StreamSnapshotFromMemory.cs
--- a/Engine/Source/Programs/Shared/EpicGames.Perforce.Managed/StreamSnapshotFromMemory.cs
+++ b/Engine/Source/Programs/Shared/EpicGames.Perforce.Managed/StreamSnapshotFromMemory.cs
@@ -77,9 +77,10 @@ namespace EpicGames.Perforce.Managed
 		/// Load a stream directory from a file on disk
 		/// </summary>
 		/// <param name="inputFile">File to read from</param>
+		/// <param name="defaultBasePath">Base path to use if missing inside loaded file</param>
 		/// <param name="cancellationToken">Cancellation token</param>
 		/// <returns>New StreamDirectoryInfo object</returns>
-		public static async Task<StreamSnapshotFromMemory?> TryLoadAsync(FileReference inputFile, Utf8String basePath, CancellationToken cancellationToken)
+		public static async Task<StreamSnapshotFromMemory?> TryLoadAsync(FileReference inputFile, Utf8String defaultBasePath, CancellationToken cancellationToken)
 		{
 			byte[] data = await FileReference.ReadAllBytesAsync(inputFile, cancellationToken);
 			if (!data.AsSpan().StartsWith(s_currentSignature))
@@ -90,7 +91,7 @@ namespace EpicGames.Perforce.Managed
 			CbObject rootObj = new CbObject(data.AsMemory(s_currentSignature.Length));
 
 			CbObject rootObj2 = rootObj["root"].AsObject();
-			Utf8String rootPath = rootObj2["path"].AsUtf8String(basePath);
+			Utf8String rootPath = rootObj2["path"].AsUtf8String(defaultBasePath);
 			StreamTreeRef root = new StreamTreeRef(rootPath, rootObj2);
 
 			CbArray array = rootObj["items"].AsArray();
