Horde: Simplify ownership semantics for TreeNode objects, to allow sharing objects between different writers.

* TreeNode objects now have a revision number, incremented each time they are modified, and no longer have a back reference to the reference that "owns" them.
* Revision number is captured whenever a node is added to a writer.
* TreeNodeRefs clear out the object they're tracking when written, if the revision number hasn't changed, and also no longer have a back reference to the node that "owns" them.
* New OnExpand() / OnCollapse() callback methods on TreeNodeRef allow custom behavior when a ref is about to transition state.
* See comment above TreeNodeRef for a breakdown of the various states and state transitions for refs.

diff --git a/Engine/Source/Programs/Shared/EpicGames.Horde.Tests/BundleTests.cs b/Engine/Source/Programs/Shared/EpicGames.Horde.Tests/BundleTests.cs
--- a/Engine/Source/Programs/Shared/EpicGames.Horde.Tests/BundleTests.cs
+++ b/Engine/Source/Programs/Shared/EpicGames.Horde.Tests/BundleTests.cs
@@ -13,7 +13,6 @@ using System.Linq;
 using System.Text;
 using System.Threading;
 using System.Threading.Tasks;
-using static System.Formats.Asn1.AsnWriter;
 
 namespace EpicGames.Horde.Tests
 {
diff --git a/Engine/Source/Programs/Shared/EpicGames.Horde/Storage/Nodes/CbNode.cs b/Engine/Source/Programs/Shared/EpicGames.Horde/Storage/Nodes/CbNode.cs
--- a/Engine/Source/Programs/Shared/EpicGames.Horde/Storage/Nodes/CbNode.cs
+++ b/Engine/Source/Programs/Shared/EpicGames.Horde/Storage/Nodes/CbNode.cs
@@ -21,7 +21,7 @@ namespace EpicGames.Horde.Storage.Nodes
 		/// <summary>
 		/// Imported nodes
 		/// </summary>
-		public IReadOnlyDictionary<IoHash, NodeLocator> References { get; }
+		public IReadOnlyDictionary<IoHash, TreeNodeRef> References { get; }
 
 		/// <summary>
 		/// Constructor
@@ -31,7 +31,7 @@ namespace EpicGames.Horde.Storage.Nodes
 		public CbNode(CbObject obj, IReadOnlyDictionary<IoHash, NodeLocator> references)
 		{
 			Object = obj;
-			References = references;
+			References = references.ToDictionary(x => x.Key, x => new TreeNodeRef(x.Key, x.Value));
 		}
 
 		/// <summary>
@@ -41,7 +41,7 @@ namespace EpicGames.Horde.Storage.Nodes
 		public CbNode(ITreeNodeReader reader)
 		{
 			Object = new CbObject(reader.ReadFixedLengthBytes(reader.Length));
-			References = reader.References;
+			References = reader.References.ToDictionary(x => x.Key, x => new TreeNodeRef(x.Key, x.Value));
 		}
 
 		/// <inheritdoc/>
@@ -51,9 +51,6 @@ namespace EpicGames.Horde.Storage.Nodes
 		}
 
 		/// <inheritdoc/>
-		public override IEnumerable<TreeNodeRef> EnumerateRefs()
-		{
-			return References.Select(x => new TreeNodeRef(this, x.Key, x.Value));
-		}
+		public override IEnumerable<TreeNodeRef> EnumerateRefs() => References.Values;
 	}
 }
diff --git a/Engine/Source/Programs/Shared/EpicGames.Horde/Storage/Nodes/DirectoryNode.cs b/Engine/Source/Programs/Shared/EpicGames.Horde/Storage/Nodes/DirectoryNode.cs
--- a/Engine/Source/Programs/Shared/EpicGames.Horde/Storage/Nodes/DirectoryNode.cs
+++ b/Engine/Source/Programs/Shared/EpicGames.Horde/Storage/Nodes/DirectoryNode.cs
@@ -72,7 +72,7 @@ namespace EpicGames.Horde.Storage.Nodes
 		/// <summary>
 		/// Cached length of this node
 		/// </summary>
-		readonly long _cachedLength;
+		long _cachedLength;
 
 		/// <summary>
 		/// Constructor
@@ -135,6 +135,14 @@ namespace EpicGames.Horde.Storage.Nodes
 			await node.CopyToFileAsync(reader, file, cancellationToken);
 		}
 
+		/// <inheritdoc/>
+		protected override void OnCollapse()
+		{
+			base.OnCollapse();
+
+			_cachedLength = Target!.Length;
+		}
+
 		/// <inheritdoc/>
 		public override string ToString() => Name.ToString();
 	}
diff --git a/Engine/Source/Programs/Shared/EpicGames.Horde/Storage/TreeNode.cs b/Engine/Source/Programs/Shared/EpicGames.Horde/Storage/TreeNode.cs
--- a/Engine/Source/Programs/Shared/EpicGames.Horde/Storage/TreeNode.cs
+++ b/Engine/Source/Programs/Shared/EpicGames.Horde/Storage/TreeNode.cs
@@ -4,7 +4,6 @@ using System;
 using System.Collections.Concurrent;
 using System.Collections.Generic;
 using System.Reflection;
-using System.Reflection.Emit;
 using System.Threading;
 using System.Threading.Tasks;
 using EpicGames.Core;
@@ -17,9 +16,10 @@ namespace EpicGames.Horde.Storage
 	public abstract class TreeNode
 	{
 		/// <summary>
-		/// The ref which points to this node
+		/// Revision number of the node. Incremented whenever the node is modified, and used to track whether nodes are modified between 
+		/// writes starting and completing.
 		/// </summary>
-		public TreeNodeRef? IncomingRef { get; internal set; }
+		public uint Revision { get; private set; }
 
 		/// <summary>
 		/// Default constructor
@@ -42,10 +42,7 @@ namespace EpicGames.Horde.Storage
 		/// </summary>
 		protected void MarkAsDirty()
 		{
-			if (IncomingRef != null)
-			{
-				IncomingRef.MarkAsDirty();
-			}
+			Revision++;
 		}
 
 		/// <summary>
diff --git a/Engine/Source/Programs/Shared/EpicGames.Horde/Storage/TreeNodeRef.cs b/Engine/Source/Programs/Shared/EpicGames.Horde/Storage/TreeNodeRef.cs
--- a/Engine/Source/Programs/Shared/EpicGames.Horde/Storage/TreeNodeRef.cs
+++ b/Engine/Source/Programs/Shared/EpicGames.Horde/Storage/TreeNodeRef.cs
@@ -9,15 +9,26 @@ using EpicGames.Core;
 namespace EpicGames.Horde.Storage
 {
 	/// <summary>
-	/// Stores a reference from a parent to child node. The reference may be to a node in memory, or to a node in the storage system.
+	/// Stores a reference to a node, which may be in memory or in the storage system.
+	/// 
+	/// Refs may be in the following states:
+	///  * In storage 
+	///      * Hash and Locator are valid, Target is null, IsDirty() returns false.
+	///      * Writing or calling Collapse() is a no-op.
+	///  * In memory and target was expanded, has not been modified 
+	///      * Hash and Locator are valid, Target is valid, IsDirty() returns false.
+	///      * Writing or calling Collapse() on a ref transitions it to the "in storage" state.
+	///  * In memory and target is new
+	///      * Hash and Locator are invalid, Target is valid, IsDirty() returns true.
+	///      * Writing a ref transitions it to the "in storage" state. Calling Collapse is a no-op.
+	///  * In memory but target has been modified
+	///      * Hash and Locator are set but may not reflect the current node state, Target is valid, IsDirty() returns true. 
+	///      * Writing a ref transitions it to the "in storage" state. Calling Collapse is a no-op.
+	///      
+	/// The <see cref="OnCollapse"/> and <see cref="OnExpand"/> methods allow overriden implementations to cache information about the target.
 	/// </summary>
 	public class TreeNodeRef
 	{
-		/// <summary>
-		/// Node that owns this ref
-		/// </summary>
-		internal TreeNode? _owner;
-
 		/// <summary>
 		/// Hash of the referenced node. Invalid for nodes in memory.
 		/// </summary>
@@ -31,69 +42,26 @@ namespace EpicGames.Horde.Storage
 		/// <summary>
 		/// The target node, or null if the node is not resident in memory.
 		/// </summary>
-		public TreeNode? Target
-		{
-			get => _target;
-			set
-			{
-				if (value != _target)
-				{
-					if (value == null)
-					{
-						_target = value;
-					}
-					else
-					{
-						if (value.IncomingRef != null)
-						{
-							throw new ArgumentException("Target node may not be part of an existing tree.");
-						}
-						if (_target != null)
-						{
-							_target.IncomingRef = null;
-						}
-						_target = value;
-						MarkAsDirty();
-					}
-				}
-			}
-		}
-
-		/// <summary>
-		/// Node pointed to by this ref
-		/// </summary>
-		private TreeNode? _target;
-
-		/// <summary>
-		/// Whether this ref is dirty
-		/// </summary>
-		private bool _dirty;
+		public TreeNode? Target { get; private set; }
 
 		/// <summary>
 		/// Creates a reference to a node in memory.
 		/// </summary>
 		/// <param name="target">Node to reference</param>
-		protected internal TreeNodeRef(TreeNode target)
+		public TreeNodeRef(TreeNode target)
 		{
-			Debug.Assert(target != null);
 			Target = target;
-			_dirty = true;
 		}
 
 		/// <summary>
 		/// Creates a reference to a node with the given hash
 		/// </summary>
-		/// <param name="owner">Node which owns the ref</param>
 		/// <param name="hash">Hash of the referenced node</param>
 		/// <param name="locator">Locator for the node</param>
-		internal TreeNodeRef(TreeNode owner, IoHash hash, NodeLocator locator)
+		public TreeNodeRef(IoHash hash, NodeLocator locator)
 		{
-			_owner = owner;
-
 			Hash = hash;
 			Locator = locator;
-
-			_dirty = false;
 		}
 
 		/// <summary>
@@ -105,8 +73,6 @@ namespace EpicGames.Horde.Storage
 			TreeNodeRefData data = reader.ReadRefData();
 			Hash = data.Hash;
 			Locator = data.Locator;
-
-			_dirty = false;
 			Debug.Assert(Hash != IoHash.Zero);
 		}
 
@@ -114,51 +80,49 @@ namespace EpicGames.Horde.Storage
 		/// Determines whether the the referenced node has modified from the last version written to storage
 		/// </summary>
 		/// <returns></returns>
-		public bool IsDirty() => _dirty;
+		public bool IsDirty() => Target != null && (Target.Revision != 0 || !Locator.IsValid());
 
 		/// <summary>
 		/// Update the reference to refer to a node in memory.
 		/// </summary>
 		public void MarkAsDirty()
 		{
-			Debug.Assert(Target != null);
+			if (Target == null)
+			{
+				throw new InvalidOperationException("Cannot mark a ref as dirty without having expanded it.");
+			}
 
 			Hash = default;
 			Locator = default;
-
-			if (!_dirty)
-			{
-				_dirty = true;
-				if (_owner != null && _owner.IncomingRef != null)
-				{
-					_owner.IncomingRef.MarkAsDirty();
-				}
-			}
 		}
 
 		/// <summary>
-		/// Update the reference to refer to a location in storage.
+		/// Updates the hash and revision number for the ref.
 		/// </summary>
 		/// <param name="hash">Hash of the node</param>
-		/// <param name="locator">Location of the node</param>
-		internal void MarkAsWritten(IoHash hash, NodeLocator locator)
+		/// <returns></returns>
+		internal void MarkAsPendingWrite(IoHash hash)
 		{
-			if (hash == Hash)
-			{
-				Locator = locator;
-				_dirty = false;
-			}
+			Hash = hash;
+			OnCollapse();
 		}
 
 		/// <summary>
-		/// Updates the hash and revision number for the ref.
+		/// Update the reference to refer to a location in storage.
 		/// </summary>
 		/// <param name="hash">Hash of the node</param>
-		/// <returns></returns>
-		internal void MarkAsPendingWrite(IoHash hash)
+		/// <param name="locator">Location of the node</param>
+		/// <param name="revision">Revision number that was written</param>
+		internal void MarkAsWritten(IoHash hash, NodeLocator locator, uint revision)
 		{
-			Hash = hash;
-			_dirty = false;
+			if (Hash == hash)
+			{
+				Locator = locator;
+				if (Target != null && Target.Revision == revision)
+				{
+					Target = null;
+				}
+			}
 		}
 
 		/// <summary>
@@ -178,9 +142,35 @@ namespace EpicGames.Horde.Storage
 			if (Target == null)
 			{
 				Target = await reader.ReadNodeAsync(Locator, cancellationToken);
-				Target!.IncomingRef = this;
+				OnExpand();
+			}
+			return Target;
+		}
+
+		/// <summary>
+		/// Collapse the current node
+		/// </summary>
+		public void Collapse()
+		{
+			if (Target != null && !IsDirty())
+			{
+				OnCollapse();
+				Target = null;
 			}
-			return Target!;
+		}
+
+		/// <summary>
+		/// Callback after a node is expanded, allowing overridden implementations to cache any information about the target
+		/// </summary>
+		protected virtual void OnExpand()
+		{
+		}
+
+		/// <summary>
+		/// Callback before a node is collapsed, allowing overridden implementations to cache any information about the target
+		/// </summary>
+		protected virtual void OnCollapse()
+		{
 		}
 	}
 
@@ -193,11 +183,7 @@ namespace EpicGames.Horde.Storage
 		/// <summary>
 		/// Accessor for the target node
 		/// </summary>
-		public new T? Target
-		{
-			get => (T?)base.Target;
-			set => base.Target = value;
-		}
+		public new T? Target => (T?)base.Target;
 
 		/// <summary>
 		/// Constructor
@@ -210,11 +196,10 @@ namespace EpicGames.Horde.Storage
 		/// <summary>
 		/// Constructor
 		/// </summary>
-		/// <param name="owner">Node which owns the ref</param>
 		/// <param name="hash">Hash of the referenced node</param>
 		/// <param name="locator">Locator for the node</param>
-		internal TreeNodeRef(TreeNode owner, IoHash hash, NodeLocator locator) 
-			: base(owner, hash, locator)
+		internal TreeNodeRef(IoHash hash, NodeLocator locator) 
+			: base(hash, locator)
 		{
 		}
 
diff --git a/Engine/Source/Programs/Shared/EpicGames.Horde/Storage/TreeWriter.cs b/Engine/Source/Programs/Shared/EpicGames.Horde/Storage/TreeWriter.cs
--- a/Engine/Source/Programs/Shared/EpicGames.Horde/Storage/TreeWriter.cs
+++ b/Engine/Source/Programs/Shared/EpicGames.Horde/Storage/TreeWriter.cs
@@ -46,6 +46,7 @@ namespace EpicGames.Horde.Storage
 			public readonly int Length;
 			public readonly IReadOnlyList<IoHash> RefHashes;
 			public readonly TreeNodeRef NodeRef;
+			public readonly uint Revision;
 
 			public NodeInfo(BundleType type, IoHash hash, int length, IReadOnlyList<IoHash> refs, TreeNodeRef nodeRef)
 			{
@@ -54,6 +55,7 @@ namespace EpicGames.Horde.Storage
 				Length = length;
 				RefHashes = refs;
 				NodeRef = nodeRef;
+				Revision = nodeRef.Target!.Revision;
 			}
 		}
 
@@ -158,19 +160,7 @@ namespace EpicGames.Horde.Storage
 		/// <inheritdoc/>
 		public async Task<IoHash> WriteAsync(TreeNodeRef nodeRef, CancellationToken cancellationToken = default)
 		{
-			// If the target node hasn't been modified, use the existing state.
-			if (!nodeRef.IsDirty())
-			{
-				// Make sure the locator is valid. The node may be queued for writing but not flushed to disk yet.
-				if (nodeRef.Locator.IsValid())
-				{
-					_hashToNode.TryAdd(nodeRef.Hash, nodeRef.Locator);
-					nodeRef.Target = null;
-				}
-				return nodeRef.Hash;
-			}
-
-			// Check we actually have a target node
+			// Check we actually have a target node. If we don't, we don't need to write anything.
 			TreeNode? node = nodeRef.Target;
 			if (node == null)
 			{
@@ -179,28 +169,38 @@ namespace EpicGames.Horde.Storage
 				return nodeRef.Hash;
 			}
 
-			// Update the incoming ref
-			if (node.IncomingRef == null)
-			{
-				node.IncomingRef = nodeRef;
-			}
-			else if(node.IncomingRef != nodeRef)
-			{
-				throw new InvalidOperationException("Node cannot have multiple incoming references");
-			}
-
-			// Write all the referenced nodes first
+			// Write all the nodes it references, and mark the ref as dirty if any of them change.
 			Dictionary<TreeNodeRef, IoHash> nextRefToHash = new Dictionary<TreeNodeRef, IoHash>();
 			foreach (TreeNodeRef nextRef in node.EnumerateRefs())
 			{
+				IoHash prevHash = nextRef.Hash;
 				IoHash nextHash = await WriteAsync(nextRef, cancellationToken);
 				Debug.Assert(nextHash != IoHash.Zero);
+
 				nextRefToHash.Add(nextRef, nextHash);
+
+				if (prevHash != nextHash)
+				{
+					nodeRef.MarkAsDirty();
+				}
+			}
+
+			// If the target node hasn't been modified, use the existing serialized state.
+			if (!nodeRef.IsDirty())
+			{
+				// Make sure the locator is valid. The node may be queued for writing but not flushed to disk yet.
+				if (nodeRef.Locator.IsValid())
+				{
+					_hashToNode.TryAdd(nodeRef.Hash, nodeRef.Locator);
+					nodeRef.Collapse();
+				}
+				return nodeRef.Hash;
 			}
 
 			// Serialize the node
 			NodeWriter nodeWriter = new NodeWriter(nextRefToHash);
 			nodeRef.Target!.Serialize(nodeWriter);
+			BundleType bundleType = nodeRef.Target.GetBundleType();
 
 			// Get the hash for the new blob
 			ReadOnlySequence<byte> data = nodeWriter.Data.AsSequence();
@@ -208,14 +208,14 @@ namespace EpicGames.Horde.Storage
 			nodeRef.MarkAsPendingWrite(hash);
 
 			// Check if we're already tracking a node with the same hash
-			NodeInfo nodeInfo = new NodeInfo(nodeRef.Target.GetBundleType(), hash, (int)data.Length, nextRefToHash.Values.Distinct().ToArray(), nodeRef);
+			NodeInfo nodeInfo = new NodeInfo(bundleType, hash, (int)data.Length, nextRefToHash.Values.Distinct().ToArray(), nodeRef);
 			if (_hashToNode.TryGetValue(hash, out NodeLocator locator))
 			{
 				// If the node is in the lookup but not currently valid, it's already queued for writing. Add this ref to the list of refs that needs fixing up,
 				// so we can update it after flushing.
 				if (locator.IsValid())
 				{
-					nodeRef.MarkAsWritten(hash, locator);
+					nodeRef.MarkAsWritten(hash, locator, nodeInfo.Revision);
 				}
 				else
 				{
@@ -379,7 +379,7 @@ namespace EpicGames.Horde.Storage
 					NodeLocator nodeLocator = new NodeLocator(locator, idx);
 
 					NodeInfo nodeInfo = _queue[idx];
-					nodeInfo.NodeRef.MarkAsWritten(nodeInfo.Hash, nodeLocator);
+					nodeInfo.NodeRef.MarkAsWritten(nodeInfo.Hash, nodeLocator, nodeInfo.Revision);
 
 					_hashToNode[nodeInfo.Hash] = nodeLocator;
 				}
@@ -391,7 +391,7 @@ namespace EpicGames.Horde.Storage
 					NodeInfo nodeInfo = _secondaryQueue[refIdx];
 					if (_hashToNode.TryGetValue(nodeInfo.Hash, out NodeLocator refLocator))
 					{
-						nodeInfo.NodeRef.MarkAsWritten(nodeInfo.Hash, refLocator);
+						nodeInfo.NodeRef.MarkAsWritten(nodeInfo.Hash, refLocator, nodeInfo.Revision);
 					}
 				}
 				_secondaryQueue.RemoveRange(0, refIdx);
