EpicGames.Perforce: Fix malformed format string in exception message.

diff --git a/Engine/Source/Programs/Shared/EpicGames.Perforce/IPerforceOutput.cs b/Engine/Source/Programs/Shared/EpicGames.Perforce/IPerforceOutput.cs
--- a/Engine/Source/Programs/Shared/EpicGames.Perforce/IPerforceOutput.cs
+++ b/Engine/Source/Programs/Shared/EpicGames.Perforce/IPerforceOutput.cs
@@ -233,7 +233,7 @@ namespace EpicGames.Perforce
 				ReadOnlyMemory<byte> data = perforce.Data;
 				if (data.Length > 0 && responses.Count == 0 && data.Span[0] != '{')
 				{
-					throw new PerforceException("Unexpected response from server (expected '{'):{0}", FormatDataAsString(data.Span));
+					throw new PerforceException("Unexpected response from server (expected '{{'):{0}", FormatDataAsString(data.Span));
 				}
 
 				// Parse the responses from the current buffer
