Horde: Move functionality for manipulating chunked and indexed log files into EpicGames.Horde.

diff --git a/Engine/Source/Programs/Shared/EpicGames.Horde.Tests/EpicGames.Horde.Tests.csproj b/Engine/Source/Programs/Shared/EpicGames.Horde.Tests/EpicGames.Horde.Tests.csproj
--- a/Engine/Source/Programs/Shared/EpicGames.Horde.Tests/EpicGames.Horde.Tests.csproj
+++ b/Engine/Source/Programs/Shared/EpicGames.Horde.Tests/EpicGames.Horde.Tests.csproj
@@ -16,6 +16,7 @@
       <PrivateAssets>all</PrivateAssets>
       <IncludeAssets>runtime; build; native; contentfiles; analyzers; buildtransitive</IncludeAssets>
     </PackageReference>
+    <PackageReference Include="System.Linq.Async" Version="6.0.1" />
   </ItemGroup>
 
   <ItemGroup>
@@ -23,4 +24,19 @@
     <ProjectReference Include="..\EpicGames.Horde\EpicGames.Horde.csproj" />
   </ItemGroup>
 
+  <ItemGroup>
+    <Compile Update="Properties\Resources.Designer.cs">
+      <DesignTime>True</DesignTime>
+      <AutoGen>True</AutoGen>
+      <DependentUpon>Resources.resx</DependentUpon>
+    </Compile>
+  </ItemGroup>
+
+  <ItemGroup>
+    <EmbeddedResource Update="Properties\Resources.resx">
+      <Generator>ResXFileCodeGenerator</Generator>
+      <LastGenOutput>Resources.Designer.cs</LastGenOutput>
+    </EmbeddedResource>
+  </ItemGroup>
+
 </Project>
diff --git a/Engine/Source/Programs/Shared/EpicGames.Horde.Tests/LogTests.cs b/Engine/Source/Programs/Shared/EpicGames.Horde.Tests/LogTests.cs
new file mode 100644
--- /dev/null
+++ b/Engine/Source/Programs/Shared/EpicGames.Horde.Tests/LogTests.cs
@@ -0,0 +1,216 @@
+// Copyright Epic Games, Inc. All Rights Reserved.
+
+using System;
+using System.Collections.Generic;
+using System.Linq;
+using System.Text;
+using System.Threading;
+using System.Threading.Tasks;
+using EpicGames.Horde.Logs;
+using EpicGames.Horde.Storage;
+using EpicGames.Horde.Storage.Backends;
+using EpicGames.Horde.Tests.Properties;
+using Microsoft.Extensions.Caching.Memory;
+using Microsoft.Extensions.Logging.Abstractions;
+using Microsoft.VisualStudio.TestTools.UnitTesting;
+
+namespace EpicGames.Horde.Tests
+{
+	[TestClass]
+	public sealed class LogTests : IDisposable
+	{
+		readonly IMemoryCache _cache;
+		readonly MemoryStorageClient _store;
+		readonly TreeReader _reader;
+
+		private readonly byte[] _data = Resources.TextFile;
+
+		public LogTests()
+		{
+			_cache = new MemoryCache(new MemoryCacheOptions());
+			_store = new MemoryStorageClient();
+			_reader = new TreeReader(_store, _cache, NullLogger.Instance);
+		}
+
+		public void Dispose()
+		{
+			_cache.Dispose();
+		}
+
+		[TestMethod]
+		public void NgramTests()
+		{
+			NgramSetBuilder builder = new NgramSetBuilder();
+
+			ulong[] values = { 1, 2, 3, 5, 8, 13, 21, 34, 55, 89, 144, 233, 377, 610, 987, 1597, 2584, 4181, 6765, 10946, 0xfedcba9876543210UL };
+			foreach (ulong value in values)
+			{
+				builder.Add(value);
+			}
+
+			NgramSet trie = builder.ToNgramSet();
+			Assert.IsTrue(Enumerable.SequenceEqual(trie.EnumerateRange(0, UInt64.MaxValue), values));
+			Assert.IsTrue(Enumerable.SequenceEqual(trie.EnumerateRange(0, 90), values.Where(x => x <= 90)));
+			Assert.IsTrue(Enumerable.SequenceEqual(trie.EnumerateRange(2, 89), values.Where(x => x >= 2 && x <= 89)));
+		}
+
+		[TestMethod]
+		public async Task IndexTests()
+		{
+			// Write the test data to the log file in blocks
+			LogBuilder builder = new LogBuilder(LogFormat.Text, 1, 1, NullLogger.Instance);
+
+			int readOffset = 0;
+			int readLineIndex = 0;
+			while (readOffset < _data.Length)
+			{
+				int length = 0;
+				int lineCount = 0;
+
+				for (int idx = readOffset; idx < Math.Min(_data.Length, readOffset + 1883); idx++)
+				{
+					if (_data[idx] == '\n')
+					{
+						length = (idx + 1) - readOffset;
+						lineCount++;
+						break;
+					}
+				}
+
+				builder.WriteData(_data.AsMemory(readOffset, length));
+
+				readOffset += length;
+				readLineIndex += lineCount;
+			}
+
+			// Flush it to storage, and read the finished log node
+			NodeLocator locator = await builder.FlushAsync(new TreeWriter(_store), CancellationToken.None);
+			LogNode log = await _reader.ReadNodeAsync<LogNode>(locator);
+
+			// Read the data back out and check it's the same
+			byte[] readData = new byte[_data.Length];
+
+			int writeOffset = 0;
+			await foreach (ReadOnlyMemory<byte> line in log.ReadLogLinesAsync(_reader, 0, CancellationToken.None))
+			{
+				line.CopyTo(readData.AsMemory(writeOffset));
+				writeOffset += line.Length;
+			}
+
+			int equalSize = 0;
+			while(equalSize < _data.Length && _data[equalSize] == readData[equalSize])
+			{
+				equalSize++;
+			}
+			Assert.AreEqual(equalSize, readOffset);
+
+			// Test some searches
+			LogIndexNode index = await log.IndexRef.ExpandAsync(_reader);
+			await SearchLogDataTestAsync(index);
+		}
+
+		[TestMethod]
+		public async Task PartialTokenTests()
+		{
+			// Generate the test data
+			string[] lines =
+			{
+				"abcdefghi\n",
+				"jklmno123\n",
+				"pqrst99uv\n",
+				"wx\n"
+			};
+
+			LogBuilder builder = new LogBuilder(LogFormat.Text, 1, 1, NullLogger.Instance);
+			for (int lineIdx = 0; lineIdx < lines.Length; lineIdx++)
+			{
+				builder.WriteData(Encoding.UTF8.GetBytes(lines[lineIdx]));
+			}
+			NodeLocator rootNodeLocator = await builder.FlushAsync(new TreeWriter(_store), CancellationToken.None);
+
+			// Read it back in and test the index
+			LogNode rootNode = await _reader.ReadNodeAsync<LogNode>(rootNodeLocator);
+			LogIndexNode index = await rootNode.IndexRef.ExpandAsync(_reader);
+		
+			for (int lineIdx = 0; lineIdx < lines.Length; lineIdx++)
+			{
+				for (int strLen = 1; strLen < 7; strLen++)
+				{
+					for (int strOfs = 0; strOfs + strLen < lines[lineIdx].Length - 1; strOfs++)
+					{
+						string str = lines[lineIdx].Substring(strOfs, strLen);
+
+						SearchStats stats = new SearchStats();
+						List<int> results = await index.Search(_reader, 0, new SearchTerm(str), stats, CancellationToken.None).ToListAsync();
+						Assert.AreEqual(1, results.Count);
+						Assert.AreEqual(lineIdx, results[0]);
+
+						Assert.AreEqual(1, stats.NumScannedBlocks);
+						Assert.AreEqual(3, stats.NumSkippedBlocks);
+						Assert.AreEqual(0, stats.NumFalsePositiveBlocks);
+					}
+				}
+			}
+		}
+		/*
+		[TestMethod]
+		public async Task AppendIndexTests()
+		{
+			JobId jobId = JobId.GenerateNewId();
+			ILogFile logFile = await _logFileService.CreateLogFileAsync(jobId, null, LogType.Text);
+
+			logFile = await WriteLogDataAsync(logFile, 0, 0, Encoding.UTF8.GetBytes("abc\n"), true);
+			logFile = await WriteLogDataAsync(logFile, 4, 1, Encoding.UTF8.GetBytes("def\n"), true);
+			logFile = await WriteLogDataAsync(logFile, 8, 2, Encoding.UTF8.GetBytes("ghi\n"), false);
+
+			await ((LogFileService)_logFileService).FlushPendingWritesAsync();
+
+			{
+				LogSearchStats stats = new LogSearchStats();
+				List<int> results = await _logFileService.SearchLogDataAsync(logFile, "abc", 0, 5, stats, CancellationToken.None);
+				Assert.AreEqual(1, results.Count);
+				Assert.AreEqual(0, results[0]);
+
+				Assert.AreEqual(2, stats.NumScannedBlocks); // abc + ghi (no index yet because it hasn't been flushed)
+				Assert.AreEqual(1, stats.NumSkippedBlocks); // def
+				Assert.AreEqual(0, stats.NumFalsePositiveBlocks);
+			}
+			{
+				LogSearchStats stats = new LogSearchStats();
+				List<int> results = await _logFileService.SearchLogDataAsync(logFile, "def", 0, 5, stats, CancellationToken.None);
+				Assert.AreEqual(1, results.Count);
+				Assert.AreEqual(1, results[0]);
+
+				Assert.AreEqual(2, stats.NumScannedBlocks); // def + ghi (no index yet because it hasn't been flushed)
+				Assert.AreEqual(1, stats.NumSkippedBlocks); // abc
+				Assert.AreEqual(0, stats.NumFalsePositiveBlocks);
+			}
+			{
+				LogSearchStats stats = new LogSearchStats();
+				List<int> results = await _logFileService.SearchLogDataAsync(logFile, "ghi", 0, 5, stats, CancellationToken.None);
+				Assert.AreEqual(1, results.Count);
+				Assert.AreEqual(2, results[0]);
+
+				Assert.AreEqual(1, stats.NumScannedBlocks); // ghi
+				Assert.AreEqual(2, stats.NumSkippedBlocks); // abc + def
+				Assert.AreEqual(0, stats.NumFalsePositiveBlocks);
+			}
+		}
+*/
+		async Task SearchLogDataTestAsync(LogIndexNode index)
+		{
+			await SearchLogDataTestAsync(index, "HISPANIOLA", 0, 4, new[] { 1503, 1520, 1525, 1595 });
+			await SearchLogDataTestAsync(index, "Hispaniola", 0, 4, new[] { 1503, 1520, 1525, 1595 });
+			await SearchLogDataTestAsync(index, "HizpaniolZ", 0, 4, Array.Empty<int>());
+			await SearchLogDataTestAsync(index, "Pieces of eight!", 0, 100, new[] { 2227, 2228, 5840, 5841, 7520 });
+			await SearchLogDataTestAsync(index, "NEWSLETTER", 0, 100, new[] { 7886 });
+		}
+
+		async Task SearchLogDataTestAsync(LogIndexNode index, string text, int firstLine, int count, int[] expectedLines)
+		{
+			SearchStats stats = new SearchStats();
+			List<int> lines = await index.Search(_reader, firstLine, new SearchTerm(text), stats, CancellationToken.None).Take(count).ToListAsync();
+			Assert.IsTrue(lines.SequenceEqual(expectedLines));
+		}
+	}
+}
diff --git a/Engine/Source/Programs/Shared/EpicGames.Horde/Logs/LogChunkNode.cs b/Engine/Source/Programs/Shared/EpicGames.Horde/Logs/LogChunkNode.cs
new file mode 100644
--- /dev/null
+++ b/Engine/Source/Programs/Shared/EpicGames.Horde/Logs/LogChunkNode.cs
@@ -0,0 +1,683 @@
+// Copyright Epic Games, Inc. All Rights Reserved.
+
+using System;
+using System.Collections;
+using System.Collections.Generic;
+using System.IO;
+using System.Linq;
+using System.Text.Json;
+using System.Text;
+using EpicGames.Core;
+using EpicGames.Horde.Storage;
+using Microsoft.Extensions.Logging;
+
+namespace EpicGames.Horde.Logs
+{
+	/// <summary>
+	/// Read-only buffer for log text, with indexed line offsets.
+	/// </summary>
+	[TreeNode("{7020B6CA-0F72-4174-B6AA-06AA60A3EF30}", 1)]
+	public class LogChunkNode : TreeNode
+	{
+		/// <summary>
+		/// Provides access to the lines for this chunk through a list interface
+		/// </summary>
+		class LinesCollection : IReadOnlyList<Utf8String>
+		{
+			readonly LogChunkNode _owner;
+
+			public LinesCollection(LogChunkNode owner) => _owner = owner;
+
+			/// <inheritdoc/>
+			public Utf8String this[int index] => _owner.GetLine(index);
+
+			/// <inheritdoc/>
+			public int Count => _owner.LineCount;
+
+			/// <inheritdoc/>
+			public IEnumerator<Utf8String> GetEnumerator()
+			{
+				for (int idx = 0; idx < Count; idx++)
+				{
+					yield return _owner.GetLine(idx);
+				}
+			}
+
+			/// <inheritdoc/>
+			IEnumerator IEnumerable.GetEnumerator() => GetEnumerator();
+		}
+
+		/// <summary>
+		/// Empty log chunk
+		/// </summary>
+		public static LogChunkNode Empty { get; } = new LogChunkNode(Array.Empty<byte>(), new int[1]);
+
+		/// <summary>
+		/// The raw text data. Contains a complete set of lines followed by newline characters.
+		/// </summary>
+		public ReadOnlyMemory<byte> Data { get; }
+
+		/// <summary>
+		/// Span for the raw text data.
+		/// </summary>
+		public ReadOnlySpan<byte> Span => Data.Span;
+
+		/// <summary>
+		/// Accessor for the lines in this chunk
+		/// </summary>
+		public IReadOnlyList<Utf8String> Lines { get; }
+
+		/// <summary>
+		/// Offsets of lines within the data object, including a sentinel for the end of the data (LineCount + 1 entries).
+		/// </summary>
+		public IReadOnlyList<int> LineOffsets { get; }
+
+		/// <summary>
+		/// Length of this chunk
+		/// </summary>
+		public int Length => Data.Length;
+
+		/// <summary>
+		/// Number of lines in the block (excluding the sentinel).
+		/// </summary>
+		public int LineCount => LineOffsets.Count - 1;
+
+		/// <summary>
+		/// Default constructor
+		/// </summary>
+		public LogChunkNode()
+			: this(Empty.Data, Empty.LineOffsets)
+		{
+		}
+
+		/// <summary>
+		/// Constructor
+		/// </summary>
+		/// <param name="data">Data to construct from</param>
+		public LogChunkNode(ReadOnlyMemory<byte> data)
+			: this(data, FindLineOffsets(data.Span))
+		{
+		}
+
+		/// <summary>
+		/// Constructor
+		/// </summary>
+		internal LogChunkNode(ReadOnlyMemory<byte> data, IReadOnlyList<int> lineOffsets)
+		{
+			Data = data;
+			Lines = new LinesCollection(this);
+			LineOffsets = lineOffsets;
+		}
+
+		/// <summary>
+		/// Deserializing constructor
+		/// </summary>
+		/// <param name="reader">Reader to pull data from</param>
+		public LogChunkNode(ITreeNodeReader reader)
+			: this(reader.ReadVariableLengthBytes())
+		{
+		}
+
+		/// <inheritdoc/>
+		public override void Serialize(ITreeNodeWriter writer)
+		{
+			writer.WriteVariableLengthBytes(Data.Span);
+		}
+
+		/// <inheritdoc/>
+		public override IEnumerable<TreeNodeRef> EnumerateRefs() => Enumerable.Empty<TreeNodeRef>();
+
+		/// <summary>
+		/// Accessor for an individual line
+		/// </summary>
+		/// <param name="idx">Index of the line to retrieve</param>
+		/// <returns>Line at the given index</returns>
+		public Utf8String GetLine(int idx) => new Utf8String(Data.Slice(LineOffsets[idx], LineOffsets[idx + 1] - LineOffsets[idx] - 1));
+
+		/// <summary>
+		/// Find the line index for a particular offset
+		/// </summary>
+		/// <param name="offset">Offset within the text</param>
+		/// <returns>The line index</returns>
+		public int GetLineIndexForOffset(int offset)
+		{
+			int lineIdx = LineOffsets.BinarySearch(offset);
+			if (lineIdx < 0)
+			{
+				lineIdx = ~lineIdx - 1;
+			}
+			return lineIdx;
+		}
+
+		/// <summary>
+		/// Creates a new list of line offsets for the given text
+		/// </summary>
+		/// <param name="data"></param>
+		/// <returns></returns>
+		public static List<int> FindLineOffsets(ReadOnlySpan<byte> data)
+		{
+			List<int> lineOffsets = new List<int>();
+			lineOffsets.Add(0);
+			UpdateLineOffsets(data, 0, lineOffsets);
+			return lineOffsets;
+		}
+
+		/// <summary>
+		/// Updates the length of this chunk, computing all the newline offsets
+		/// </summary>
+		/// <param name="data">Text to search for line endings</param>
+		/// <param name="start">Start offset within the text buffer</param>
+		/// <param name="lineOffsets">Offsets of each line within the text</param>
+		public static void UpdateLineOffsets(ReadOnlySpan<byte> data, int start, List<int> lineOffsets)
+		{
+			if (start < data.Length)
+			{
+				// Make sure the data ends with a newline
+				if (data[data.Length - 1] != '\n')
+				{
+					throw new InvalidDataException("Chunk data must end with a newline");
+				}
+
+				// Calculate the new number of newlines
+				for (int idx = start; idx < data.Length; idx++)
+				{
+					if (data[idx] == '\n')
+					{
+						lineOffsets.Add(idx + 1);
+					}
+				}
+			}
+		}
+	}
+
+	/// <summary>
+	/// Reference to a chunk of text, with information about its placement in the larger log file
+	/// </summary>
+	public class LogChunkRef : TreeNodeRef<LogChunkNode>
+	{
+		/// <summary>
+		/// First line within the file
+		/// </summary>
+		public int LineIndex { get; }
+
+		/// <summary>
+		/// Number of lines in this block
+		/// </summary>
+		public int LineCount { get; }
+
+		/// <summary>
+		/// Constructor
+		/// </summary>
+		/// <param name="lineIndex">Index of the first line within this block</param>
+		/// <param name="target">Referenced log text</param>
+		public LogChunkRef(int lineIndex, LogChunkNode target)
+			: base(target)
+		{
+			LineIndex = lineIndex;
+			LineCount = target.LineCount;
+		}
+
+		/// <summary>
+		/// Deserializing constructor
+		/// </summary>
+		/// <param name="reader"></param>
+		public LogChunkRef(ITreeNodeReader reader)
+			: base(reader)
+		{
+			LineIndex = (int)reader.ReadUnsignedVarInt();
+			LineCount = (int)reader.ReadUnsignedVarInt();
+		}
+
+		/// <inheritdoc/>
+		public new void Serialize(ITreeNodeWriter writer)
+		{
+			base.Serialize(writer);
+
+			writer.WriteUnsignedVarInt(LineIndex);
+			writer.WriteUnsignedVarInt(LineCount);
+		}
+	}
+
+	/// <summary>
+	/// Builder for <see cref="LogChunkNode"/> objects.
+	/// </summary>
+	public class LogChunkBuilder
+	{
+		/// <summary>
+		/// Accessor for Data
+		/// </summary>
+		byte[] _data;
+
+		/// <summary>
+		/// Current used length of the buffer
+		/// </summary>
+		int _length;
+
+		/// <summary>
+		/// Offsets of the start of each line within the data
+		/// </summary>
+		readonly List<int> _lineOffsets = new List<int> { 0 };
+
+		/// <summary>
+		/// Current length of the buffer
+		/// </summary>
+		public int Length => _length;
+
+		/// <summary>
+		/// Number of lines in this buffer
+		/// </summary>
+		public int LineCount => _lineOffsets.Count - 1;
+
+		/// <summary>
+		/// Capacity of the buffer
+		/// </summary>
+		public int Capacity => _data.Length;
+
+		/// <summary>
+		/// Constructor
+		/// </summary>
+		public LogChunkBuilder(int maxLength = 64 * 1024)
+		{
+			_data = new byte[maxLength];
+		}
+
+		/// <summary>
+		/// Constructor
+		/// </summary>
+		/// <param name="data">Data to initialize this chunk with. Ownership of this array is transferred to the chunk, and its length determines the chunk size.</param>
+		/// <param name="length">Number of valid bytes within the initial data array</param>
+		public LogChunkBuilder(byte[] data, int length)
+			: this(data, length, LogChunkNode.FindLineOffsets(data.AsSpan(0, length)))
+		{
+		}
+
+		/// <summary>
+		/// Constructor
+		/// </summary>
+		/// <param name="data"></param>
+		/// <param name="length"></param>
+		/// <param name="lineOffsets"></param>
+		private LogChunkBuilder(byte[] data, int length, List<int> lineOffsets)
+		{
+			_data = data;
+			_length = length;
+			_lineOffsets = lineOffsets;
+		}
+
+		/// <summary>
+		/// Clear the contents of the buffer
+		/// </summary>
+		public void Clear()
+		{
+			_length = 0;
+			_lineOffsets.RemoveRange(1, _lineOffsets.Count - 1);
+		}
+
+		/// <summary>
+		/// Create a new chunk data object with the given data appended. The internal buffers are reused, with the assumption that
+		/// there is no contention over writing to the same location in the chunk.
+		/// </summary>
+		/// <param name="textData">The data to append</param>
+		/// <returns>New chunk data object</returns>
+		public void Append(ReadOnlySpan<byte> textData)
+		{
+			textData.CopyTo(_data.AsSpan(_length, textData.Length));
+
+			int prevLength = _length;
+			_length += textData.Length;
+
+			LogChunkNode.UpdateLineOffsets(_data.AsSpan(0, _length), prevLength, _lineOffsets);
+		}
+
+		/// <summary>
+		/// Appends JSON text from another buffer as plain text in this one
+		/// </summary>
+		public void AppendJsonAsPlainText(ReadOnlySpan<byte> inputLine, ILogger logger)
+		{
+			// Make sure the output buffer is large enough
+			int requiredSpace = inputLine.Length;
+			if (_length + requiredSpace > _data.Length)
+			{
+				Array.Resize(ref _data, _length + requiredSpace);
+			}
+
+			// Convert the line to plain text
+			try
+			{
+				_length = ConvertToPlainText(inputLine, _data, _length);
+			}
+			catch (Exception ex)
+			{
+				inputLine.CopyTo(_data.AsSpan(_length));
+				_length += inputLine.Length;
+				logger.LogWarning(ex, "Exception while attempting to parse log text as JSON. Line: \"{Line}\"", Encoding.UTF8.GetString(inputLine).Trim());
+			}
+			_lineOffsets.Add(_length);
+		}
+
+		/// <summary>
+		/// Appends JSON text from another buffer as plain text in this one
+		/// </summary>
+		public void AppendJsonAsPlainText(LogChunkNode srcText, int srcLineIndex, int srcLineCount, ILogger logger)
+		{
+			for (int idx = 0; idx < srcLineCount; idx++)
+			{
+				int lineOffset = srcText.LineOffsets[srcLineIndex + idx];
+				int nextLineOffset = srcText.LineOffsets[srcLineIndex + idx + 1];
+				ReadOnlySpan<byte> inputLine = srcText.Data.Slice(lineOffset, nextLineOffset - lineOffset).Span;
+				AppendJsonAsPlainText(inputLine, logger);
+			}
+		}
+
+		/// <summary>
+		/// Determines if the given line is empty
+		/// </summary>
+		/// <param name="input">The input data</param>
+		/// <returns>True if the given text is empty</returns>
+		public static bool IsEmptyOrWhitespace(ReadOnlySpan<byte> input)
+		{
+			for (int idx = 0; idx < input.Length; idx++)
+			{
+				byte v = input[idx];
+				if (v != (byte)'\n' && v != '\r' && v != ' ')
+				{
+					return false;
+				}
+			}
+			return true;
+		}
+
+		/// <summary>
+		/// Converts a JSON log line to plain text
+		/// </summary>
+		/// <param name="input">The JSON data</param>
+		/// <param name="output">Output buffer for the converted line</param>
+		/// <param name="outputOffset">Offset within the buffer to write the converted data</param>
+		/// <returns></returns>
+		public static int ConvertToPlainText(ReadOnlySpan<byte> input, byte[] output, int outputOffset)
+		{
+			if (IsEmptyOrWhitespace(input))
+			{
+				output[outputOffset] = (byte)'\n';
+				return outputOffset + 1;
+			}
+
+			Utf8JsonReader reader = new Utf8JsonReader(input);
+			if (reader.Read() && reader.TokenType == JsonTokenType.StartObject)
+			{
+				while (reader.Read() && reader.TokenType == JsonTokenType.PropertyName)
+				{
+					if (!reader.ValueTextEquals("message"))
+					{
+						reader.Skip();
+						continue;
+					}
+					if (!reader.Read() || reader.TokenType != JsonTokenType.String)
+					{
+						reader.Skip();
+						continue;
+					}
+
+					int unescapedLength = UnescapeUtf8(reader.ValueSpan, output.AsSpan(outputOffset));
+					outputOffset += unescapedLength;
+
+					output[outputOffset] = (byte)'\n';
+					outputOffset++;
+
+					break;
+				}
+			}
+			return outputOffset;
+		}
+
+		/// <summary>
+		/// Unescape a json utf8 string
+		/// </summary>
+		/// <param name="source">Source span of bytes</param>
+		/// <param name="target">Target span of bytes</param>
+		/// <returns>Length of the converted data</returns>
+		static int UnescapeUtf8(ReadOnlySpan<byte> source, Span<byte> target)
+		{
+			int length = 0;
+			for (; ; )
+			{
+				// Copy up to the next backslash
+				int backslash = source.IndexOf((byte)'\\');
+				if (backslash == -1)
+				{
+					source.CopyTo(target);
+					length += source.Length;
+					break;
+				}
+				else if (backslash > 0)
+				{
+					source.Slice(0, backslash).CopyTo(target);
+					source = source.Slice(backslash);
+					target = target.Slice(backslash);
+					length += backslash;
+				}
+
+				// Check what the escape code is
+				if (source[1] == 'u')
+				{
+					char[] chars = { (char)((StringUtils.ParseHexByte(source, 2) << 8) | StringUtils.ParseHexByte(source, 4)) };
+					int encodedLength = Encoding.UTF8.GetBytes(chars.AsSpan(), target);
+					source = source.Slice(6);
+					target = target.Slice(encodedLength);
+					length += encodedLength;
+				}
+				else
+				{
+					target[0] = source[1] switch
+					{
+						(byte)'\"' => (byte)'\"',
+						(byte)'\\' => (byte)'\\',
+						(byte)'b' => (byte)'\b',
+						(byte)'f' => (byte)'\f',
+						(byte)'n' => (byte)'\n',
+						(byte)'r' => (byte)'\r',
+						(byte)'t' => (byte)'\t',
+						_ => source[1]
+					};
+					source = source.Slice(2);
+					target = target.Slice(1);
+					length++;
+				}
+			}
+			return length;
+		}
+
+		/// <summary>
+		/// Shrinks the data allocated by this chunk to the minimum required
+		/// </summary>
+		public void Shrink()
+		{
+			if (_data.Length > _length)
+			{
+				Array.Resize(ref _data, _length);
+			}
+		}
+
+		/// <summary>
+		/// Create a <see cref="LogChunkNode"/> object from the current state
+		/// </summary>
+		/// <returns></returns>
+		public LogChunkNode ToLogChunk() => new LogChunkNode(_data.AsMemory(0, _length).ToArray(), _lineOffsets.ToArray());
+	}
+
+	/// <summary>
+	/// Builds a sequence of log chunks
+	/// </summary>
+	class LogChunkSequenceBuilder
+	{
+		readonly List<LogChunkNode> _chunks = new List<LogChunkNode>();
+		readonly LogChunkBuilder _nextChunkBuilder = new LogChunkBuilder();
+		int _flushedLength;
+		int _flushedLineCount;
+
+		/// <summary>
+		/// Desired size for each chunk
+		/// </summary>
+		public int ChunkSize { get; }
+
+		/// <summary>
+		/// The complete chunks. Note that this does not include data which has not yet been flushed.
+		/// </summary>
+		public IReadOnlyList<LogChunkNode> Chunks => _chunks;
+
+		/// <summary>
+		/// Total length of the sequence
+		/// </summary>
+		public int Length => _flushedLength + _nextChunkBuilder.Length;
+
+		/// <summary>
+		/// Number of lines in this builder
+		/// </summary>
+		public int LineCount => _flushedLineCount + _nextChunkBuilder.LineCount;
+
+		/// <summary>
+		/// Constructor
+		/// </summary>
+		/// <param name="chunkSize">Desired size for each chunk. Each chunk will be limited to this size.</param>
+		public LogChunkSequenceBuilder(int chunkSize)
+		{
+			ChunkSize = chunkSize;
+		}
+
+		/// <summary>
+		/// Clear the current contents of the buffer
+		/// </summary>
+		public void Clear()
+		{
+			_chunks.Clear();
+			_nextChunkBuilder.Clear();
+			_flushedLength = 0;
+			_flushedLineCount = 0;
+		}
+
+		/// <summary>
+		/// Remove a number of chunks from the start of the builder
+		/// </summary>
+		/// <param name="count">Number of chunks to remove</param>
+		public void Remove(int count)
+		{
+			for (int idx = 0; idx < count; idx++)
+			{
+				LogChunkNode chunk = _chunks[idx];
+				_flushedLength -= chunk.Length;
+				_flushedLineCount -= chunk.LineCount;
+			}
+			_chunks.RemoveRange(0, count);
+		}
+
+		/// <inheritdoc cref="LogChunkBuilder.Append(ReadOnlySpan{Byte})"/>
+		public void Append(ReadOnlySpan<byte> textData)
+		{
+			if (textData.Length == 0)
+			{
+				return;
+			}
+			if (textData[^1] != (byte)'\n')
+			{
+				throw new ArgumentException("Text data to append must end with a newline", nameof(textData));
+			}
+
+			while (textData.Length > 0)
+			{
+				ReadOnlySpan<byte> lineData = textData.Slice(0, textData.IndexOf((byte)'\n') + 1);
+
+				CreateOutputSpace(lineData.Length);
+				_nextChunkBuilder.Append(lineData);
+
+				textData = textData.Slice((int)lineData.Length);
+			}
+		}
+
+		/// <inheritdoc cref="LogChunkBuilder.AppendJsonAsPlainText(ReadOnlySpan{Byte}, ILogger)"/>
+		public void AppendJsonAsPlainText(ReadOnlySpan<byte> textData, ILogger logger)
+		{
+			if (textData.Length == 0)
+			{
+				return;
+			}
+			if (textData[^1] != (byte)'\n')
+			{
+				throw new ArgumentException("Text data to append must end with a newline", nameof(textData));
+			}
+
+			while (textData.Length > 0)
+			{
+				ReadOnlySpan<byte> lineData = textData.Slice(0, textData.IndexOf((byte)'\n') + 1);
+
+				CreateOutputSpace(lineData.Length);
+				_nextChunkBuilder.AppendJsonAsPlainText(lineData, logger);
+
+				textData = textData.Slice((int)lineData.Length);
+			}
+		}
+
+		/// <summary>
+		/// Flushes the current contents of the builder
+		/// </summary>
+		public void Flush()
+		{
+			if (_nextChunkBuilder.Length > 0)
+			{
+				LogChunkNode nextChunk = _nextChunkBuilder.ToLogChunk();
+				_chunks.Add(nextChunk);
+				_nextChunkBuilder.Clear();
+				_flushedLength += nextChunk.Length;
+				_flushedLineCount += nextChunk.LineCount;
+			}
+		}
+
+		/// <summary>
+		/// Enumerate lines starting at the given index
+		/// </summary>
+		/// <param name="startIdx">Index to start from</param>
+		/// <returns>Sequence of lines</returns>
+		public IEnumerable<Utf8String> EnumerateLines(int startIdx = 0)
+		{
+			int lineIdx = startIdx;
+			foreach (LogChunkNode chunk in _chunks)
+			{
+				for (; lineIdx < chunk.LineCount; lineIdx++)
+				{
+					yield return chunk.GetLine(lineIdx);
+				}
+				lineIdx -= chunk.LineCount;
+			}
+		}
+
+		/// <summary>
+		/// Flushes the current chunk if necessary to provide the requested space
+		/// </summary>
+		/// <param name="requiredSpace">Space required in <see cref="_nextChunkBuilder"/></param>
+		void CreateOutputSpace(int requiredSpace)
+		{
+			if (_nextChunkBuilder.Length + requiredSpace > ChunkSize)
+			{
+				Flush();
+			}
+		}
+	}
+
+	/// <summary>
+	/// Extension methods for ILogText
+	/// </summary>
+	public static class LogChunkExtensions
+	{
+		/// <summary>
+		/// Converts a log text instance to plain text
+		/// </summary>
+		/// <param name="logText">The text to convert</param>
+		/// <param name="logger">Logger for conversion warnings</param>
+		/// <returns>The plain text instance</returns>
+		public static LogChunkNode ConvertJsonToPlainText(this LogChunkNode logText, ILogger logger)
+		{
+			LogChunkBuilder other = new LogChunkBuilder();
+			other.AppendJsonAsPlainText(logText, 0, logText.LineCount, logger);
+			return other.ToLogChunk();
+		}
+	}
+}
diff --git a/Engine/Source/Programs/Shared/EpicGames.Horde/Logs/LogIndexNode.cs b/Engine/Source/Programs/Shared/EpicGames.Horde/Logs/LogIndexNode.cs
new file mode 100644
--- /dev/null
+++ b/Engine/Source/Programs/Shared/EpicGames.Horde/Logs/LogIndexNode.cs
@@ -0,0 +1,413 @@
+// Copyright Epic Games, Inc. All Rights Reserved.
+
+using System;
+using System.Collections.Generic;
+using System.IO;
+using System.Linq;
+using System.Runtime.CompilerServices;
+using System.Threading;
+using EpicGames.Core;
+using EpicGames.Horde.Storage;
+using OpenTracing;
+using OpenTracing.Util;
+
+namespace EpicGames.Horde.Logs
+{
+	/// <summary>
+	/// Index for a log file.
+	/// 
+	/// The index consists of a sequence of compressed, plain text blocks (see <see cref="LogChunkRef"/>, and a 
+	/// set of (ngram, block index) pairs encoded as 64-bit integers (see <see cref="NgramSet"/>). 
+	/// 
+	/// Each ngram is a 1-4 byte sequence of utf8 bytes, padded out to 32-bits (see <see cref="Ngram"/>).
+	/// 
+	/// When performing a text search, the search term is split into a set of ngrams, and the set queried for blocks containing
+	/// them all. Matching blocks are decompressed and scanned for matches using a simplified Knuth-Morris-Pratt search.
+	/// 
+	/// Since alignment of ngrams may not match alignment of ngrams in the search term, we offset the search term by
+	/// 1-4 bytes and include the union of blocks matching at any offset.
+	/// </summary>
+	[TreeNode("{BAE1A00E-FD63-474E-A804-8081E20134F9}", 1)]
+	public class LogIndexNode : TreeNode
+	{
+		/// <summary>
+		/// Version number for serialized data
+		/// </summary>
+		const int CurrentVersion = 1;
+
+		/// <summary>
+		/// Index for tokens into the block list
+		/// </summary>
+		readonly NgramSet _ngramSet;
+
+		/// <summary>
+		/// Number of bits in the index devoted to the block index
+		/// </summary>
+		readonly int _numChunkBits;
+
+		/// <summary>
+		/// List of text blocks
+		/// </summary>
+		readonly LogChunkRef[] _plainTextChunkRefs;
+
+		/// <summary>
+		/// Empty index data
+		/// </summary>
+		public static LogIndexNode Empty { get; } = new LogIndexNode(NgramSet.Empty, 0, Array.Empty<LogChunkRef>());
+
+		/// <summary>
+		/// Public accessor for the plain text chunks
+		/// </summary>
+		public IReadOnlyList<LogChunkRef> PlainTextChunkRefs => _plainTextChunkRefs;
+
+		/// <summary>
+		/// Number of lines covered by the index
+		/// </summary>
+		public int LineCount => (_plainTextChunkRefs.Length > 0) ? (_plainTextChunkRefs[^1].LineIndex + _plainTextChunkRefs[^1].LineCount) : 0;
+
+		/// <summary>
+		/// Constructor
+		/// </summary>
+		/// <param name="ngramSet">Index into the text chunks</param>
+		/// <param name="numChunkBits">Number of bits devoted to the chunks index</param>
+		/// <param name="plainTextChunkRefs">Plain text chunks for this log file</param>
+		public LogIndexNode(NgramSet ngramSet, int numChunkBits, LogChunkRef[] plainTextChunkRefs)
+		{
+			_ngramSet = ngramSet;
+			_numChunkBits = numChunkBits;
+			_plainTextChunkRefs = plainTextChunkRefs;
+		}
+
+		/// <summary>
+		/// Deserialization constructor
+		/// </summary>
+		/// <param name="reader">Reader for data</param>
+		public LogIndexNode(ITreeNodeReader reader)
+		{
+			int version = (int)reader.ReadUnsignedVarInt();
+			if (version != CurrentVersion)
+			{
+				throw new InvalidDataException("Incorrect version number for log data");
+			}
+
+			_ngramSet = reader.ReadNgramSet();
+			_numChunkBits = (int)reader.ReadUnsignedVarInt();
+			_plainTextChunkRefs = reader.ReadVariableLengthArray(() => new LogChunkRef(reader));
+		}
+
+		/// <inheritdoc/>
+		public override void Serialize(ITreeNodeWriter writer)
+		{
+			writer.WriteUnsignedVarInt(CurrentVersion);
+			writer.WriteNgramSet(_ngramSet);
+			writer.WriteUnsignedVarInt(_numChunkBits);
+			writer.WriteVariableLengthArray(_plainTextChunkRefs, x => x.Serialize(writer));
+		}
+
+		/// <inheritdoc/>
+		public override IEnumerable<TreeNodeRef> EnumerateRefs() => _plainTextChunkRefs;
+
+		/// <summary>
+		/// Appends a set of text blocks to this index
+		/// </summary>
+		/// <param name="appendPlainTextChunks">Text blocks to append</param>
+		/// <returns>New log index with the given blocks appended</returns>
+		public LogIndexNode Append(IReadOnlyList<LogChunkNode> appendPlainTextChunks)
+		{
+			using IScope scope = GlobalTracer.Instance.BuildSpan("LogIndex.Append").StartActive();
+
+			// Create the new blocks array
+			LogChunkRef[] newChunks = new LogChunkRef[_plainTextChunkRefs.Length + appendPlainTextChunks.Count];
+			for (int idx = 0; idx < _plainTextChunkRefs.Length; idx++)
+			{
+				newChunks[idx] = _plainTextChunkRefs[idx];
+			}
+
+			int lineIndex = LineCount;
+			for (int idx = 0; idx < appendPlainTextChunks.Count; idx++)
+			{
+				LogChunkNode appendBlock = appendPlainTextChunks[idx];
+				newChunks[_plainTextChunkRefs.Length + idx] = new LogChunkRef(lineIndex, appendPlainTextChunks[idx]);
+				lineIndex += appendBlock.LineCount;
+			}
+
+			// Figure out how many bits to devote to the block size
+			int newNumBlockBits = 0;
+			while (newChunks.Length > (1 << newNumBlockBits))
+			{
+				newNumBlockBits += 4;
+			}
+
+			// Add the existing index data into a new ngram set
+			NgramSetBuilder newNgramSetBuilder = new NgramSetBuilder();
+
+			ulong blockMask = ((1UL << _numChunkBits) - 1);
+			foreach (ulong value in _ngramSet)
+			{
+				ulong token = value >> _numChunkBits;
+				ulong blockIdx = value & blockMask;
+
+				ulong newToken = (token << newNumBlockBits) | (blockIdx);
+				newNgramSetBuilder.Add(newToken);
+			}
+
+			// Add the new blocks to the index
+			for (int idx = 0; idx < appendPlainTextChunks.Count; idx++)
+			{
+				LogChunkNode appendTextBlock = appendPlainTextChunks[idx];
+				int blockIdx = _plainTextChunkRefs.Length + idx;
+				Ngram.Decompose(appendTextBlock.Span, token => newNgramSetBuilder.Add((token << newNumBlockBits) | (uint)blockIdx));
+			}
+
+			// Create the new index
+			NgramSet newNgramSet = newNgramSetBuilder.ToNgramSet();
+			return new LogIndexNode(newNgramSet, newNumBlockBits, newChunks);
+		}
+
+		/// <summary>
+		/// Search for the given text in the index
+		/// </summary>
+		/// <param name="reader">Reader for node data</param>
+		/// <param name="firstLineIndex">First line index to search from</param>
+		/// <param name="text">Text to search for</param>
+		/// <param name="stats">Receives stats for the search</param>
+		/// <param name="cancellationToken">Cancellation token for the operation</param>
+		/// <returns>List of line numbers for the text</returns>
+		public async IAsyncEnumerable<int> Search(TreeReader reader, int firstLineIndex, SearchTerm text, SearchStats stats, [EnumeratorCancellation] CancellationToken cancellationToken)
+		{
+			int lastBlockCount = 0;
+			foreach (int blockIdx in EnumeratePossibleChunks(text.Bytes, firstLineIndex))
+			{
+				LogChunkRef indexChunk = _plainTextChunkRefs[blockIdx];
+
+				stats.NumScannedBlocks++;
+				stats.NumDecompressedBlocks += (indexChunk.Target == null)? 1 : 0;
+
+				stats.NumSkippedBlocks += blockIdx - lastBlockCount;
+				lastBlockCount = blockIdx + 1;
+
+				// Decompress the text
+				LogChunkNode chunk = await indexChunk.ExpandAsync(reader, cancellationToken);
+
+				// Find the initial offset within this block
+				int offset = 0;
+				if(firstLineIndex > indexChunk.LineIndex)
+				{
+					int lineIndexWithinBlock = firstLineIndex - indexChunk.LineIndex;
+					offset = chunk.LineOffsets[lineIndexWithinBlock];
+				}
+
+				// Search within this block
+				for(; ;)
+				{
+					// Find the next offset
+					int nextOffset = chunk.Data.Span.FindNextOcurrence(offset, text);
+					if(nextOffset == -1)
+					{
+						stats.NumScannedBytes += chunk.Data.Length - offset;
+						break;
+					}
+
+					// Update the stats
+					stats.NumScannedBytes += nextOffset - offset;
+					offset = nextOffset;
+
+					// Check it's not another match within the same line
+					int lineIndexWithinBlock = chunk.GetLineIndexForOffset(offset);
+					yield return indexChunk.LineIndex + lineIndexWithinBlock;
+
+					// Move to the next line
+					offset = chunk.LineOffsets[lineIndexWithinBlock + 1];
+				}
+
+				// If the last scanned bytes is zero, we didn't have any matches from this chunk
+				if(offset == 0 && _ngramSet != null)
+				{
+					stats.NumFalsePositiveBlocks++;
+				}
+			}
+			stats.NumSkippedBlocks += _plainTextChunkRefs.Length - lastBlockCount;
+		}
+
+		/// <summary>
+		/// Search for the given text in the index
+		/// </summary>
+		/// <param name="text">Text to search for</param>
+		/// <param name="lineIndex">The first </param>
+		/// <returns>List of line numbers for the text</returns>
+		IEnumerable<int> EnumeratePossibleChunks(ReadOnlyMemory<byte> text, int lineIndex)
+		{
+			// Find the starting chunk index
+			int chunkIdx = _plainTextChunkRefs.BinarySearch(x => x.LineIndex, lineIndex);
+			if(chunkIdx < 0)
+			{
+				chunkIdx = Math.Max(~chunkIdx - 1, 0);
+			}
+
+			// Make sure we're not starting out of range
+			if (chunkIdx < _plainTextChunkRefs.Length)
+			{
+				// In order to be considered a possible chunk for a positive match, the chunk must contain an exact list of tokens parsed from the
+				// search text, along with a set of partial matches (for substrings that are not guaranteed to start on token boundaries). The former
+				// is relatively cheap to compute, and done by stepping through all enumerators for each token in a wave, only returning chunks
+				// which all contain the token. At the very least, this must contain chunks in the index from chunkIdx onwards.
+				List<IEnumerator<int>> enumerators = new List<IEnumerator<int>>();
+				enumerators.Add(Enumerable.Range(chunkIdx, _plainTextChunkRefs.Length - chunkIdx).GetEnumerator());
+
+				// The second aspect of the search requires a more expensive search through the trie. This is done by running a set of arbitrary 
+				// delegates to filter the matches returned from the enumerators.
+				List<Predicate<int>> predicates = new List<Predicate<int>>();
+
+				// If the index has a trie, tokenize the input and generate a list of enumerators and predicates.
+				if (_ngramSet != null)
+				{
+					// Find a list of filters for matching chunks
+					HashSet<ulong> tokens = new HashSet<ulong>();
+					for (int tokenPos = 0; tokenPos < text.Length;)
+					{
+						ReadOnlySpan<byte> token = Ngram.GetText(text.Span, tokenPos);
+						if (tokenPos == 0)
+						{
+							GetUnalignedTokenPredicate(token, token.Length == text.Length, predicates);
+						}
+						else
+						{
+							GetAlignedTokenPredicate(token, tokenPos + token.Length == text.Length, tokens, predicates);
+						}
+						tokenPos += token.Length;
+					}
+
+					// Create an enumerator for each token
+					foreach (ulong token in tokens)
+					{
+						ulong minValue = (token << _numChunkBits) | (ulong)(long)chunkIdx;
+						ulong maxValue = minValue + (1UL << _numChunkBits) - 1;
+
+						IEnumerator<int> enumerator = _ngramSet.EnumerateRange(minValue, maxValue).Select(x => (int)(x - minValue)).GetEnumerator();
+						if (!enumerator.MoveNext())
+						{
+							yield break;
+						}
+
+						enumerators.Add(enumerator);
+					}
+				}
+
+				// Enumerate the matches
+				for (; ; )
+				{
+					// Advance all the enumerators that are behind the current chunk index. If they are all equal, we have a match.
+					bool match = true;
+					foreach (IEnumerator<int> enumerator in enumerators)
+					{
+						while (enumerator.Current < chunkIdx)
+						{
+							if (!enumerator.MoveNext())
+							{
+								yield break;
+							}
+						}
+
+						if (enumerator.Current > chunkIdx)
+						{
+							chunkIdx = enumerator.Current;
+							match = false;
+						}
+					}
+
+					// Return the match and move to the next chunk
+					if (match)
+					{
+						if (predicates.All(predicate => predicate(chunkIdx)))
+						{
+							yield return chunkIdx;
+						}
+						chunkIdx++;
+					}
+				}
+			}
+		}
+
+		/// <summary>
+		/// Gets predicates for matching a token that starts 
+		/// </summary>
+		/// <param name="text">The token text</param>
+		/// <param name="allowPartialMatch">Whether to allow a partial match of the token</param>
+		/// <param name="tokens">Set of aligned tokens that are required</param>
+		/// <param name="predicates">List of predicates for the search</param>
+		void GetAlignedTokenPredicate(ReadOnlySpan<byte> text, bool allowPartialMatch, HashSet<ulong> tokens, List<Predicate<int>> predicates)
+		{
+			for (int offset = 0; offset < text.Length; offset += Ngram.MaxBytes)
+			{
+				ulong token = Ngram.GetWindowedValue(text, offset);
+				if (offset + Ngram.MaxBytes > text.Length && allowPartialMatch)
+				{
+					ulong tokenMask = Ngram.GetWindowedMask(text, offset, true);
+					predicates.Add(chunkIdx => ChunkContainsToken(chunkIdx, token, tokenMask));
+					break;
+				}
+				tokens.Add(token);
+			}
+		}
+
+		/// <summary>
+		/// Generates a predicate for matching a token which may or may not start on a regular token boundary
+		/// </summary>
+		/// <param name="text">The token text</param>
+		/// <param name="allowPartialMatch">Whether to allow a partial match of the token</param>
+		/// <param name="predicates">List of predicates for the search</param>
+		void GetUnalignedTokenPredicate(ReadOnlySpan<byte> text, bool allowPartialMatch, List<Predicate<int>> predicates)
+		{
+			byte[] textCopy = text.ToArray();
+
+			Lazy<HashSet<int>> lazyChunks = new Lazy<HashSet<int>>(() =>
+			{
+				HashSet<int> union = new HashSet<int>();
+				for (int shift = 0; shift < Ngram.MaxBytes; shift++)
+				{
+					HashSet<int> chunks = new HashSet<int>(ChunksContainingNgram(textCopy.AsSpan(), -shift, allowPartialMatch));
+					for (int offset = -shift + Ngram.MaxBytes; offset < textCopy.Length && chunks.Count > 0; offset += Ngram.MaxBytes)
+					{
+						chunks.IntersectWith(ChunksContainingNgram(textCopy.AsSpan(), offset, allowPartialMatch));
+					}
+					if (chunks.Count > 0)
+					{
+						union.UnionWith(chunks);
+					}
+				}
+				return union;
+			});
+
+			predicates.Add(chunkIdx => lazyChunks.Value.Contains(chunkIdx));
+		}
+
+		/// <summary>
+		/// Tests whether a chunk contains a particular token
+		/// </summary>
+		/// <param name="chunkIdx">Index of the chunk to search</param>
+		/// <param name="token">The token to test</param>
+		/// <param name="tokenMask">Mask of which bits in the token are valid</param>
+		/// <returns>True if the given block contains a token</returns>
+		bool ChunkContainsToken(int chunkIdx, ulong token, ulong tokenMask)
+		{
+			token = (token << _numChunkBits) | (uint)chunkIdx;
+			tokenMask = (tokenMask << _numChunkBits) | ((1UL << _numChunkBits) - 1);
+			return _ngramSet!.EnumerateValues((value, valueMask) => (value & tokenMask) == (token & valueMask)).Any();
+		}
+
+		/// <summary>
+		/// Tests whether a chunk contains a particular ngram
+		/// </summary>
+		/// <param name="text">The token to test</param>
+		/// <param name="offset">Offset of the window into the token to test</param>
+		/// <param name="allowPartialMatch">Whether to allow a partial match of the token</param>
+		/// <returns>True if the given block contains a token</returns>
+		IEnumerable<int> ChunksContainingNgram(ReadOnlySpan<byte> text, int offset, bool allowPartialMatch)
+		{
+			ulong token = Ngram.GetWindowedValue(text, offset) << _numChunkBits;
+			ulong tokenMask = Ngram.GetWindowedMask(text, offset, allowPartialMatch) << _numChunkBits;
+			ulong blockMask = (1UL << _numChunkBits) - 1;
+			return _ngramSet!.EnumerateValues((value, valueMask) => (value & tokenMask) == (token & valueMask)).Select(x => (int)(x & blockMask)).Distinct();
+		}
+	}
+}
diff --git a/Engine/Source/Programs/Shared/EpicGames.Horde/Logs/LogNode.cs b/Engine/Source/Programs/Shared/EpicGames.Horde/Logs/LogNode.cs
new file mode 100644
--- /dev/null
+++ b/Engine/Source/Programs/Shared/EpicGames.Horde/Logs/LogNode.cs
@@ -0,0 +1,351 @@
+// Copyright Epic Games, Inc. All Rights Reserved.
+
+using System;
+using System.Collections.Generic;
+using System.Data.Common;
+using System.Linq;
+using System.Reflection;
+using System.Threading.Tasks;
+using System.Threading;
+using EpicGames.Core;
+using EpicGames.Horde.Storage;
+using EpicGames.Serialization;
+using Microsoft.Extensions.Logging;
+using System.Diagnostics;
+using System.Runtime.CompilerServices;
+
+namespace EpicGames.Horde.Logs
+{
+	/// <summary>
+	/// Format for the log file
+	/// </summary>
+	public enum LogFormat
+	{
+		/// <summary>
+		/// Text data
+		/// </summary>
+		Text = 0,
+
+		/// <summary>
+		/// Json data
+		/// </summary>
+		Json = 1,
+	}
+
+	/// <summary>
+	/// Represents an entire log
+	/// </summary>
+	[TreeNode("{274DF8F7-9E87-4B4F-8AD5-318CDB25AD33}", 1)]
+	public class LogNode : TreeNode
+	{
+		/// <summary>
+		/// Default value for an empty log file
+		/// </summary>
+		public static LogNode Empty { get; } = new LogNode(LogFormat.Json, 0, Array.Empty<LogChunkRef>(), new TreeNodeRef<LogIndexNode>(LogIndexNode.Empty));
+
+		/// <summary>
+		/// Format for this log file
+		/// </summary>
+		public LogFormat Format { get; }
+
+		/// <summary>
+		/// Total number of lines
+		/// </summary>
+		public int LineCount { get; }
+
+		/// <summary>
+		/// Text blocks for this chunk
+		/// </summary>
+		public IReadOnlyList<LogChunkRef> TextChunkRefs { get; }
+
+		/// <summary>
+		/// Index for this log
+		/// </summary>
+		public TreeNodeRef<LogIndexNode> IndexRef { get; }
+
+		/// <summary>
+		/// Deserializing constructor
+		/// </summary>
+		public LogNode(LogFormat format, int lineCount, IReadOnlyList<LogChunkRef> textChunkRefs, TreeNodeRef<LogIndexNode> indexRef)
+		{
+			Format = format;
+			LineCount = lineCount;
+			TextChunkRefs = textChunkRefs.ToArray();
+			IndexRef = indexRef;
+		}
+
+		/// <summary>
+		/// Deserializing constructor
+		/// </summary>
+		/// <param name="reader">Reader to draw data from</param>
+		public LogNode(ITreeNodeReader reader)
+		{
+			Format = (LogFormat)reader.ReadUInt8();
+			LineCount = (int)reader.ReadUnsignedVarInt();
+			IndexRef = reader.ReadRef<LogIndexNode>();
+			TextChunkRefs = reader.ReadList(() => new LogChunkRef(reader));
+		}
+
+		/// <inheritdoc/>
+		public override void Serialize(ITreeNodeWriter writer)
+		{
+			writer.WriteUInt8((byte)Format);
+			writer.WriteUnsignedVarInt(LineCount);
+			writer.WriteRef(IndexRef);
+			writer.WriteList(TextChunkRefs, x => x.Serialize(writer));
+		}
+
+		/// <inheritdoc/>
+		public override IEnumerable<TreeNodeRef> EnumerateRefs()
+		{
+			yield return IndexRef;
+
+			foreach (TreeNodeRef<LogChunkNode> textChunkRef in TextChunkRefs)
+			{
+				yield return textChunkRef;
+			}
+		}
+	}
+
+	/// <summary>
+	/// Assists building log files through trees of <see cref="LogNode"/>, <see cref="LogIndexNode"/> and <see cref="LogChunkNode"/> nodes.
+	/// </summary>
+	class LogBuilder
+	{
+		/// <summary>
+		/// Default maximum size for a log text block
+		/// </summary>
+		public const int DefaultTextBlockLength = 256 * 1024;
+
+		/// <summary>
+		/// Default maximum size for an index text block
+		/// </summary>
+		public const int DefaultIndexBlockLength = 64 * 1024;
+
+		/// <summary>
+		/// Number of lines written to the log
+		/// </summary>
+		public int LineCount => _root.LineCount + _textBuilder.LineCount;
+
+		/// <summary>
+		/// Number of lines flushed to storage
+		/// </summary>
+		public int FlushedLineCount => _root.LineCount;
+
+		readonly LogFormat _format;
+
+		// Data for the log file which has been flushed to disk so far
+		LogNode _root = LogNode.Empty;
+		LogIndexNode _index = LogIndexNode.Empty;
+
+		// Json data read but not flushed
+		readonly LogChunkSequenceBuilder _textBuilder;
+		readonly LogChunkSequenceBuilder _indexTextBuilder;
+
+		// Lock object for access to the fields above
+		readonly object _lockObject = new object();
+
+		// Inner log device; used for messaging encoding errors.
+		readonly ILogger _logger;
+
+		/// <summary>
+		/// Constructor
+		/// </summary>
+		/// <param name="format">Format for data in the log file</param>
+		/// <param name="logger"></param>
+		public LogBuilder(LogFormat format, ILogger logger)
+			: this(format, DefaultTextBlockLength, DefaultIndexBlockLength, logger)
+		{
+		}
+
+		/// <summary>
+		/// Constructor
+		/// </summary>
+		/// <param name="format">Format of data in the log file</param>
+		/// <param name="maxTextBlockLength">maximum size for a regular text block</param>
+		/// <param name="maxIndexBlockLength">Maximum size for an index text block</param>
+		/// <param name="logger">Logger for conversion errors</param>
+		public LogBuilder(LogFormat format, int maxTextBlockLength, int maxIndexBlockLength, ILogger logger)
+		{
+			_format = format;
+			_textBuilder = new LogChunkSequenceBuilder(maxTextBlockLength);
+			_indexTextBuilder = new LogChunkSequenceBuilder(maxIndexBlockLength);
+			_logger = logger;
+		}
+
+		/// <summary>
+		/// Read data from the unflushed log tail
+		/// </summary>
+		/// <param name="firstLineIdx">The first line to read, from the end of the flushed data</param>
+		/// <param name="maxLength"></param>
+		/// <returns></returns>
+		public ReadOnlyMemory<byte> ReadTailData(int firstLineIdx, int maxLength)
+		{
+			lock (_lockObject)
+			{
+				// Measure the size of buffer required for the tail data
+				int length = 0;
+				int lineCount = 0;
+				foreach (Utf8String line in _textBuilder.EnumerateLines(firstLineIdx))
+				{
+					int nextLength = length + line.Length;
+					if (length > 0 && nextLength > maxLength)
+					{
+						break;
+					}
+					length = nextLength;
+					lineCount++;
+				}
+
+				// Allocate the buffer
+				byte[] buffer = new byte[length];
+
+				// Copy lines into the buffer
+				Span<byte> output = buffer.AsSpan();
+				foreach (Utf8String line in _textBuilder.EnumerateLines(firstLineIdx).Take(lineCount))
+				{
+					line.Span.CopyTo(output);
+					output = output.Slice(line.Length);
+				}
+				Debug.Assert(output.Length == 0);
+
+				return buffer;
+			}
+		}
+
+		/// <summary>
+		/// Append JSON data to the end of the log
+		/// </summary>
+		/// <param name="data">Log data to append</param>
+		public void WriteData(ReadOnlyMemory<byte> data)
+		{
+			lock (_lockObject)
+			{
+				ReadOnlyMemory<byte> remaining = data;
+				for (; ; )
+				{
+					int newlineIdx = remaining.Span.IndexOf((byte)'\n');
+					if (newlineIdx == -1)
+					{
+						break;
+					}
+
+					ReadOnlyMemory<byte> line = remaining.Slice(0, newlineIdx + 1);
+					_textBuilder.Append(line.Span);
+
+					if (_format == LogFormat.Json)
+					{
+						_indexTextBuilder.AppendJsonAsPlainText(line.Span, _logger);
+					}
+					else
+					{
+						_indexTextBuilder.Append(line.Span);
+					}
+
+					remaining = remaining.Slice(newlineIdx + 1);
+				}
+			}
+		}
+
+		/// <summary>
+		/// Flushes the written data to the log
+		/// </summary>
+		/// <param name="writer">Writer for the output nodes</param>
+		/// <param name="cancellationToken">Cancellation token for the operation</param>
+		public async Task<NodeLocator> FlushAsync(TreeWriter writer, CancellationToken cancellationToken)
+		{
+			// Capture the new data that needs to be written
+			IReadOnlyList<LogChunkNode> writeTextChunks;
+			IReadOnlyList<LogChunkNode> writeIndexTextChunks;
+
+			lock (_lockObject)
+			{
+				_textBuilder.Flush();
+				writeTextChunks = _textBuilder.Chunks.ToArray();
+
+				_indexTextBuilder.Flush();
+				writeIndexTextChunks = _indexTextBuilder.Chunks.ToArray();
+			}
+
+			// Flush any complete chunks to storage
+			LogIndexNode newIndex = _index.Append(writeIndexTextChunks);
+
+			List<LogChunkRef> newJsonChunkRefs = new List<LogChunkRef>(_root.TextChunkRefs);
+			int lineCount = _root.LineCount;
+			foreach (LogChunkNode writeTextChunk in writeTextChunks)
+			{
+				newJsonChunkRefs.Add(new LogChunkRef(lineCount, writeTextChunk));
+				lineCount += writeTextChunk.LineCount;
+			}
+
+			LogNode newRoot = new LogNode(_format, newIndex.LineCount, newJsonChunkRefs, new TreeNodeRef<LogIndexNode>(newIndex));
+
+			IoHash newRootHash = await writer.WriteAsync(new TreeNodeRef(newRoot), cancellationToken);
+			await writer.FlushAsync(cancellationToken);
+
+			NodeLocator newRootLocator = writer.GetLocator(newRootHash);
+
+			// Update the new state
+			lock (_lockObject)
+			{
+				_root = newRoot;
+				_index = newIndex;
+				_textBuilder.Remove(writeTextChunks.Count);
+				_indexTextBuilder.Remove(writeIndexTextChunks.Count);
+			}
+
+			return newRootLocator;
+		}
+	}
+
+	/// <summary>
+	/// Extension methods
+	/// </summary>
+	public static class LogNodeExtensions
+	{
+		/// <summary>
+		/// Reads lines from a line
+		/// </summary>
+		/// <param name="logNode">Log to read from</param>
+		/// <param name="reader">Reader to pull nodes from</param>
+		/// <param name="cancellationToken">Cancellation token</param>
+		/// <returns>Sequence of line buffers</returns>
+		public static async IAsyncEnumerable<ReadOnlyMemory<byte>> ReadLogAsync(this LogNode logNode, TreeReader reader, [EnumeratorCancellation] CancellationToken cancellationToken)
+		{
+			foreach (LogChunkRef textChunkRef in logNode.TextChunkRefs)
+			{
+				LogChunkNode textChunk = await textChunkRef.ExpandAsync(reader, cancellationToken);
+				yield return textChunk.Data;
+			}
+		}
+
+		/// <summary>
+		/// Reads lines from a line
+		/// </summary>
+		/// <param name="logNode">Log to read from</param>
+		/// <param name="reader">Reader to pull nodes from</param>
+		/// <param name="index">Zero-based index of the first line to read from</param>
+		/// <param name="cancellationToken">Cancellation token</param>
+		/// <returns>Sequence of line buffers</returns>
+		public static async IAsyncEnumerable<ReadOnlyMemory<byte>> ReadLogLinesAsync(this LogNode logNode, TreeReader reader, int index, [EnumeratorCancellation] CancellationToken cancellationToken)
+		{
+			foreach (LogChunkRef textChunkRef in logNode.TextChunkRefs)
+			{
+				int lineIdx = Math.Max(index - textChunkRef.LineIndex, 0);
+				if (lineIdx < textChunkRef.LineCount)
+				{
+					LogChunkNode textChunk = await textChunkRef.ExpandAsync(reader, cancellationToken);
+
+					int offset = textChunk.LineOffsets[lineIdx];
+					for (; lineIdx < textChunk.LineCount; lineIdx++)
+					{
+						int nextOffset = textChunk.LineOffsets[lineIdx + 1];
+						ReadOnlyMemory<byte> line = textChunk.Data.Slice(offset, nextOffset - offset);
+						yield return line;
+						offset = nextOffset;
+					}
+				}
+			}
+		}
+	}
+}
diff --git a/Engine/Source/Programs/Shared/EpicGames.Horde/Logs/Ngram.cs b/Engine/Source/Programs/Shared/EpicGames.Horde/Logs/Ngram.cs
new file mode 100644
--- /dev/null
+++ b/Engine/Source/Programs/Shared/EpicGames.Horde/Logs/Ngram.cs
@@ -0,0 +1,192 @@
+// Copyright Epic Games, Inc. All Rights Reserved.
+
+using System;
+
+namespace EpicGames.Horde.Logs
+{
+	/// <summary>
+	/// Functionality for decomposing log text into ngrams.
+	/// </summary>
+	public static class Ngram
+	{
+		/// <summary>
+		/// Maximum number of bytes in each ngram
+		/// </summary>
+		public const int MaxBytes = 4;
+
+		/// <summary>
+		/// Number of bits in each ngram
+		/// </summary>
+		public const int MaxBits = MaxBytes * 8;
+
+		/// <summary>
+		/// Lookup from input byte to token type
+		/// </summary>
+		static readonly byte[] s_tokenTypes = GetTokenTypes();
+
+		/// <summary>
+		/// Lookup from input byte to token char
+		/// </summary>
+		static readonly byte[] s_tokenChars = GetTokenChars();
+
+		/// <summary>
+		/// Gets a single token
+		/// </summary>
+		/// <param name="text">The text to parse</param>
+		/// <returns>The token value</returns>
+		public static ulong Read(ReadOnlySpan<byte> text)
+		{
+			ulong token = 0;
+			Decompose(text, x => token = x);
+			return token;
+		}
+
+		/// <summary>
+		/// Decompose a span of text into tokens
+		/// </summary>
+		/// <param name="text">Text to scan</param>
+		/// <param name="addNgram">Receives a set of tokens</param>
+		public static void Decompose(ReadOnlySpan<byte> text, Action<ulong> addNgram)
+		{
+			if (text.Length > 0)
+			{
+				int type = s_tokenTypes[text[0]];
+				int numTokenBits = 8;
+				ulong token = s_tokenChars[text[0]];
+
+				for (int textIdx = 1; textIdx < text.Length; textIdx++)
+				{
+					byte nextChar = s_tokenChars[text[textIdx]];
+					int nextType = s_tokenTypes[nextChar];
+					if (type != nextType || numTokenBits + 8 > MaxBits)
+					{
+						addNgram(token << (MaxBits - numTokenBits));
+						token = 0;
+						numTokenBits = 0;
+						type = nextType;
+					}
+					token = (token << 8) | nextChar;
+					numTokenBits += 8;
+				}
+
+				addNgram(token << (MaxBits - numTokenBits));
+			}
+		}
+
+		/// <summary>
+		/// Decompose a string to a set of ngrams
+		/// </summary>
+		/// <param name="text">Text to scan</param>
+		/// <returns></returns>
+		public static NgramSet Decompose(ReadOnlySpan<byte> text)
+		{
+			NgramSetBuilder builder = new NgramSetBuilder();
+			Decompose(text, builder.Add);
+			return builder.ToNgramSet();
+		}
+
+		/// <summary>
+		/// Gets the length of the first token in the given span
+		/// </summary>
+		/// <param name="text">The text to search</param>
+		/// <param name="pos">Start position for the search</param>
+		/// <returns>Length of the first token</returns>
+		public static ReadOnlySpan<byte> GetText(ReadOnlySpan<byte> text, int pos)
+		{
+			int type = s_tokenTypes[text[pos]];
+			for (int end = pos + 1; ; end++)
+			{
+				if (end == text.Length || s_tokenTypes[text[end]] != type)
+				{
+					return text.Slice(pos, end - pos);
+				}
+			}
+		}
+
+		/// <summary>
+		/// Gets the length of the first token in the given span
+		/// </summary>
+		/// <param name="text">The text to search</param>
+		/// <param name="offset">Offset of the window to read from the token</param>
+		/// <returns>Length of the first token</returns>
+		public static ulong GetWindowedValue(ReadOnlySpan<byte> text, int offset)
+		{
+			ulong token = 0;
+			for (int idx = 0; idx < MaxBytes; idx++)
+			{
+				token <<= 8;
+				if (offset >= 0 && offset < text.Length)
+				{
+					token |= s_tokenChars[text[offset]];
+				}
+				offset++;
+			}
+			return token;
+		}
+
+		/// <summary>
+		/// Gets the length of the first token in the given span
+		/// </summary>
+		/// <param name="text">The text to search</param>
+		/// <param name="offset">Offset of the window to read from the token</param>
+		/// <param name="allowPartialMatch">Whether to allow only matching the start of the string</param>
+		/// <returns>Length of the first token</returns>
+		public static ulong GetWindowedMask(ReadOnlySpan<byte> text, int offset, bool allowPartialMatch)
+		{
+			ulong token = 0;
+			for (int idx = 0; idx < MaxBytes; idx++)
+			{
+				token <<= 8;
+				if (offset >= 0 && (offset < text.Length || !allowPartialMatch))
+				{
+					token |= 0xff;
+				}
+				offset++;
+			}
+			return token;
+		}
+
+		/// <summary>
+		/// Build the lookup table for token types
+		/// </summary>
+		/// <returns>Array whose elements map from an input byte to token type</returns>
+		static byte[] GetTokenTypes()
+		{
+			byte[] charTypes = new byte[256];
+			for (int idx = 'a'; idx <= 'z'; idx++)
+			{
+				charTypes[idx] = 1;
+			}
+			for (int idx = 'A'; idx <= 'Z'; idx++)
+			{
+				charTypes[idx] = 1;
+			}
+			for (int idx = '0'; idx <= '9'; idx++)
+			{
+				charTypes[idx] = 2;
+			}
+			charTypes[' '] = 3;
+			charTypes['\t'] = 3;
+			charTypes['\n'] = 4;
+			return charTypes;
+		}
+
+		/// <summary>
+		/// Build the lookup table for token types
+		/// </summary>
+		/// <returns>Array whose elements map from an input byte to token type</returns>
+		static byte[] GetTokenChars()
+		{
+			byte[] chars = new byte[256];
+			for(int idx = 0; idx < 256; idx++)
+			{
+				chars[idx] = (byte)idx;
+			}
+			for (int idx = 'A'; idx <= 'Z'; idx++)
+			{
+				chars[idx] = (byte)('a' + idx - 'A');
+			}
+			return chars;
+		}
+	}
+}
diff --git a/Engine/Source/Programs/Shared/EpicGames.Horde/Logs/NgramSet.cs b/Engine/Source/Programs/Shared/EpicGames.Horde/Logs/NgramSet.cs
new file mode 100644
--- /dev/null
+++ b/Engine/Source/Programs/Shared/EpicGames.Horde/Logs/NgramSet.cs
@@ -0,0 +1,312 @@
+// Copyright Epic Games, Inc. All Rights Reserved.
+
+using System;
+using System.Collections;
+using System.Collections.Generic;
+using System.Runtime.InteropServices;
+using EpicGames.Core;
+
+namespace EpicGames.Horde.Logs
+{
+	/// <summary>
+	/// A sparse, space-efficient set of 64-bit values. Implemented as a trie backed by a flat lookup table.
+	/// 
+	/// Each 64-bit value in the set is decomposed into 4-bit fragments, and each node in the trie contains a 2^4=16-bit mask 
+	/// indicating which child nodes exist. The array of nodes is stored in a flat buffer, with a predictable order, with 
+	/// the children of a particular node stored contiguously, breadth first. 
+	/// 
+	/// Doing so allows constructing a lookup table for the first child of each parent node with a single pass of
+	/// the buffer, allowing efficient traversal of the tree to satisfy queries.
+	/// 
+	/// In practice, only the top 32-bits of values stored in the trie are used for encoding ngram information. The
+	/// bottom 32 bits are used to index a block number, allowing querying the existence of ngrams and their
+	/// rough location.
+	/// </summary>
+	public class NgramSet : IEnumerable<ulong>
+	{
+		/// <summary>
+		/// Stack item for traversing the tree
+		/// </summary>
+		struct StackItem
+		{
+			/// <summary>
+			/// The current node index
+			/// </summary>
+			public int _index;
+
+			/// <summary>
+			/// Value in the current node (0-15)
+			/// </summary>
+			public int _value;
+		}
+
+		/// <summary>
+		/// Delegate for filtering values during a tree traversal
+		/// </summary>
+		/// <param name="value">The current value</param>
+		/// <param name="mask">Mask for which bits in the value are valid</param>
+		/// <returns>True if values matching the given mask should be enumerated</returns>
+		public delegate bool VisitorDelegate(ulong value, ulong mask);
+
+		/// <summary>
+		/// Height of the tree
+		/// </summary>
+		const int Height = sizeof(ulong) * 2;
+
+		/// <summary>
+		/// Array of bitmasks for each node in the tree
+		/// </summary>
+		public IReadOnlyList<ushort> NodeData => _nodeData;
+
+		/// <summary>
+		/// Array of bitmasks for each node in the tree
+		/// </summary>
+		readonly ushort[] _nodeData;
+
+		/// <summary>
+		/// Array of child offsets for each node. Excludes the last layer of the tree.
+		/// </summary>
+		readonly int[] _firstChildIndex;
+
+		/// <summary>
+		/// Empty index definition
+		/// </summary>
+		public static NgramSet Empty { get; } = new NgramSet(new ushort[1]);
+
+		/// <summary>
+		/// Constructor
+		/// </summary>
+		/// <param name="nodeData">Node data</param>
+		public NgramSet(ushort[] nodeData)
+		{
+			_nodeData = nodeData;
+			_firstChildIndex = CreateChildLookup(nodeData);
+		}
+
+		/// <summary>
+		/// Tests whether the given value is in the trie
+		/// </summary>
+		/// <param name="value">The value to check for</param>
+		/// <returns>True if the value is in the trie</returns>
+		public bool Contains(ulong value)
+		{
+			int index = 0;
+			for (int shift = (sizeof(ulong) * 8) - 4; shift >= 0; shift -= 4)
+			{
+				int mask = _nodeData[index];
+				int flag = 1 << (int)((value >> shift) & 15);
+				if ((mask & flag) == 0)
+				{
+					return false;
+				}
+
+				index = _firstChildIndex[index];
+				for (; ; )
+				{
+					mask &= (mask - 1);
+					if ((mask & flag) == 0)
+					{
+						break;
+					}
+					index++;
+				}
+			}
+			return true;
+		}
+
+		/// <inheritdoc/>
+		public IEnumerator<ulong> GetEnumerator()
+		{
+			return EnumerateRange(UInt64.MinValue, UInt64.MaxValue).GetEnumerator();
+		}
+
+		/// <inheritdoc/>
+		IEnumerator IEnumerable.GetEnumerator() => GetEnumerator();
+
+		/// <summary>
+		/// Enumerate all values matching a given filter
+		/// </summary>
+		/// <param name="predicate">Predicate for which values to include</param>
+		/// <returns>Values satisfying the given predicate</returns>
+		public IEnumerable<ulong> EnumerateValues(VisitorDelegate predicate)
+		{
+			int depth = 0;
+			ulong value = 0;
+
+			StackItem[] stack = new StackItem[Height];
+			stack[1]._index = _firstChildIndex[0];
+
+			for (; ; )
+			{
+				StackItem current = stack[depth];
+				if (current._value >= 16)
+				{
+					// Move up the tree if we've enumerated all the branches at the current level
+					depth--;
+					if (depth < 0)
+					{
+						yield break;
+					}
+					stack[depth]._value++;
+
+					// Increment the child index too. These are stored sequentially for a given parent. The value will be cleared when we recurse into it.
+					stack[depth + 1]._index++;
+				}
+				else if ((_nodeData[current._index] & (1 << current._value)) == 0)
+				{
+					// This branch does not exist. Skip it.
+					stack[depth]._value++;
+				}
+				else
+				{
+					// Get the value and mask for the current node
+					int shift = (stack.Length - depth - 1) * 4;
+					ulong mask = ~((1UL << shift) - 1);
+					value = (value & (mask << 4)) | ((ulong)(uint)current._value << shift);
+
+					if (!predicate(value, mask))
+					{
+						// This node is excluded, skip it
+						stack[depth]._value++;
+						if (depth + 1 < stack.Length)
+						{
+							stack[depth + 1]._index++;
+						}
+					}
+					else if(depth + 1 < stack.Length)
+					{
+						// Move down the tree
+						depth++;
+						stack[depth]._value = 0;
+						if (depth + 1 < stack.Length)
+						{
+							stack[depth + 1]._index = _firstChildIndex[stack[depth]._index];
+						}
+					}
+					else
+					{
+						// Yield the current value
+						yield return value;
+						stack[depth]._value++;
+					}
+				}
+			}
+		}
+
+		/// <summary>
+		/// Enumerates all values in the trie between the given ranges
+		/// </summary>
+		/// <param name="minValue">Minimum value to enumerate</param>
+		/// <param name="maxValue">Maximum value to enumerate</param>
+		/// <returns>Sequence of values</returns>
+		public IEnumerable<ulong> EnumerateRange(ulong minValue, ulong maxValue)
+		{
+			return EnumerateValues((value, mask) => (value >= (minValue & mask) && value <= (maxValue & mask)));
+		}
+
+		/// <summary>
+		/// Creates a lookup for child node offsets from raw node data
+		/// </summary>
+		/// <param name="nodeData">Array of masks for each node</param>
+		/// <returns>Array of offsets</returns>
+		static int[] CreateChildLookup(ushort[] nodeData)
+		{
+			List<int> childOffsets = new List<int>();
+			if (nodeData.Length > 0)
+			{
+				int nodeCount = 1;
+
+				int index = 0;
+				int childIndex = nodeCount;
+
+				for (int level = 0; level < Height; level++)
+				{
+					int nextNodeCount = 0;
+					for (int idx = 0; idx < nodeCount; idx++)
+					{
+						ushort node = nodeData[index++];
+
+						int numChildren = CountBits(node);
+						childOffsets.Add(childIndex);
+						childIndex += numChildren;
+
+						nextNodeCount += numChildren;
+					}
+					nodeCount = nextNodeCount;
+				}
+			}
+			return childOffsets.ToArray();
+		}
+
+		/// <summary>
+		/// Count the number of set bits in the given value
+		/// </summary>
+		/// <param name="value">Value to test</param>
+		/// <returns>Number of set bits</returns>
+		static int CountBits(ushort value)
+		{
+			int count = value;
+			count = (count & 0b0101010101010101) + ((count >> 1) & 0b0101010101010101);
+			count = (count & 0b0011001100110011) + ((count >> 2) & 0b0011001100110011);
+			count = (count & 0b0000111100001111) + ((count >> 4) & 0b0000111100001111);
+			count = (count & 0b0000000011111111) + ((count >> 8) & 0b0000000011111111);
+			return count;
+		}
+
+		/// <summary>
+		/// Read a trie from the given buffer
+		/// </summary>
+		/// <param name="reader">Reader to read from</param>
+		/// <returns>New trie</returns>
+		public static NgramSet Read(IMemoryReader reader)
+		{
+			ReadOnlyMemory<byte> nodes = reader.ReadVariableLengthBytesWithInt32Length();
+			ushort[] nodeData = MemoryMarshal.Cast<byte, ushort>(nodes.Span).ToArray();
+			return new NgramSet(nodeData);
+		}
+
+		/// <summary>
+		/// Write this trie to the given buffer
+		/// </summary>
+		/// <param name="writer">Writer to output to</param>
+		public void Write(IMemoryWriter writer)
+		{
+			writer.WriteVariableLengthBytesWithInt32Length(MemoryMarshal.AsBytes<ushort>(_nodeData));
+		}
+
+		/// <summary>
+		/// Gets the serialized size of this trie
+		/// </summary>
+		/// <returns></returns>
+		public int GetSerializedSize()
+		{
+			return (sizeof(int) + _nodeData.Length * sizeof(ushort));
+		}
+	}
+
+	/// <summary>
+	/// Extension methods for serializing tries
+	/// </summary>
+	public static class NgramSetExtensions
+	{
+		/// <summary>
+		/// Read a trie from the given buffer
+		/// </summary>
+		/// <param name="reader">Reader to read from</param>
+		/// <returns>New trie</returns>
+		public static NgramSet ReadNgramSet(this IMemoryReader reader)
+		{
+			return NgramSet.Read(reader);
+		}
+
+		/// <summary>
+		/// Write this trie to the given buffer
+		/// </summary>
+		/// <param name="writer">Writer to output to</param>
+		/// <param name="set">Trie to write</param>
+		public static void WriteNgramSet(this IMemoryWriter writer, NgramSet set)
+		{
+			set.Write(writer);
+		}
+	}
+}
diff --git a/Engine/Source/Programs/Shared/EpicGames.Horde/Logs/NgramSetBuilder.cs b/Engine/Source/Programs/Shared/EpicGames.Horde/Logs/NgramSetBuilder.cs
new file mode 100644
--- /dev/null
+++ b/Engine/Source/Programs/Shared/EpicGames.Horde/Logs/NgramSetBuilder.cs
@@ -0,0 +1,110 @@
+// Copyright Epic Games, Inc. All Rights Reserved.
+
+using System.Collections.Generic;
+
+namespace EpicGames.Horde.Logs
+{
+	/// <summary>
+	/// Structure used for building compact <see cref="NgramSet"/> instances
+	/// </summary>
+	public class NgramSetBuilder
+	{
+		/// <summary>
+		/// Node within the trie
+		/// </summary>
+		class Node
+		{
+			public Node?[]? _children;
+		}
+
+		/// <summary>
+		/// The root node
+		/// </summary>
+		readonly Node _root;
+
+		/// <summary>
+		/// Default constructor
+		/// </summary>
+		public NgramSetBuilder()
+		{
+			_root = new Node();
+		}
+
+		/// <summary>
+		/// Adds a value to the trie
+		/// </summary>
+		/// <param name="value">Value to add</param>
+		public void Add(ulong value)
+		{
+			// Loop through the tree until we've added the item
+			Node leaf = _root;
+			for (int shift = (sizeof(ulong) * 8) - 4; shift >= 0; shift -= 4)
+			{
+				int index = (int)(value >> shift) & 15;
+				leaf._children ??= new Node[16];
+				leaf._children[index] ??= new Node();
+				leaf = leaf._children[index]!;
+			}
+		}
+
+		/// <summary>
+		/// Searches for the given item in the trie
+		/// </summary>
+		/// <param name="value">Value to add</param>
+		public bool Contains(ulong value)
+		{
+			// Loop through the tree until we've added the item
+			Node leaf = _root;
+			for (int shift = (sizeof(ulong) * 8) - 4; shift >= 0; shift -= 4)
+			{
+				int index = (int)(value >> shift) & 15;
+				if (leaf._children == null)
+				{
+					return false;
+				}
+				if (leaf._children[index] == null)
+				{
+					return false;
+				}
+				leaf = leaf._children[index]!;
+			}
+			return true;
+		}
+
+		/// <summary>
+		/// Creates a <see cref="NgramSet"/> from this data
+		/// </summary>
+		/// <returns></returns>
+		public NgramSet ToNgramSet()
+		{
+			List<ushort> values = new List<ushort>();
+
+			List<Node> nodes = new List<Node>();
+			nodes.Add(_root);
+
+			for(int bits = 0; bits < (sizeof(ulong) * 8); bits += 4)
+			{
+				List<Node> nextNodes = new List<Node>();
+				foreach(Node node in nodes)
+				{
+					ushort value = 0;
+					if (node._children != null)
+					{
+						for (int idx = 0; idx < node._children.Length; idx++)
+						{
+							if(node._children[idx] != null)
+							{
+								value |= (ushort)(1 << idx);
+								nextNodes.Add(node._children[idx]!);
+							}
+						}
+					}
+					values.Add(value);
+				}
+				nodes = nextNodes;
+			}
+
+			return new NgramSet(values.ToArray());
+		}
+	}
+}
diff --git a/Engine/Source/Programs/Shared/EpicGames.Horde/Logs/SearchStats.cs b/Engine/Source/Programs/Shared/EpicGames.Horde/Logs/SearchStats.cs
new file mode 100644
--- /dev/null
+++ b/Engine/Source/Programs/Shared/EpicGames.Horde/Logs/SearchStats.cs
@@ -0,0 +1,35 @@
+// Copyright Epic Games, Inc. All Rights Reserved.
+
+namespace EpicGames.Horde.Logs
+{
+	/// <summary>
+	/// Stats for a search
+	/// </summary>
+	public class SearchStats
+	{
+		/// <summary>
+		/// Number of blocks that were scanned
+		/// </summary>
+		public int NumScannedBlocks { get; set; }
+
+		/// <summary>
+		/// Number of bytes that had to be scanned for results
+		/// </summary>
+		public int NumScannedBytes { get; set; }
+
+		/// <summary>
+		/// Number of blocks that were skipped
+		/// </summary>
+		public int NumSkippedBlocks { get; set; }
+
+		/// <summary>
+		/// Number of blocks that had to be decompressed
+		/// </summary>
+		public int NumDecompressedBlocks { get; set; }
+
+		/// <summary>
+		/// Number of blocks that were searched but did not contain the search term
+		/// </summary>
+		public int NumFalsePositiveBlocks { get; set; }
+	}
+}
diff --git a/Engine/Source/Programs/Shared/EpicGames.Horde/Logs/SearchTerm.cs b/Engine/Source/Programs/Shared/EpicGames.Horde/Logs/SearchTerm.cs
new file mode 100644
--- /dev/null
+++ b/Engine/Source/Programs/Shared/EpicGames.Horde/Logs/SearchTerm.cs
@@ -0,0 +1,185 @@
+// Copyright Epic Games, Inc. All Rights Reserved.
+
+using System;
+using System.Collections.Generic;
+using System.Text;
+
+namespace EpicGames.Horde.Logs
+{
+	/// <summary>
+	/// Stores cached information about a utf8 search term 
+	/// </summary>
+	public class SearchTerm
+	{
+		/// <summary>
+		/// The search text
+		/// </summary>
+		public string Text { get; }
+
+		/// <summary>
+		/// The utf-8 bytes to search for
+		/// </summary>
+		public ReadOnlyMemory<byte> Bytes { get; }
+
+		/// <summary>
+		/// Normalized (lowercase) utf-8 bytes to search for
+		/// </summary>
+		readonly byte[] _searchBytes;
+
+		/// <summary>
+		/// Skip table for comparisons
+		/// </summary>
+		readonly byte[] _skipTable;
+
+		/// <summary>
+		/// Constructor
+		/// </summary>
+		/// <param name="text">The text to search for</param>
+		public SearchTerm(string text)
+		{
+			byte[] bytes = Encoding.UTF8.GetBytes(text);
+			Bytes = bytes;
+			Text = text;
+
+			// Find the byte sequence to search for, in lowercase
+			_searchBytes = new byte[bytes.Length];
+			for (int idx = 0; idx < _searchBytes.Length; idx++)
+			{
+				if (bytes[idx] >= 'A' && bytes[idx] <= 'Z')
+				{
+					_searchBytes[idx] = (byte)('a' + (bytes[idx] - 'A'));
+				}
+				else
+				{
+					_searchBytes[idx] = bytes[idx];
+				}
+			}
+
+			// Build a table indicating how many characters to skip before attempting the next comparison
+			_skipTable = new byte[256];
+			for (int idx = 0; idx < 256; idx++)
+			{
+				_skipTable[idx] = (byte)_searchBytes.Length;
+			}
+			for (int idx = 0; idx < _searchBytes.Length - 1; idx++)
+			{
+				byte character = _searchBytes[idx];
+
+				byte skipBytes = (byte)(_searchBytes.Length - 1 - idx);
+				_skipTable[character] = skipBytes;
+
+				if (character >= 'a' && character <= 'z')
+				{
+					_skipTable['A' + (character - 'a')] = skipBytes;
+				}
+			}
+		}
+
+		/// <summary>
+		/// Find all ocurrences of the text in the given buffer
+		/// </summary>
+		/// <param name="buffer">The buffer to search</param>
+		/// <param name="text">The text to search for</param>
+		/// <returns>Sequence of offsets within the buffer</returns>
+		public static IEnumerable<int> FindOcurrences(ReadOnlyMemory<byte> buffer, SearchTerm text)
+		{
+			for(int offset = 0; ;offset++)
+			{
+				offset = FindNextOcurrence(buffer.Span, offset, text);
+				if(offset == -1)
+				{
+					break;
+				}
+				yield return offset;
+			}
+		}
+
+		/// <summary>
+		/// Perform a case sensitive search for the next occurerence of the search term in a given buffer
+		/// </summary>
+		/// <param name="buffer">The buffer to search</param>
+		/// <param name="offset">Starting offset for the search</param>
+		/// <param name="text">The text to search for</param>
+		/// <returns>Offset of the next occurence, or -1</returns>
+		public static int FindNextOcurrence(ReadOnlySpan<byte> buffer, int offset, SearchTerm text)
+		{
+			while (offset + text._searchBytes.Length <= buffer.Length)
+			{
+				if (Matches(buffer, offset, text))
+				{
+					return offset;
+				}
+				else
+				{
+					offset += text._skipTable[buffer[offset + text._searchBytes.Length - 1]];
+				}
+			}
+			return -1;
+		}
+
+		/// <summary>
+		/// Compare the search term against the given buffer
+		/// </summary>
+		/// <param name="buffer">The buffer to search</param>
+		/// <param name="offset">Starting offset for the search</param>
+		/// <param name="text">The text to search for</param>
+		/// <returns>True if the text matches, false otherwise</returns>
+		public static bool Matches(ReadOnlySpan<byte> buffer, int offset, SearchTerm text)
+		{
+			for (int idx = text._searchBytes.Length - 1; idx >= 0; idx--)
+			{
+				byte character = buffer[offset + idx];
+				if (character >= 'A' && character <= 'Z')
+				{
+					character = (byte)('a' + (character - 'A'));
+				}
+				if (character != text._searchBytes[idx])
+				{
+					return false;
+				}
+			}
+			return true;
+		}
+	}
+
+	/// <summary>
+	/// Stores cached information about a utf8 search term 
+	/// </summary>
+	public static class SearchTextExtensions
+	{
+		/// <summary>
+		/// Find all ocurrences of the text in the given buffer
+		/// </summary>
+		/// <param name="buffer">The buffer to search</param>
+		/// <param name="text">The text to search for</param>
+		/// <returns>Sequence of offsets within the buffer</returns>
+		public static IEnumerable<int> FindOcurrences(this ReadOnlyMemory<byte> buffer, SearchTerm text)
+		{
+			return SearchTerm.FindOcurrences(buffer, text);
+		}
+
+		/// <summary>
+		/// Perform a case sensitive search for the next occurerence of the search term in a given buffer
+		/// </summary>
+		/// <param name="buffer">The buffer to search</param>
+		/// <param name="offset">Starting offset for the search</param>
+		/// <param name="text">The text to search for</param>
+		/// <returns>Offset of the next occurence, or -1</returns>
+		public static int FindNextOcurrence(this ReadOnlySpan<byte> buffer, int offset, SearchTerm text)
+		{
+			return SearchTerm.FindNextOcurrence(buffer, offset, text);
+		}
+
+		/// <summary>
+		/// Compare the search term against the given buffer
+		/// </summary>
+		/// <param name="buffer">The buffer to search</param>
+		/// <param name="offset">Starting offset for the search</param>
+		/// <param name="text">The text to search for</param>
+		/// <returns>True if the text matches, false otherwise</returns>
+		public static bool Matches(this ReadOnlySpan<byte> buffer, int offset, SearchTerm text)
+		{
+			return SearchTerm.Matches(buffer, offset, text);
+		}
+	}
+}
diff --git a/Engine/Source/Programs/Shared/EpicGames.Horde/Storage/TreeReader.cs b/Engine/Source/Programs/Shared/EpicGames.Horde/Storage/TreeReader.cs
--- a/Engine/Source/Programs/Shared/EpicGames.Horde/Storage/TreeReader.cs
+++ b/Engine/Source/Programs/Shared/EpicGames.Horde/Storage/TreeReader.cs
@@ -54,7 +54,10 @@ namespace EpicGames.Horde.Storage
 		{
 			typeof(Nodes.DirectoryNode),
 			typeof(Nodes.LeafFileNode),
-			typeof(Nodes.InteriorFileNode)
+			typeof(Nodes.InteriorFileNode),
+			typeof(Logs.LogNode),
+			typeof(Logs.LogChunkNode),
+			typeof(Logs.LogIndexNode),
 		};
 
 		/// <summary>
