EpicGames.Core: Add constructor to make a Utf8String from a ReadOnlySpan<char>.

diff --git a/Engine/Source/Programs/Shared/EpicGames.Core/Utf8String.cs b/Engine/Source/Programs/Shared/EpicGames.Core/Utf8String.cs
--- a/Engine/Source/Programs/Shared/EpicGames.Core/Utf8String.cs
+++ b/Engine/Source/Programs/Shared/EpicGames.Core/Utf8String.cs
@@ -52,6 +52,18 @@ namespace EpicGames.Core
 			Memory = Encoding.UTF8.GetBytes(text);
 		}
 
+		/// <summary>
+		/// Constructor
+		/// </summary>
+		/// <param name="text">Text to construct from</param>
+		public Utf8String(ReadOnlySpan<char> text)
+		{
+			int length = Encoding.UTF8.GetByteCount(text);
+			byte[] buffer = new byte[length];
+			Encoding.UTF8.GetBytes(text, buffer);
+			Memory = buffer;
+		}
+
 		/// <summary>
 		/// Constructor
 		/// </summary>
