EpicGames.Horde: Add utility method to allow parsing a NodeLocator from a string.

diff --git a/Engine/Source/Programs/Shared/EpicGames.Horde/Storage/IStorageClient.cs b/Engine/Source/Programs/Shared/EpicGames.Horde/Storage/IStorageClient.cs
--- a/Engine/Source/Programs/Shared/EpicGames.Horde/Storage/IStorageClient.cs
+++ b/Engine/Source/Programs/Shared/EpicGames.Horde/Storage/IStorageClient.cs
@@ -4,9 +4,11 @@ using System;
 using System.Collections.Generic;
 using System.Diagnostics;
 using System.Diagnostics.CodeAnalysis;
+using System.Globalization;
 using System.IO;
 using System.Reflection;
 using System.Reflection.Metadata;
+using System.Text;
 using System.Threading;
 using System.Threading.Tasks;
 using EpicGames.Core;
@@ -105,6 +107,19 @@ namespace EpicGames.Horde.Storage
 		/// </summary>
 		public bool IsValid() => Blob.IsValid();
 
+		/// <summary>
+		/// Parse a string as a node locator
+		/// </summary>
+		/// <param name="text">Text to parse</param>
+		/// <returns></returns>
+		public static NodeLocator Parse(string text)
+		{
+			int hashIdx = text.IndexOf('#', StringComparison.Ordinal);
+			int exportIdx = Int32.Parse(text.AsSpan(hashIdx + 1), NumberStyles.None, CultureInfo.InvariantCulture);
+			BlobLocator blobLocator = new BlobLocator(new Utf8String(text.AsSpan(0, hashIdx)));
+			return new NodeLocator(blobLocator, exportIdx);
+		}
+
 		/// <inheritdoc/>
 		public override bool Equals([NotNullWhen(true)] object? obj) => obj is NodeLocator locator && Equals(locator);
 
