Horde: Fix warning.

diff --git a/Engine/Source/Programs/Shared/EpicGames.Horde.Tests/LogTests.cs b/Engine/Source/Programs/Shared/EpicGames.Horde.Tests/LogTests.cs
--- a/Engine/Source/Programs/Shared/EpicGames.Horde.Tests/LogTests.cs
+++ b/Engine/Source/Programs/Shared/EpicGames.Horde.Tests/LogTests.cs
@@ -4,6 +4,7 @@ using System;
 using System.Collections.Generic;
 using System.Linq;
 using System.Text;
+using System.Resources;
 using System.Threading;
 using System.Threading.Tasks;
 using EpicGames.Horde.Logs;
@@ -14,6 +15,8 @@ using Microsoft.Extensions.Caching.Memory;
 using Microsoft.Extensions.Logging.Abstractions;
 using Microsoft.VisualStudio.TestTools.UnitTesting;
 
+[assembly: NeutralResourcesLanguageAttribute("en-US")]
+
 namespace EpicGames.Horde.Tests
 {
 	[TestClass]
diff --git a/Engine/Source/Programs/Shared/EpicGames.Horde/Logs/LogNode.cs b/Engine/Source/Programs/Shared/EpicGames.Horde/Logs/LogNode.cs
--- a/Engine/Source/Programs/Shared/EpicGames.Horde/Logs/LogNode.cs
+++ b/Engine/Source/Programs/Shared/EpicGames.Horde/Logs/LogNode.cs
@@ -110,7 +110,7 @@ namespace EpicGames.Horde.Logs
 	/// <summary>
 	/// Assists building log files through trees of <see cref="LogNode"/>, <see cref="LogIndexNode"/> and <see cref="LogChunkNode"/> nodes.
 	/// </summary>
-	class LogBuilder
+	public class LogBuilder
 	{
 		/// <summary>
 		/// Default maximum size for a log text block
