Horde: Change StringId to UTF8.

diff --git a/Engine/Source/Programs/Shared/EpicGames.Horde/StringId.cs b/Engine/Source/Programs/Shared/EpicGames.Horde/StringId.cs
--- a/Engine/Source/Programs/Shared/EpicGames.Horde/StringId.cs
+++ b/Engine/Source/Programs/Shared/EpicGames.Horde/StringId.cs
@@ -1,6 +1,7 @@
 // Copyright Epic Games, Inc. All Rights Reserved.
 
 using System;
+using EpicGames.Core;
 
 namespace EpicGames.Horde
 {
@@ -12,22 +13,21 @@ namespace EpicGames.Horde
 		/// <summary>
 		/// The text representing this id
 		/// </summary>
-		public string Text { get; }
+		public Utf8String Text { get; }
 
 		/// <summary>
 		/// Constructor
 		/// </summary>
 		/// <param name="text">Unique id for the string</param>
-		public StringId(string text)
+		public StringId(Utf8String text)
 		{
-			Text = text;
-			ValidateArgument(text, nameof(text));
+			Text = ValidateArgument(text, nameof(text));
 		}
 
 		/// <summary>
 		/// Checks whether this StringId is set
 		/// </summary>
-		public bool IsEmpty => String.IsNullOrEmpty(Text);
+		public bool IsEmpty => Text.IsEmpty;
 
 		/// <summary>
 		/// Validates the given string as a StringId, normalizing it if necessary.
@@ -35,7 +35,7 @@ namespace EpicGames.Horde
 		/// <param name="text">Text to validate as a StringId</param>
 		/// <param name="paramName">Name of the parameter to show if invalid characters are returned.</param>
 		/// <returns></returns>
-		public static string ValidateArgument(string text, string paramName)
+		public static Utf8String ValidateArgument(Utf8String text, string paramName)
 		{
 			if (text.Length == 0)
 			{
@@ -50,14 +50,12 @@ namespace EpicGames.Horde
 
 			for (int idx = 0; idx < text.Length; idx++)
 			{
-				char character = text[idx];
+				char character = (char)text[idx];
 				if (!IsValidCharacter(character))
 				{
 					if (character >= 'A' && character <= 'Z')
 					{
-#pragma warning disable CA1308 // Normalize strings to uppercase
-						text = text.ToLowerInvariant();
-#pragma warning restore CA1308 // Normalize strings to uppercase
+						text = ToLower(text);
 					}
 					else
 					{
@@ -69,6 +67,26 @@ namespace EpicGames.Horde
 			return text;
 		}
 
+		/// <summary>
+		/// Converts a utf8 string to lowercase
+		/// </summary>
+		/// <param name="text"></param>
+		/// <returns></returns>
+		static Utf8String ToLower(Utf8String text)
+		{
+			byte[] output = new byte[text.Length];
+			for (int idx = 0; idx < text.Length; idx++)
+			{
+				byte character = text[idx];
+				if (character >= 'A' && character <= 'Z')
+				{
+					character = (byte)((character - 'A') + 'a');
+				}
+				output[idx] = character;
+			}
+			return new Utf8String(output);
+		}
+
 		/// <summary>
 		/// Checks whether the given character is valid within a string id
 		/// </summary>
@@ -95,13 +113,13 @@ namespace EpicGames.Horde
 		public override bool Equals(object? obj) => obj is StringId id && Equals(id);
 
 		/// <inheritdoc/>
-		public override int GetHashCode() => Text.GetHashCode(StringComparison.Ordinal);
+		public override int GetHashCode() => Text.GetHashCode();
 
 		/// <inheritdoc/>
-		public bool Equals(StringId other) => Text.Equals(other.Text, StringComparison.Ordinal);
+		public bool Equals(StringId other) => Text.Equals(other.Text);
 
 		/// <inheritdoc/>
-		public override string ToString() => Text;
+		public override string ToString() => Text.ToString();
 
 		/// <summary>
 		/// Compares two string ids for equality
