Fixed issue in C# UHT where it would fail to generate an error when preprocessor lines appear in an inheritence list.

diff --git a/Engine/Source/Programs/Shared/EpicGames.UHT/Parsers/UhtClassParser.cs b/Engine/Source/Programs/Shared/EpicGames.UHT/Parsers/UhtClassParser.cs
--- a/Engine/Source/Programs/Shared/EpicGames.UHT/Parsers/UhtClassParser.cs
+++ b/Engine/Source/Programs/Shared/EpicGames.UHT/Parsers/UhtClassParser.cs
@@ -98,7 +98,7 @@ namespace EpicGames.UHT.Parsers
 					topScope.TokenReader.Optional("final");
 
 					// Parse the inheritance
-					UhtParserHelpers.ParseInheritance(topScope.TokenReader, topScope.Session.Config!, out UhtToken superIdentifier, out List<UhtToken[]>? baseIdentifiers);
+					UhtParserHelpers.ParseInheritance(parentScope.HeaderParser, topScope.Session.Config!, out UhtToken superIdentifier, out List<UhtToken[]>? baseIdentifiers);
 					classObj.SuperIdentifier = superIdentifier;
 					classObj.BaseIdentifiers = baseIdentifiers;
 
diff --git a/Engine/Source/Programs/Shared/EpicGames.UHT/Parsers/UhtHeaderFileParser.cs b/Engine/Source/Programs/Shared/EpicGames.UHT/Parsers/UhtHeaderFileParser.cs
--- a/Engine/Source/Programs/Shared/EpicGames.UHT/Parsers/UhtHeaderFileParser.cs
+++ b/Engine/Source/Programs/Shared/EpicGames.UHT/Parsers/UhtHeaderFileParser.cs
@@ -209,6 +209,12 @@ namespace EpicGames.UHT.Parsers
 		/// </summary>
 		public UhtPropertyParser? PropertyParser { get; set; } = null;
 
+		/// <summary>
+		/// If set, the preprocessor is run in a C++ UHT compatibility mode where only a subset
+		/// of #if class of preprocessor statements are allowed.
+		/// </summary>
+		public string? RestrictedPreprocessorContext { get; set; } = null;
+
 		/// <summary>
 		/// Stack of current #if states
 		/// </summary>
@@ -639,6 +645,7 @@ namespace EpicGames.UHT.Parsers
 
 			if (directive.IsValue("error"))
 			{
+				CheckRestrictedMode();
 				if (isBeingIncluded)
 				{
 					throw new UhtException(TokenReader, directive.InputLine, "#error directive encountered");
@@ -646,10 +653,12 @@ namespace EpicGames.UHT.Parsers
 			}
 			else if (directive.IsValue("pragma"))
 			{
+				CheckRestrictedMode();
 				// Ignore all pragmas
 			}
 			else if (directive.IsValue("linenumber"))
 			{
+				CheckRestrictedMode();
 				if (!lineTokenReader.TryOptionalConstInt(out int newInputLine))
 				{
 					throw new UhtException(TokenReader, directive.InputLine, "Missing line number in line number directive");
@@ -658,6 +667,7 @@ namespace EpicGames.UHT.Parsers
 			}
 			else if (directive.IsValue("include"))
 			{
+				CheckRestrictedMode();
 				if (isBeingIncluded)
 				{
 					if (SpottedAutogeneratedHeaderInclude)
@@ -699,6 +709,10 @@ namespace EpicGames.UHT.Parsers
 			{
 				checkClearComments = true;
 				PushCompilerDirective(ParserConditional(lineTokenReader));
+				if (IsRestrictedDirective(GetCurrentNonCompositeCompilerDirective()))
+				{
+					CheckRestrictedMode();
+				}
 			}
 			else if (directive.IsValue("ifdef") || directive.IsValue("ifndef"))
 			{
@@ -716,6 +730,10 @@ namespace EpicGames.UHT.Parsers
 						$"Mixing {GetCompilerDirectiveText(oldCompilerDirective)} with {GetCompilerDirectiveText(newCompilerDirective)} in an #elif preprocessor block is not supported");
 				}
 				PushCompilerDirective(newCompilerDirective);
+				if (IsRestrictedDirective(GetCurrentNonCompositeCompilerDirective()))
+				{
+					CheckRestrictedMode();
+				}
 			}
 			else if (directive.IsValue("else"))
 			{
@@ -741,6 +759,10 @@ namespace EpicGames.UHT.Parsers
 						PushCompilerDirective(oldCompilerDirective);
 						break;
 				}
+				if (IsRestrictedDirective(GetCurrentNonCompositeCompilerDirective()))
+				{
+					CheckRestrictedMode();
+				}
 			}
 			else if (directive.IsValue("endif"))
 			{
@@ -748,9 +770,11 @@ namespace EpicGames.UHT.Parsers
 			}
 			else if (directive.IsValue("define"))
 			{
+				CheckRestrictedMode();
 			}
 			else if (directive.IsValue("undef"))
 			{
+				CheckRestrictedMode();
 			}
 			else
 			{
@@ -762,6 +786,14 @@ namespace EpicGames.UHT.Parsers
 			return checkClearComments;
 		}
 
+		private void CheckRestrictedMode()
+		{
+			if (RestrictedPreprocessorContext != null)
+			{
+				TokenReader.LogError($"Preprocessor statement not allowed while {RestrictedPreprocessorContext}");
+			}
+		}
+
 		private UhtCompilerDirective ParserConditional(UhtTokenBufferReader lineTokenReader)
 		{
 			// Get any possible ! and the identifier
@@ -895,6 +927,27 @@ namespace EpicGames.UHT.Parsers
 			}
 		}
 
+		private bool IsRestrictedDirective(UhtCompilerDirective compilerDirective)
+		{
+			switch (compilerDirective)
+			{
+				case UhtCompilerDirective.CPPBlock:
+				case UhtCompilerDirective.NotCPPBlock:
+				case UhtCompilerDirective.ZeroBlock:
+				case UhtCompilerDirective.OneBlock:
+				case UhtCompilerDirective.Unrecognized:
+					return false;
+
+				case UhtCompilerDirective.WithEditor:
+				case UhtCompilerDirective.WithEditorOnlyData:
+				case UhtCompilerDirective.WithHotReload:
+					return true;
+
+				default:
+					throw new UhtIceException("Unknown compiler directive flag");
+			}
+		}
+
 		private static bool SupportsElif(UhtCompilerDirective compilerDirective)
 		{
 			return
diff --git a/Engine/Source/Programs/Shared/EpicGames.UHT/Parsers/UhtParserHelpers.cs b/Engine/Source/Programs/Shared/EpicGames.UHT/Parsers/UhtParserHelpers.cs
--- a/Engine/Source/Programs/Shared/EpicGames.UHT/Parsers/UhtParserHelpers.cs
+++ b/Engine/Source/Programs/Shared/EpicGames.UHT/Parsers/UhtParserHelpers.cs
@@ -18,30 +18,42 @@ namespace EpicGames.UHT.Parsers
 		/// <summary>
 		/// Parse the inheritance 
 		/// </summary>
-		/// <param name="tokenReader">Token reader</param>
+		/// <param name="headerFileParser">Header file being parsed</param>
 		/// <param name="config">Configuration</param>
 		/// <param name="superIdentifier">Output super identifier</param>
 		/// <param name="baseIdentifiers">Output base identifiers</param>
-		public static void ParseInheritance(IUhtTokenReader tokenReader, IUhtConfig config, out UhtToken superIdentifier, out List<UhtToken[]>? baseIdentifiers)
+		public static void ParseInheritance(UhtHeaderFileParser headerFileParser, IUhtConfig config, out UhtToken superIdentifier, out List<UhtToken[]>? baseIdentifiers)
 		{
-			UhtToken superIdentifierTemp = new();
-			List<UhtToken[]>? baseIdentifiersTemp = null;
-			tokenReader.OptionalInheritance(
-				(ref UhtToken identifier) =>
-				{
-					config.RedirectTypeIdentifier(ref identifier);
-					superIdentifierTemp = identifier;
-				},
-				(UhtTokenList identifier) =>
-				{
-					if (baseIdentifiersTemp == null)
+
+			// TODO: C++ UHT doesn't allow preprocessor statements inside of inheritance lists
+			string? restrictedPreprocessorContext = headerFileParser.RestrictedPreprocessorContext;
+			headerFileParser.RestrictedPreprocessorContext = "parsing inheritance list";
+
+			try
+			{
+				UhtToken superIdentifierTemp = new();
+				List<UhtToken[]>? baseIdentifiersTemp = null;
+				headerFileParser.TokenReader.OptionalInheritance(
+					(ref UhtToken identifier) =>
 					{
-						baseIdentifiersTemp = new List<UhtToken[]>();
-					}
-					baseIdentifiersTemp.Add(identifier.ToArray());
-				});
-			superIdentifier = superIdentifierTemp;
-			baseIdentifiers = baseIdentifiersTemp;
+						config.RedirectTypeIdentifier(ref identifier);
+						superIdentifierTemp = identifier;
+					},
+					(UhtTokenList identifier) =>
+					{
+						if (baseIdentifiersTemp == null)
+						{
+							baseIdentifiersTemp = new List<UhtToken[]>();
+						}
+						baseIdentifiersTemp.Add(identifier.ToArray());
+					});
+				superIdentifier = superIdentifierTemp;
+				baseIdentifiers = baseIdentifiersTemp;
+			}
+			finally
+			{
+				headerFileParser.RestrictedPreprocessorContext = restrictedPreprocessorContext;
+			}
 		}
 
 		/// <summary>
diff --git a/Engine/Source/Programs/Shared/EpicGames.UHT/Parsers/UhtScriptStructParser.cs b/Engine/Source/Programs/Shared/EpicGames.UHT/Parsers/UhtScriptStructParser.cs
--- a/Engine/Source/Programs/Shared/EpicGames.UHT/Parsers/UhtScriptStructParser.cs
+++ b/Engine/Source/Programs/Shared/EpicGames.UHT/Parsers/UhtScriptStructParser.cs
@@ -110,7 +110,7 @@ namespace EpicGames.UHT.Parsers
 					topScope.TokenReader.Optional("final");
 
 					// Parse the inheritance
-					UhtParserHelpers.ParseInheritance(topScope.TokenReader, topScope.Session.Config!, out UhtToken superIdentifier, out List<UhtToken[]>? baseIdentifiers);
+					UhtParserHelpers.ParseInheritance(parentScope.HeaderParser, topScope.Session.Config!, out UhtToken superIdentifier, out List<UhtToken[]>? baseIdentifiers);
 					scriptStruct.SuperIdentifier = superIdentifier;
 					scriptStruct.BaseIdentifiers = baseIdentifiers;
 
