EpicGames.Perforce: Add test fixture server with accompanying C# test runner

Includes a simple sync test for ManagedWorkspace as an example.

diff --git a/Engine/Source/Programs/Shared/EpicGames.Perforce.Fixture/BasePerforceFixtureTest.cs b/Engine/Source/Programs/Shared/EpicGames.Perforce.Fixture/BasePerforceFixtureTest.cs
new file mode 100644
--- /dev/null
+++ b/Engine/Source/Programs/Shared/EpicGames.Perforce.Fixture/BasePerforceFixtureTest.cs
@@ -0,0 +1,136 @@
+// Copyright Epic Games, Inc. All Rights Reserved.
+
+using System;
+using System.IO;
+using EpicGames.Core;
+using Microsoft.Extensions.Logging;
+using Microsoft.VisualStudio.TestTools.UnitTesting;
+
+namespace EpicGames.Perforce.Fixture;
+
+/// <summary>
+/// Base test class for creating tests that use the Perforce fixture server.
+/// Every test invocation checks for the <see cref="EnvVar" />.
+/// If present, test continues and a PerforceConnection using that connection string is created.
+/// If not present, the test is marked as inconclusive, which is safe to ignore and doesn't fail the test suite.
+/// </summary>
+[TestClass]
+public abstract class BasePerforceFixtureTest : IDisposable
+{
+	private const string EnvVar = "EPICGAMES_P4_FIXTURE_SERVER_URL";
+	private const string UseDefaultMarker = "default";
+	private const string DefaultUrl = "p4://test.user@localhost:1666";
+
+	/// <summary>
+	/// Console-based logger factory
+	/// </summary>
+	protected ILoggerFactory LoggerFactory { get; }
+	
+	/// <summary>
+	/// Temporary scratch dir, deleted after every test.
+	/// Provides a location for (managed) workspace root or cache files. 
+	/// </summary>
+	protected DirectoryReference TempDir { get; }
+	
+	/// <summary>
+	/// A Perforce connection to the fixture test server
+	/// </summary>
+	protected PerforceConnection PerforceConnection { get; private set; }
+
+	private readonly string _clientName;
+	
+	protected BasePerforceFixtureTest()
+	{
+		LoggerFactory = Microsoft.Extensions.Logging.LoggerFactory.Create(builder => builder.AddConsole());
+		
+		// PerforceConnection is set in a [TestInitialize] method as Assert.Inconclusive cannot be raised in the constructor.
+		// That initialize method will get called after the constructor ensuring the connection is instantiated.
+		PerforceConnection = null!;
+		
+		_clientName = "test-fixture-" + Guid.NewGuid().ToString()[..8];
+		TempDir = new DirectoryReference(Path.Join(Path.GetTempPath(), "p4-" + _clientName));
+		DirectoryReference.CreateDirectory(TempDir);
+	}
+
+	[TestInitialize]
+	public void SetUpPerforceConnection()
+	{
+		string? serverUrl = Environment.GetEnvironmentVariable(EnvVar);
+		if (serverUrl == null)
+		{
+			// Only the first ~120 chars of the message is printed by MSTest runner. Try to keep it below that.
+			Assert.Inconclusive($"P4 fixture server not configured. Set env var {EnvVar} to enable. See {nameof(BasePerforceFixtureTest)}.");
+		}
+		
+		PerforceConnection = GetPerforceConnection(serverUrl!);
+	}
+
+	[TestCleanup]
+	public void RemoveTempDir()
+	{
+		if (Directory.Exists(TempDir.FullName))
+		{
+			// Remove the read-only flags set by the P4 client
+			foreach (string filePath in Directory.EnumerateFiles(TempDir.FullName, "*", SearchOption.AllDirectories))
+			{
+				FileInfo fileInfo = new (filePath) { IsReadOnly = false };
+			}
+			
+			Directory.Delete(TempDir.FullName, true);
+		}
+	}
+
+	private PerforceConnection GetPerforceConnection(string connectionString)
+	{
+		connectionString = connectionString == UseDefaultMarker ? DefaultUrl : connectionString;
+		Uri uri = new (connectionString);
+		
+		Assert.AreEqual("p4", uri.Scheme, "P4 fixture server URL must start with p4://");
+
+		string serverAndPort = $"{uri.Host}:{uri.Port}";
+		(string username, string? password) = GetCredentialsFromUri(uri);
+		
+		PerforceSettings perforceSettings = new (serverAndPort, username)
+		{
+			AppName = "Perforce Test Fixture",
+			ClientName = _clientName,
+			PreferNativeClient = true
+		};
+
+		if (password != null)
+		{
+			perforceSettings.Password = password;
+		}
+
+		return new PerforceConnection(perforceSettings, LoggerFactory.CreateLogger<PerforceConnection>());
+	}
+
+	private static (string username, string? password) GetCredentialsFromUri(Uri uri)
+	{
+		if (!uri.UserInfo.Contains(':', StringComparison.Ordinal))
+		{
+			return (uri.UserInfo, null);
+		}
+
+		string[] parts = uri.UserInfo.Split(":");
+		return (parts[0], parts[1]);
+	}
+
+	public void Dispose()
+	{
+		Dispose(true);
+		GC.SuppressFinalize(this);
+	}
+	
+	/// <summary>
+	/// Standard Dispose pattern method
+	/// </summary>
+	/// <param name="disposing"></param>
+	protected virtual void Dispose(bool disposing)
+	{
+		if (disposing)
+		{
+			LoggerFactory.Dispose();
+		}
+	}
+}
\ No newline at end of file
diff --git a/Engine/Source/Programs/Shared/EpicGames.Perforce.Fixture/EpicGames.Perforce.Fixture.csproj b/Engine/Source/Programs/Shared/EpicGames.Perforce.Fixture/EpicGames.Perforce.Fixture.csproj
new file mode 100644
--- /dev/null
+++ b/Engine/Source/Programs/Shared/EpicGames.Perforce.Fixture/EpicGames.Perforce.Fixture.csproj
@@ -0,0 +1,38 @@
+<Project Sdk="Microsoft.NET.Sdk">
+
+  <PropertyGroup>
+    <TargetFramework>net6.0</TargetFramework>
+    <Configuration Condition=" '$(Configuration)' == '' ">Development</Configuration>
+    <OutputType>Library</OutputType>
+    <Deterministic>true</Deterministic>
+    <RootNamespace>EpicGames.Perforce.Fixture</RootNamespace>
+    <GenerateAssemblyCompanyAttribute>false</GenerateAssemblyCompanyAttribute>
+    <GenerateAssemblyInformationalVersionAttribute>false</GenerateAssemblyInformationalVersionAttribute>
+    <GenerateAssemblyProductAttribute>false</GenerateAssemblyProductAttribute>
+    <GenerateAssemblyVersionAttribute>false</GenerateAssemblyVersionAttribute>
+    <Configurations>Debug;Release;Development</Configurations>
+    <IsPackable>false</IsPackable>
+    <Nullable>enable</Nullable>
+  </PropertyGroup>
+  <ItemGroup>
+    <Compile Include="..\..\Shared\MetaData.cs">
+      <Link>Properties\MetaData.cs</Link>
+    </Compile>
+  </ItemGroup>
+
+  <ItemGroup>
+    <PackageReference Include="Microsoft.Extensions.Logging.Console" Version="6.0.0" />
+    <PackageReference Include="Microsoft.NET.Test.Sdk" Version="17.3.0" />
+    <PackageReference Include="MSTest.TestAdapter" Version="2.2.10" />
+    <PackageReference Include="MSTest.TestFramework" Version="2.2.10" />
+    <PackageReference Include="coverlet.collector" Version="3.1.2">
+      <PrivateAssets>all</PrivateAssets>
+      <IncludeAssets>runtime; build; native; contentfiles; analyzers; buildtransitive</IncludeAssets>
+    </PackageReference>
+  </ItemGroup>
+
+  <ItemGroup>
+    <ProjectReference Include="..\EpicGames.Perforce.Managed\EpicGames.Perforce.Managed.csproj" />
+  </ItemGroup>
+
+</Project>
diff --git a/Engine/Source/Programs/Shared/EpicGames.Perforce.Fixture/setup.sh b/Engine/Source/Programs/Shared/EpicGames.Perforce.Fixture/setup.sh
new file mode 100644
--- /dev/null
+++ b/Engine/Source/Programs/Shared/EpicGames.Perforce.Fixture/setup.sh
@@ -0,0 +1,103 @@
+#!/bin/bash
+
+set -euo pipefail # Exit on any error
+
+export PORT=1666
+export P4PORT=localhost:$PORT
+
+# Start P4D
+p4d -r /app/data -p $PORT &
+
+# The time needed for P4D to start listening on the TCP port varies greatly from 2 to +10 secs
+# Waiting for the port to be available instead is more reliable than an arbitrary sleep
+echo "Waiting for server to listen on port $PORT..."
+while ! nc -z localhost $PORT </dev/null; do sleep 1; done
+
+# Set client.readonly.dir for partitioned client types
+mkdir -p /app/part-db-have
+p4 configure set client.readonly.dir=/app/part-db-have
+
+# Create a dummy user
+p4 user -f -i << EOF
+User:	test.user
+Type:   standard
+Email:	test.user@localhost
+FullName:	Test User
+AuthMethod:	perforce
+EOF
+
+# Create a stream depot
+p4 depot -i << EOF
+Depot:          Foo
+Owner:          test.user
+Description:
+                Created by test.user.
+
+Type:           stream
+StreamDepth:    //Foo/2
+Map:            Foo/...
+EOF
+
+# Create an actual stream
+p4 stream -i << EOF
+Stream: //Foo/Main
+Owner:  test.user
+Name:   Main
+Parent: none
+Type:   mainline
+Description:
+        Main stream for Foo depot
+
+Options:        allsubmit unlocked notoparent nofromparent mergedown
+ParentView:     inherit
+Paths:
+        share ...
+EOF
+
+export P4CLIENT=generate-fixture-data
+export CLIENT_ROOT=/tmp/${P4CLIENT}
+p4 client -i << EOF
+Client: ${P4CLIENT}
+Root: ${CLIENT_ROOT}
+Stream: //Foo/Main
+EOF
+
+# Create the workspace root dir and sync the empty stream
+mkdir -p $CLIENT_ROOT
+cd $CLIENT_ROOT
+p4 sync
+
+
+# Changelist #2
+mkdir $CLIENT_ROOT/Data
+echo "This is change main.cpp #1" > $CLIENT_ROOT/main.cpp && p4 add $CLIENT_ROOT/main.cpp
+echo "This is change main.h #1" > $CLIENT_ROOT/main.h && p4 add $CLIENT_ROOT/main.h
+echo "This is change common.h #1" > $CLIENT_ROOT/common.h && p4 add $CLIENT_ROOT/common.h
+echo "This is change unused.cpp #1" > $CLIENT_ROOT/unused.cpp && p4 add $CLIENT_ROOT/unused.cpp
+echo "This is change data.txt #1" > $CLIENT_ROOT/Data/data.txt && p4 add $CLIENT_ROOT/Data/data.txt
+p4 submit -d "Initial import"
+
+# Changelist #3
+p4 edit main.cpp
+echo "This is change main.cpp #2" > $CLIENT_ROOT/main.cpp
+p4 submit -d "Improvement to main.cpp"
+
+# Changelist #4
+p4 delete unused.cpp
+p4 submit -d "Delete an unused file"
+
+# Changelist #5
+p4 edit common.h
+p4 move common.h shared.h
+p4 submit -d "Rename common.h to shared.h"
+
+# Changelist #6
+p4 edit main.cpp
+p4 edit main.h
+echo "This is change main.cpp #3" > $CLIENT_ROOT/main.cpp
+echo "This is change main.h #2" > $CLIENT_ROOT/main.h
+p4 submit -d "Some updates to main"
+
+
+# Stop server running in background
+kill "$(pidof p4d)"
\ No newline at end of file
diff --git a/Engine/Source/Programs/Shared/EpicGames.Perforce.Fixture/start.bat b/Engine/Source/Programs/Shared/EpicGames.Perforce.Fixture/start.bat
new file mode 100755
--- /dev/null
+++ b/Engine/Source/Programs/Shared/EpicGames.Perforce.Fixture/start.bat
@@ -0,0 +1,6 @@
+@echo off
+IF NOT DEFINED DOCKER_EXECUTABLE (SET DOCKER_EXECUTABLE=docker)
+
+set TAG=perforce-fixture
+%DOCKER_EXECUTABLE% build --tag %TAG% .
+%DOCKER_EXECUTABLE% run -it --rm --publish 1666:1666 %TAG%
\ No newline at end of file
diff --git a/Engine/Source/Programs/Shared/EpicGames.Perforce.Managed.Tests/EpicGames.Perforce.Managed.Tests.csproj b/Engine/Source/Programs/Shared/EpicGames.Perforce.Managed.Tests/EpicGames.Perforce.Managed.Tests.csproj
--- a/Engine/Source/Programs/Shared/EpicGames.Perforce.Managed.Tests/EpicGames.Perforce.Managed.Tests.csproj
+++ b/Engine/Source/Programs/Shared/EpicGames.Perforce.Managed.Tests/EpicGames.Perforce.Managed.Tests.csproj
@@ -32,6 +32,7 @@
 
   <ItemGroup>
     <ProjectReference Include="..\EpicGames.Perforce.Managed\EpicGames.Perforce.Managed.csproj" />
+    <ProjectReference Include="..\EpicGames.Perforce.Fixture\EpicGames.Perforce.Fixture.csproj" />
   </ItemGroup>
 
 </Project>
diff --git a/Engine/Source/Programs/Shared/EpicGames.Perforce.Managed.Tests/ManagedWorkspaceTest.cs b/Engine/Source/Programs/Shared/EpicGames.Perforce.Managed.Tests/ManagedWorkspaceTest.cs
new file mode 100644
--- /dev/null
+++ b/Engine/Source/Programs/Shared/EpicGames.Perforce.Managed.Tests/ManagedWorkspaceTest.cs
@@ -0,0 +1,32 @@
+// Copyright Epic Games, Inc. All Rights Reserved.
+
+using System;
+using System.IO;
+using System.Threading;
+using System.Threading.Tasks;
+using EpicGames.Perforce.Fixture;
+using Microsoft.Extensions.Logging;
+using Microsoft.VisualStudio.TestTools.UnitTesting;
+
+namespace EpicGames.Perforce.Managed.Tests;
+
+[TestClass]
+public class ManagedWorkspaceTest : BasePerforceFixtureTest
+{
+	[TestMethod]
+	public async Task SimpleSync()
+	{
+		ILogger<ManagedWorkspace> logger = LoggerFactory.CreateLogger<ManagedWorkspace>();
+		ManagedWorkspace workspace = await ManagedWorkspace.LoadOrCreateAsync(Environment.MachineName, TempDir, true, logger, CancellationToken.None);
+		
+		await workspace.SetupAsync(PerforceConnection, "//Foo/Main", CancellationToken.None);
+
+		await workspace.SyncAsync(PerforceConnection, "//Foo/Main", 6, Array.Empty<string>(), true,
+			false, null, CancellationToken.None);
+
+		Assert.IsTrue(File.Exists(Path.Join(TempDir.FullName, "Sync", "main.cpp")));
+		Assert.IsTrue(File.Exists(Path.Join(TempDir.FullName, "Sync", "main.h")));
+		Assert.IsTrue(File.Exists(Path.Join(TempDir.FullName, "Sync", "shared.h")));
+		Assert.IsTrue(File.Exists(Path.Join(TempDir.FullName, "Sync", "Data", "data.txt")));
+	}
+}
\ No newline at end of file
