EpicGames.Core: Swallow exception on abandoned child process reads.

diff --git a/Engine/Source/Programs/Shared/EpicGames.Core/ManagedProcess.cs b/Engine/Source/Programs/Shared/EpicGames.Core/ManagedProcess.cs
--- a/Engine/Source/Programs/Shared/EpicGames.Core/ManagedProcess.cs
+++ b/Engine/Source/Programs/Shared/EpicGames.Core/ManagedProcess.cs
@@ -1060,6 +1060,7 @@ namespace EpicGames.Core
 				for (; ; )
 				{
 					Task<int> readTask = StdOut.ReadAsync(buffer, 0, bufferSize, cancellationToken);
+					_ = readTask.ContinueWith(x => _ = x.Exception, TaskContinuationOptions.OnlyOnFaulted);
 
 					Task completedTask = await Task.WhenAny(readTask, taskCompletionSource.Task);
 					cancellationToken.ThrowIfCancellationRequested();
