UGS: Fix the informational version string not being used to configure the P4 connection.

diff --git a/Engine/Source/Programs/Shared/EpicGames.Perforce/PerforceSettings.cs b/Engine/Source/Programs/Shared/EpicGames.Perforce/PerforceSettings.cs
--- a/Engine/Source/Programs/Shared/EpicGames.Perforce/PerforceSettings.cs
+++ b/Engine/Source/Programs/Shared/EpicGames.Perforce/PerforceSettings.cs
@@ -58,7 +58,7 @@ namespace EpicGames.Perforce
 	/// </summary>
 	public class PerforceSettings : IPerforceSettings
 	{
-		static readonly AssemblyName? s_entryAssemblyName = Assembly.GetEntryAssembly()!.GetName();
+		static readonly Assembly? s_entryAssembly = Assembly.GetEntryAssembly();
 
 		string? _appName;
 		string? _appVersion;
@@ -71,12 +71,12 @@ namespace EpicGames.Perforce
 		/// <summary>
 		/// Default app name to use for new settings objects
 		/// </summary>
-		public static string? DefaultAppName { get; set; } = s_entryAssemblyName?.Name;
+		public static string? DefaultAppName { get; set; } = s_entryAssembly?.GetName()?.Name;
 
 		/// <summary>
 		/// Default version string to use for new settings objects
 		/// </summary>
-		public static string? DefaultAppVersion { get; set; } = Assembly.GetExecutingAssembly().GetCustomAttribute<AssemblyInformationalVersionAttribute>()?.InformationalVersion ?? s_entryAssemblyName?.Version?.ToString();
+		public static string? DefaultAppVersion { get; set; } = s_entryAssembly?.GetCustomAttribute<AssemblyInformationalVersionAttribute>()?.InformationalVersion ?? s_entryAssembly?.GetName()?.Version?.ToString();
 
 		/// <inheritdoc/>
 		public string ServerAndPort { get; set; }
@@ -120,13 +120,6 @@ namespace EpicGames.Perforce
 			Password = environment.GetValue("P4PASSWD");
 			HostName = environment.GetValue("P4HOST");
 			ClientName = environment.GetValue("P4CLIENT");
-
-			AssemblyName entryAssemblyName = Assembly.GetEntryAssembly()!.GetName();
-			if (entryAssemblyName.Name != null)
-			{
-				AppName = entryAssemblyName.Name;
-				AppVersion = entryAssemblyName.Version?.ToString() ?? String.Empty;
-			}
 		}
 
 		/// <summary>
@@ -138,8 +131,6 @@ namespace EpicGames.Perforce
 		{
 			ServerAndPort = serverAndPort;
 			UserName = userName;
-			AppName = Default.AppName;
-			AppVersion = Default.AppVersion;
 		}
 
 		/// <summary>
@@ -152,8 +143,8 @@ namespace EpicGames.Perforce
 			Password = other.Password;
 			HostName = other.HostName;
 			ClientName = other.ClientName;
-			AppName = other.AppName;
-			AppVersion = other.AppVersion;
+			_appName = other.AppName;
+			_appVersion = other.AppVersion;
 			PreferNativeClient = other.PreferNativeClient;
 		}
 
