Fix static analyzer warning.

diff --git a/Engine/Source/Programs/Shared/EpicGames.Core/ManagedProcess.cs b/Engine/Source/Programs/Shared/EpicGames.Core/ManagedProcess.cs
--- a/Engine/Source/Programs/Shared/EpicGames.Core/ManagedProcess.cs
+++ b/Engine/Source/Programs/Shared/EpicGames.Core/ManagedProcess.cs
@@ -1060,7 +1060,7 @@ namespace EpicGames.Core
 				for (; ; )
 				{
 					Task<int> readTask = StdOut.ReadAsync(buffer, 0, bufferSize, cancellationToken);
-					_ = readTask.ContinueWith(x => _ = x.Exception, TaskContinuationOptions.OnlyOnFaulted);
+					_ = readTask.ContinueWith(x => _ = x.Exception, CancellationToken.None, TaskContinuationOptions.OnlyOnFaulted, TaskScheduler.Default);
 
 					Task completedTask = await Task.WhenAny(readTask, taskCompletionSource.Task);
 					cancellationToken.ThrowIfCancellationRequested();
