Fixed C# UHT to now allow exponential notion in specifiers and meta data.

diff --git a/Engine/Source/Programs/Shared/EpicGames.UHT/Parsers/UhtSpecifierParser.cs b/Engine/Source/Programs/Shared/EpicGames.UHT/Parsers/UhtSpecifierParser.cs
--- a/Engine/Source/Programs/Shared/EpicGames.UHT/Parsers/UhtSpecifierParser.cs
+++ b/Engine/Source/Programs/Shared/EpicGames.UHT/Parsers/UhtSpecifierParser.cs
@@ -455,6 +455,15 @@ namespace EpicGames.UHT.Parsers
 					}
 					break;
 
+				case UhtTokenType.FloatConst:
+					// TODO: C++ UHT compat 
+					if (token.Value.Span.Contains('e') || token.Value.Span.Contains('E'))
+					{
+						tokenReader.LogError($"Metadata and specifier values can not contain the exponential notion at this time");
+					}
+					builder.Append(token.GetConstantValue(respectQuotes));
+					break;
+
 				default:
 					builder.Append(token.GetConstantValue(respectQuotes));
 					break;
