UHT: Support for RigVM Lazy arguments

Ran local UHT compat checks

diff --git a/Engine/Source/Programs/Shared/EpicGames.UHT/Exporters/CodeGen/UhtHeaderCodeGeneratorHFile.cs b/Engine/Source/Programs/Shared/EpicGames.UHT/Exporters/CodeGen/UhtHeaderCodeGeneratorHFile.cs
--- a/Engine/Source/Programs/Shared/EpicGames.UHT/Exporters/CodeGen/UhtHeaderCodeGeneratorHFile.cs
+++ b/Engine/Source/Programs/Shared/EpicGames.UHT/Exporters/CodeGen/UhtHeaderCodeGeneratorHFile.cs
@@ -194,19 +194,44 @@ namespace EpicGames.UHT.Exporters.CodeGen
 										additionalParameters = ", RigVMExecuteContext.GetSlice().GetIndex()";
 									}
 
+									string getDataMethod = "GetData";
+									if (parameter.IsLazy)
+									{
+										getDataMethod = $"GetDataLazily<{parameter.TypeOriginal(false, false)}>";
+									}
+
 									if (parameter.IsArray)
 									{
 										string extendedType = parameter.ExtendedType();
+
+										if (!parameter.IsLazy)
+										{
+											builder
+												.Append("\t\tTArray")
+												.Append(extendedType)
+												.Append("& ")
+												.Append(paramNameOriginal)
+												.Append(" = ")
+												.Append("*(TArray")
+												.Append(extendedType)
+												.Append("*)");
+										}
+										else
+										{
+											builder
+												.Append("\t\tconst ")
+												.Append(paramTypeOriginal)
+												.Append("& ")
+												.Append(paramNameOriginal)
+												.Append(" = ");
+										}
+
 										builder
-											.Append("\t\tTArray")
-											.Append(extendedType)
-											.Append("& ")
-											.Append(paramNameOriginal)
-											.Append(" = *(TArray")
-											.Append(extendedType)
-											.Append("*)RigVMMemoryHandles[")
+											.Append("RigVMMemoryHandles[")
 											.Append(operandIndex)
-											.Append("].GetData(false")
+											.Append("].")
+											.Append(getDataMethod)
+											.Append("(false")
 											.Append(additionalParameters)
 											.Append("); \\\r\n");
 										operandIndex++;
@@ -226,11 +251,22 @@ namespace EpicGames.UHT.Exporters.CodeGen
 											.Append(variableType)
 											.Append(' ')
 											.Append(paramNameOriginal)
-											.Append(" = *(")
-											.Append(extractedType)
-											.Append("*)RigVMMemoryHandles[")
+											.Append(" = ");
+
+										if (!parameter.IsLazy)
+										{
+											builder
+												.Append("*(")
+												.Append(extractedType)
+												.Append("*)");
+										}
+
+										builder
+											.Append("RigVMMemoryHandles[")
 											.Append(operandIndex)
-											.Append("].GetData(false")
+											.Append("].")
+											.Append(getDataMethod)
+											.Append("(false")
 											.Append(additionalParameters)
 											.Append("); \\\r\n");
 										operandIndex++;
diff --git a/Engine/Source/Programs/Shared/EpicGames.UHT/Types/UhtScriptStruct.cs b/Engine/Source/Programs/Shared/EpicGames.UHT/Types/UhtScriptStruct.cs
--- a/Engine/Source/Programs/Shared/EpicGames.UHT/Types/UhtScriptStruct.cs
+++ b/Engine/Source/Programs/Shared/EpicGames.UHT/Types/UhtScriptStruct.cs
@@ -60,6 +60,11 @@ namespace EpicGames.UHT.Types
 		/// Set if the property is an enum as byte
 		/// </summary>
 		IsEnumAsByte = 0x00040000,
+
+		/// <summary>
+		/// Computes the value lazily
+		/// </summary>
+		IsLazy = 0x00080000,
 	}
 
 	/// <summary>
@@ -188,6 +193,12 @@ namespace EpicGames.UHT.Types
 		[JsonIgnore]
 		public bool IsEnumAsByte => ParameterFlags.HasAnyFlags(UhtRigVMParameterFlags.IsEnumAsByte);
 
+		/// <summary>
+		/// True if the parameter should be computed lazily
+		/// </summary>
+		[JsonIgnore]
+		public bool IsLazy => ParameterFlags.HasAnyFlags(UhtRigVMParameterFlags.IsLazy);
+
 		/// <summary>
 		/// True if the parameter is an array
 		/// </summary>
@@ -209,6 +220,7 @@ namespace EpicGames.UHT.Types
 			ParameterFlags |= property.MetaData.ContainsKey("Output") ? UhtRigVMParameterFlags.Output : UhtRigVMParameterFlags.None;
 			ParameterFlags |= property.IsEditorOnlyProperty ? UhtRigVMParameterFlags.EditorOnly : UhtRigVMParameterFlags.None;
 			ParameterFlags |= property.MetaData.ContainsKey("Singleton") ? UhtRigVMParameterFlags.Singleton : UhtRigVMParameterFlags.None;
+			ParameterFlags |= property.MetaData.ContainsKey("Lazy") ? UhtRigVMParameterFlags.IsLazy : UhtRigVMParameterFlags.None;
 
 			if (property.MetaData.ContainsKey("Visible"))
 			{
@@ -236,7 +248,7 @@ namespace EpicGames.UHT.Types
 				if (property.PropertyCaps.HasAnyFlags(UhtPropertyCaps.IsRigVMArray))
 				{
 					ParameterFlags |= UhtRigVMParameterFlags.IsArray;
-					if (IsConst())
+					if (IsConst() && !IsLazy)
 					{
 						string extendedType = ExtendedType(false);
 						CastName = $"{Name}_{index}_Array";
@@ -275,10 +287,11 @@ namespace EpicGames.UHT.Types
 		/// Get the type of the parameter
 		/// </summary>
 		/// <param name="castType">If true, return the cast type</param>
+		/// <param name="wrapLazyType">If true, return the wrapped lazy type as needed</param>
 		/// <returns>Parameter type</returns>
-		public string TypeOriginal(bool castType = false)
+		public string TypeOriginal(bool castType = false, bool wrapLazyType = true)
 		{
-			return castType && CastType != null ? CastType : Type;
+			return GetLazyType(castType && CastType != null ? CastType : Type, wrapLazyType);
 		}
 
 		/// <summary>
@@ -299,7 +312,7 @@ namespace EpicGames.UHT.Types
 		/// <returns>Base parameter type</returns>
 		public string BaseType(bool castType = false)
 		{
-			string typeOriginal = TypeOriginal(castType);
+			string typeOriginal = TypeOriginal(castType, false);
 
 			int lesserPos = typeOriginal.IndexOf('<', StringComparison.Ordinal);
 			if (lesserPos >= 0)
@@ -319,7 +332,7 @@ namespace EpicGames.UHT.Types
 		/// <returns>Template arguments of the type</returns>
 		public string ExtendedType(bool castType = false)
 		{
-			string typeOriginal = TypeOriginal(castType);
+			string typeOriginal = TypeOriginal(castType, false);
 
 			int lesserPos = typeOriginal.IndexOf('<', StringComparison.Ordinal);
 			if (lesserPos >= 0)
@@ -389,6 +402,21 @@ namespace EpicGames.UHT.Types
 			return IsConst() ? TypeConstRef(castType) : TypeRef(castType);
 		}
 
+		/// <summary>
+		/// Return the type wrapped with a lazy struct as needed
+		/// </summary>
+		/// <param name="typeToWrap">The type to wrap with a lazy struct</param>
+		/// <param name="wrapLazyType">If true the type will be wrapped as needed</param>
+		/// <returns>Type wrapped as needed</returns>
+		public string GetLazyType(string typeToWrap, bool wrapLazyType)
+		{
+			if (IsLazy && wrapLazyType)
+			{
+				return $"TRigVMLazyValue<{typeToWrap}>";
+			}
+			return typeToWrap;
+		}
+
 		/// <summary>
 		/// Return a variable declaration for the parameter
 		/// </summary>
@@ -910,6 +938,10 @@ namespace EpicGames.UHT.Types
 						}
 						else
 						{
+							if (parameter.IsLazy && parameter.Output)
+							{
+								this.LogError($"RigVM Struct {this.SourceName} - Member {parameter.Name} is both an output and a lazy input.");
+							}
 							RigVMStructInfo.Members.Add(parameter);
 						}
 					}
diff --git a/Engine/Source/Programs/UnrealHeaderTool/Private/CodeGenerator.cpp b/Engine/Source/Programs/UnrealHeaderTool/Private/CodeGenerator.cpp
--- a/Engine/Source/Programs/UnrealHeaderTool/Private/CodeGenerator.cpp
+++ b/Engine/Source/Programs/UnrealHeaderTool/Private/CodeGenerator.cpp
@@ -3441,15 +3441,35 @@ void FNativeClassHeaderGenerator::ExportGeneratedStructBodyMacros(FOutputDevice&
 				AdditionalParameters = SliceContextParameter;
 			}
 
+			FString GetDataMethod = TEXT("GetData");
+			if(Parameter.bIsLazy)
+			{
+				static constexpr TCHAR GetDataLazilyFormat[] = TEXT("GetDataLazily<%s>");
+				GetDataMethod = FString::Printf(GetDataLazilyFormat, *Parameter.TypeOriginal(false, false));
+			}
+
 			if (Parameter.IsArray())
 			{
-				FString ExtendedType = Parameter.ExtendedType();
-				RigVMStubProlog.Add(FString::Printf(TEXT("TArray%s& %s = *(TArray%s*)RigVMMemoryHandles[%d].GetData(false%s);"),
-					*ExtendedType,
-					*ParamNameOriginal,
-					*ExtendedType,
-					OperandIndex,
-					*AdditionalParameters));
+				if(!Parameter.bIsLazy)
+				{
+					FString ExtendedType = Parameter.ExtendedType();
+					RigVMStubProlog.Add(FString::Printf(TEXT("TArray%s& %s = *(TArray%s*)RigVMMemoryHandles[%d].%s(false%s);"),
+						*ExtendedType,
+						*ParamNameOriginal,
+						*ExtendedType,
+						OperandIndex,
+						*GetDataMethod,
+						*AdditionalParameters));
+				}
+				else
+				{
+					RigVMStubProlog.Add(FString::Printf(TEXT("%s %s = RigVMMemoryHandles[%d].%s(false%s);"),
+						*Parameter.TypeConstRef(false),
+						*ParamNameOriginal,
+						OperandIndex,
+						*GetDataMethod,
+						*AdditionalParameters));
+				}
 
 				OperandIndex++;
 			}
@@ -3462,13 +3482,19 @@ void FNativeClassHeaderGenerator::ExportGeneratedStructBodyMacros(FOutputDevice&
 				{
 					ExtractedType = Parameter.TypeOriginal(true);
 				}
+				
 				FString ParameterCast = FString::Printf(TEXT("*(%s*)"), *ExtractedType);
+				if(Parameter.bIsLazy)
+				{
+					ParameterCast.Reset();
+				}
 
-				RigVMStubProlog.Add(FString::Printf(TEXT("%s %s = %sRigVMMemoryHandles[%d].GetData(false%s);"),
+				RigVMStubProlog.Add(FString::Printf(TEXT("%s %s = %sRigVMMemoryHandles[%d].%s(false%s);"),
 				*VariableType,
 				*ParamNameOriginal,
 				*ParameterCast,
 				OperandIndex,
+				*GetDataMethod,
 				*AdditionalParameters));
 
 				OperandIndex++;
diff --git a/Engine/Source/Programs/UnrealHeaderTool/Private/HeaderParser.cpp b/Engine/Source/Programs/UnrealHeaderTool/Private/HeaderParser.cpp
--- a/Engine/Source/Programs/UnrealHeaderTool/Private/HeaderParser.cpp
+++ b/Engine/Source/Programs/UnrealHeaderTool/Private/HeaderParser.cpp
@@ -6628,6 +6628,7 @@ const FName FHeaderParser::NAME_InputText(TEXT("Input"));
 const FName FHeaderParser::NAME_OutputText(TEXT("Output"));
 const FName FHeaderParser::NAME_ConstantText(TEXT("Constant"));
 const FName FHeaderParser::NAME_VisibleText(TEXT("Visible"));
+const FName FHeaderParser::NAME_LazyText(TEXT("Lazy"));
 
 const FName FHeaderParser::NAME_SingletonText(TEXT("Singleton"));
 
@@ -6686,6 +6687,13 @@ void FHeaderParser::ParseRigVMMethodParameters(FUnrealStructDefinitionInfo& Stru
 			Parameter.bOutput = false;
 		}
 
+		Parameter.bIsLazy = PropertyDef->HasMetaData(NAME_LazyText);
+
+		if(Parameter.bOutput && Parameter.bIsLazy)
+		{
+			LogError(TEXT("RigVM Struct '%s' - Member '%s' is both an output and a lazy input."), *StructDef.GetName(), *Parameter.Name);
+		}
+
 		if (Parameter.bEditorOnly)
 		{
 			LogError(TEXT("RigVM Struct '%s' - Member '%s' is editor only - WITH_EDITORONLY_DATA not allowed on structs with RIGVM_METHOD."), *StructDef.GetName(), *Parameter.Name, *MemberCPPType);
@@ -6718,7 +6726,7 @@ void FHeaderParser::ParseRigVMMethodParameters(FUnrealStructDefinitionInfo& Stru
 		{
 			ExtendedCPPType = FString::Printf(TEXT("<%s>"), *ExtendedCPPType.LeftChop(1).RightChop(1));
 			
-			if(Parameter.IsConst())
+			if(Parameter.IsConst() && !Parameter.bIsLazy)
 			{
 				ExtendedCPPType = FString::Printf(TEXT("<const %s>"), *ExtendedCPPType.LeftChop(1).RightChop(1));
 				Parameter.CastName = FString::Printf(TEXT("%s_%d_Array"), *Parameter.Name, StructRigVMInfo.Members.Num());
diff --git a/Engine/Source/Programs/UnrealHeaderTool/Private/HeaderParser.h b/Engine/Source/Programs/UnrealHeaderTool/Private/HeaderParser.h
--- a/Engine/Source/Programs/UnrealHeaderTool/Private/HeaderParser.h
+++ b/Engine/Source/Programs/UnrealHeaderTool/Private/HeaderParser.h
@@ -663,6 +663,7 @@ public:
 	static const FName NAME_OutputText;
 	static const FName NAME_ConstantText;
 	static const FName NAME_VisibleText;
+	static const FName NAME_LazyText;
 	static const FName NAME_SingletonText;
 	static const TCHAR* TArrayText;
 	static const TCHAR* TEnumAsByteText;
diff --git a/Engine/Source/Programs/UnrealHeaderTool/Private/ParserHelper.h b/Engine/Source/Programs/UnrealHeaderTool/Private/ParserHelper.h
--- a/Engine/Source/Programs/UnrealHeaderTool/Private/ParserHelper.h
+++ b/Engine/Source/Programs/UnrealHeaderTool/Private/ParserHelper.h
@@ -534,6 +534,7 @@ struct FRigVMParameter
 		, bEditorOnly(false)
 		, bIsEnum(false)
 		, bIsEnumAsByte(false)
+		, bIsLazy(false)
 	{
 	}
 
@@ -550,15 +551,16 @@ struct FRigVMParameter
 	bool bEditorOnly;
 	bool bIsEnum;
 	bool bIsEnumAsByte;
+	bool bIsLazy;
 
 	const FString& NameOriginal(bool bCastName = false) const
 	{
 		return (bCastName && !CastName.IsEmpty()) ? CastName : Name;
 	}
 
-	const FString& TypeOriginal(bool bCastType = false) const
+	FString TypeOriginal(bool bCastType = false, bool bWrapLazyType = true) const
 	{
-		return (bCastType && !CastType.IsEmpty()) ? CastType : Type;
+		return GetLazyType((bCastType && !CastType.IsEmpty()) ? CastType : Type, bWrapLazyType);
 	}
 
 	FString Declaration(bool bCastType = false, bool bCastName = false) const
@@ -568,7 +570,7 @@ struct FRigVMParameter
 
 	FString BaseType(bool bCastType = false) const
 	{
-		const FString& String = TypeOriginal(bCastType);
+		const FString& String = TypeOriginal(bCastType, false);
 		int32 LesserPos = 0;
 		if (String.FindChar(TEXT('<'), LesserPos))
 		{
@@ -579,7 +581,7 @@ struct FRigVMParameter
 
 	FString ExtendedType(bool bCastType = false) const
 	{
-		const FString& String = TypeOriginal(bCastType);
+		const FString& String = TypeOriginal(bCastType, false);
 		int32 LesserPos = 0;
 		if (String.FindChar(TEXT('<'), LesserPos))
 		{
@@ -619,6 +621,16 @@ struct FRigVMParameter
 		return IsConst() ? TypeConstRef(bCastType) : TypeRef(bCastType);
 	}
 
+	FString GetLazyType(const FString& InType, bool bWrapLazyType) const
+	{
+		if(bIsLazy && bWrapLazyType)
+		{
+			static constexpr TCHAR LazyTypeFormat[] = TEXT("TRigVMLazyValue<%s>");
+			return FString::Printf(LazyTypeFormat, *InType);
+		}
+		return InType;
+	}
+
 	FString Variable(bool bCastType = false, bool bCastName = false) const
 	{
 		return FString::Printf(TEXT("%s %s"), *TypeVariableRef(bCastType), *NameOriginal(bCastName));
diff --git a/Engine/Source/Programs/UnrealHeaderTool/Resources/UHTDebugging/RigVMStructHeader.h b/Engine/Source/Programs/UnrealHeaderTool/Resources/UHTDebugging/RigVMStructHeader.h
--- a/Engine/Source/Programs/UnrealHeaderTool/Resources/UHTDebugging/RigVMStructHeader.h
+++ b/Engine/Source/Programs/UnrealHeaderTool/Resources/UHTDebugging/RigVMStructHeader.h
@@ -161,4 +161,25 @@ struct CONTROLRIG_API FRigUnit_BeginExecution
 	// The execution result
 	UPROPERTY(EditAnywhere, Transient, DisplayName = "Execute", Category = "BeginExecution", meta = (Output))
 	FUHTTestExecuteContext ExecuteContext;
-};
\ No newline at end of file
+};
+
+USTRUCT()
+struct FRigVMLazyStruct
+{
+	GENERATED_BODY()
+
+	RIGVM_METHOD()
+	void Compute();
+
+	UPROPERTY(meta = (Input, Lazy))
+	float A;
+
+	UPROPERTY(meta = (Input, Lazy))
+	TArray<float> B;
+
+	UPROPERTY(meta = (Input))
+	TArray<float> C;
+
+	UPROPERTY(meta = (Output))
+	float Result;
+};
