EpicGames.Perforce: Run managed tests in parallel

diff --git a/Engine/Source/Programs/Shared/EpicGames.Perforce.Managed.Tests/AssemblyInfo.cs b/Engine/Source/Programs/Shared/EpicGames.Perforce.Managed.Tests/AssemblyInfo.cs
new file mode 100644
--- /dev/null
+++ b/Engine/Source/Programs/Shared/EpicGames.Perforce.Managed.Tests/AssemblyInfo.cs
@@ -0,0 +1,5 @@
+// Copyright Epic Games, Inc. All Rights Reserved.
+
+using Microsoft.VisualStudio.TestTools.UnitTesting;
+
+[assembly: Parallelize(Workers = 10, Scope = ExecutionScope.MethodLevel)]
\ No newline at end of file
