EpicGames.Perforce: Add fixture data mirroring that of depot and streams

diff --git a/Engine/Source/Programs/Shared/EpicGames.Perforce.Fixture/BasePerforceFixtureTest.cs b/Engine/Source/Programs/Shared/EpicGames.Perforce.Fixture/BasePerforceFixtureTest.cs
--- a/Engine/Source/Programs/Shared/EpicGames.Perforce.Fixture/BasePerforceFixtureTest.cs
+++ b/Engine/Source/Programs/Shared/EpicGames.Perforce.Fixture/BasePerforceFixtureTest.cs
@@ -37,6 +37,11 @@ public abstract class BasePerforceFixtureTest : IDisposable
 	/// </summary>
 	protected PerforceConnection PerforceConnection { get; private set; }
 
+	/// <summary>
+	/// Fixture data mirroring how the stream/depot looks
+	/// </summary>
+	protected PerforceFixture Fixture { get; } = new ();
+
 	private readonly string _clientName;
 	
 	protected BasePerforceFixtureTest()
diff --git a/Engine/Source/Programs/Shared/EpicGames.Perforce.Fixture/PerforceFixture.cs b/Engine/Source/Programs/Shared/EpicGames.Perforce.Fixture/PerforceFixture.cs
new file mode 100644
--- /dev/null
+++ b/Engine/Source/Programs/Shared/EpicGames.Perforce.Fixture/PerforceFixture.cs
@@ -0,0 +1,162 @@
+// Copyright Epic Games, Inc. All Rights Reserved.
+
+using System;
+using System.Collections.Generic;
+using System.IO;
+using System.Linq;
+using System.Text;
+using EpicGames.Core;
+using Microsoft.VisualStudio.TestTools.UnitTesting;
+
+namespace EpicGames.Perforce.Fixture;
+
+public class DepotFileFixture
+{
+	public string DepotFile { get; }
+	public string ClientFile { get; }
+	public long Size { get; }
+	public string Digest { get; }
+	public string Content { get; }
+	
+	public DepotFileFixture(string depotFile, string clientFile, string content)
+	{
+		DepotFile = depotFile;
+		ClientFile = clientFile;
+		Content = content;
+		Size = content.Length;
+		Digest = PerforceFixture.CalcMd5(content);
+	}
+}
+
+public class ChangelistFixture
+{
+	public int Number { get; }
+	public string Description { get; }
+	
+	/// <summary>
+	/// List of files in stream as how they would appear locally on disk, when synced to this changelist
+	/// (after any view maps have been applied)
+	/// </summary>
+	public IReadOnlyList<DepotFileFixture> StreamFiles { get; }
+
+	public ChangelistFixture(int number, string description, List<DepotFileFixture> streamFiles)
+	{
+		Number = number;
+		Description = description;
+		StreamFiles = streamFiles;
+	}
+
+	/// <summary>
+	/// Assert directory contains exactly the files described by stream
+	/// </summary>
+	/// <param name="clientRoot">Client/workspace root directory</param>
+	public void AssertAllFilesOnDisk(string clientRoot)
+	{
+		(List<string> actual, List<string> expected) = GetFiles(clientRoot);
+		CollectionAssert.AreEqual(expected, actual);
+	}
+	
+	/// <summary>
+	/// Assert directory does contain the files described by stream
+	/// </summary>
+	/// <param name="clientRoot">Client/workspace root directory</param>
+	public void AssertNotAllFilesOnDisk(string clientRoot)
+	{
+		(List<string> actual, List<string> expected) = GetFiles(clientRoot);
+		CollectionAssert.AreNotEqual(expected, actual);
+	}
+
+	public (List<string> localFiles, List<string> streamFiles) GetFiles(string clientRoot)
+	{
+		EnumerationOptions options = new () { RecurseSubdirectories = true };
+		List<string> localFiles = Directory.EnumerateFiles(clientRoot, "*", options)
+			.Select(x => Path.GetRelativePath(clientRoot, x))
+			.Select(x => x.Replace("\\", "/", StringComparison.Ordinal))
+			.ToList();
+		List<string> streamFiles = new(StreamFiles.Select(x => x.ClientFile));
+		
+		localFiles.Sort();
+		streamFiles.Sort();
+		
+		return (localFiles, streamFiles);
+	}
+}
+
+public class StreamFixture
+{
+	public string Root { get; }
+	public IReadOnlyList<ChangelistFixture> Changelists { get; }
+
+	public StreamFixture(string root, IReadOnlyList<ChangelistFixture> changelists)
+	{
+		Root = root;
+		Changelists = changelists;
+	}
+}
+
+public class PerforceFixture
+{
+	public StreamFixture StreamFooMain { get; } = new ("//Foo/Main",
+		new List<ChangelistFixture>
+		{
+			new(0, "<placeholder>", new List<DepotFileFixture>()),
+			new(1, "<placeholder>", new List<DepotFileFixture>()),
+			new(2, "Initial import",
+				new List<DepotFileFixture>()
+				{
+					new("//Foo/Main/main.cpp", "main.cpp", "This is change main.cpp #1"),
+					new("//Foo/Main/main.h", "main.h", "This is change main.h #1"),
+					new("//Foo/Main/common.h", "common.h", "This is change common.h #1"),
+					new("//Foo/Main/unused.cpp", "unused.cpp", "This is change unused.cpp #1"),
+					new("//Foo/Main/Data/data.txt", "Data/data.txt", "This is change data.txt #1"),
+				}),
+			new(3, "Improvement to main.cpp",
+				new List<DepotFileFixture>()
+				{
+					new("//Foo/Main/main.cpp", "main.cpp", "This is change main.cpp #2"),
+					new("//Foo/Main/main.h", "main.h", "This is change main.h #1"),
+					new("//Foo/Main/common.h", "common.h", "This is change common.h #1"),
+					new("//Foo/Main/unused.cpp", "unused.cpp", "This is change unused.cpp #1"),
+					new("//Foo/Main/Data/data.txt", "Data/data.txt", "This is change data.txt #1"),
+				}),
+			new(4, "Delete an unused file",
+				new List<DepotFileFixture>()
+				{
+					new("//Foo/Main/main.cpp", "main.cpp", "This is change main.cpp #2"),
+					new("//Foo/Main/main.h", "main.h", "This is change main.h #1"),
+					new("//Foo/Main/common.h", "common.h", "This is change common.h #1"),
+					new("//Foo/Main/Data/data.txt", "Data/data.txt", "This is change data.txt #1"),
+				}),
+			new(5, "Rename common.h to shared.h",
+				new List<DepotFileFixture>()
+				{
+					new("//Foo/Main/main.cpp", "main.cpp", "This is change main.cpp #2"),
+					new("//Foo/Main/main.h", "main.h", "This is change main.h #1"),
+					new("//Foo/Main/shared.h", "shared.h", "This is change common.h #1"),
+					new("//Foo/Main/Data/data.txt", "Data/data.txt", "This is change data.txt #1"),
+				}),
+			new(6, "Some updates to main",
+				new List<DepotFileFixture>()
+				{
+					new("//Foo/Main/main.cpp", "main.cpp", "This is change main.cpp #3"),
+					new("//Foo/Main/main.h", "main.h", "This is change main.h #2"),
+					new("//Foo/Main/shared.h", "shared.h", "This is change common.h #1"),
+					new("//Foo/Main/Data/data.txt", "Data/data.txt", "This is change data.txt #1"),
+				}),
+			new(7, "Add more data",
+				new List<DepotFileFixture>()
+				{
+					new("//Foo/Main/main.cpp", "main.cpp", "This is change main.cpp #3"),
+					new("//Foo/Main/main.h", "main.h", "This is change main.h #2"),
+					new("//Foo/Main/shared.h", "shared.h", "This is change common.h #1"),
+					new("//Foo/Main/Data/data.txt", "Data/data.txt", "This is change data.txt #1"),
+					new("//Foo/Main/Data/moredata.txt", "Data/moredata.txt", "This is change moredata.txt #1"),
+				}),
+		});
+
+	public static string CalcMd5(string content)
+	{
+		byte[] data = Encoding.ASCII.GetBytes(content);
+		return Md5Hash.Compute(data).ToString();
+	}
+}
\ No newline at end of file
diff --git a/Engine/Source/Programs/Shared/EpicGames.Perforce.Fixture/setup.sh b/Engine/Source/Programs/Shared/EpicGames.Perforce.Fixture/setup.sh
--- a/Engine/Source/Programs/Shared/EpicGames.Perforce.Fixture/setup.sh
+++ b/Engine/Source/Programs/Shared/EpicGames.Perforce.Fixture/setup.sh
@@ -98,6 +98,9 @@ echo "This is change main.cpp #3" > $CLIENT_ROOT/main.cpp
 echo "This is change main.h #2" > $CLIENT_ROOT/main.h
 p4 submit -d "Some updates to main"
 
+# Changelist #7
+echo "This is change moredata.txt #1" > $CLIENT_ROOT/Data/moredata.txt && p4 add $CLIENT_ROOT/Data/moredata.txt
+p4 submit -d "Add more data"
 
 # Stop server running in background
 kill "$(pidof p4d)"
\ No newline at end of file
diff --git a/Engine/Source/Programs/Shared/EpicGames.Perforce.Managed.Tests/ManagedWorkspaceTest.cs b/Engine/Source/Programs/Shared/EpicGames.Perforce.Managed.Tests/ManagedWorkspaceTest.cs
--- a/Engine/Source/Programs/Shared/EpicGames.Perforce.Managed.Tests/ManagedWorkspaceTest.cs
+++ b/Engine/Source/Programs/Shared/EpicGames.Perforce.Managed.Tests/ManagedWorkspaceTest.cs
@@ -13,20 +13,44 @@ namespace EpicGames.Perforce.Managed.Tests;
 [TestClass]
 public class ManagedWorkspaceTest : BasePerforceFixtureTest
 {
+	private string SyncDir => Path.Join(TempDir.FullName, "Sync");
+	private string StreamName => Fixture.StreamFooMain.Root;
+	
 	[TestMethod]
-	public async Task SimpleSync()
+	public async Task SyncFromScratch()
 	{
-		ILogger<ManagedWorkspace> logger = LoggerFactory.CreateLogger<ManagedWorkspace>();
-		ManagedWorkspace workspace = await ManagedWorkspace.LoadOrCreateAsync(Environment.MachineName, TempDir, true, logger, CancellationToken.None);
-		
-		await workspace.SetupAsync(PerforceConnection, "//Foo/Main", CancellationToken.None);
+		ManagedWorkspace workspace = await GetManagedWorkspace();
+		await workspace.SetupAsync(PerforceConnection, StreamName, CancellationToken.None);
+
+		await workspace.SyncAsync(PerforceConnection, StreamName, 6, Array.Empty<string>(), true,
+			false, null, CancellationToken.None);
+
+		Fixture.StreamFooMain.Changelists[6].AssertAllFilesOnDisk(SyncDir);
+		Fixture.StreamFooMain.Changelists[7].AssertNotAllFilesOnDisk(SyncDir); // Contains an extra file so should fail
+	}
+	
+	[TestMethod]
+	public async Task SyncIncrementally()
+	{
+		ManagedWorkspace ws = await GetManagedWorkspace();
+		await ws.SetupAsync(PerforceConnection, StreamName, CancellationToken.None);
 
-		await workspace.SyncAsync(PerforceConnection, "//Foo/Main", 6, Array.Empty<string>(), true,
+		await ws.SyncAsync(PerforceConnection, StreamName, 6, Array.Empty<string>(), true,
 			false, null, CancellationToken.None);
 
-		Assert.IsTrue(File.Exists(Path.Join(TempDir.FullName, "Sync", "main.cpp")));
-		Assert.IsTrue(File.Exists(Path.Join(TempDir.FullName, "Sync", "main.h")));
-		Assert.IsTrue(File.Exists(Path.Join(TempDir.FullName, "Sync", "shared.h")));
-		Assert.IsTrue(File.Exists(Path.Join(TempDir.FullName, "Sync", "Data", "data.txt")));
+		Fixture.StreamFooMain.Changelists[6].AssertAllFilesOnDisk(SyncDir);
+		Fixture.StreamFooMain.Changelists[7].AssertNotAllFilesOnDisk(SyncDir); // Contains an extra file so should fail
+		
+		await ws.SyncAsync(PerforceConnection, StreamName, 7, Array.Empty<string>(), true,
+			false, null, CancellationToken.None);
+
+		Fixture.StreamFooMain.Changelists[6].AssertNotAllFilesOnDisk(SyncDir);
+		Fixture.StreamFooMain.Changelists[7].AssertAllFilesOnDisk(SyncDir);
+	}
+
+	private async Task<ManagedWorkspace> GetManagedWorkspace(bool overwrite = true)
+	{
+		ILogger<ManagedWorkspace> logger = LoggerFactory.CreateLogger<ManagedWorkspace>();
+		return await ManagedWorkspace.LoadOrCreateAsync(Environment.MachineName, TempDir, overwrite, logger, CancellationToken.None);
 	}
 }
\ No newline at end of file
