Fixed C# UHT to properly generate an error when "const" is used inappropriately in a function or property declaration.

diff --git a/Engine/Source/Programs/Shared/EpicGames.UHT/Parsers/UhtPropertyParser.cs b/Engine/Source/Programs/Shared/EpicGames.UHT/Parsers/UhtPropertyParser.cs
--- a/Engine/Source/Programs/Shared/EpicGames.UHT/Parsers/UhtPropertyParser.cs
+++ b/Engine/Source/Programs/Shared/EpicGames.UHT/Parsers/UhtPropertyParser.cs
@@ -1295,13 +1295,11 @@ namespace EpicGames.UHT.Parsers
 		[SuppressMessage("Style", "IDE0060:Remove unused parameter", Justification = "Attribute accessed method")]
 		private static UhtProperty? DefaultProperty(UhtPropertyResolvePhase resolvePhase, UhtPropertySettings propertySettings, IUhtTokenReader tokenReader, UhtToken matchedToken)
 		{
+			UhtProperty? property = null;
 			UhtSession session = propertySettings.Outer.Session;
 			int typeStartPos = tokenReader.PeekToken().InputStartPos;
 
-			if (tokenReader.TryOptional("const"))
-			{
-				propertySettings.MetaData.Add(UhtNames.NativeConst, "");
-			}
+			bool gotConst = tokenReader.TryOptional("const");
 
 			UhtFindOptions findOptions = UhtFindOptions.DelegateFunction | UhtFindOptions.Enum | UhtFindOptions.Class | UhtFindOptions.ScriptStruct;
 			if (tokenReader.TryOptional("enum"))
@@ -1339,29 +1337,35 @@ namespace EpicGames.UHT.Parsers
 					}
 				}
 
-				return new UhtEnumProperty(propertySettings, enumObj);
+				property = new UhtEnumProperty(propertySettings, enumObj);
 			}
 			else if (type is UhtScriptStruct scriptStruct)
 			{
-				return new UhtStructProperty(propertySettings, scriptStruct);
+				property = new UhtStructProperty(propertySettings, scriptStruct);
 			}
 			else if (type is UhtFunction function)
 			{
 				if (!function.FunctionFlags.HasAnyFlags(EFunctionFlags.MulticastDelegate))
 				{
-					return new UhtDelegateProperty(propertySettings, function);
+					property = new UhtDelegateProperty(propertySettings, function);
 				}
 				else if (function.FunctionType == UhtFunctionType.SparseDelegate)
 				{
-					return new UhtMulticastSparseDelegateProperty(propertySettings, function);
+					property = new UhtMulticastSparseDelegateProperty(propertySettings, function);
 				}
 				else
 				{
-					return new UhtMulticastInlineDelegateProperty(propertySettings, function);
+					property = new UhtMulticastInlineDelegateProperty(propertySettings, function);
 				}
 			}
 			else if (type is UhtClass classObj)
 			{
+				if (gotConst)
+				{
+					propertySettings.MetaData.Add(UhtNames.NativeConst, "");
+					gotConst = false;
+				}
+
 				// Const after variable type but before pointer symbol
 				if (tokenReader.TryOptional("const"))
 				{
@@ -1382,22 +1386,36 @@ namespace EpicGames.UHT.Parsers
 
 				if (classObj.ClassFlags.HasAnyFlags(EClassFlags.Interface))
 				{
-					return new UhtInterfaceProperty(propertySettings, classObj);
+					property = new UhtInterfaceProperty(propertySettings, classObj);
 				}
 				else if (classObj.IsChildOf(session.UClass))
 				{
 					// UObject here specifies that there is no limiter
-					return new UhtClassProperty(propertySettings, classObj, session.UObject);
+					property = new UhtClassProperty(propertySettings, classObj, session.UObject);
 				}
 				else
 				{
-					return new UhtObjectProperty(propertySettings, classObj);
+					property = new UhtObjectProperty(propertySettings, classObj);
 				}
 			}
 			else
 			{
 				throw new UhtIceException("Unexpected type found");
 			}
+
+			if (gotConst)
+			{
+				if (propertySettings.PropertyCategory == UhtPropertyCategory.Member)
+				{
+					tokenReader.LogError("Const properties are not supported.");
+				}
+				else
+				{
+					tokenReader.LogError($"Inappropriate keyword 'const' on variable of type {type.SourceName}");
+				}
+			}
+
+			return property;
 		}
 	}
 }
