Fixed issue where functions would report a conflict with a property when it shouldn't have.

diff --git a/Engine/Source/Programs/Shared/EpicGames.UHT/Types/UhtFunction.cs b/Engine/Source/Programs/Shared/EpicGames.UHT/Types/UhtFunction.cs
--- a/Engine/Source/Programs/Shared/EpicGames.UHT/Types/UhtFunction.cs
+++ b/Engine/Source/Programs/Shared/EpicGames.UHT/Types/UhtFunction.cs
@@ -759,7 +759,7 @@ namespace EpicGames.UHT.Types
 			}
 
 			// Make sure this isn't a duplicate type.  In some ways this has already been check above
-			UhtType? existingType = Outer!.FindType(UhtFindOptions.PropertyOrFunction | UhtFindOptions.EngineName | UhtFindOptions.SelfOnly, EngineName);
+			UhtType? existingType = Outer!.FindType(UhtFindOptions.Function | UhtFindOptions.EngineName | UhtFindOptions.SelfOnly, EngineName);
 			if (existingType != null && existingType != this)
 			{
 				this.LogError($"'{EngineName}' conflicts with '{existingType.FullName}'");
