UAT: Disable PDB generation on win-arm64. This is a known issue in NET 6 ([REDACTED_PROFILE_URL]/roslyn/issues/59421).

diff --git a/Engine/Source/Programs/Shared/EpicGames.Build/EpicGames.Build.csproj b/Engine/Source/Programs/Shared/EpicGames.Build/EpicGames.Build.csproj
--- a/Engine/Source/Programs/Shared/EpicGames.Build/EpicGames.Build.csproj
+++ b/Engine/Source/Programs/Shared/EpicGames.Build/EpicGames.Build.csproj
@@ -10,7 +10,7 @@
     <GenerateTargetFrameworkAttribute>false</GenerateTargetFrameworkAttribute>
     <Configurations>Debug;Release;Development;Analyze</Configurations>
     <Nullable>enable</Nullable>
-    <DebugType>pdbonly</DebugType>
+    <DebugType Condition="'$(OS)' != 'Windows_NT' or '$(PROCESSOR_ARCHITECTURE)' != 'ARM64'">pdbonly</DebugType> <!-- [REDACTED_PROFILE_URL]/roslyn/issues/59421 -->
     <ProduceReferenceAssemblyInOutDir>true</ProduceReferenceAssemblyInOutDir>
     <AllowUnsafeBlocks>true</AllowUnsafeBlocks>
     <RunAnalyzersDuringBuild>false</RunAnalyzersDuringBuild>
diff --git a/Engine/Source/Programs/Shared/EpicGames.Core/EpicGames.Core.csproj b/Engine/Source/Programs/Shared/EpicGames.Core/EpicGames.Core.csproj
--- a/Engine/Source/Programs/Shared/EpicGames.Core/EpicGames.Core.csproj
+++ b/Engine/Source/Programs/Shared/EpicGames.Core/EpicGames.Core.csproj
@@ -10,7 +10,7 @@
     <Configurations>Debug;Release;Development;Analyze</Configurations>
     <Nullable>enable</Nullable>
     <AllowUnsafeBlocks>true</AllowUnsafeBlocks>
-    <DebugType>pdbonly</DebugType>
+    <DebugType Condition="'$(OS)' != 'Windows_NT' or '$(PROCESSOR_ARCHITECTURE)' != 'ARM64'">pdbonly</DebugType>
     <RunAnalyzersDuringBuild>false</RunAnalyzersDuringBuild>
   </PropertyGroup>
   <PropertyGroup Condition=" '$(Configuration)|$(Platform)' == 'Debug|AnyCPU' ">
diff --git a/Engine/Source/Programs/Shared/EpicGames.IoHash/EpicGames.IoHash.csproj b/Engine/Source/Programs/Shared/EpicGames.IoHash/EpicGames.IoHash.csproj
--- a/Engine/Source/Programs/Shared/EpicGames.IoHash/EpicGames.IoHash.csproj
+++ b/Engine/Source/Programs/Shared/EpicGames.IoHash/EpicGames.IoHash.csproj
@@ -9,8 +9,7 @@
     <GenerateTargetFrameworkAttribute>false</GenerateTargetFrameworkAttribute>
     <Configurations>Debug;Release;Development</Configurations>
     <Nullable>enable</Nullable>
-    <DebugType>full</DebugType>
-    <DebugType>pdbonly</DebugType>
+    <DebugType Condition="'$(OS)' != 'Windows_NT' or '$(PROCESSOR_ARCHITECTURE)' != 'ARM64'">pdbonly</DebugType>
   </PropertyGroup>
   <PropertyGroup Condition=" '$(Configuration)|$(Platform)' == 'Debug|AnyCPU' ">
     <DocumentationFile></DocumentationFile>
diff --git a/Engine/Source/Programs/Shared/EpicGames.MsBuild/EpicGames.MsBuild.csproj b/Engine/Source/Programs/Shared/EpicGames.MsBuild/EpicGames.MsBuild.csproj
--- a/Engine/Source/Programs/Shared/EpicGames.MsBuild/EpicGames.MsBuild.csproj
+++ b/Engine/Source/Programs/Shared/EpicGames.MsBuild/EpicGames.MsBuild.csproj
@@ -9,7 +9,7 @@
     <GenerateAssemblyInfo>false</GenerateAssemblyInfo>
     <GenerateTargetFrameworkAttribute>false</GenerateTargetFrameworkAttribute>
     <Configurations>Debug;Release;Development</Configurations>
-    <DebugType>pdbonly</DebugType>
+    <DebugType Condition="'$(OS)' != 'Windows_NT' or '$(PROCESSOR_ARCHITECTURE)' != 'ARM64'">pdbonly</DebugType> <!-- [REDACTED_PROFILE_URL]/roslyn/issues/59421 -->
     <ProduceReferenceAssemblyInOutDir>true</ProduceReferenceAssemblyInOutDir>
   </PropertyGroup>
   <PropertyGroup Condition=" '$(Configuration)|$(Platform)' == 'Debug|AnyCPU' ">
diff --git a/Engine/Source/Programs/Shared/EpicGames.Serialization/EpicGames.Serialization.csproj b/Engine/Source/Programs/Shared/EpicGames.Serialization/EpicGames.Serialization.csproj
--- a/Engine/Source/Programs/Shared/EpicGames.Serialization/EpicGames.Serialization.csproj
+++ b/Engine/Source/Programs/Shared/EpicGames.Serialization/EpicGames.Serialization.csproj
@@ -9,7 +9,7 @@
     <GenerateTargetFrameworkAttribute>false</GenerateTargetFrameworkAttribute>
     <Configurations>Debug;Release;Development;Analyze</Configurations>
     <Nullable>enable</Nullable>
-    <DebugType>pdbonly</DebugType>
+    <DebugType Condition="'$(OS)' != 'Windows_NT' or '$(PROCESSOR_ARCHITECTURE)' != 'ARM64'">pdbonly</DebugType> <!-- [REDACTED_PROFILE_URL]/roslyn/issues/59421 -->
     <RunAnalyzersDuringBuild>false</RunAnalyzersDuringBuild>
   </PropertyGroup>
   <PropertyGroup Condition=" '$(Configuration)|$(Platform)' == 'Debug|AnyCPU' ">
diff --git a/Engine/Source/Programs/Shared/EpicGames.UHT/EpicGames.UHT.csproj b/Engine/Source/Programs/Shared/EpicGames.UHT/EpicGames.UHT.csproj
--- a/Engine/Source/Programs/Shared/EpicGames.UHT/EpicGames.UHT.csproj
+++ b/Engine/Source/Programs/Shared/EpicGames.UHT/EpicGames.UHT.csproj
@@ -12,7 +12,7 @@
     <Nullable>enable</Nullable>
     <AllowUnsafeBlocks>true</AllowUnsafeBlocks>
     <GenerateDocumentationFile>True</GenerateDocumentationFile>
-    <DebugType>pdbonly</DebugType>
+    <DebugType Condition="'$(OS)' != 'Windows_NT' or '$(PROCESSOR_ARCHITECTURE)' != 'ARM64'">pdbonly</DebugType> <!-- [REDACTED_PROFILE_URL]/roslyn/issues/59421 -->
     <EnforceCodeStyleInBuild>true</EnforceCodeStyleInBuild>
     <ProduceReferenceAssemblyInOutDir>true</ProduceReferenceAssemblyInOutDir>
     <RunAnalyzersDuringBuild>false</RunAnalyzersDuringBuild>
