Horde: Allow any arbitrary base URL for storage. This allows using different endpoints other than the generic user-facing /api/v1/storage/{ns} route.

diff --git a/Engine/Source/Programs/Horde/Horde.Agent/Commands/Bundles/CreateCommand.cs b/Engine/Source/Programs/Horde/Horde.Agent/Commands/Bundles/CreateCommand.cs
--- a/Engine/Source/Programs/Horde/Horde.Agent/Commands/Bundles/CreateCommand.cs
+++ b/Engine/Source/Programs/Horde/Horde.Agent/Commands/Bundles/CreateCommand.cs
@@ -42,22 +42,16 @@ namespace Horde.Agent.Commands.Bundles
 		class HttpStorageClientOwner : IStorageClientOwner
 		{
 			public IMemoryCache Cache { get; }
-			readonly HttpClient _httpClient;
-			readonly HttpClient _redirectHttpClient;
 			public IStorageClient Store { get; }
 
-			public HttpStorageClientOwner(NamespaceId namespaceId, Uri baseUri, ILogger logger)
+			public HttpStorageClientOwner(Uri baseUri, ILogger logger)
 			{
 				Cache = new MemoryCache(new MemoryCacheOptions());
-				_httpClient = new HttpClient { BaseAddress = baseUri };
-				_redirectHttpClient = new HttpClient();
-				Store = new HttpStorageClient(namespaceId, _httpClient, _redirectHttpClient, logger);
+				Store = new HttpStorageClient(() => new HttpClient { BaseAddress = baseUri }, () => new HttpClient(), logger);
 			}
 
 			public void Dispose()
 			{
-				_redirectHttpClient.Dispose();
-				_httpClient.Dispose();
 				Cache.Dispose();
 			}
 		}
@@ -80,7 +74,7 @@ namespace Horde.Agent.Commands.Bundles
 		{
 			if (Http)
 			{
-				return new HttpStorageClientOwner(NamespaceId, new Uri(Server), logger);
+				return new HttpStorageClientOwner(new Uri($"{Server}/api/v1/storage/{NamespaceId}/"), logger);
 			}
 			else
 			{
diff --git a/Engine/Source/Programs/Horde/Horde.Agent/Execution/JobExecutor.cs b/Engine/Source/Programs/Horde/Horde.Agent/Execution/JobExecutor.cs
--- a/Engine/Source/Programs/Horde/Horde.Agent/Execution/JobExecutor.cs
+++ b/Engine/Source/Programs/Horde/Horde.Agent/Execution/JobExecutor.cs
@@ -667,12 +667,12 @@ namespace Horde.Agent.Execution
 			using MemoryCache cache = new MemoryCache(new MemoryCacheOptions { });
 
 			using HttpClient httpClient = _httpClientFactory.CreateClient(StorageHttpClientName);
-			httpClient.BaseAddress = _session.ServerUrl;
+			httpClient.BaseAddress = new Uri($"{_session.ServerUrl}/api/v1/storage/default");
 			httpClient.DefaultRequestHeaders.Authorization = new AuthenticationHeaderValue("Bearer", _session.Token);
 
 			using HttpClient redirectHttpClient = _httpClientFactory.CreateClient(StorageHttpClientName);
 
-			HttpStorageClient storage = new HttpStorageClient(new NamespaceId("default"), httpClient, redirectHttpClient, logger);
+			HttpStorageClient storage = new HttpStorageClient(CreateHttpClient, CreateHttpRedirectClient, logger);
 			//			FileStorageClient storage = new FileStorageClient(DirectoryReference.Combine(sharedStorageDir, "bundles"), cache, logger);
 			TreeReader reader = new TreeReader(storage, cache, logger);
 			logger.LogInformation("Using Horde-managed shared storage via {SharedStorageDir}", sharedStorageDir);
@@ -912,6 +912,19 @@ namespace Horde.Agent.Execution
 			return true;
 		}
 
+		HttpClient CreateHttpClient()
+		{
+			HttpClient httpClient = _httpClientFactory.CreateClient(StorageHttpClientName);
+			httpClient.BaseAddress = new Uri($"{_session.ServerUrl}/api/v1/storage/default");
+			httpClient.DefaultRequestHeaders.Authorization = new AuthenticationHeaderValue("Bearer", _session.Token);
+			return httpClient;
+		}
+
+		HttpClient CreateHttpRedirectClient()
+		{
+			return _httpClientFactory.CreateClient(StorageHttpClientName);
+		}
+
 		protected async Task<int> ExecuteAutomationToolAsync(BeginStepResponse step, DirectoryReference workspaceDir, string arguments, IStorageClient? store, ILogger logger, CancellationToken cancellationToken)
 		{
 			int result;
diff --git a/Engine/Source/Programs/Shared/EpicGames.Horde/Storage/Backends/HttpStorageClient.cs b/Engine/Source/Programs/Shared/EpicGames.Horde/Storage/Backends/HttpStorageClient.cs
--- a/Engine/Source/Programs/Shared/EpicGames.Horde/Storage/Backends/HttpStorageClient.cs
+++ b/Engine/Source/Programs/Shared/EpicGames.Horde/Storage/Backends/HttpStorageClient.cs
@@ -34,20 +34,18 @@ namespace EpicGames.Horde.Storage.Backends
 			public int ExportIdx { get; set; }
 		}
 
-		readonly NamespaceId _namespaceId;
-		readonly HttpClient _httpClient;
-		readonly HttpClient _redirectHttpClient;
+		readonly Func<HttpClient> _createClient;
+		readonly Func<HttpClient> _createRedirectClient;
 		readonly ILogger _logger;
 		bool _supportsUploadRedirects = true;
 
 		/// <summary>
 		/// Constructor
 		/// </summary>
-		public HttpStorageClient(NamespaceId namespaceId, HttpClient httpClient, HttpClient redirectHttpClient, ILogger logger) 
+		public HttpStorageClient(Func<HttpClient> createClient, Func<HttpClient> createRedirectClient, ILogger logger) 
 		{
-			_namespaceId = namespaceId;
-			_httpClient = httpClient;
-			_redirectHttpClient = redirectHttpClient;
+			_createClient = createClient;
+			_createRedirectClient = createRedirectClient;
 			_logger = logger;
 		}
 
@@ -57,11 +55,14 @@ namespace EpicGames.Horde.Storage.Backends
 		public override async Task<Stream> ReadBlobAsync(BlobLocator locator, CancellationToken cancellationToken = default)
 		{
 			_logger.LogDebug("Reading {Locator}", locator);
-			using (HttpRequestMessage request = new HttpRequestMessage(HttpMethod.Get, $"api/v1/storage/{_namespaceId}/blobs/{locator}"))
+			using (HttpClient httpClient = _createClient())
 			{
-				HttpResponseMessage response = await _httpClient.SendAsync(request, cancellationToken);
-				response.EnsureSuccessStatusCode();
-				return await response.Content.ReadAsStreamAsync(cancellationToken);
+				using (HttpRequestMessage request = new HttpRequestMessage(HttpMethod.Get, $"blobs/{locator}"))
+				{
+					HttpResponseMessage response = await httpClient.SendAsync(request, cancellationToken);
+					response.EnsureSuccessStatusCode();
+					return await response.Content.ReadAsStreamAsync(cancellationToken);
+				}
 			}
 		}
 
@@ -69,11 +70,14 @@ namespace EpicGames.Horde.Storage.Backends
 		public override async Task<Stream> ReadBlobRangeAsync(BlobLocator locator, int offset, int length, CancellationToken cancellationToken = default)
 		{
 			_logger.LogDebug("Reading {Locator} ({Offset}+{Length})", locator, offset, length);
-			using (HttpRequestMessage request = new HttpRequestMessage(HttpMethod.Get, $"api/v1/storage/{_namespaceId}/blobs/{locator}?offset={offset}&length={length}"))
+			using (HttpClient httpClient = _createClient())
 			{
-				HttpResponseMessage response = await _httpClient.SendAsync(request, cancellationToken);
-				response.EnsureSuccessStatusCode();
-				return await response.Content.ReadAsStreamAsync(cancellationToken);
+				using (HttpRequestMessage request = new HttpRequestMessage(HttpMethod.Get, $"blobs/{locator}?offset={offset}&length={length}"))
+				{
+					HttpResponseMessage response = await httpClient.SendAsync(request, cancellationToken);
+					response.EnsureSuccessStatusCode();
+					return await response.Content.ReadAsStreamAsync(cancellationToken);
+				}
 			}
 		}
 
@@ -84,19 +88,22 @@ namespace EpicGames.Horde.Storage.Backends
 
 			if (_supportsUploadRedirects)
 			{
-				WriteBlobResponse redirectResponse = await SendWriteRequestAsync(null, prefix, cancellationToken);
-				if (redirectResponse.UploadUrl != null)
+				using (HttpClient redirectHttpClient = _createRedirectClient())
 				{
-					using (HttpResponseMessage uploadResponse = await _redirectHttpClient.PutAsync(redirectResponse.UploadUrl, streamContent, cancellationToken))
+					WriteBlobResponse redirectResponse = await SendWriteRequestAsync(null, prefix, cancellationToken);
+					if (redirectResponse.UploadUrl != null)
 					{
-						if (!uploadResponse.IsSuccessStatusCode)
+						using (HttpResponseMessage uploadResponse = await redirectHttpClient.PutAsync(redirectResponse.UploadUrl, streamContent, cancellationToken))
 						{
-							string body = await uploadResponse.Content.ReadAsStringAsync(cancellationToken);
-							throw new StorageException($"Unable to upload data to redirected URL: {body}", null);
+							if (!uploadResponse.IsSuccessStatusCode)
+							{
+								string body = await uploadResponse.Content.ReadAsStringAsync(cancellationToken);
+								throw new StorageException($"Unable to upload data to redirected URL: {body}", null);
+							}
 						}
+						_logger.LogDebug("Written {Locator} (using redirect)", redirectResponse.Locator);
+						return redirectResponse.Locator;
 					}
-					_logger.LogDebug("Written {Locator} (using redirect)", redirectResponse.Locator);
-					return redirectResponse.Locator;
 				}
 			}
 
@@ -108,23 +115,26 @@ namespace EpicGames.Horde.Storage.Backends
 
 		async Task<WriteBlobResponse> SendWriteRequestAsync(StreamContent? streamContent, Utf8String prefix = default, CancellationToken cancellationToken = default)
 		{
-			using (HttpRequestMessage request = new HttpRequestMessage(HttpMethod.Post, $"api/v1/storage/{_namespaceId}/blobs"))
+			using (HttpClient httpClient = _createClient())
 			{
-				using StringContent stringContent = new StringContent(prefix.ToString());
-
-				MultipartFormDataContent form = new MultipartFormDataContent();
-				if (streamContent != null)
+				using (HttpRequestMessage request = new HttpRequestMessage(HttpMethod.Post, $"blobs"))
 				{
-					form.Add(streamContent, "file", "filename");
-				}
-				form.Add(stringContent, "prefix");
+					using StringContent stringContent = new StringContent(prefix.ToString());
 
-				request.Content = form;
-				using (HttpResponseMessage response = await _httpClient.SendAsync(request, cancellationToken))
-				{
-					response.EnsureSuccessStatusCode();
-					WriteBlobResponse? data = await response.Content.ReadFromJsonAsync<WriteBlobResponse>(cancellationToken: cancellationToken);
-					return data!;
+					MultipartFormDataContent form = new MultipartFormDataContent();
+					if (streamContent != null)
+					{
+						form.Add(streamContent, "file", "filename");
+					}
+					form.Add(stringContent, "prefix");
+
+					request.Content = form;
+					using (HttpResponseMessage response = await httpClient.SendAsync(request, cancellationToken))
+					{
+						response.EnsureSuccessStatusCode();
+						WriteBlobResponse? data = await response.Content.ReadFromJsonAsync<WriteBlobResponse>(cancellationToken: cancellationToken);
+						return data!;
+					}
 				}
 			}
 		}
@@ -137,11 +147,14 @@ namespace EpicGames.Horde.Storage.Backends
 		public override async Task DeleteRefAsync(RefName name, CancellationToken cancellationToken)
 		{
 			_logger.LogDebug("Deleting ref {RefName}", name);
-			using (HttpRequestMessage request = new HttpRequestMessage(HttpMethod.Delete, $"api/v1/storage/{_namespaceId}/refs/{name}"))
+			using (HttpClient httpClient = _createClient())
 			{
-				using (HttpResponseMessage response = await _httpClient.SendAsync(request, cancellationToken))
+				using (HttpRequestMessage request = new HttpRequestMessage(HttpMethod.Delete, $"refs/{name}"))
 				{
-					response.EnsureSuccessStatusCode();
+					using (HttpResponseMessage response = await httpClient.SendAsync(request, cancellationToken))
+					{
+						response.EnsureSuccessStatusCode();
+					}
 				}
 			}
 		}
@@ -149,26 +162,29 @@ namespace EpicGames.Horde.Storage.Backends
 		/// <inheritdoc/>
 		public override async Task<NodeLocator> TryReadRefTargetAsync(RefName name, DateTime cacheTime = default, CancellationToken cancellationToken = default)
 		{
-			using (HttpRequestMessage request = new HttpRequestMessage(HttpMethod.Get, $"api/v1/storage/{_namespaceId}/refs/{name}"))
+			using (HttpClient httpClient = _createClient())
 			{
-				using (HttpResponseMessage response = await _httpClient.SendAsync(request, cancellationToken))
+				using (HttpRequestMessage request = new HttpRequestMessage(HttpMethod.Get, $"refs/{name}"))
 				{
-					if (response.StatusCode == HttpStatusCode.NotFound)
-					{
-						_logger.LogDebug("Read ref {RefName} -> None", name);
-						return default;
-					}
-					else if (!response.IsSuccessStatusCode)
+					using (HttpResponseMessage response = await httpClient.SendAsync(request, cancellationToken))
 					{
-						_logger.LogError("Unable to read ref {RefName} (status: {StatusCode}, body: {Body})", name, response.StatusCode, await response.Content.ReadAsStringAsync(cancellationToken));
-						throw new StorageException($"Unable to read ref '{name}'", null);
-					}
-					else
-					{
-						response.EnsureSuccessStatusCode();
-						ReadRefResponse? data = await response.Content.ReadFromJsonAsync<ReadRefResponse>(cancellationToken: cancellationToken);
-						_logger.LogDebug("Read ref {RefName} -> {Blob}#{ExportIdx}", name, data!.Blob, data!.ExportIdx);
-						return new NodeLocator(data!.Blob, data!.ExportIdx);
+						if (response.StatusCode == HttpStatusCode.NotFound)
+						{
+							_logger.LogDebug("Read ref {RefName} -> None", name);
+							return default;
+						}
+						else if (!response.IsSuccessStatusCode)
+						{
+							_logger.LogError("Unable to read ref {RefName} (status: {StatusCode}, body: {Body})", name, response.StatusCode, await response.Content.ReadAsStringAsync(cancellationToken));
+							throw new StorageException($"Unable to read ref '{name}'", null);
+						}
+						else
+						{
+							response.EnsureSuccessStatusCode();
+							ReadRefResponse? data = await response.Content.ReadFromJsonAsync<ReadRefResponse>(cancellationToken: cancellationToken);
+							_logger.LogDebug("Read ref {RefName} -> {Blob}#{ExportIdx}", name, data!.Blob, data!.ExportIdx);
+							return new NodeLocator(data!.Blob, data!.ExportIdx);
+						}
 					}
 				}
 			}
@@ -178,9 +194,12 @@ namespace EpicGames.Horde.Storage.Backends
 		public override async Task WriteRefTargetAsync(RefName name, NodeLocator target, RefOptions? options = null, CancellationToken cancellationToken = default)
 		{
 			_logger.LogDebug("Writing ref {RefName} -> {Blob}#{ExportIdx}", name, target.Blob, target.ExportIdx);
-			using (HttpResponseMessage response = await _httpClient.PutAsync($"api/v1/storage/{_namespaceId}/refs/{name}", new { locator = target.Blob, exportIdx = target.ExportIdx, options }, cancellationToken))
+			using (HttpClient httpClient = _createClient())
 			{
-				response.EnsureSuccessStatusCode();
+				using (HttpResponseMessage response = await httpClient.PutAsync($"refs/{name}", new { locator = target.Blob, exportIdx = target.ExportIdx, options }, cancellationToken))
+				{
+					response.EnsureSuccessStatusCode();
+				}
 			}
 		}
 
diff --git a/Engine/Source/Programs/Shared/EpicGames.Horde/Storage/BlobLocator.cs b/Engine/Source/Programs/Shared/EpicGames.Horde/Storage/BlobLocator.cs
--- a/Engine/Source/Programs/Shared/EpicGames.Horde/Storage/BlobLocator.cs
+++ b/Engine/Source/Programs/Shared/EpicGames.Horde/Storage/BlobLocator.cs
@@ -73,8 +73,15 @@ namespace EpicGames.Horde.Storage
 		/// <param name="inner"></param>
 		public BlobLocator(Utf8String inner)
 		{
-			Inner = inner;
-			ValidateArgument(nameof(inner), inner);
+			if (inner.Length == 0)
+			{
+				Inner = default;
+			}
+			else
+			{
+				Inner = inner;
+				ValidateArgument(nameof(inner), inner);
+			}
 		}
 
 		/// <summary>
