Fixed issue where removed UAT plugins would not have their compile script state file deleted and thus cause future -NoCompile runs to load old plugins.

diff --git a/Engine/Source/Programs/Shared/EpicGames.Build/System/CompileScriptModules.cs b/Engine/Source/Programs/Shared/EpicGames.Build/System/CompileScriptModules.cs
--- a/Engine/Source/Programs/Shared/EpicGames.Build/System/CompileScriptModules.cs
+++ b/Engine/Source/Programs/Shared/EpicGames.Build/System/CompileScriptModules.cs
@@ -339,8 +339,9 @@ namespace UnrealBuildBase
 			}
 			else
 			{
-				// when the engine is not installed, delete any build record .json file that is not valid
-				foreach (CsProjBuildRecordEntry Entry in BuildRecords.Values.Where(x => x.Status == CsProjBuildRecordStatus.Invalid))
+				// When the engine is not installed, delete any build record .json file that is not valid.
+				// It is important that we do !Valid to delete both invalid (out of date) and unknown (where csproj no long exists).
+				foreach (CsProjBuildRecordEntry Entry in BuildRecords.Values.Where(x => x.Status != CsProjBuildRecordStatus.Valid))
 				{
 					if (Entry.BuildRecordFile != null)
 					{
