Horde: Add byte offsets and lengths to log chunks, for compatibility with legacy implementation.

diff --git a/Engine/Source/Programs/Shared/EpicGames.Horde/Logs/LogChunkNode.cs b/Engine/Source/Programs/Shared/EpicGames.Horde/Logs/LogChunkNode.cs
--- a/Engine/Source/Programs/Shared/EpicGames.Horde/Logs/LogChunkNode.cs
+++ b/Engine/Source/Programs/Shared/EpicGames.Horde/Logs/LogChunkNode.cs
@@ -205,16 +205,29 @@ namespace EpicGames.Horde.Logs
 		/// </summary>
 		public int LineCount { get; }
 
+		/// <summary>
+		/// Offset within the entire log file
+		/// </summary>
+		public long Offset { get; }
+
+		/// <summary>
+		/// Length of this chunk
+		/// </summary>
+		public int Length { get; }
+
 		/// <summary>
 		/// Constructor
 		/// </summary>
 		/// <param name="lineIndex">Index of the first line within this block</param>
+		/// <param name="offset">Offset within the log file</param>
 		/// <param name="target">Referenced log text</param>
-		public LogChunkRef(int lineIndex, LogChunkNode target)
+		public LogChunkRef(int lineIndex, long offset, LogChunkNode target)
 			: base(target)
 		{
 			LineIndex = lineIndex;
 			LineCount = target.LineCount;
+			Offset = offset;
+			Length = target.Length;
 		}
 
 		/// <summary>
@@ -226,6 +239,8 @@ namespace EpicGames.Horde.Logs
 		{
 			LineIndex = (int)reader.ReadUnsignedVarInt();
 			LineCount = (int)reader.ReadUnsignedVarInt();
+			Offset = (long)reader.ReadUnsignedVarInt();
+			Length = (int)reader.ReadUnsignedVarInt();
 		}
 
 		/// <inheritdoc/>
@@ -235,6 +250,8 @@ namespace EpicGames.Horde.Logs
 
 			writer.WriteUnsignedVarInt(LineIndex);
 			writer.WriteUnsignedVarInt(LineCount);
+			writer.WriteUnsignedVarInt((ulong)Offset);
+			writer.WriteUnsignedVarInt(Length);
 		}
 	}
 
diff --git a/Engine/Source/Programs/Shared/EpicGames.Horde/Logs/LogIndexNode.cs b/Engine/Source/Programs/Shared/EpicGames.Horde/Logs/LogIndexNode.cs
--- a/Engine/Source/Programs/Shared/EpicGames.Horde/Logs/LogIndexNode.cs
+++ b/Engine/Source/Programs/Shared/EpicGames.Horde/Logs/LogIndexNode.cs
@@ -65,6 +65,11 @@ namespace EpicGames.Horde.Logs
 		/// </summary>
 		public int LineCount => (_plainTextChunkRefs.Length > 0) ? (_plainTextChunkRefs[^1].LineIndex + _plainTextChunkRefs[^1].LineCount) : 0;
 
+		/// <summary>
+		/// Length of the text data
+		/// </summary>
+		public long Length => (_plainTextChunkRefs.Length > 0) ? (_plainTextChunkRefs[^1].Offset + _plainTextChunkRefs[^1].Length) : 0;
+
 		/// <summary>
 		/// Constructor
 		/// </summary>
@@ -124,11 +129,13 @@ namespace EpicGames.Horde.Logs
 			}
 
 			int lineIndex = LineCount;
+			long offset = Length;
 			for (int idx = 0; idx < appendPlainTextChunks.Count; idx++)
 			{
 				LogChunkNode appendBlock = appendPlainTextChunks[idx];
-				newChunks[_plainTextChunkRefs.Length + idx] = new LogChunkRef(lineIndex, appendPlainTextChunks[idx]);
+				newChunks[_plainTextChunkRefs.Length + idx] = new LogChunkRef(lineIndex, offset, appendPlainTextChunks[idx]);
 				lineIndex += appendBlock.LineCount;
+				offset += appendBlock.Length;
 			}
 
 			// Figure out how many bits to devote to the block size
diff --git a/Engine/Source/Programs/Shared/EpicGames.Horde/Logs/LogNode.cs b/Engine/Source/Programs/Shared/EpicGames.Horde/Logs/LogNode.cs
--- a/Engine/Source/Programs/Shared/EpicGames.Horde/Logs/LogNode.cs
+++ b/Engine/Source/Programs/Shared/EpicGames.Horde/Logs/LogNode.cs
@@ -41,7 +41,7 @@ namespace EpicGames.Horde.Logs
 		/// <summary>
 		/// Default value for an empty log file
 		/// </summary>
-		public static LogNode Empty { get; } = new LogNode(LogFormat.Json, 0, Array.Empty<LogChunkRef>(), new TreeNodeRef<LogIndexNode>(LogIndexNode.Empty));
+		public static LogNode Empty { get; } = new LogNode(LogFormat.Json, 0, 0, Array.Empty<LogChunkRef>(), new TreeNodeRef<LogIndexNode>(LogIndexNode.Empty));
 
 		/// <summary>
 		/// Format for this log file
@@ -53,6 +53,11 @@ namespace EpicGames.Horde.Logs
 		/// </summary>
 		public int LineCount { get; }
 
+		/// <summary>
+		/// Length of the log file
+		/// </summary>
+		public long Length { get; }
+
 		/// <summary>
 		/// Text blocks for this chunk
 		/// </summary>
@@ -66,10 +71,11 @@ namespace EpicGames.Horde.Logs
 		/// <summary>
 		/// Deserializing constructor
 		/// </summary>
-		public LogNode(LogFormat format, int lineCount, IReadOnlyList<LogChunkRef> textChunkRefs, TreeNodeRef<LogIndexNode> indexRef)
+		public LogNode(LogFormat format, int lineCount, long length, IReadOnlyList<LogChunkRef> textChunkRefs, TreeNodeRef<LogIndexNode> indexRef)
 		{
 			Format = format;
 			LineCount = lineCount;
+			Length = length;
 			TextChunkRefs = textChunkRefs.ToArray();
 			IndexRef = indexRef;
 		}
@@ -82,6 +88,7 @@ namespace EpicGames.Horde.Logs
 		{
 			Format = (LogFormat)reader.ReadUInt8();
 			LineCount = (int)reader.ReadUnsignedVarInt();
+			Length = (long)reader.ReadUnsignedVarInt();
 			IndexRef = reader.ReadRef<LogIndexNode>();
 			TextChunkRefs = reader.ReadList(() => new LogChunkRef(reader));
 		}
@@ -91,6 +98,7 @@ namespace EpicGames.Horde.Logs
 		{
 			writer.WriteUInt8((byte)Format);
 			writer.WriteUnsignedVarInt(LineCount);
+			writer.WriteUnsignedVarInt((ulong)Length);
 			writer.WriteRef(IndexRef);
 			writer.WriteList(TextChunkRefs, x => x.Serialize(writer));
 		}
@@ -272,13 +280,15 @@ namespace EpicGames.Horde.Logs
 
 			List<LogChunkRef> newJsonChunkRefs = new List<LogChunkRef>(_root.TextChunkRefs);
 			int lineCount = _root.LineCount;
+			long length = _root.Length;
 			foreach (LogChunkNode writeTextChunk in writeTextChunks)
 			{
-				newJsonChunkRefs.Add(new LogChunkRef(lineCount, writeTextChunk));
+				newJsonChunkRefs.Add(new LogChunkRef(lineCount, length, writeTextChunk));
 				lineCount += writeTextChunk.LineCount;
+				length += writeTextChunk.Length;
 			}
 
-			LogNode newRoot = new LogNode(_format, newIndex.LineCount, newJsonChunkRefs, new TreeNodeRef<LogIndexNode>(newIndex));
+			LogNode newRoot = new LogNode(_format, lineCount, length, newJsonChunkRefs, new TreeNodeRef<LogIndexNode>(newIndex));
 
 			IoHash newRootHash = await writer.WriteAsync(new TreeNodeRef(newRoot), cancellationToken);
 			await writer.FlushAsync(cancellationToken);
