Horde: Fix NamespaceId parameters no longer being validated correctly.

diff --git a/Engine/Source/Programs/Shared/EpicGames.Horde/Storage/NamespaceId.cs b/Engine/Source/Programs/Shared/EpicGames.Horde/Storage/NamespaceId.cs
--- a/Engine/Source/Programs/Shared/EpicGames.Horde/Storage/NamespaceId.cs
+++ b/Engine/Source/Programs/Shared/EpicGames.Horde/Storage/NamespaceId.cs
@@ -19,7 +19,7 @@ namespace EpicGames.Horde.Storage
 		/// <summary>
 		/// The text representing this id
 		/// </summary>
-		public Utf8String Text { get; }
+		public StringId Text { get; }
 
 		/// <summary>
 		/// Constructor
@@ -27,7 +27,7 @@ namespace EpicGames.Horde.Storage
 		/// <param name="text">Unique id for the namespace</param>
 		public NamespaceId(Utf8String text)
 		{
-			Text = text;
+			Text = new StringId(text);
 		}
 
 		/// <inheritdoc/>
diff --git a/Engine/Source/Programs/Shared/EpicGames.Horde/StringId.cs b/Engine/Source/Programs/Shared/EpicGames.Horde/StringId.cs
--- a/Engine/Source/Programs/Shared/EpicGames.Horde/StringId.cs
+++ b/Engine/Source/Programs/Shared/EpicGames.Horde/StringId.cs
@@ -15,6 +15,16 @@ namespace EpicGames.Horde
 		/// </summary>
 		public Utf8String Text { get; }
 
+		/// <summary>
+		/// Accessor for the string bytes
+		/// </summary>
+		public ReadOnlySpan<byte> Span => Text.Span;
+
+		/// <summary>
+		/// Accessor for the string bytes
+		/// </summary>
+		public ReadOnlyMemory<byte> Memory => Text.Memory;
+
 		/// <summary>
 		/// Constructor
 		/// </summary>
