EpicGames.Perforce.Managed: Add tests for syncing without have-table and cache files

Also add stricter verification of digests after sync

diff --git a/Engine/Source/Programs/Shared/EpicGames.Perforce.Fixture/PerforceFixture.cs b/Engine/Source/Programs/Shared/EpicGames.Perforce.Fixture/PerforceFixture.cs
--- a/Engine/Source/Programs/Shared/EpicGames.Perforce.Fixture/PerforceFixture.cs
+++ b/Engine/Source/Programs/Shared/EpicGames.Perforce.Fixture/PerforceFixture.cs
@@ -10,6 +10,8 @@ using Microsoft.VisualStudio.TestTools.UnitTesting;
 
 namespace EpicGames.Perforce.Fixture;
 
+using FileSet = HashSet<(string clientFile, long size, string digest)>;
+
 public class DepotFileFixture
 {
 	public string DepotFile { get; }
@@ -50,20 +52,56 @@ public class ChangelistFixture
 	/// Assert directory contains exactly the files described by stream
 	/// </summary>
 	/// <param name="clientRoot">Client/workspace root directory</param>
-	public void AssertAllFilesOnDisk(string clientRoot)
+	public void AssertDepotFiles(string clientRoot)
 	{
-		(List<string> actual, List<string> expected) = GetFiles(clientRoot);
-		CollectionAssert.AreEqual(expected, actual);
+		(FileSet actual, FileSet expected) = GetFileSets(clientRoot);
+
+		if (!expected.SetEquals(actual))
+		{
+			List<(string clientFile, long size, string digest)> expectedList = new(expected);
+			List<(string clientFile, long size, string digest)> actualList = new(actual);
+			
+			expectedList.Sort();
+			actualList.Sort();
+			
+			Console.WriteLine("Expected ------------------------------------------------------");
+			foreach ((string clientFile, long size, string digest) in expectedList)
+			{
+				Console.WriteLine($"{clientFile,-20} | {size,5} | {digest}");
+			}
+			
+			Console.WriteLine("");
+			Console.WriteLine("Actual --------------------------------------------------------");
+			foreach ((string clientFile, long size, string digest) in actualList)
+			{
+				Console.WriteLine($"{clientFile,-20} | {size,5} | {digest}");
+			}
+
+			Assert.Fail("Files on disk does not match files in stream at given CL");
+		}
 	}
 	
-	/// <summary>
-	/// Assert directory does contain the files described by stream
-	/// </summary>
-	/// <param name="clientRoot">Client/workspace root directory</param>
-	public void AssertNotAllFilesOnDisk(string clientRoot)
+	private (FileSet localFiles, FileSet streamFiles) GetFileSets(string clientRoot)
 	{
-		(List<string> actual, List<string> expected) = GetFiles(clientRoot);
-		CollectionAssert.AreNotEqual(expected, actual);
+		EnumerationOptions options = new () { RecurseSubdirectories = true };
+
+		FileSet streamFiles = new(StreamFiles.Select(x => (x.ClientFile, x.Size, x.Digest)));
+		FileSet localFiles = new(Directory.EnumerateFiles(clientRoot, "*", options)
+			.Select(x => Path.GetRelativePath(clientRoot, x))
+			.Select(x => x.Replace("\\", "/", StringComparison.Ordinal))
+			.Select(clientFile =>
+			{
+				string absPath = Path.Join(clientRoot, clientFile);
+				string content = File.ReadAllText(absPath);
+				
+				// Since content of files in fixture are only single lines, the trick below works to workaround
+				// differences in line endings after sync (client vs server).
+				content = content.Replace("\r\n", "", StringComparison.Ordinal);
+				long size = content.Length;
+				return (clientFile, size, PerforceFixture.CalcMd5(content));
+			}));
+		
+		return (localFiles, streamFiles);
 	}
 
 	public (List<string> localFiles, List<string> streamFiles) GetFiles(string clientRoot)
diff --git a/Engine/Source/Programs/Shared/EpicGames.Perforce.Managed.Tests/ManagedWorkspaceTest.cs b/Engine/Source/Programs/Shared/EpicGames.Perforce.Managed.Tests/ManagedWorkspaceTest.cs
--- a/Engine/Source/Programs/Shared/EpicGames.Perforce.Managed.Tests/ManagedWorkspaceTest.cs
+++ b/Engine/Source/Programs/Shared/EpicGames.Perforce.Managed.Tests/ManagedWorkspaceTest.cs
@@ -4,6 +4,7 @@ using System;
 using System.IO;
 using System.Threading;
 using System.Threading.Tasks;
+using EpicGames.Core;
 using EpicGames.Perforce.Fixture;
 using Microsoft.Extensions.Logging;
 using Microsoft.VisualStudio.TestTools.UnitTesting;
@@ -15,7 +16,14 @@ public class ManagedWorkspaceTest : BasePerforceFixtureTest
 {
 	private string SyncDir => Path.Join(TempDir.FullName, "Sync");
 	private string StreamName => Fixture.StreamFooMain.Root;
-	
+
+	private readonly ILogger<ManagedWorkspace> _mwLogger;
+
+	public ManagedWorkspaceTest()
+	{
+		_mwLogger = LoggerFactory.CreateLogger<ManagedWorkspace>();		
+	}
+
 	[TestMethod]
 	public async Task SyncFromScratch()
 	{
@@ -25,8 +33,7 @@ public class ManagedWorkspaceTest : BasePerforceFixtureTest
 		await workspace.SyncAsync(PerforceConnection, StreamName, 6, Array.Empty<string>(), true,
 			false, null, CancellationToken.None);
 
-		Fixture.StreamFooMain.Changelists[6].AssertAllFilesOnDisk(SyncDir);
-		Fixture.StreamFooMain.Changelists[7].AssertNotAllFilesOnDisk(SyncDir); // Contains an extra file so should fail
+		Fixture.StreamFooMain.Changelists[6].AssertDepotFiles(SyncDir);
 	}
 	
 	[TestMethod]
@@ -35,22 +42,64 @@ public class ManagedWorkspaceTest : BasePerforceFixtureTest
 		ManagedWorkspace ws = await GetManagedWorkspace();
 		await ws.SetupAsync(PerforceConnection, StreamName, CancellationToken.None);
 
-		await ws.SyncAsync(PerforceConnection, StreamName, 6, Array.Empty<string>(), true,
-			false, null, CancellationToken.None);
+		await ws.SyncAsync(PerforceConnection, StreamName, 6, Array.Empty<string>(), true, false, null, CancellationToken.None);
+		Fixture.StreamFooMain.Changelists[6].AssertDepotFiles(SyncDir);
+		
+		await ws.SyncAsync(PerforceConnection, StreamName, 7, Array.Empty<string>(), true, false, null, CancellationToken.None);
+		Fixture.StreamFooMain.Changelists[7].AssertDepotFiles(SyncDir);
+		
+		// Go back one change number
+		await ws.SyncAsync(PerforceConnection, StreamName, 6, Array.Empty<string>(), true, false, null, CancellationToken.None);
+		Fixture.StreamFooMain.Changelists[6].AssertDepotFiles(SyncDir);
+	}
+	
+	[TestMethod]
+	public async Task SyncIncrementallyWithoutHaveTable()
+	{
+		ManagedWorkspace ws = await GetManagedWorkspace();
+		await ws.SetupAsync(PerforceConnection, StreamName, CancellationToken.None);
 
-		Fixture.StreamFooMain.Changelists[6].AssertAllFilesOnDisk(SyncDir);
-		Fixture.StreamFooMain.Changelists[7].AssertNotAllFilesOnDisk(SyncDir); // Contains an extra file so should fail
+		await ws.SyncWithoutHaveTableAsync(PerforceConnection, StreamName, 6, Array.Empty<string>(), true, false, null, CancellationToken.None);
+		Fixture.StreamFooMain.Changelists[6].AssertDepotFiles(SyncDir);
 		
-		await ws.SyncAsync(PerforceConnection, StreamName, 7, Array.Empty<string>(), true,
-			false, null, CancellationToken.None);
+		await ws.SyncWithoutHaveTableAsync(PerforceConnection, StreamName, 7, Array.Empty<string>(), true, false, null, CancellationToken.None);
+		Fixture.StreamFooMain.Changelists[7].AssertDepotFiles(SyncDir);
+		
+		// Go back one change number
+		await ws.SyncWithoutHaveTableAsync(PerforceConnection, StreamName, 6, Array.Empty<string>(), true, false, null, CancellationToken.None);
+		Fixture.StreamFooMain.Changelists[6].AssertDepotFiles(SyncDir);
+	}
+	
+	[TestMethod]
+	public async Task SyncUsingCacheFiles()
+	{
+		ManagedWorkspace ws = await GetManagedWorkspace();
+		await ws.SetupAsync(PerforceConnection, StreamName, CancellationToken.None);
 
-		Fixture.StreamFooMain.Changelists[6].AssertNotAllFilesOnDisk(SyncDir);
-		Fixture.StreamFooMain.Changelists[7].AssertAllFilesOnDisk(SyncDir);
+		FileReference GetCacheFilePath(int changeNumber)
+		{
+			return new FileReference(Path.Join(TempDir.FullName, $"CacheFile-{changeNumber}.bin")); 
+		}
+		
+		// Sync and create a new cache file per change number
+		for (int i = 2; i < Fixture.StreamFooMain.Changelists.Count; i++)
+		{
+			Console.WriteLine("Syncing CL " + i);
+			await ws.SyncAsync(PerforceConnection, StreamName, i, Array.Empty<string>(), true, false, GetCacheFilePath(i), CancellationToken.None);
+			Fixture.StreamFooMain.Changelists[i].AssertDepotFiles(SyncDir);
+		}
+		
+		// Sync again but using the cache files created above
+		for (int i = 2; i < Fixture.StreamFooMain.Changelists.Count; i++)
+		{
+			Console.WriteLine("Syncing CL " + i);
+			await ws.SyncAsync(PerforceConnection, StreamName, i, Array.Empty<string>(), true, false, GetCacheFilePath(i), CancellationToken.None);
+			Fixture.StreamFooMain.Changelists[i].AssertDepotFiles(SyncDir);
+		}
 	}
 
 	private async Task<ManagedWorkspace> GetManagedWorkspace(bool overwrite = true)
 	{
-		ILogger<ManagedWorkspace> logger = LoggerFactory.CreateLogger<ManagedWorkspace>();
-		return await ManagedWorkspace.LoadOrCreateAsync(Environment.MachineName, TempDir, overwrite, logger, CancellationToken.None);
+		return await ManagedWorkspace.LoadOrCreateAsync(Environment.MachineName, TempDir, overwrite, _mwLogger, CancellationToken.None);
 	}
 }
\ No newline at end of file
