Fixed performance issue with shadow detection and symbol lookup.

diff --git a/Engine/Source/Programs/Shared/EpicGames.UHT/Types/UhtClass.cs b/Engine/Source/Programs/Shared/EpicGames.UHT/Types/UhtClass.cs
--- a/Engine/Source/Programs/Shared/EpicGames.UHT/Types/UhtClass.cs
+++ b/Engine/Source/Programs/Shared/EpicGames.UHT/Types/UhtClass.cs
@@ -781,7 +781,7 @@ namespace EpicGames.UHT.Types
 									}
 									else
 									{
-										VisibleType = false;
+										Session.HideTypeInSymbolTable(this);
 										result = false;
 									}
 								}
diff --git a/Engine/Source/Programs/Shared/EpicGames.UHT/Types/UhtProperty.cs b/Engine/Source/Programs/Shared/EpicGames.UHT/Types/UhtProperty.cs
--- a/Engine/Source/Programs/Shared/EpicGames.UHT/Types/UhtProperty.cs
+++ b/Engine/Source/Programs/Shared/EpicGames.UHT/Types/UhtProperty.cs
@@ -1554,7 +1554,7 @@ namespace EpicGames.UHT.Types
 			}
 
 			// Shadowing checks are done at this level, not in the properties themselves
-			if (options.HasAnyFlags(UhtValidationOptions.Shadowing))
+			if (options.HasAnyFlags(UhtValidationOptions.Shadowing) && PropertyCategory != UhtPropertyCategory.Return)
 			{
 
 				// First check for duplicate name in self and then duplicate name in parents
diff --git a/Engine/Source/Programs/Shared/EpicGames.UHT/Types/UhtType.cs b/Engine/Source/Programs/Shared/EpicGames.UHT/Types/UhtType.cs
--- a/Engine/Source/Programs/Shared/EpicGames.UHT/Types/UhtType.cs
+++ b/Engine/Source/Programs/Shared/EpicGames.UHT/Types/UhtType.cs
@@ -740,12 +740,6 @@ namespace EpicGames.UHT.Types
 		/// </summary>
 		public abstract string EngineClassName { get; }
 
-		/// <summary>
-		/// If true, then this type is visible to the symbol table.
-		/// </summary>
-		[JsonIgnore]
-		public bool VisibleType { get; set; } = true;
-
 		/// <summary>
 		/// Return true if the type has been deprecated
 		/// </summary>
diff --git a/Engine/Source/Programs/Shared/EpicGames.UHT/Utils/UhtSession.cs b/Engine/Source/Programs/Shared/EpicGames.UHT/Utils/UhtSession.cs
--- a/Engine/Source/Programs/Shared/EpicGames.UHT/Utils/UhtSession.cs
+++ b/Engine/Source/Programs/Shared/EpicGames.UHT/Utils/UhtSession.cs
@@ -1604,6 +1604,19 @@ namespace EpicGames.UHT.Utils
 			}
 		}
 
+		/// <summary>
+		/// Hide the given type in the symbol table
+		/// </summary>
+		/// <param name="typeToHide"></param>
+		public void HideTypeInSymbolTable(UhtType typeToHide)
+		{
+			_sourceNameSymbolTable.Hide(typeToHide, typeToHide.SourceName);
+			if (typeToHide.EngineType.HasEngineName())
+			{
+				_engineNameSymbolTable.Hide(typeToHide, typeToHide.EngineName);
+			}
+		}
+
 		/// <summary>
 		/// Return an ordered enumeration of all messages.
 		/// </summary>
@@ -2115,13 +2128,10 @@ namespace EpicGames.UHT.Utils
 					}
 				}
 
-				if (type.VisibleType)
+				_sourceNameSymbolTable.Add(type, type.SourceName);
+				if (engineExtendedType.HasEngineName())
 				{
-					_sourceNameSymbolTable.Add(type, type.SourceName);
-					if (engineExtendedType.HasEngineName())
-					{
-						_engineNameSymbolTable.Add(type, type.EngineName);
-					}
+					_engineNameSymbolTable.Add(type, type.EngineName);
 				}
 			}
 
diff --git a/Engine/Source/Programs/Shared/EpicGames.UHT/Utils/UhtSymbolTable.cs b/Engine/Source/Programs/Shared/EpicGames.UHT/Utils/UhtSymbolTable.cs
--- a/Engine/Source/Programs/Shared/EpicGames.UHT/Utils/UhtSymbolTable.cs
+++ b/Engine/Source/Programs/Shared/EpicGames.UHT/Utils/UhtSymbolTable.cs
@@ -147,6 +147,28 @@ namespace EpicGames.UHT.Utils
 			throw new UhtIceException("Attempt to replace a type that wasn't found");
 		}
 
+		/// <summary>
+		/// Hide the given type in the symbol table
+		/// </summary>
+		/// <param name="typeToHide">Type to be hidden</param>
+		/// <param name="name">The name of the type.</param>
+		/// <exception cref="UhtIceException">Thrown if the symbol wasn't found.</exception>
+		public void Hide(UhtType typeToHide, string name)
+		{
+			if (_caselessDictionary.TryGetValue(name, out int symbolIndex))
+			{
+				for (; symbolIndex != 0; symbolIndex = _symbols[symbolIndex].NextIndex)
+				{
+					if (_symbols[symbolIndex].Type == typeToHide)
+					{
+						_symbols[symbolIndex].MatchOptions = 0;
+						return;
+					}
+				}
+			}
+			throw new UhtIceException("Attempt to hide a type that wasn't found");
+		}
+
 		/// <summary>
 		/// Lookup the given name using cased string compare.
 		/// </summary>
@@ -378,7 +400,6 @@ namespace EpicGames.UHT.Utils
 		private bool IsMatch(UhtFindOptions options, int symbolIndex, int casedIndex)
 		{
 			return (casedIndex == 0 || casedIndex == _symbols[symbolIndex].CasedIndex) &&
-				_symbols[symbolIndex].Type.VisibleType &&
 				(_symbols[symbolIndex].MatchOptions & options) != 0;
 		}
 	}
