Fixed C# UHT to generate a circular dependency error when an include file includes itself.

diff --git a/Engine/Source/Programs/Shared/EpicGames.UHT/Parsers/UhtHeaderFileParser.cs b/Engine/Source/Programs/Shared/EpicGames.UHT/Parsers/UhtHeaderFileParser.cs
--- a/Engine/Source/Programs/Shared/EpicGames.UHT/Parsers/UhtHeaderFileParser.cs
+++ b/Engine/Source/Programs/Shared/EpicGames.UHT/Parsers/UhtHeaderFileParser.cs
@@ -693,12 +693,7 @@ namespace EpicGames.UHT.Parsers
 						{
 							SpottedAutogeneratedHeaderInclude = true;
 						}
-						if (includeNameString.Span.Contains(".generated.h", StringComparison.Ordinal))
-						{
-							string refName = Path.GetFileName(includeNameString.ToString()).Replace(".generated.h", ".h", StringComparison.Ordinal);
-							HeaderFile.AddReferencedHeader(refName, true);
-						}
-						else
+						if (!includeNameString.Span.Contains(".generated.h", StringComparison.Ordinal))
 						{
 							HeaderFile.AddReferencedHeader(includeNameString.ToString(), true);
 						}
diff --git a/Engine/Source/Programs/Shared/EpicGames.UHT/Types/UhtHeaderFile.cs b/Engine/Source/Programs/Shared/EpicGames.UHT/Types/UhtHeaderFile.cs
--- a/Engine/Source/Programs/Shared/EpicGames.UHT/Types/UhtHeaderFile.cs
+++ b/Engine/Source/Programs/Shared/EpicGames.UHT/Types/UhtHeaderFile.cs
@@ -274,6 +274,13 @@ namespace EpicGames.UHT.Types
 		/// <param name="isIncludedFile">True if this is a directly included file</param>
 		public void AddReferencedHeader(UhtHeaderFile headerFile, bool isIncludedFile)
 		{
+
+			// Ignore direct references to myself
+			if (!isIncludedFile && headerFile == this)
+			{
+				return;
+			}
+
 			lock (_referencedHeaders)
 			{
 				// Check for a duplicate
diff --git a/Engine/Source/Programs/Shared/EpicGames.UHT/Utils/UhtSession.cs b/Engine/Source/Programs/Shared/EpicGames.UHT/Utils/UhtSession.cs
--- a/Engine/Source/Programs/Shared/EpicGames.UHT/Utils/UhtSession.cs
+++ b/Engine/Source/Programs/Shared/EpicGames.UHT/Utils/UhtSession.cs
@@ -2162,10 +2162,7 @@ namespace EpicGames.UHT.Utils
 					states[visit.HeaderFileTypeIndex] = TopologicalState.Temporary;
 					foreach (UhtHeaderFile referenced in visit.ReferencedHeadersNoLock)
 					{
-						if (visit != referenced)
-						{
-							TopologicalVisit(states, referenced, headerStack);
-						}
+						TopologicalVisit(states, referenced, headerStack);
 					}
 					states[visit.HeaderFileTypeIndex] = TopologicalState.Permanent;
 					_sortedHeaderFiles.Add(visit);
