EpicGames.Perforce: Fix threads being leaked if an exception is thrown during initialization for NativePerforceConnection.

diff --git a/Engine/Source/Programs/Shared/EpicGames.Perforce/NativePerforceConnection.cs b/Engine/Source/Programs/Shared/EpicGames.Perforce/NativePerforceConnection.cs
--- a/Engine/Source/Programs/Shared/EpicGames.Perforce/NativePerforceConnection.cs
+++ b/Engine/Source/Programs/Shared/EpicGames.Perforce/NativePerforceConnection.cs
@@ -330,9 +330,18 @@ namespace EpicGames.Perforce
 		/// <returns></returns>
 		public static async Task<NativePerforceConnection> CreateAsync(IPerforceSettings settings, ILogger logger)
 		{
-			NativePerforceConnection connection = new NativePerforceConnection(settings, logger);
-			await connection.ConnectAsync();
-			return connection;
+			NativePerforceConnection? connection = null;
+			try
+			{
+				connection = new NativePerforceConnection(settings, logger);
+				await connection.ConnectAsync();
+				return connection;
+			}
+			catch
+			{
+				connection?.Dispose();
+				throw;
+			}
 		}
 
 		/// <summary>
