EpicGames.Perforce: Use separate client when syncing without have table

Re-using the same client and mixing different sync behavior led to issues.
Also update fixture tests to verify state of the have table.

diff --git a/Engine/Source/Programs/Horde/Horde.Agent/Commands/Workspace/SyncCommand.cs b/Engine/Source/Programs/Horde/Horde.Agent/Commands/Workspace/SyncCommand.cs
--- a/Engine/Source/Programs/Horde/Horde.Agent/Commands/Workspace/SyncCommand.cs
+++ b/Engine/Source/Programs/Horde/Horde.Agent/Commands/Workspace/SyncCommand.cs
@@ -52,7 +52,7 @@ namespace Horde.Agent.Commands.Workspace
 			List<string> expandedFilters = ExpandFilters(Filters);
 
 			using IPerforceConnection perforceClient = await perforce.WithClientAsync(ClientName);
-			await repo.SyncAsync(perforceClient, StreamName, changeNumber, expandedFilters, !IncrementalSync, FakeSync, CacheFile, CancellationToken.None);
+			await repo.SyncAsync(perforceClient, StreamName, changeNumber, expandedFilters, !IncrementalSync, FakeSync, true, CacheFile, CancellationToken.None);
 
 			if (PreflightChange != -1)
 			{
diff --git a/Engine/Source/Programs/Horde/Horde.Agent/Execution/PerforceExecutor.cs b/Engine/Source/Programs/Horde/Horde.Agent/Execution/PerforceExecutor.cs
--- a/Engine/Source/Programs/Horde/Horde.Agent/Execution/PerforceExecutor.cs
+++ b/Engine/Source/Programs/Horde/Horde.Agent/Execution/PerforceExecutor.cs
@@ -115,14 +115,7 @@ namespace Horde.Agent.Execution
 
 				// Sync the workspace
 				int syncPreflightChange = (_job.ClonedPreflightChange != 0) ? _job.ClonedPreflightChange : _job.PreflightChange;
-				if (_syncWithoutHaveTable)
-				{
-					await _workspace.SyncAsync(_job.Change, syncPreflightChange, null, false, cancellationToken);					
-				}
-				else
-				{
-					await _workspace.SyncAsync(_job.Change, syncPreflightChange, null, true, cancellationToken);
-				}
+				await _workspace.SyncAsync(_job.Change, syncPreflightChange, null, _syncWithoutHaveTable, cancellationToken);
 
 				// Remove any cached BuildGraph manifests
 				DirectoryReference manifestDir = DirectoryReference.Combine(_workspace.WorkspaceDir, "Engine", "Saved", "BuildGraph");
diff --git a/Engine/Source/Programs/Horde/Horde.Agent/Utility/WorkspaceInfo.cs b/Engine/Source/Programs/Horde/Horde.Agent/Utility/WorkspaceInfo.cs
--- a/Engine/Source/Programs/Horde/Horde.Agent/Utility/WorkspaceInfo.cs
+++ b/Engine/Source/Programs/Horde/Horde.Agent/Utility/WorkspaceInfo.cs
@@ -295,19 +295,12 @@ namespace Horde.Agent.Utility
 		/// <param name="change">The changelist to sync to</param>
 		/// <param name="preflightChange">Change to preflight, or 0</param>
 		/// <param name="cacheFile">Path to the cache file to use</param>
-		/// <param name="useHaveTable">Use client's have table when resolving contents</param>
+		/// <param name="useHaveTable">Whether to update have table during sync</param>
 		/// <param name="cancellationToken">Cancellation token</param>
 		/// <returns>Async task</returns>
 		public async Task SyncAsync(int change, int preflightChange, FileReference? cacheFile, bool useHaveTable, CancellationToken cancellationToken)
 		{
-			if (useHaveTable)
-			{
-				await Repository.SyncAsync(PerforceClient, StreamName, change, View, RemoveUntrackedFiles, false, cacheFile, cancellationToken);
-			}
-			else
-			{
-				await Repository.SyncWithoutHaveTableAsync(PerforceClient, StreamName, change, View, RemoveUntrackedFiles, false, cacheFile, cancellationToken);
-			}
+			await Repository.SyncAsync(PerforceClient, StreamName, change, View, RemoveUntrackedFiles, false, useHaveTable, cacheFile, cancellationToken);
 
 			// Purge the cache for incremental workspaces
 			if (!RemoveUntrackedFiles)
diff --git a/Engine/Source/Programs/Shared/EpicGames.Perforce.Fixture/PerforceFixture.cs b/Engine/Source/Programs/Shared/EpicGames.Perforce.Fixture/PerforceFixture.cs
--- a/Engine/Source/Programs/Shared/EpicGames.Perforce.Fixture/PerforceFixture.cs
+++ b/Engine/Source/Programs/Shared/EpicGames.Perforce.Fixture/PerforceFixture.cs
@@ -5,6 +5,8 @@ using System.Collections.Generic;
 using System.IO;
 using System.Linq;
 using System.Text;
+using System.Threading;
+using System.Threading.Tasks;
 using EpicGames.Core;
 using Microsoft.VisualStudio.TestTools.UnitTesting;
 
@@ -17,21 +19,25 @@ public class DepotFileFixture
 	public string DepotFile { get; }
 	public string ClientFile { get; }
 	public long Size { get; }
+	public int Revision { get; }
 	public string Digest { get; }
 	public string Content { get; }
 	
-	public DepotFileFixture(string depotFile, string clientFile, string content)
+	public DepotFileFixture(string depotFile, string clientFile, string content, int revision)
 	{
 		DepotFile = depotFile;
 		ClientFile = clientFile;
-		Content = content;
 		Size = content.Length;
+		Revision = revision;
 		Digest = PerforceFixture.CalcMd5(content);
+		Content = content;
 	}
 }
 
 public class ChangelistFixture
 {
+	public const string Placeholder = "<placeholder>";
+	
 	public int Number { get; }
 	public string Description { get; }
 	
@@ -81,6 +87,41 @@ public class ChangelistFixture
 		}
 	}
 	
+	/// <summary>
+	/// Assert files in stream for this changelist matches the client's have table
+	/// </summary>
+	/// <param name="perforce">Perforce connection</param>
+	public async Task AssertHaveTableAsync(IPerforceConnection perforce)
+	{
+		List<HaveRecord> haveRecords = await perforce.HaveAsync(new FileSpecList(), CancellationToken.None).ToListAsync();
+		HashSet<(string depotFile, int rev)> actual = new(haveRecords.Select((x) => (x.DepotFile, x.HaveRev)));
+		HashSet<(string depotFile, int rev)> expected = new(StreamFiles.Select((x) => (x.DepotFile, x.Revision)));
+		
+		if (!expected.SetEquals(actual))
+		{
+			List<(string depotFile, int rev)> expectedList = new(expected);
+			List<(string depotFile, int rev)> actualList = new(actual);
+			
+			expectedList.Sort();
+			actualList.Sort();
+			
+			Console.WriteLine("Expected ------------------------------------------------------");
+			foreach ((string depotFile, int rev) in expectedList)
+			{
+				Console.WriteLine($"{depotFile,-20} | {rev,5}");
+			}
+			
+			Console.WriteLine("");
+			Console.WriteLine("Actual --------------------------------------------------------");
+			foreach ((string depotFile, int rev) in actualList)
+			{
+				Console.WriteLine($"{depotFile,-20} | {rev,5}");
+			}
+
+			Assert.Fail("Files in stream does not match files in client's have table");
+		}
+	}
+	
 	private (FileSet localFiles, FileSet streamFiles) GetFileSets(string clientRoot)
 	{
 		EnumerationOptions options = new () { RecurseSubdirectories = true };
@@ -123,13 +164,27 @@ public class ChangelistFixture
 public class StreamFixture
 {
 	public string Root { get; }
-	public IReadOnlyList<ChangelistFixture> Changelists { get; }
+	public IEnumerable<ChangelistFixture> Changelists { get; }
 
 	public StreamFixture(string root, IReadOnlyList<ChangelistFixture> changelists)
 	{
 		Root = root;
 		Changelists = changelists;
 	}
+
+	public ChangelistFixture GetChangelist(int changeNum)
+	{
+		foreach (ChangelistFixture changelist in Changelists)
+		{
+			if (changelist.Number == changeNum)
+			{
+				return changelist;
+			}
+		}
+
+		// Ignore null check to make tests more readable
+		return null!;
+	}
 }
 
 public class PerforceFixture
@@ -137,58 +192,58 @@ public class PerforceFixture
 	public StreamFixture StreamFooMain { get; } = new ("//Foo/Main",
 		new List<ChangelistFixture>
 		{
-			new(0, "<placeholder>", new List<DepotFileFixture>()),
-			new(1, "<placeholder>", new List<DepotFileFixture>()),
+			new(0, ChangelistFixture.Placeholder, new List<DepotFileFixture>()),
+			new(1, ChangelistFixture.Placeholder, new List<DepotFileFixture>()),
 			new(2, "Initial import",
 				new List<DepotFileFixture>()
 				{
-					new("//Foo/Main/main.cpp", "main.cpp", "This is change main.cpp #1"),
-					new("//Foo/Main/main.h", "main.h", "This is change main.h #1"),
-					new("//Foo/Main/common.h", "common.h", "This is change common.h #1"),
-					new("//Foo/Main/unused.cpp", "unused.cpp", "This is change unused.cpp #1"),
-					new("//Foo/Main/Data/data.txt", "Data/data.txt", "This is change data.txt #1"),
+					new("//Foo/Main/main.cpp", "main.cpp", "This is change main.cpp #1", 1),
+					new("//Foo/Main/main.h", "main.h", "This is change main.h #1", 1),
+					new("//Foo/Main/common.h", "common.h", "This is change common.h #1", 1),
+					new("//Foo/Main/unused.cpp", "unused.cpp", "This is change unused.cpp #1", 1),
+					new("//Foo/Main/Data/data.txt", "Data/data.txt", "This is change data.txt #1", 1),
 				}),
 			new(3, "Improvement to main.cpp",
 				new List<DepotFileFixture>()
 				{
-					new("//Foo/Main/main.cpp", "main.cpp", "This is change main.cpp #2"),
-					new("//Foo/Main/main.h", "main.h", "This is change main.h #1"),
-					new("//Foo/Main/common.h", "common.h", "This is change common.h #1"),
-					new("//Foo/Main/unused.cpp", "unused.cpp", "This is change unused.cpp #1"),
-					new("//Foo/Main/Data/data.txt", "Data/data.txt", "This is change data.txt #1"),
+					new("//Foo/Main/main.cpp", "main.cpp", "This is change main.cpp #2", 2),
+					new("//Foo/Main/main.h", "main.h", "This is change main.h #1", 1),
+					new("//Foo/Main/common.h", "common.h", "This is change common.h #1", 1),
+					new("//Foo/Main/unused.cpp", "unused.cpp", "This is change unused.cpp #1", 1),
+					new("//Foo/Main/Data/data.txt", "Data/data.txt", "This is change data.txt #1", 1),
 				}),
 			new(4, "Delete an unused file",
 				new List<DepotFileFixture>()
 				{
-					new("//Foo/Main/main.cpp", "main.cpp", "This is change main.cpp #2"),
-					new("//Foo/Main/main.h", "main.h", "This is change main.h #1"),
-					new("//Foo/Main/common.h", "common.h", "This is change common.h #1"),
-					new("//Foo/Main/Data/data.txt", "Data/data.txt", "This is change data.txt #1"),
+					new("//Foo/Main/main.cpp", "main.cpp", "This is change main.cpp #2", 2),
+					new("//Foo/Main/main.h", "main.h", "This is change main.h #1", 1),
+					new("//Foo/Main/common.h", "common.h", "This is change common.h #1", 1),
+					new("//Foo/Main/Data/data.txt", "Data/data.txt", "This is change data.txt #1", 1),
 				}),
 			new(5, "Rename common.h to shared.h",
 				new List<DepotFileFixture>()
 				{
-					new("//Foo/Main/main.cpp", "main.cpp", "This is change main.cpp #2"),
-					new("//Foo/Main/main.h", "main.h", "This is change main.h #1"),
-					new("//Foo/Main/shared.h", "shared.h", "This is change common.h #1"),
-					new("//Foo/Main/Data/data.txt", "Data/data.txt", "This is change data.txt #1"),
+					new("//Foo/Main/main.cpp", "main.cpp", "This is change main.cpp #2", 2),
+					new("//Foo/Main/main.h", "main.h", "This is change main.h #1", 1),
+					new("//Foo/Main/shared.h", "shared.h", "This is change common.h #1", 1),
+					new("//Foo/Main/Data/data.txt", "Data/data.txt", "This is change data.txt #1", 1),
 				}),
 			new(6, "Some updates to main",
 				new List<DepotFileFixture>()
 				{
-					new("//Foo/Main/main.cpp", "main.cpp", "This is change main.cpp #3"),
-					new("//Foo/Main/main.h", "main.h", "This is change main.h #2"),
-					new("//Foo/Main/shared.h", "shared.h", "This is change common.h #1"),
-					new("//Foo/Main/Data/data.txt", "Data/data.txt", "This is change data.txt #1"),
+					new("//Foo/Main/main.cpp", "main.cpp", "This is change main.cpp #3", 3),
+					new("//Foo/Main/main.h", "main.h", "This is change main.h #2", 2),
+					new("//Foo/Main/shared.h", "shared.h", "This is change common.h #1", 1),
+					new("//Foo/Main/Data/data.txt", "Data/data.txt", "This is change data.txt #1", 1),
 				}),
 			new(7, "Add more data",
 				new List<DepotFileFixture>()
 				{
-					new("//Foo/Main/main.cpp", "main.cpp", "This is change main.cpp #3"),
-					new("//Foo/Main/main.h", "main.h", "This is change main.h #2"),
-					new("//Foo/Main/shared.h", "shared.h", "This is change common.h #1"),
-					new("//Foo/Main/Data/data.txt", "Data/data.txt", "This is change data.txt #1"),
-					new("//Foo/Main/Data/moredata.txt", "Data/moredata.txt", "This is change moredata.txt #1"),
+					new("//Foo/Main/main.cpp", "main.cpp", "This is change main.cpp #3", 3),
+					new("//Foo/Main/main.h", "main.h", "This is change main.h #2", 2),
+					new("//Foo/Main/shared.h", "shared.h", "This is change common.h #1", 1),
+					new("//Foo/Main/Data/data.txt", "Data/data.txt", "This is change data.txt #1", 1),
+					new("//Foo/Main/Data/moredata.txt", "Data/moredata.txt", "This is change moredata.txt #1", 1),
 				}),
 		});
 
diff --git a/Engine/Source/Programs/Shared/EpicGames.Perforce.Managed.Tests/ManagedWorkspaceTest.cs b/Engine/Source/Programs/Shared/EpicGames.Perforce.Managed.Tests/ManagedWorkspaceTest.cs
--- a/Engine/Source/Programs/Shared/EpicGames.Perforce.Managed.Tests/ManagedWorkspaceTest.cs
+++ b/Engine/Source/Programs/Shared/EpicGames.Perforce.Managed.Tests/ManagedWorkspaceTest.cs
@@ -1,7 +1,9 @@
 // Copyright Epic Games, Inc. All Rights Reserved.
 
 using System;
+using System.Collections.Generic;
 using System.IO;
+using System.Linq;
 using System.Threading;
 using System.Threading.Tasks;
 using EpicGames.Core;
@@ -16,6 +18,7 @@ public class ManagedWorkspaceTest : BasePerforceFixtureTest
 {
 	private string SyncDir => Path.Join(TempDir.FullName, "Sync");
 	private string StreamName => Fixture.StreamFooMain.Root;
+	private StreamFixture Stream => Fixture.StreamFooMain;
 
 	private readonly ILogger<ManagedWorkspace> _mwLogger;
 
@@ -25,56 +28,87 @@ public class ManagedWorkspaceTest : BasePerforceFixtureTest
 	}
 
 	[TestMethod]
-	public async Task SyncFromScratch()
+	public async Task SyncSingleChangelist()
 	{
-		ManagedWorkspace workspace = await GetManagedWorkspace();
-		await workspace.SetupAsync(PerforceConnection, StreamName, CancellationToken.None);
+		ManagedWorkspace ws = await GetManagedWorkspace();
+		await AssertHaveTableFileCount(0);
 
-		await workspace.SyncAsync(PerforceConnection, StreamName, 6, Array.Empty<string>(), true,
-			false, null, CancellationToken.None);
+		await SyncAsync(ws, 6);
+		Stream.GetChangelist(6).AssertDepotFiles(SyncDir);
+		await Stream.GetChangelist(6).AssertHaveTableAsync(PerforceConnection);
+	}
+	
+	[TestMethod]
+	public async Task SyncSingleChangelistWithoutHaveTable()
+	{
+		ManagedWorkspace ws = await GetManagedWorkspace();
+		await AssertHaveTableFileCount(0);
 
-		Fixture.StreamFooMain.Changelists[6].AssertDepotFiles(SyncDir);
+		await SyncAsync(ws, 6, useHaveTable: false);
+		Stream.GetChangelist(6).AssertDepotFiles(SyncDir);
+		await AssertHaveTableFileCount(0);
 	}
 	
 	[TestMethod]
-	public async Task SyncIncrementally()
+	public async Task SyncBackwardsToOlderChangelist()
 	{
 		ManagedWorkspace ws = await GetManagedWorkspace();
-		await ws.SetupAsync(PerforceConnection, StreamName, CancellationToken.None);
 
-		await ws.SyncAsync(PerforceConnection, StreamName, 6, Array.Empty<string>(), true, false, null, CancellationToken.None);
-		Fixture.StreamFooMain.Changelists[6].AssertDepotFiles(SyncDir);
+		await SyncAsync(ws, 6);
+		Stream.GetChangelist(6).AssertDepotFiles(SyncDir);
+		await Stream.GetChangelist(6).AssertHaveTableAsync(PerforceConnection);
 		
-		await ws.SyncAsync(PerforceConnection, StreamName, 7, Array.Empty<string>(), true, false, null, CancellationToken.None);
-		Fixture.StreamFooMain.Changelists[7].AssertDepotFiles(SyncDir);
+		await SyncAsync(ws, 7);
+		Stream.GetChangelist(7).AssertDepotFiles(SyncDir);
+		await Stream.GetChangelist(7).AssertHaveTableAsync(PerforceConnection);
 		
-		// Go back one change number
-		await ws.SyncAsync(PerforceConnection, StreamName, 6, Array.Empty<string>(), true, false, null, CancellationToken.None);
-		Fixture.StreamFooMain.Changelists[6].AssertDepotFiles(SyncDir);
+		// Go back one changelist
+		await SyncAsync(ws, 6);
+		Stream.GetChangelist(6).AssertDepotFiles(SyncDir);
+		await Stream.GetChangelist(6).AssertHaveTableAsync(PerforceConnection);
 	}
 	
 	[TestMethod]
-	public async Task SyncIncrementallyWithoutHaveTable()
+	public async Task SyncBackwardsToOlderChangelistWithoutHaveTable()
 	{
 		ManagedWorkspace ws = await GetManagedWorkspace();
-		await ws.SetupAsync(PerforceConnection, StreamName, CancellationToken.None);
+		
+		await SyncAsync(ws, 6, useHaveTable: false);
+		Stream.GetChangelist(6).AssertDepotFiles(SyncDir);
+		await AssertHaveTableFileCount(0);
 
-		await ws.SyncWithoutHaveTableAsync(PerforceConnection, StreamName, 6, Array.Empty<string>(), true, false, null, CancellationToken.None);
-		Fixture.StreamFooMain.Changelists[6].AssertDepotFiles(SyncDir);
+		await SyncAsync(ws, 7, useHaveTable: false);
+		Stream.GetChangelist(7).AssertDepotFiles(SyncDir);
+		await AssertHaveTableFileCount(0);
 		
-		await ws.SyncWithoutHaveTableAsync(PerforceConnection, StreamName, 7, Array.Empty<string>(), true, false, null, CancellationToken.None);
-		Fixture.StreamFooMain.Changelists[7].AssertDepotFiles(SyncDir);
+		// Go back one changelist
+		await SyncAsync(ws, 6, useHaveTable: false);
+		Stream.GetChangelist(6).AssertDepotFiles(SyncDir);
+		await AssertHaveTableFileCount(0);
+	}
+
+	[TestMethod]
+	public async Task SyncMixingUseOfHaveTable()
+	{
+		ManagedWorkspace ws = await GetManagedWorkspace();
+
+		await SyncAsync(ws, 1, useHaveTable: false);
+		Stream.GetChangelist(1).AssertDepotFiles(SyncDir);
+		await AssertHaveTableFileCount(0);
 		
-		// Go back one change number
-		await ws.SyncWithoutHaveTableAsync(PerforceConnection, StreamName, 6, Array.Empty<string>(), true, false, null, CancellationToken.None);
-		Fixture.StreamFooMain.Changelists[6].AssertDepotFiles(SyncDir);
+		await SyncAsync(ws, 2, useHaveTable: true);
+		Stream.GetChangelist(2).AssertDepotFiles(SyncDir);
+		await Stream.GetChangelist(2).AssertHaveTableAsync(PerforceConnection);
+		
+		await SyncAsync(ws, 3, useHaveTable: false);
+		Stream.GetChangelist(3).AssertDepotFiles(SyncDir);
+		await Stream.GetChangelist(2).AssertHaveTableAsync(PerforceConnection); // Have table should only match what was synced previously
 	}
-	
+
 	[TestMethod]
 	public async Task SyncUsingCacheFiles()
 	{
 		ManagedWorkspace ws = await GetManagedWorkspace();
-		await ws.SetupAsync(PerforceConnection, StreamName, CancellationToken.None);
 
 		FileReference GetCacheFilePath(int changeNumber)
 		{
@@ -82,24 +116,46 @@ public class ManagedWorkspaceTest : BasePerforceFixtureTest
 		}
 		
 		// Sync and create a new cache file per change number
-		for (int i = 2; i < Fixture.StreamFooMain.Changelists.Count; i++)
+		foreach (ChangelistFixture cl in Stream.Changelists)
 		{
-			Console.WriteLine("Syncing CL " + i);
-			await ws.SyncAsync(PerforceConnection, StreamName, i, Array.Empty<string>(), true, false, GetCacheFilePath(i), CancellationToken.None);
-			Fixture.StreamFooMain.Changelists[i].AssertDepotFiles(SyncDir);
+			await ws.SyncAsync(PerforceConnection, StreamName, cl.Number, Array.Empty<string>(), true, false, true, GetCacheFilePath(cl.Number), CancellationToken.None);
+			cl.AssertDepotFiles(SyncDir);
+			await cl.AssertHaveTableAsync(PerforceConnection);
 		}
 		
 		// Sync again but using the cache files created above
-		for (int i = 2; i < Fixture.StreamFooMain.Changelists.Count; i++)
+		foreach (ChangelistFixture cl in Stream.Changelists.Reverse())
 		{
-			Console.WriteLine("Syncing CL " + i);
-			await ws.SyncAsync(PerforceConnection, StreamName, i, Array.Empty<string>(), true, false, GetCacheFilePath(i), CancellationToken.None);
-			Fixture.StreamFooMain.Changelists[i].AssertDepotFiles(SyncDir);
+			await ws.SyncAsync(PerforceConnection, StreamName, cl.Number, Array.Empty<string>(), true, false, true, GetCacheFilePath(cl.Number), CancellationToken.None);
+			cl.AssertDepotFiles(SyncDir);
+			await cl.AssertHaveTableAsync(PerforceConnection);
 		}
 	}
 
-	private async Task<ManagedWorkspace> GetManagedWorkspace(bool overwrite = true)
+	private async Task<ManagedWorkspace> GetManagedWorkspace()
+	{
+		ManagedWorkspace ws = await ManagedWorkspace.CreateAsync(Environment.MachineName, TempDir, _mwLogger, CancellationToken.None);
+		await ws.SetupAsync(PerforceConnection, StreamName, CancellationToken.None);
+		return ws;
+	}
+	
+	private async Task SyncAsync(ManagedWorkspace managedWorkspace, int changeNumber, bool useHaveTable = true, FileReference? cacheFile = null)
+	{
+		await managedWorkspace.SyncAsync(PerforceConnection, StreamName, changeNumber, Array.Empty<string>(), true, false, useHaveTable, cacheFile, CancellationToken.None);
+	}
+	
+	private async Task AssertHaveTableFileCount(int expected)
 	{
-		return await ManagedWorkspace.LoadOrCreateAsync(Environment.MachineName, TempDir, overwrite, _mwLogger, CancellationToken.None);
+		List<HaveRecord> haveRecords = await PerforceConnection.HaveAsync(new FileSpecList(), CancellationToken.None).ToListAsync();
+
+		if (haveRecords.Count != expected)
+		{
+			Console.WriteLine("Have table contains:");
+			foreach (HaveRecord haveRecord in haveRecords)
+			{
+				Console.WriteLine(haveRecord.DepotFile + "#" + haveRecord.HaveRev);
+			}
+			Assert.Fail($"Actual have table file count does not match expected count. Actual={haveRecords.Count} Expected={expected}");
+		}
 	}
 }
\ No newline at end of file
diff --git a/Engine/Source/Programs/Shared/EpicGames.Perforce.Managed/ManagedWorkspace.cs b/Engine/Source/Programs/Shared/EpicGames.Perforce.Managed/ManagedWorkspace.cs
--- a/Engine/Source/Programs/Shared/EpicGames.Perforce.Managed/ManagedWorkspace.cs
+++ b/Engine/Source/Programs/Shared/EpicGames.Perforce.Managed/ManagedWorkspace.cs
@@ -134,6 +134,11 @@ namespace EpicGames.Perforce.Managed
 		/// Name of the main data file for a repository
 		/// </summary>
 		const string DataFileName = "Repository.dat";
+		
+		/// <summary>
+		/// Suffix added to client names that doesn't update the have table (syncing with -p option)
+		/// </summary>
+		const string NoHaveTableSuffix = "-NoHaveTable";
 
 		/// <summary>
 		/// Name of the host
@@ -778,9 +783,10 @@ namespace EpicGames.Perforce.Managed
 		/// <param name="view">View of the workspace</param>
 		/// <param name="removeUntracked">Whether to remove untracked files from the workspace</param>
 		/// <param name="fakeSync">Whether to simulate the syncing operation rather than actually getting files from the server</param>
+		/// <param name="useHaveTable">Whether to update have table on the server during sync</param>
 		/// <param name="cacheFile">If set, uses the given file to cache the contents of the workspace. This can improve sync times when multiple machines sync the same workspace.</param>
 		/// <param name="cancellationToken">Cancellation token</param>
-		public async Task SyncAsync(IPerforceConnection perforce, string streamName, int changeNumber, IReadOnlyList<string> view, bool removeUntracked, bool fakeSync, FileReference? cacheFile, CancellationToken cancellationToken)
+		public async Task SyncAsync(IPerforceConnection perforce, string streamName, int changeNumber, IReadOnlyList<string> view, bool removeUntracked, bool fakeSync, bool useHaveTable, FileReference? cacheFile, CancellationToken cancellationToken)
 		{
 			Stopwatch timer = Stopwatch.StartNew();
 			if (changeNumber == -1)
@@ -792,6 +798,7 @@ namespace EpicGames.Perforce.Managed
 				_logger.LogInformation("Syncing to {StreamName} at CL {CL}", streamName, changeNumber);
 			}
 
+			using (_logger.WithProperty("useHaveTable", useHaveTable).BeginScope())
 			using (_logger.BeginIndentScope("  "))
 			{
 				// Update the client to the current stream
@@ -803,94 +810,46 @@ namespace EpicGames.Perforce.Managed
 					changeNumber = await GetLatestClientChangeAsync(perforce, cancellationToken);
 				}
 
-				// Revert any open files
-				await RevertInternalAsync(perforce, cancellationToken);
-
-				// Force the P4 metadata to match up
-				Task updateHaveTableTask = Task.Run(() => UpdateClientHaveTableAsync(perforce, changeNumber, view, cancellationToken), cancellationToken);
-
-				// Clean the current workspace
-				await CleanInternalAsync(removeUntracked, cancellationToken);
-
-				// Wait for the have table update to finish
-				await updateHaveTableTask;
-
-				// Update the state of the current stream, if necessary
-				StreamSnapshot? contents;
-				if (cacheFile == null)
-				{
-					contents = await FindClientContentsAsync(perforce, changeNumber, cancellationToken);
-				}
-				else
+				if (useHaveTable)
 				{
-					contents = await TryLoadClientContentsAsync(cacheFile, streamName, cancellationToken);
-					if (contents == null)
-					{
-						contents = await FindAndSaveClientContentsAsync(perforce, streamName, changeNumber, cacheFile, true, cancellationToken);
-					}
-				}
+					// Revert any open files
+					await RevertInternalAsync(perforce, cancellationToken);
 
-				// Sync all the appropriate files
-				await RemoveFilesFromWorkspaceAsync(contents, cancellationToken);
-				await AddFilesToWorkspaceAsync(perforce, contents, fakeSync, cancellationToken);
-			}
-
-			_logger.LogInformation("Completed in {ElapsedTime}s", $"{timer.Elapsed.TotalSeconds:0.0}");
-		}
-		
-		/// <summary>
-		/// Switches to the given stream
-		/// </summary>
-		/// <param name="perforce">The perforce connection</param>
-		/// <param name="streamName">Name of the stream to sync</param>
-		/// <param name="changeNumber">Changelist number to sync. -1 to sync to latest.</param>
-		/// <param name="view">View of the workspace</param>
-		/// <param name="removeUntracked">Whether to remove untracked files from the workspace</param>
-		/// <param name="fakeSync">Whether to simulate the syncing operation rather than actually getting files from the server</param>
-		/// <param name="cacheFile">If set, uses the given file to cache the contents of the workspace. This can improve sync times when multiple machines sync the same workspace.</param>
-		/// <param name="cancellationToken">Cancellation token</param>
-		public async Task SyncWithoutHaveTableAsync(IPerforceConnection perforce, string streamName, int changeNumber, IReadOnlyList<string> view, bool removeUntracked, bool fakeSync, FileReference? cacheFile, CancellationToken cancellationToken)
-		{
-			Stopwatch timer = Stopwatch.StartNew();
-			if (changeNumber == -1)
-			{
-				_logger.LogInformation("Syncing without have-table to {StreamName} at latest", streamName);
-			}
-			else
-			{
-				_logger.LogInformation("Syncing without have-table to {StreamName} at CL {CL}", streamName, changeNumber);
-			}
+					// Force the P4 metadata to match up
+					Task updateHaveTableTask = Task.Run(() => UpdateClientHaveTableAsync(perforce, changeNumber, view, cancellationToken), cancellationToken);
 
-			using (_logger.WithProperty("useHaveTable", true).BeginScope())
-			using (_logger.BeginIndentScope("  "))
-			{
-				// Update the client to the current stream
-				await UpdateClientAsync(perforce, streamName, cancellationToken);
+					// Clean the current workspace
+					await CleanInternalAsync(removeUntracked, cancellationToken);
 
-				// Get the latest change number
-				if (changeNumber == -1)
-				{
-					changeNumber = await GetLatestClientChangeAsync(perforce, cancellationToken);
+					// Wait for the have table update to finish
+					await updateHaveTableTask;
 				}
 
 				// Update the state of the current stream, if necessary
 				StreamSnapshot? contents;
 				if (cacheFile == null)
 				{
-					contents = await FindClientContentsWithoutHaveTableAsync(perforce, streamName, changeNumber, cancellationToken);
+					if (useHaveTable)
+					{
+						contents = await FindClientContentsAsync(perforce, changeNumber, cancellationToken);						
+					}
+					else
+					{
+						contents = await FindClientContentsWithoutHaveTableAsync(perforce, streamName, changeNumber, cancellationToken);
+					}
 				}
 				else
 				{
 					contents = await TryLoadClientContentsAsync(cacheFile, streamName, cancellationToken);
 					if (contents == null)
 					{
-						contents = await FindAndSaveClientContentsAsync(perforce, streamName, changeNumber, cacheFile, false, cancellationToken);
+						contents = await FindAndSaveClientContentsAsync(perforce, streamName, changeNumber, cacheFile, useHaveTable, cancellationToken);
 					}
 				}
 
 				// Sync all the appropriate files
 				await RemoveFilesFromWorkspaceAsync(contents, cancellationToken);
-				await AddFilesToWorkspaceAsync(perforce, contents, fakeSync, cancellationToken);
+				await AddFilesToWorkspaceAsync(perforce, contents, fakeSync, useHaveTable, cancellationToken);
 			}
 
 			_logger.LogInformation("Completed in {ElapsedTime}s", $"{timer.Elapsed.TotalSeconds:0.0}");
@@ -1123,7 +1082,7 @@ namespace EpicGames.Perforce.Managed
 
 					StreamSnapshot contents = streamState[idx].Item2;
 					await RemoveFilesFromWorkspaceAsync(contents, cancellationToken);
-					await AddFilesToWorkspaceAsync(request.PerforceClient, contents, fakeSync, cancellationToken);
+					await AddFilesToWorkspaceAsync(request.PerforceClient, contents, fakeSync, true, cancellationToken);
 				}
 			}
 
@@ -1143,54 +1102,82 @@ namespace EpicGames.Perforce.Managed
 		/// <returns>Async task</returns>
 		public async Task DeleteClientAsync(IPerforceConnection perforceClient, CancellationToken cancellationToken)
 		{
-			PerforceResponse response = await perforceClient.TryDeleteClientAsync(DeleteClientOptions.None, perforceClient.Settings.ClientName!, cancellationToken);
-			if (response.Error != null && response.Error.Generic != PerforceGenericCode.Unknown)
+			async Task DeleteAsync(string clientName)
 			{
-				if (response.Error.Generic == PerforceGenericCode.NotYet)
+				PerforceResponse response = await perforceClient.TryDeleteClientAsync(DeleteClientOptions.None, clientName, cancellationToken);
+				if (response.Error != null && response.Error.Generic != PerforceGenericCode.Unknown)
 				{
-					await RevertInternalAsync(perforceClient, cancellationToken);
-					response = await perforceClient.TryDeleteClientAsync(DeleteClientOptions.None, perforceClient.Settings.ClientName!, cancellationToken);
+					if (response.Error.Generic == PerforceGenericCode.NotYet)
+					{
+						await RevertInternalAsync(perforceClient, cancellationToken);
+						response = await perforceClient.TryDeleteClientAsync(DeleteClientOptions.None, clientName, cancellationToken);
+					}
+					response.EnsureSuccess();
 				}
-				response.EnsureSuccess();
+				_createdClients.Remove(clientName);
 			}
-			_createdClients.Remove(perforceClient.Settings.ClientName!);
+
+			await DeleteAsync(GetClientName(perforceClient, true));
+			await DeleteAsync(GetClientName(perforceClient, false));
 		}
 
-		/// <summary>
-		/// Sets the stream for the current client
-		/// </summary>
-		/// <param name="perforceClient">The Perforce connection</param>
-		/// <param name="streamName">New stream for the client</param>
-		/// <param name="cancellationToken">Cancellation token</param>
-		private async Task UpdateClientAsync(IPerforceConnection perforceClient, string streamName, CancellationToken cancellationToken)
+		private static string GetClientName(IPerforceConnection perforceClient, bool useHaveTable)
 		{
-			// Create or update the client if it doesn't exist already
-			ClientRecord? client;
-			if (!_createdClients.TryGetValue(perforceClient.Settings.ClientName!, out client) || client.Stream != streamName)
+			if (!useHaveTable)
 			{
-				using (Trace("UpdateClient"))
-				using (ILoggerProgress status = _logger.BeginProgressScope("Updating client..."))
-				{
-					Stopwatch timer = Stopwatch.StartNew();
+				return perforceClient.Settings.ClientName + NoHaveTableSuffix;
+			}
+			return perforceClient.Settings.ClientName!;
+		}
+		
+		private async Task<ClientRecord> GetOrCreateClientAsync(IPerforceConnection perforceClient, string streamName, bool useHaveTable, CancellationToken cancellationToken)
+		{
+			string clientName = GetClientName(perforceClient, useHaveTable);
+			if (_createdClients.TryGetValue(clientName, out ClientRecord? client) && client.Stream == streamName)
+			{
+				return client;
+			}
 
-					client = new ClientRecord(perforceClient.Settings.ClientName!, perforceClient.Settings.UserName!, _workspaceDir.FullName);
-					client.Host = _hostName;
-					client.Stream = streamName;
-					client.Type = "partitioned";
+			using (Trace("UpdateClient"))
+			using (ILoggerProgress status = _logger.BeginProgressScope("Updating client..."))
+			{
+				Stopwatch timer = Stopwatch.StartNew();
 
-					using IPerforceConnection perforce = await perforceClient.WithoutClientAsync();
+				client = new ClientRecord(clientName, perforceClient.Settings.UserName!, _workspaceDir.FullName);
+				client.Host = _hostName;
+				client.Stream = streamName;
+				
+				// Partitioned and read-only types store their have table separately on the server, compared to normal (writeable) clients
+				// Clients that sync without updating the have table cannot submit so they're marked as read-only. 
+				client.Type = useHaveTable ? "partitioned" : "readonly";
 
-					PerforceResponse response = await perforce.TryCreateClientAsync(client, cancellationToken);
-					if (!response.Succeeded)
-					{
-						await perforceClient.TryDeleteClientAsync(DeleteClientOptions.None, perforceClient.Settings.ClientName!, cancellationToken);
-						await perforceClient.CreateClientAsync(client, cancellationToken);
-					}
+				using IPerforceConnection perforce = await perforceClient.WithoutClientAsync();
 
-					status.Progress = $"({timer.Elapsed.TotalSeconds:0.0}s)";
+				PerforceResponse response = await perforce.TryCreateClientAsync(client, cancellationToken);
+				if (!response.Succeeded)
+				{
+					await perforceClient.TryDeleteClientAsync(DeleteClientOptions.None, clientName, cancellationToken);
+					await perforceClient.CreateClientAsync(client, cancellationToken);
 				}
-				_createdClients[perforceClient.Settings.ClientName!] = client;
+
+				status.Progress = $"({timer.Elapsed.TotalSeconds:0.0}s)";
 			}
+			_createdClients[clientName] = client;
+
+			return client;
+		}
+		
+		/// <summary>
+		/// Sets the stream for the current client
+		/// </summary>
+		/// <param name="perforceClient">The Perforce connection</param>
+		/// <param name="streamName">New stream for the client</param>
+		/// <param name="cancellationToken">Cancellation token</param>
+		private async Task UpdateClientAsync(IPerforceConnection perforceClient, string streamName, CancellationToken cancellationToken)
+		{
+			// Ensure a client exists both with and without have table use, using the correct stream
+			await GetOrCreateClientAsync(perforceClient, streamName, true, cancellationToken);
+			await GetOrCreateClientAsync(perforceClient, streamName, false, cancellationToken);
 
 			// Update the config file with the name of the client
 			FileReference configFile = FileReference.Combine(_baseDir, "p4.ini");
@@ -1639,7 +1626,7 @@ namespace EpicGames.Perforce.Managed
 		/// <param name="basePath">Base path for the stream</param>
 		/// <param name="changeNumber">The change number being synced. This must be specified in order to get the digest at the correct revision.</param>
 		/// <param name="cacheFile">Location of the file to save the cached contents</param>
-		/// <param name="useHaveTable">Use client's have table when resolving contents</param>
+		/// <param name="useHaveTable">Whether to update the have table during sync</param>
 		/// <param name="cancellationToken">Cancellation token</param>
 		/// <returns>Contents of the workspace</returns>
 		private async Task<StreamSnapshotFromMemory> FindAndSaveClientContentsAsync(IPerforceConnection perforceClient, Utf8String basePath, int changeNumber, FileReference cacheFile, bool useHaveTable, CancellationToken cancellationToken)
@@ -1806,8 +1793,9 @@ namespace EpicGames.Perforce.Managed
 		/// <param name="client">The client connection</param>
 		/// <param name="stream">Contents of the stream</param>
 		/// <param name="fakeSync">Whether to simulate the sync operation, rather than actually syncing files</param>
+		/// <param name="useHaveTable">Whether to update the have table during sync</param>
 		/// <param name="cancellationToken">Cancellation token</param>
-		private async Task AddFilesToWorkspaceAsync(IPerforceConnection client, StreamSnapshot stream, bool fakeSync, CancellationToken cancellationToken)
+		private async Task AddFilesToWorkspaceAsync(IPerforceConnection client, StreamSnapshot stream, bool fakeSync, bool useHaveTable, CancellationToken cancellationToken)
 		{
 			// Make sure the repair flag is reset
 			await RunOptionalRepairAsync(cancellationToken);
@@ -1915,7 +1903,7 @@ namespace EpicGames.Perforce.Managed
 							{
 								(int batchBeginIdx, int batchEndIdx) = batches[nextBatchIdx];
 
-								Task task = Task.Run(() => SyncBatch(client, filesToSync, batchBeginIdx, batchEndIdx, fakeSync, cancellationToken), cancellationToken);
+								Task task = Task.Run(() => SyncBatch(client, filesToSync, batchBeginIdx, batchEndIdx, fakeSync, useHaveTable, cancellationToken), cancellationToken);
 								tasks[task] = nextBatchIdx++;
 							}
 
@@ -1967,9 +1955,10 @@ namespace EpicGames.Perforce.Managed
 		/// <param name="beginIdx">First file to sync</param>
 		/// <param name="endIdx">Index of the last file to sync (exclusive)</param>
 		/// <param name="fakeSync">Whether to fake a sync</param>
+		/// <param name="useHaveTable">Whether to update the have table during sync</param>
 		/// <param name="cancellationToken">Cancellation token for the request</param>
 		/// <returns>Async task</returns>
-		async Task SyncBatch(IPerforceConnection client, WorkspaceFileToSync[] filesToSync, int beginIdx, int endIdx, bool fakeSync, CancellationToken cancellationToken)
+		async Task SyncBatch(IPerforceConnection client, WorkspaceFileToSync[] filesToSync, int beginIdx, int endIdx, bool fakeSync, bool useHaveTable, CancellationToken cancellationToken)
 		{
 			if (fakeSync)
 			{
@@ -1991,10 +1980,20 @@ namespace EpicGames.Perforce.Managed
 					}
 				}
 
-				using PerforceConnection clientWithFileList = new PerforceConnection(client.Settings, client.Logger);
-				clientWithFileList.GlobalOptions.Add($"-x\"{syncFileName}\"");
-
-				await clientWithFileList.SyncAsync(SyncOptions.Force | SyncOptions.FullDepotSyntax, -1, Array.Empty<string>(), cancellationToken).ToListAsync(cancellationToken);
+				if (useHaveTable)
+				{
+					using PerforceConnection clientWithFileList = new (client.Settings, client.Logger);
+					clientWithFileList.GlobalOptions.Add($"-x\"{syncFileName}\"");
+					await clientWithFileList.SyncAsync(SyncOptions.Force | SyncOptions.FullDepotSyntax, -1, Array.Empty<string>(), cancellationToken).ToListAsync(cancellationToken);
+				}
+				else
+				{
+					// Ensure a client with an empty have table is used to not interfere with the DoNotUpdateHaveList option.
+					using PerforceConnection clientWithFileList = new (client.Settings, client.Logger);
+					clientWithFileList.ClientName = GetClientName(client, false);
+					clientWithFileList.GlobalOptions.Add($"-x\"{syncFileName}\"");
+					await clientWithFileList.SyncAsync(SyncOptions.DoNotUpdateHaveList | SyncOptions.FullDepotSyntax, -1, Array.Empty<string>(), cancellationToken).ToListAsync(cancellationToken);
+				}
 			}
 		}
 
