EpicGames.Core: Add an extension method to run a conversion function on each element in an array.

diff --git a/Engine/Source/Programs/Horde/Horde.Build/Jobs/TestData/TestDataService.cs b/Engine/Source/Programs/Horde/Horde.Build/Jobs/TestData/TestDataService.cs
--- a/Engine/Source/Programs/Horde/Horde.Build/Jobs/TestData/TestDataService.cs
+++ b/Engine/Source/Programs/Horde/Horde.Build/Jobs/TestData/TestDataService.cs
@@ -98,8 +98,8 @@ namespace Horde.Build.Jobs.TestData
 		/// <returns></returns>
 		public async Task<List<ITestDataRef>> FindTestRefs(StreamId[] streamIds, TestMetaId[] metaIds, string[]? testIds = null, string[]? suiteIds = null, DateTime? minCreateTime = null, DateTime? maxCreateTime = null, int? minChange = null, int? maxChange = null)
 		{
-			TestId[]? tids = testIds?.ConvertAll(x => new TestId(x)).ToArray();
-			TestSuiteId[]? sids = suiteIds?.ConvertAll(x => new TestSuiteId(x)).ToArray();
+			TestId[]? tids = testIds?.ConvertAll(x => new TestId(x));
+			TestSuiteId[]? sids = suiteIds?.ConvertAll(x => new TestSuiteId(x));
 
 			return await _testData.FindTestRefs(streamIds, metaIds, tids, sids, minCreateTime, maxCreateTime, minChange, maxChange);
 		}
diff --git a/Engine/Source/Programs/Shared/EpicGames.Core/ArrayExtensions.cs b/Engine/Source/Programs/Shared/EpicGames.Core/ArrayExtensions.cs
new file mode 100644
--- /dev/null
+++ b/Engine/Source/Programs/Shared/EpicGames.Core/ArrayExtensions.cs
@@ -0,0 +1,30 @@
+// Copyright Epic Games, Inc. All Rights Reserved.
+
+using System;
+
+namespace EpicGames.Core
+{
+	/// <summary>
+	/// Extension methods for arrays.
+	/// </summary>
+	public static class ArrayExtensions
+	{
+		/// <summary>
+		/// Converts an array of elements to a different type
+		/// </summary>
+		/// <typeparam name="TIn">Type of the input elements</typeparam>
+		/// <typeparam name="TOut">Type of the output elements</typeparam>
+		/// <param name="input">The array to convert</param>
+		/// <param name="func">Conversion function for each element</param>
+		/// <returns>Array of elements of the output type</returns>
+		public static TOut[] ConvertAll<TIn, TOut>(this TIn[] input, Func<TIn, TOut> func)
+		{
+			TOut[] output = new TOut[input.Length];
+			for (int idx = 0; idx < input.Length; idx++)
+			{
+				output[idx] = func(input[idx]);
+			}
+			return output;
+		}
+	}
+}
