EpicGames.Core: Add a utility class for building UTF8 strings.

diff --git a/Engine/Source/Programs/Shared/EpicGames.Core/MemoryWriter.cs b/Engine/Source/Programs/Shared/EpicGames.Core/MemoryWriter.cs
--- a/Engine/Source/Programs/Shared/EpicGames.Core/MemoryWriter.cs
+++ b/Engine/Source/Programs/Shared/EpicGames.Core/MemoryWriter.cs
@@ -91,10 +91,15 @@ namespace EpicGames.Core
 		byte[] _data;
 		int _length;
 
+		/// <summary>
+		/// Returns the span that has been written to
+		/// </summary>
+		public Span<byte> WrittenSpan => _data.AsSpan(0, _length);
+
 		/// <summary>
 		/// Returns the memory that has been written to
 		/// </summary>
-		public ReadOnlyMemory<byte> WrittenMemory => _data.AsMemory(0, _length);
+		public Memory<byte> WrittenMemory => _data.AsMemory(0, _length);
 
 		/// <summary>
 		/// Constructor
diff --git a/Engine/Source/Programs/Shared/EpicGames.Core/Utf8StringBuilder.cs b/Engine/Source/Programs/Shared/EpicGames.Core/Utf8StringBuilder.cs
new file mode 100644
--- /dev/null
+++ b/Engine/Source/Programs/Shared/EpicGames.Core/Utf8StringBuilder.cs
@@ -0,0 +1,90 @@
+// Copyright Epic Games, Inc. All Rights Reserved.
+
+using System;
+using System.Collections.Generic;
+using System.Text;
+
+namespace EpicGames.Core
+{
+	/// <summary>
+	/// Represents a memory region which can be treated as a utf-8 string.
+	/// </summary>
+	public struct Utf8StringBuilder
+	{
+		readonly ArrayMemoryWriter _writer;
+
+		/// <summary>
+		/// Returns the length of this string
+		/// </summary>
+		public int Length => _writer.Length;
+
+		/// <summary>
+		/// Accessor for the written span
+		/// </summary>
+		public Span<byte> WrittenSpan => _writer.WrittenSpan;
+
+		/// <summary>
+		/// Accessor for the written memory
+		/// </summary>
+		public Memory<byte> WrittenMemory => _writer.WrittenMemory;
+
+		/// <summary>
+		/// Accessor for the written data
+		/// </summary>
+		public Utf8String WrittenString => new Utf8String(_writer.WrittenMemory);
+
+		/// <summary>
+		/// Constructor
+		/// </summary>
+		public Utf8StringBuilder()
+			: this(256)
+		{
+		}
+
+		/// <summary>
+		/// Constructor
+		/// </summary>
+		public Utf8StringBuilder(int initialSize)
+			: this(new ArrayMemoryWriter(initialSize))
+		{
+		}
+
+		/// <summary>
+		/// Constructor
+		/// </summary>
+		public Utf8StringBuilder(ArrayMemoryWriter writer)
+		{
+			_writer = writer;
+		}
+
+		/// <summary>
+		/// Append a character to the end of this builder
+		/// </summary>
+		/// <param name="ch">Character to append</param>
+		public void Append(byte ch)
+		{
+			_writer.WriteUInt8(ch);
+		}
+
+		/// <summary>
+		/// Appends a string to the end of this builder
+		/// </summary>
+		/// <param name="text">Text to append</param>
+		public void Append(Utf8String text)
+		{
+			_writer.WriteFixedLengthBytes(text.Span);
+		}
+
+		/// <summary>
+		/// Converts the written memory to a utf8 string
+		/// </summary>
+		/// <returns></returns>
+		public Utf8String ToUtf8String()
+		{
+			return new Utf8String(_writer.WrittenMemory);
+		}
+
+		/// <inheritdoc/>
+		public override string ToString() => ToUtf8String().ToString();
+	}
+}
