Horde: Refactor EpicGames.Redis to have an interface more in line with StackExchange.Redis. Typed keys are now analagous to RedisKey, rather than capturing an IDatabaseAsync interface as well, making the use of transactions more natural.

diff --git a/Engine/Source/Programs/Horde/Horde.Build.Tests/LogTailServiceTests.cs b/Engine/Source/Programs/Horde/Horde.Build.Tests/LogTailServiceTests.cs
--- a/Engine/Source/Programs/Horde/Horde.Build.Tests/LogTailServiceTests.cs
+++ b/Engine/Source/Programs/Horde/Horde.Build.Tests/LogTailServiceTests.cs
@@ -6,6 +6,7 @@ using System.Text;
 using System.Threading;
 using System.Threading.Tasks;
 using EpicGames.Core;
+using EpicGames.Redis;
 using Horde.Build.Logs;
 using Horde.Build.Server;
 using Horde.Build.Utilities;
@@ -80,15 +81,15 @@ namespace Horde.Build.Tests
 			await tailService.EnableTailingAsync(_logId, 30);
 			await tailService.AppendAsync(_logId, 30, Encoding.UTF8.GetBytes("line30\nline31\nline32\nline33\nline34\nline35\nline36\nline37\nline38\nline39\nline40\n"));
 
-			Assert.AreEqual(4, await tailService.RedisTailData(_logId).LengthAsync());
+			Assert.AreEqual(4, await redisService.GetDatabase().SortedSetLengthAsync(LogTailService.TailDataKey(_logId)));
 
 			await tailService.FlushAsync(_logId, 33);
 
-			Assert.AreEqual(4, await tailService.RedisTailData(_logId).LengthAsync());
+			Assert.AreEqual(4, await redisService.GetDatabase().SortedSetLengthAsync(LogTailService.TailDataKey(_logId)));
 
 			await Clock.AdvanceAsync(TimeSpan.FromMinutes(1.0));
 
-			Assert.AreEqual(3, await tailService.RedisTailData(_logId).LengthAsync());
+			Assert.AreEqual(3, await redisService.GetDatabase().SortedSetLengthAsync(LogTailService.TailDataKey(_logId)));
 
 			{
 				List<Utf8String> lines = await tailService.ReadAsync(_logId, 25, 20);
@@ -123,13 +124,13 @@ namespace Horde.Build.Tests
 			await tailService.EnableTailingAsync(_logId, 30);
 			await tailService.AppendAsync(_logId, 30, Encoding.UTF8.GetBytes("foo\n"));
 
-			Assert.IsTrue(await redisService.GetDatabase().KeyExistsAsync(tailService.RedisTailNext(_logId).Key));
-			Assert.IsTrue(await redisService.GetDatabase().KeyExistsAsync(tailService.RedisTailData(_logId).Key));
+			Assert.IsTrue(await redisService.GetDatabase().KeyExistsAsync(LogTailService.TailNextKey(_logId)));
+			Assert.IsTrue(await redisService.GetDatabase().KeyExistsAsync(LogTailService.TailDataKey(_logId)));
 
 			await Clock.AdvanceAsync(TimeSpan.FromMinutes(5.0));
 
-			Assert.IsFalse(await redisService.GetDatabase().KeyExistsAsync(tailService.RedisTailNext(_logId).Key));
-			Assert.IsFalse(await redisService.GetDatabase().KeyExistsAsync(tailService.RedisTailData(_logId).Key));
+			Assert.IsFalse(await redisService.GetDatabase().KeyExistsAsync(LogTailService.TailNextKey(_logId)));
+			Assert.IsFalse(await redisService.GetDatabase().KeyExistsAsync(LogTailService.TailDataKey(_logId)));
 		}
 
 		[TestMethod]
diff --git a/Engine/Source/Programs/Horde/Horde.Build/Agents/AgentCollection.cs b/Engine/Source/Programs/Horde/Horde.Build/Agents/AgentCollection.cs
--- a/Engine/Source/Programs/Horde/Horde.Build/Agents/AgentCollection.cs
+++ b/Engine/Source/Programs/Horde/Horde.Build/Agents/AgentCollection.cs
@@ -560,7 +560,7 @@ namespace Horde.Build.Agents
 		/// <inheritdoc/>
 		public async Task<IDisposable> SubscribeToUpdateEventsAsync(Action<AgentId> onUpdate)
 		{
-			return await _redisService.GetDatabase().Multiplexer.GetSubscriber().SubscribeAsync(_updateEventChannel, (channel, agentId) => onUpdate(agentId));
+			return await _redisService.GetDatabase().Multiplexer.SubscribeAsync(_updateEventChannel, (channel, agentId) => onUpdate(agentId));
 		}
 	}
 }
diff --git a/Engine/Source/Programs/Horde/Horde.Build/Agents/AgentService.cs b/Engine/Source/Programs/Horde/Horde.Build/Agents/AgentService.cs
--- a/Engine/Source/Programs/Horde/Horde.Build/Agents/AgentService.cs
+++ b/Engine/Source/Programs/Horde/Horde.Build/Agents/AgentService.cs
@@ -83,7 +83,7 @@ namespace Horde.Build.Agents
 		readonly ILogger _logger;
 		readonly ITicker _ticker;
 		
-		readonly RedisString<AgentRateTable> _agentRateTableData;
+		readonly RedisStringKey<AgentRateTable> _agentRateTableData = new RedisStringKey<AgentRateTable>("agent-rates");
 		readonly RedisService _redisService;
 
 		// Lazily updated costs for different agent types
@@ -108,8 +108,7 @@ namespace Horde.Build.Agents
 			_sessions = sessions;
 			_aclService = aclService;
 			_downtimeService = downtimeService;
-			_agentRateTableData = new RedisString<AgentRateTable>(redisService.ConnectionPool, "agent-rates");
-			_agentRateTable = new AsyncCachedValue<AgentRateTable?>(() => _agentRateTableData.GetAsync(), TimeSpan.FromSeconds(2.0));//.FromMinutes(5.0));
+			_agentRateTable = new AsyncCachedValue<AgentRateTable?>(() => redisService.GetDatabase().StringGetAsync(_agentRateTableData), TimeSpan.FromSeconds(2.0));//.FromMinutes(5.0));
 			_poolsList = new AsyncCachedValue<List<IPool>>(() => poolCollection.GetAsync(), TimeSpan.FromSeconds(30.0));
 			_dogStatsd = dogStatsd;
 			_taskSources = taskSources.ToArray();
@@ -899,7 +898,7 @@ namespace Horde.Build.Agents
 		/// <returns></returns>
 		public async Task UpdateRateTableAsync(List<AgentRateConfig> entries)
 		{
-			await _agentRateTableData.SetAsync(new AgentRateTable { Entries = entries });
+			await _redisService.GetDatabase().StringSetAsync(_agentRateTableData, new AgentRateTable { Entries = entries });
 		}
 
 		/// <summary>
@@ -909,7 +908,7 @@ namespace Horde.Build.Agents
 		/// <returns>Hourly rate of running the given agent</returns>
 		public async ValueTask<double?> GetRateAsync(AgentId agentId)
 		{
-			RedisKey key = $"agent-rate/{agentId}";
+			RedisStringKey<double> key = $"agent-rate/{agentId}";
 
 			// Try to get the current value
 			RedisValue value = await _redisService.GetDatabase().StringGetAsync(key);
diff --git a/Engine/Source/Programs/Horde/Horde.Build/Compute/RedisMessageQueue.cs b/Engine/Source/Programs/Horde/Horde.Build/Compute/RedisMessageQueue.cs
--- a/Engine/Source/Programs/Horde/Horde.Build/Compute/RedisMessageQueue.cs
+++ b/Engine/Source/Programs/Horde/Horde.Build/Compute/RedisMessageQueue.cs
@@ -95,31 +95,31 @@ namespace Horde.Build.Compute
 		/// </summary>
 		/// <param name="channelId"></param>
 		/// <returns></returns>
-		RedisList<T> GetChannel(string channelId)
+		RedisListKey<T> GetChannel(string channelId)
 		{
-			return new RedisList<T>(_redis, _keyPrefix.Append(channelId));
+			return new RedisListKey<T>(_keyPrefix.Append(channelId));
 		}
 
 		/// <inheritdoc/>
 		public async Task PostAsync(string channelId, T message)
 		{
-			RedisList<T> channel = GetChannel(channelId);
+			RedisListKey<T> channel = GetChannel(channelId);
 
-			long length = await channel.RightPushAsync(message);
+			long length = await _redis.ListRightPushAsync(channel, message);
 			if (length == 1)
 			{
 				await _redis.PublishAsync(_updateChannel, channelId, CommandFlags.FireAndForget);
 			}
-			await _redis.KeyExpireAsync(channel.Key, ExpireTime, CommandFlags.FireAndForget);
+			await _redis.KeyExpireAsync(channel, ExpireTime, flags: CommandFlags.FireAndForget);
 		}
 
-		static async Task<bool> ReadMessagesAsync(RedisList<T> list, List<T> messages)
+		async Task<bool> ReadMessagesAsync(RedisListKey<T> list, List<T> messages)
 		{
-			T? message = await list.LeftPopAsync();
+			T? message = await _redis.ListLeftPopAsync(list);
 			while(message != null)
 			{
 				messages.Add(message);
-				message = await list.LeftPopAsync();
+				message = await _redis.ListLeftPopAsync(list);
 			}
 			return messages.Count > 0;
 		}
@@ -137,7 +137,7 @@ namespace Horde.Build.Compute
 		{
 			List<T> messages = new List<T>();
 
-			RedisList<T> channel = GetChannel(channelId);
+			RedisListKey<T> channel = GetChannel(channelId);
 			while (!cancellationToken.IsCancellationRequested && !await ReadMessagesAsync(channel, messages))
 			{
 				// Register for notifications on this channel
diff --git a/Engine/Source/Programs/Horde/Horde.Build/Compute/RedisTaskScheduler.cs b/Engine/Source/Programs/Horde/Horde.Build/Compute/RedisTaskScheduler.cs
--- a/Engine/Source/Programs/Horde/Horde.Build/Compute/RedisTaskScheduler.cs
+++ b/Engine/Source/Programs/Horde/Horde.Build/Compute/RedisTaskScheduler.cs
@@ -87,8 +87,8 @@ namespace Horde.Build.Compute
 
 		readonly RedisConnectionPool _redisConnectionPool;
 		readonly RedisKey _baseKey;
-		readonly RedisSet<TQueueId> _queueIndex;
-		readonly RedisHash<TQueueId, DateTime> _activeQueues; // Queues which are actively being dequeued from
+		readonly RedisSetKey<TQueueId> _queueIndex;
+		readonly RedisHashKey<TQueueId, DateTime> _activeQueues; // Queues which are actively being dequeued from
 		ReadOnlyHashSet<TQueueId> _localActiveQueues = new HashSet<TQueueId>();
 		readonly Stopwatch _resetActiveQueuesTimer = Stopwatch.StartNew();
 		readonly List<Listener> _listeners = new List<Listener>();
@@ -107,8 +107,8 @@ namespace Horde.Build.Compute
 		{
 			_redisConnectionPool = redisConnectionPool;
 			_baseKey = baseKey;
-			_queueIndex = new RedisSet<TQueueId>(_redisConnectionPool, baseKey.Append("index"));
-			_activeQueues = new RedisHash<TQueueId, DateTime>(_redisConnectionPool, baseKey.Append("active"));
+			_queueIndex = new RedisSetKey<TQueueId>(baseKey.Append("index"));
+			_activeQueues = new RedisHashKey<TQueueId, DateTime>(baseKey.Append("active"));
 			_newQueueChannel = new RedisChannel<TQueueId>(baseKey.Append("new_queues").ToString());
 			_logger = logger;
 
@@ -143,33 +143,23 @@ namespace Horde.Build.Compute
 		/// </summary>
 		/// <param name="queueId">The queue identifier</param>
 		/// <returns></returns>
-		RedisKey GetQueueKey(TQueueId queueId)
+		RedisListKey<TTask> GetQueueKey(TQueueId queueId)
 		{
-			return _baseKey.Append(RedisSerializer.Serialize(queueId).AsKey());
-		}
-
-		/// <summary>
-		/// Gets the key for a list of tasks for a particular queue
-		/// </summary>
-		/// <param name="queueId">The queue identifier</param>
-		/// <returns></returns>
-		RedisList<TTask> GetQueue(TQueueId queueId)
-		{
-			return new RedisList<TTask>(_redisConnectionPool, GetQueueKey(queueId));
+			return new RedisListKey<TTask>(_baseKey.Append(RedisSerializer.Serialize(queueId).AsKey()));
 		}
 
 		/// <summary>
 		/// Pushes a task onto either end of a queue
 		/// </summary>
-		static Task<long> PushTaskAsync(RedisList<TTask> list, TTask task, When when, CommandFlags flags, bool atFront)
+		static Task<long> PushTaskAsync(IDatabaseAsync database, RedisListKey<TTask> list, TTask task, When when, CommandFlags flags, bool atFront)
 		{
 			if (atFront)
 			{
-				return list.LeftPushAsync(task, when, flags);
+				return database.ListLeftPushAsync(list, task, when, flags);
 			}
 			else
 			{
-				return list.RightPushAsync(task, when, flags);
+				return database.ListRightPushAsync(list, task, when, flags);
 			}
 		}
 
@@ -181,25 +171,26 @@ namespace Horde.Build.Compute
 		/// <param name="atFront">Whether to add to the front of the queue</param>
 		public async Task EnqueueAsync(TQueueId queueId, TTask task, bool atFront)
 		{
-			RedisList<TTask> list = GetQueue(queueId);
+			RedisListKey<TTask> list = GetQueueKey(queueId);
 			for (; ; )
 			{
-				long newLength = await PushTaskAsync(list, task, When.Exists, CommandFlags.None, atFront);
+				IDatabase connection = _redisConnectionPool.GetDatabase();
+
+				long newLength = await PushTaskAsync(connection, list, task, When.Exists, CommandFlags.None, atFront);
 				if (newLength > 0)
 				{
 					_logger.LogInformation("Length of queue {QueueId} is {Length}", queueId, newLength);
 					break;
 				}
 
-				IDatabase redis = _redisConnectionPool.GetDatabase();
-				ITransaction transaction = redis.CreateTransaction();
-				_ = transaction.With(_queueIndex).AddAsync(queueId, CommandFlags.FireAndForget);
-				_ = PushTaskAsync(transaction.With(list), task, When.Always, CommandFlags.FireAndForget, atFront);
+				ITransaction transaction = connection.CreateTransaction();
+				_ = transaction.SetAddAsync(_queueIndex, queueId, CommandFlags.FireAndForget);
+				_ = PushTaskAsync(transaction, list, task, When.Always, CommandFlags.FireAndForget, atFront);
 
 				if (await transaction.ExecuteAsync())
 				{
 					_logger.LogInformation("Created queue {QueueId}", queueId);
-					await redis.PublishAsync(_newQueueChannel, queueId);
+					await connection.PublishAsync(_newQueueChannel, queueId);
 					break;
 				}
 
@@ -216,7 +207,7 @@ namespace Horde.Build.Compute
 		public async Task<Task<(TQueueId, TTask)?>> DequeueAsync(Func<TQueueId, ValueTask<bool>> predicate, CancellationToken token = default)
 		{
 			// Compare against all the list of cached queues to see if we can dequeue something from any of them
-			TQueueId[] queues = await _queueIndex.MembersAsync();
+			TQueueId[] queues = await _redisConnectionPool.GetDatabase().SetMembersAsync(_queueIndex);
 
 			// Try to dequeue an item from the list
 			(TQueueId, TTask)? entry = await TryAssignToLocalAgentAsync(queues, predicate);
@@ -297,12 +288,14 @@ namespace Horde.Build.Compute
 		{
 			await AddActiveQueue(queueId);
 
-			TTask? item = await GetQueue(queueId).LeftPopAsync();
+			IDatabase database = _redisConnectionPool.GetDatabase();
+
+			TTask? item = await database.ListLeftPopAsync(GetQueueKey(queueId));
 			if (item == null)
 			{
-				ITransaction transaction = _redisConnectionPool.GetDatabase().CreateTransaction();
-				transaction.AddCondition(Condition.KeyNotExists(GetQueueKey(queueId)));
-				Task<bool> wasRemoved = transaction.With(_queueIndex).RemoveAsync(queueId);
+				ITransaction transaction = database.CreateTransaction();
+				transaction.AddCondition(Condition.KeyNotExists(GetQueueKey(queueId).Inner));
+				Task<bool> wasRemoved = transaction.SetRemoveAsync(_queueIndex, queueId);
 
 				if (await transaction.ExecuteAsync() && await wasRemoved)
 				{
@@ -351,7 +344,7 @@ namespace Horde.Build.Compute
 				if (Interlocked.CompareExchange(ref _localActiveQueues, newLocalActiveQueues, localActiveQueuesCopy) == localActiveQueuesCopy)
 				{
 					_logger.LogInformation("Refreshing active queue {QueueId}", queueId);
-					await _activeQueues.SetAsync(queueId, DateTime.UtcNow);
+					await _redisConnectionPool.GetDatabase().HashSetAsync(_activeQueues, queueId, DateTime.UtcNow);
 					break;
 				}
 			}
@@ -363,12 +356,12 @@ namespace Horde.Build.Compute
 		/// <returns></returns>
 		public async Task<List<TQueueId>> GetInactiveQueuesAsync()
 		{
-			HashSet<TQueueId> keys = new HashSet<TQueueId>(await _queueIndex.MembersAsync());
+			HashSet<TQueueId> keys = new HashSet<TQueueId>(await _redisConnectionPool.GetDatabase().SetMembersAsync(_queueIndex));
 			HashSet<TQueueId> invalidKeys = new HashSet<TQueueId>();
 
 			DateTime minTime = DateTime.UtcNow - TimeSpan.FromMinutes(10.0);
 
-			HashEntry<TQueueId, DateTime>[] entries = await _activeQueues.GetAllAsync();
+			HashEntry<TQueueId, DateTime>[] entries = await _redisConnectionPool.GetDatabase().HashGetAllAsync(_activeQueues);
 			foreach (HashEntry<TQueueId, DateTime> entry in entries)
 			{
 				if (entry.Value < minTime)
@@ -383,7 +376,7 @@ namespace Horde.Build.Compute
 
 			if (invalidKeys.Count > 0)
 			{
-				await _activeQueues.DeleteAsync(invalidKeys.ToArray());
+				await _redisConnectionPool.GetDatabase().HashDeleteAsync(_activeQueues, invalidKeys.ToArray());
 			}
 
 			return keys.ToList();
@@ -391,13 +384,13 @@ namespace Horde.Build.Compute
 
 		public async Task<int> GetNumQueuedTasksAsync(Func<TQueueId, ValueTask<bool>> predicate, CancellationToken token = default)
 		{
-			HashSet<TQueueId> queueIds = new (await _queueIndex.MembersAsync());
+			HashSet<TQueueId> queueIds = new (await _redisConnectionPool.GetDatabase().SetMembersAsync(_queueIndex));
 			long totalTaskCount = 0;
 			foreach (TQueueId queueId in queueIds)
 			{
 				if (await predicate(queueId))
 				{
-					totalTaskCount += await GetQueue(queueId).LengthAsync();
+					totalTaskCount += await _redisConnectionPool.GetDatabase().ListLengthAsync(GetQueueKey(queueId));
 				}
 			}
 
@@ -408,8 +401,7 @@ namespace Horde.Build.Compute
 		{
 			Channel<TQueueId> newQueues = Channel.CreateUnbounded<TQueueId>();
 
-			ISubscriber subscriber = _redisConnectionPool.GetDatabase().Multiplexer.GetSubscriber();
-			await using RedisChannelSubscription<TQueueId>? _ = await subscriber.SubscribeAsync(_newQueueChannel, (_, v) => newQueues.Writer.TryWrite(v));
+			await using RedisSubscription _ = await _redisConnectionPool.GetDatabase().Multiplexer.SubscribeAsync(_newQueueChannel, v => newQueues.Writer.TryWrite(v));
 
 			while (await newQueues.Reader.WaitToReadAsync(cancellationToken))
 			{
@@ -427,7 +419,7 @@ namespace Horde.Build.Compute
 
 		async Task<bool> TryDispatchToNewQueueAsync(TQueueId queueId)
 		{
-			RedisList<TTask> queue = GetQueue(queueId);
+			RedisListKey<TTask> queue = GetQueueKey(queueId);
 
 			// Find a local listener that can execute the work
 			(TQueueId QueueId, TTask TaskId)? entry = null;
diff --git a/Engine/Source/Programs/Horde/Horde.Build/Jobs/Schedules/ScheduleService.cs b/Engine/Source/Programs/Horde/Horde.Build/Jobs/Schedules/ScheduleService.cs
--- a/Engine/Source/Programs/Horde/Horde.Build/Jobs/Schedules/ScheduleService.cs
+++ b/Engine/Source/Programs/Horde/Horde.Build/Jobs/Schedules/ScheduleService.cs
@@ -70,9 +70,9 @@ namespace Horde.Build.Jobs.Schedules
 		readonly IClock _clock;
 		readonly ILogger _logger;
 		readonly RedisService _redis;
-		readonly RedisKey _baseLockKey;
-		readonly RedisKey _tickLockKey; // Lock to tick the queue
-		readonly RedisSortedSet<QueueItem> _queue; // Items to tick, ordered by time
+		static readonly RedisKey s_baseLockKey = "scheduler/locks";
+		static readonly RedisKey s_tickLockKey = s_baseLockKey.Append("/tick"); // Lock to tick the queue
+		static readonly RedisSortedSetKey<QueueItem> s_queueKey = "scheduler/queue"; // Items to tick, ordered by time
 		readonly ITicker _ticker;
 
 		/// <summary>
@@ -90,9 +90,6 @@ namespace Horde.Build.Jobs.Schedules
 			_logger = logger;
 
 			_redis = redis;
-			_baseLockKey = "scheduler/locks";
-			_tickLockKey = _baseLockKey.Append("/tick");
-			_queue = new RedisSortedSet<QueueItem>(redis.GetDatabase(), "scheduler/queue");
 			if (mongoService.ReadOnlyMode)
 			{
 				_ticker = new NullTicker();
@@ -126,7 +123,7 @@ namespace Horde.Build.Jobs.Schedules
 			DateTime utcNow = _clock.UtcNow;
 
 			// Update the current queue
-			await using (RedisLock sharedLock = new (_redis.GetDatabase(), _tickLockKey))
+			await using (RedisLock sharedLock = new (_redis.GetDatabase(), s_tickLockKey))
 			{
 				if (await sharedLock.AcquireAsync(TimeSpan.FromMinutes(1.0), false))
 				{
@@ -145,7 +142,7 @@ namespace Horde.Build.Jobs.Schedules
 				}
 
 				// Acquire the lock for this schedule and update it
-				await using (RedisLock sharedLock = new RedisLock<QueueItem>(_redis.GetDatabase(), _baseLockKey, item))
+				await using (RedisLock sharedLock = new RedisLock<QueueItem>(_redis.GetDatabase(), s_baseLockKey, item))
 				{
 					if (await sharedLock.AcquireAsync(TimeSpan.FromMinutes(1.0)))
 					{
@@ -168,14 +165,15 @@ namespace Horde.Build.Jobs.Schedules
 
 		async Task<QueueItem?> PopQueueItemAsync()
 		{
+			IDatabaseAsync target = _redis.GetDatabase();
 			for (; ; )
 			{
-				QueueItem[] items = await _queue.RangeByRankAsync(0, 0);
+				QueueItem[] items = await target.SortedSetRangeByRankAsync(s_queueKey, 0, 0);
 				if (items.Length == 0)
 				{
 					return null;
 				}
-				if (await _queue.RemoveAsync(items[0]))
+				if (await target.SortedSetRemoveAsync(s_queueKey, items[0]))
 				{
 					return items[0];
 				}
@@ -185,8 +183,8 @@ namespace Horde.Build.Jobs.Schedules
 		internal async Task ResetAsync()
 		{
 			IDatabase redis = _redis.GetDatabase();
-			await redis.KeyDeleteAsync(_queue.Key);
-			await redis.KeyDeleteAsync(_tickLockKey);
+			await redis.KeyDeleteAsync(s_queueKey);
+			await redis.KeyDeleteAsync(s_tickLockKey);
 		}
 
 		internal async Task TickForTestingAsync()
@@ -224,7 +222,7 @@ namespace Horde.Build.Jobs.Schedules
 				}
 			}
 
-			await _queue.AddAsync(queueItems.ToArray());
+			await _redis.GetDatabase().SortedSetAddAsync(s_queueKey, queueItems.ToArray());
 		}
 
 		/// <summary>
diff --git a/Engine/Source/Programs/Horde/Horde.Build/Logs/LogTailService.cs b/Engine/Source/Programs/Horde/Horde.Build/Logs/LogTailService.cs
--- a/Engine/Source/Programs/Horde/Horde.Build/Logs/LogTailService.cs
+++ b/Engine/Source/Programs/Horde/Horde.Build/Logs/LogTailService.cs
@@ -24,8 +24,8 @@ namespace Horde.Build.Logs
 	/// 
 	/// * Once a log is retrieved, server adds an entry with this log id to a sorted set in Redis (see <see cref="_expireQueue"/>, scored by 
 	///   expiry time (defaulting to <see cref="ExpireAfter"/> from present) and broadcasts it to any other server instances.
-	/// * If a log is being tailed, server keeps the total number of lines in a Redis key (<see cref="RedisTailNext"/>, and appends chunks of log 
-	///   data to a sorted set for that log scored by index of the first line (see <see cref="RedisTailData"/>).
+	/// * If a log is being tailed, server keeps the total number of lines in a Redis key (<see cref="TailNextKey"/>, and appends chunks of log 
+	///   data to a sorted set for that log scored by index of the first line (see <see cref="TailDataKey"/>).
 	/// * Chunks are split on fixed boundaries, in order to allow older entries to be purged by score without having to count the number of lines
 	///   they contain first.
 	/// * Agent polls for requests to provide tail data via <see cref="LogRpcService.UpdateLogTail"/>, which calls <see cref="WaitForTailNextAsync"/>.
@@ -52,18 +52,18 @@ namespace Horde.Build.Logs
 
 		readonly ConcurrentDictionary<LogId, AsyncEvent> _notifyLogEvents = new ConcurrentDictionary<LogId, AsyncEvent>();
 		readonly RedisChannel<LogId> _tailStartChannel;
-		readonly RedisSortedSet<string> _trimQueue;
-		readonly RedisSortedSet<LogId> _expireQueue;
+		static readonly RedisSortedSetKey<string> _trimQueue = "logs:trim";
+		static readonly RedisSortedSetKey<LogId> _expireQueue = "logs:expire";
 
 		readonly IClock _clock;
 		readonly ITicker _expireTailsTicker;
 
 		readonly ILogger _logger;
 
-		RedisChannelSubscription<LogId>? _tailStartSubscription;
+		IAsyncDisposable? _tailStartSubscription;
 
-		internal RedisString<int> RedisTailNext(LogId logId) => new RedisString<int>(_redisService.ConnectionPool, $"logs:{logId}:tail-next");
-		internal RedisSortedSet<ReadOnlyMemory<byte>> RedisTailData(LogId logId) => new RedisSortedSet<ReadOnlyMemory<byte>>(_redisService.ConnectionPool, $"logs:{logId}:tail-data");
+		internal static RedisStringKey<int> TailNextKey(LogId logId) => new RedisStringKey<int>($"logs:{logId}:tail-next");
+		internal static RedisSortedSetKey<ReadOnlyMemory<byte>> TailDataKey(LogId logId) => new RedisSortedSetKey<ReadOnlyMemory<byte>>($"logs:{logId}:tail-data");
 
 		/// <summary>
 		/// Boundary to split Redis chunks along. 
@@ -92,8 +92,6 @@ namespace Horde.Build.Logs
 
 			_redisService = redisService;
 			_tailStartChannel = new RedisChannel<LogId>("logs:notify");
-			_trimQueue = new RedisSortedSet<string>(_redisService.ConnectionPool, "logs:trim");
-			_expireQueue = new RedisSortedSet<LogId>(_redisService.ConnectionPool, "logs:expire");
 
 			_clock = clock;
 			_expireTailsTicker = clock.AddSharedTicker<LogTailService>(TimeSpan.FromSeconds(30.0), ExpireTailsAsync, logger);
@@ -103,7 +101,7 @@ namespace Horde.Build.Logs
 		/// <inheritdoc/>
 		public async Task StartAsync(CancellationToken cancellationToken)
 		{
-			_tailStartSubscription = await _redisService.DatabaseSingleton.Multiplexer.GetSubscriber().SubscribeAsync(_tailStartChannel, (_, x) => OnTailStart(x));
+			_tailStartSubscription = await _redisService.SubscribeAsync(_tailStartChannel, OnTailStart);
 			await _expireTailsTicker.StartAsync();
 		}
 
@@ -121,7 +119,7 @@ namespace Horde.Build.Logs
 		{
 			double maxScore = (_clock.UtcNow - DateTime.UnixEpoch).TotalSeconds;
 
-			SortedSetEntry<string>[] trimEntries = await _trimQueue.RangeByScoreWithScoresAsync(stop: maxScore);
+			SortedSetEntry<string>[] trimEntries = await _redisService.GetDatabase().SortedSetRangeByScoreWithScoresAsync(_trimQueue, stop: maxScore);
 			foreach (SortedSetEntry<string> trimEntry in trimEntries)
 			{
 				string[] values = trimEntry.Element.Split('=');
@@ -130,18 +128,18 @@ namespace Horde.Build.Logs
 				int lineCount = Int32.Parse(values[1], NumberStyles.None, CultureInfo.InvariantCulture);
 
 				int minLineIndex = (lineCount & ~(ChunkLineCount - 1)) - 1;
-				_ = await RedisTailData(logId).RemoveRangeByScoreAsync(Double.NegativeInfinity, minLineIndex, Exclude.None, CommandFlags.FireAndForget);
+				_ = await _redisService.GetDatabase().SortedSetRemoveRangeByScoreAsync(TailDataKey(logId), Double.NegativeInfinity, minLineIndex, Exclude.None, CommandFlags.FireAndForget);
 			}
 
-			SortedSetEntry<LogId>[] expireEntries = await _expireQueue.RangeByScoreWithScoresAsync(stop: maxScore);
+			SortedSetEntry<LogId>[] expireEntries = await _redisService.GetDatabase().SortedSetRangeByScoreWithScoresAsync(_expireQueue, stop: maxScore);
 			foreach (SortedSetEntry<LogId> expireEntry in expireEntries)
 			{
 				LogId logId = expireEntry.Element;
 
 				ITransaction transaction = _redisService.GetDatabase().CreateTransaction();
-				transaction.AddCondition(Condition.SortedSetEqual(_expireQueue.Key, expireEntry.ElementValue, expireEntry.Score));
-				_ = transaction.KeyDeleteAsync(RedisTailNext(logId).Key, CommandFlags.FireAndForget);
-				_ = transaction.KeyDeleteAsync(RedisTailData(logId).Key, CommandFlags.FireAndForget);
+				transaction.AddCondition(Condition.SortedSetEqual(_expireQueue.Inner, expireEntry.ElementValue, expireEntry.Score));
+				_ = transaction.KeyDeleteAsync(TailNextKey(logId), CommandFlags.FireAndForget);
+				_ = transaction.KeyDeleteAsync(TailDataKey(logId), CommandFlags.FireAndForget);
 				_ = transaction.ExecuteAsync(CommandFlags.FireAndForget);
 			}
 		}
@@ -161,7 +159,7 @@ namespace Horde.Build.Logs
 		/// <param name="lineCount">The current flushed line count for the log file</param>
 		public async ValueTask EnableTailingAsync(LogId logId, int lineCount)
 		{
-			if (await RedisTailNext(logId).SetAsync(lineCount, when: When.NotExists))
+			if (await _redisService.GetDatabase().StringSetAsync(TailNextKey(logId), lineCount, when: When.NotExists))
 			{
 				_logger.LogDebug("Enabled tailing for log {LogId}", logId);
 			}
@@ -169,7 +167,7 @@ namespace Horde.Build.Logs
 			await _redisService.GetDatabase().PublishAsync(_tailStartChannel, logId);
 
 			double score = (_clock.UtcNow + ExpireAfter - DateTime.UnixEpoch).TotalSeconds;
-			_ = _expireQueue.AddAsync(logId, score, flags: CommandFlags.FireAndForget);
+			_ = _redisService.GetDatabase().SortedSetAddAsync(_expireQueue, logId, score, flags: CommandFlags.FireAndForget);
 		}
 
 		/// <summary>
@@ -186,13 +184,13 @@ namespace Horde.Build.Logs
 
 			if (entries.Count > 0)
 			{
-				RedisString<int> redisTailNext = RedisTailNext(logId);
-				RedisSortedSet<ReadOnlyMemory<byte>> redisTailData = RedisTailData(logId);
+				RedisStringKey<int> redisTailNext = TailNextKey(logId);
+				RedisSortedSetKey<ReadOnlyMemory<byte>> redisTailData = TailDataKey(logId);
 
 				ITransaction transaction = _redisService.GetDatabase().CreateTransaction();
-				transaction.AddCondition(Condition.StringEqual(RedisTailNext(logId).Key, tailNext));
-				_ = transaction.With(redisTailData).AddAsync(entries.ToArray(), flags: CommandFlags.FireAndForget);
-				_ = transaction.With(redisTailNext).SetAsync(newLineCount, flags: CommandFlags.FireAndForget);
+				transaction.AddCondition(Condition.StringEqual(TailNextKey(logId).Inner, RedisSerializer.Serialize(tailNext)));
+				_ = transaction.SortedSetAddAsync(redisTailData, entries.ToArray(), flags: CommandFlags.FireAndForget);
+				_ = transaction.StringSetAsync(redisTailNext, newLineCount, flags: CommandFlags.FireAndForget);
 
 				bool result = await transaction.ExecuteAsync();
 				_logger.LogTrace("Tail data for {LogId}, line {TailNext} -> {NewTailNext}, result: {Result}", logId, tailNext, newLineCount, result);
@@ -230,7 +228,7 @@ namespace Horde.Build.Logs
 
 			int nextIndex = lineIndex;
 
-			SortedSetEntry<ReadOnlyMemory<byte>>[] entries = await RedisTailData(logId).RangeByScoreWithScoresAsync(chunkLowerBound, chunkUpperBound);
+			SortedSetEntry<ReadOnlyMemory<byte>>[] entries = await _redisService.GetDatabase().SortedSetRangeByScoreWithScoresAsync(TailDataKey(logId), chunkLowerBound, chunkUpperBound);
 			foreach (SortedSetEntry<ReadOnlyMemory<byte>> entry in entries)
 			{
 				ReadOnlyMemory<byte> memory = entry.Element;
@@ -267,7 +265,7 @@ namespace Horde.Build.Logs
 		{
 			string trimEntry = $"{logId}={lineCount}";
 			double trimTime = (_clock.UtcNow + TrimAfter - DateTime.UnixEpoch).TotalSeconds;
-			_ = await _trimQueue.AddAsync(trimEntry, trimTime, flags: CommandFlags.FireAndForget);
+			_ = await _redisService.GetDatabase().SortedSetAddAsync(_trimQueue, trimEntry, trimTime, flags: CommandFlags.FireAndForget);
 		}
 
 		/// <summary>
@@ -277,8 +275,7 @@ namespace Horde.Build.Logs
 		/// <returns>Number of lines for the log file</returns>
 		public async Task<int> GetTailNextAsync(LogId logId)
 		{
-			int? tailNext = await RedisTailNext(logId).GetValueAsync();
-			return tailNext ?? -1;
+			return await _redisService.GetDatabase().StringGetAsync(TailNextKey(logId), -1);
 		}
 
 		/// <summary>
@@ -298,10 +295,10 @@ namespace Horde.Build.Logs
 					{
 						Task notifyTask = notifyEvent.Task;
 
-						int? tailNext = await RedisTailNext(logId).GetValueAsync();
-						if (tailNext != null)
+						int tailNext = await _redisService.GetDatabase().StringGetAsync(TailNextKey(logId), -1);
+						if (tailNext != -1)
 						{
-							return tailNext.Value;
+							return tailNext;
 						}
 						else if (cancellationToken.IsCancellationRequested)
 						{
diff --git a/Engine/Source/Programs/Horde/Horde.Build/Notifications/Sinks/SlackNotificationSink.cs b/Engine/Source/Programs/Horde/Horde.Build/Notifications/Sinks/SlackNotificationSink.cs
--- a/Engine/Source/Programs/Horde/Horde.Build/Notifications/Sinks/SlackNotificationSink.cs
+++ b/Engine/Source/Programs/Horde/Horde.Build/Notifications/Sinks/SlackNotificationSink.cs
@@ -221,12 +221,12 @@ namespace Horde.Build.Notifications.Sinks
 		readonly IExternalIssueService _externalIssueService;
 		readonly JsonSerializerOptions _jsonSerializerOptions;
 		readonly ITicker _escalateTicker;
-		readonly RedisSortedSet<int> _escalateIssues;
+		static readonly RedisSortedSetKey<int> _escalateIssues = "slack/escalate";
 		readonly IClock _clock;
 		readonly ILogger _logger;
 
 		readonly ITicker _issueQueueTicker;
-		readonly RedisList<int> _redisIssueQueue;
+		static readonly RedisListKey<int> _redisIssueQueue = "slack/issue-queue";
 		readonly string _redisIssueLockPrefix;
 
 		readonly HttpClient _httpClient;
@@ -252,14 +252,12 @@ namespace Horde.Build.Notifications.Sinks
 			_settings = settings.Value;
 			_messageStates = mongoService.GetCollection<MessageStateDocument>("SlackV2", keys => keys.Ascending(x => x.Recipient).Ascending(x => x.EventId), unique: true);
 			_slackUsers = mongoService.GetCollection<SlackUserDocument>("Slack.UsersV2");
-			_escalateIssues = new RedisSortedSet<int>(redisService.ConnectionPool, "slack/escalate");
 			_clock = clock;
 			_logger = logger;
 
 			_escalateTicker = clock.AddSharedTicker<SlackNotificationSink>(TimeSpan.FromMinutes(1.0), EscalateAsync, _logger);
 
 			_issueQueueTicker = clock.AddSharedTicker("slack-issues", TimeSpan.FromMinutes(1.0), ProcessIssueQueueAsync, _logger);
-			_redisIssueQueue = new RedisList<int>(redisService.ConnectionPool, "slack/issue-queue");
 			_redisIssueLockPrefix = "slack/issues/";
 
 			_httpClient = new HttpClient();
@@ -692,7 +690,7 @@ namespace Horde.Build.Notifications.Sinks
 			}
 
 			// Otherwise add it to the redis queue, and attempt to process the queue immediately.
-			await _redisIssueQueue.RightPushAsync(issue.Id);
+			await _redisService.GetDatabase().ListRightPushAsync(_redisIssueQueue, issue.Id);
 			await ProcessIssueQueueAsync(CancellationToken.None);
 		}
 
@@ -706,13 +704,13 @@ namespace Horde.Build.Notifications.Sinks
 
 			// Execute loop number of times based on the length of the queue at the start. This should bound the number of iterations while allowing us
 			// to re-queue items that we're unable to process now, without having to track whether we've reached the end of the list in its original state.
-			long count = await _redisIssueQueue.LengthAsync();
+			long count = await _redisService.GetDatabase().ListLengthAsync(_redisIssueQueue);
 			for (; count > 0; count--)
 			{
-				int issueId = await _redisIssueQueue.LeftPopAsync();
+				int issueId = await _redisService.GetDatabase().ListLeftPopAsync(_redisIssueQueue);
 				if (!testedIssueIds.Add(issueId) || !await TryUpdateIssueAsync(issueId))
 				{
-					await _redisIssueQueue.RightPushAsync(issueId);
+					await _redisService.GetDatabase().ListRightPushAsync(_redisIssueQueue, issueId);
 				}
 				cancellationToken.ThrowIfCancellationRequested();
 			}
@@ -1072,7 +1070,7 @@ namespace Horde.Build.Notifications.Sinks
 			if (workflow.EscalateAlias != null && workflow.EscalateTimes.Count > 0)
 			{
 				DateTime escalateTime = _clock.UtcNow + TimeSpan.FromMinutes(workflow.EscalateTimes[0]);
-				if (await _escalateIssues.AddAsync(issue.Id, (escalateTime - DateTime.UnixEpoch).TotalSeconds, StackExchange.Redis.When.NotExists))
+				if (await _redisService.GetDatabase().SortedSetAddAsync(_escalateIssues, issue.Id, (escalateTime - DateTime.UnixEpoch).TotalSeconds, StackExchange.Redis.When.NotExists))
 				{
 					_logger.LogInformation("First escalation time for issue {IssueId} is {Time}", issue.Id, escalateTime);
 				}
@@ -2162,7 +2160,7 @@ namespace Horde.Build.Notifications.Sinks
 			DateTime utcNow = DateTime.UtcNow;
 			double time = (utcNow - DateTime.UnixEpoch).TotalSeconds;
 
-			int[] issueIds = await _escalateIssues.RangeByScoreAsync(0, time);
+			int[] issueIds = await _redisService.GetDatabase().SortedSetRangeByScoreAsync(_escalateIssues, 0, time);
 			if (issueIds.Length > 0)
 			{
 				_logger.LogInformation("Escalating issues for {Time} ({TimeSecs})", utcNow, time);
@@ -2175,18 +2173,18 @@ namespace Horde.Build.Notifications.Sinks
 						if (nextTime == null)
 						{
 							_logger.LogInformation("Cancelling escalation for issue {IssueId}", issueId);
-							await _escalateIssues.RemoveAsync(issueId);
+							await _redisService.GetDatabase().SortedSetRemoveAsync(_escalateIssues, issueId);
 						}
 						else
 						{
 							_logger.LogInformation("Next escalation for issue {IssueId} is at timestamp {Time}", issueId, nextTime.Value);
-							await _escalateIssues.AddAsync(issueId, nextTime.Value);
+							await _redisService.GetDatabase().SortedSetAddAsync(_escalateIssues, issueId, nextTime.Value);
 						}
 					}
 					catch (SlackException ex)
 					{
 						_logger.LogError(ex, "Slack exception while escalating issue {IssueId}; cancelling.", issueId);
-						await _escalateIssues.RemoveAsync(issueId);
+						await _redisService.GetDatabase().SortedSetRemoveAsync(_escalateIssues, issueId);
 					}
 				}
 			}
diff --git a/Engine/Source/Programs/Horde/Horde.Build/Perforce/PerforceServiceCache.cs b/Engine/Source/Programs/Horde/Horde.Build/Perforce/PerforceServiceCache.cs
--- a/Engine/Source/Programs/Horde/Horde.Build/Perforce/PerforceServiceCache.cs
+++ b/Engine/Source/Programs/Horde/Horde.Build/Perforce/PerforceServiceCache.cs
@@ -628,7 +628,7 @@ namespace Horde.Build.Perforce
 					}
 				}
 
-				RedisChannelSubscription<StreamId>? subscription = null;
+				RedisSubscription? subscription = null;
 				try
 				{
 					for (; ; )
diff --git a/Engine/Source/Programs/Horde/Horde.Build/Server/RedisService.cs b/Engine/Source/Programs/Horde/Horde.Build/Server/RedisService.cs
--- a/Engine/Source/Programs/Horde/Horde.Build/Server/RedisService.cs
+++ b/Engine/Source/Programs/Horde/Horde.Build/Server/RedisService.cs
@@ -235,7 +235,7 @@ namespace Horde.Build.Server
 		}
 
 		/// <inheritdoc cref="SubscribeAsync{T}(RedisChannel{T}, Action{RedisChannel{T}, T})"/>
-		public Task<RedisChannelSubscription<T>> SubscribeAsync<T>(RedisChannel<T> channel, Action<T> callback) => SubscribeAsync(channel, (ch, x) => callback(x));
+		public Task<RedisSubscription> SubscribeAsync<T>(RedisChannel<T> channel, Action<T> callback) => SubscribeAsync(channel, (ch, x) => callback(x));
 
 		/// <summary>
 		/// Subscribe to notifications on a channel
@@ -244,10 +244,10 @@ namespace Horde.Build.Server
 		/// <param name="channel">Channel to monitor</param>
 		/// <param name="callback">Callback for new events</param>
 		/// <returns>Subscription object</returns>
-		public async Task<RedisChannelSubscription<T>> SubscribeAsync<T>(RedisChannel<T> channel, Action<RedisChannel<T>, T> callback)
+		public async Task<RedisSubscription> SubscribeAsync<T>(RedisChannel<T> channel, Action<RedisChannel<T>, T> callback)
 		{
-			ISubscriber subscriber = ConnectionPool.GetConnection().GetSubscriber();
-			return await subscriber.SubscribeAsync(channel, callback);
+			IConnectionMultiplexer connection = ConnectionPool.GetConnection();
+			return await connection.SubscribeAsync(channel, callback);
 		}
 	}
 }
diff --git a/Engine/Source/Programs/Horde/Horde.Build/Storage/StorageService.cs b/Engine/Source/Programs/Horde/Horde.Build/Storage/StorageService.cs
--- a/Engine/Source/Programs/Horde/Horde.Build/Storage/StorageService.cs
+++ b/Engine/Source/Programs/Horde/Horde.Build/Storage/StorageService.cs
@@ -935,7 +935,7 @@ namespace Horde.Build.Storage
 		static RedisKey GetRedisKey(string suffix) => $"storage:{suffix}";
 		static RedisKey GetRedisKey(NamespaceId namespaceId, string suffix) => GetRedisKey($"{namespaceId}:{suffix}");
 
-		RedisSortedSet<RedisValue> GetGcCheckSet(NamespaceId namespaceId) => new RedisSortedSet<RedisValue>(_redisService.ConnectionPool, GetRedisKey(namespaceId, "check"));
+		static RedisSortedSetKey<RedisValue> GetGcCheckSet(NamespaceId namespaceId) => new RedisSortedSetKey<RedisValue>(GetRedisKey(namespaceId, "check"));
 
 		/// <summary>
 		/// Find the next namespace to run GC on
@@ -1016,10 +1016,10 @@ namespace Horde.Build.Storage
 
 			double score = GetGcTimestamp(utcNow);
 
-			RedisSortedSet<RedisValue> checkSet = GetGcCheckSet(namespaceId);
+			RedisSortedSetKey<RedisValue> checkSet = GetGcCheckSet(namespaceId);
 			for (; ; )
 			{
-				RedisValue[] values = await checkSet.RangeByRankAsync(0, 0);
+				RedisValue[] values = await _redisService.GetDatabase().SortedSetRangeByRankAsync(checkSet, 0, 0);
 				if (values.Length == 0)
 				{
 					break;
@@ -1034,13 +1034,13 @@ namespace Horde.Build.Storage
 						if (info.Imports != null)
 						{
 							SortedSetEntry<RedisValue>[] entries = info.Imports.Select(x => new SortedSetEntry<RedisValue>(x.ToByteArray(), score)).ToArray();
-							_ = checkSet.AddAsync(entries, flags: CommandFlags.FireAndForget);
+							_ = _redisService.GetDatabase().SortedSetAddAsync(checkSet, entries, flags: CommandFlags.FireAndForget);
 							score = Math.BitIncrement(score);
 						}
 						await namespaceInfo.Client.DeleteBlobAsync(info.BlobId, cancellationToken);
 					}
 				}
-				_ = checkSet.RemoveAsync(values[0], CommandFlags.FireAndForget);
+				_ = _redisService.GetDatabase().SortedSetRemoveAsync(checkSet, values[0], CommandFlags.FireAndForget);
 			}
 
 			await _gcState.UpdateAsync(state => state.FindOrAddNamespace(namespaceId).LastTime = utcNow);
@@ -1107,7 +1107,7 @@ namespace Horde.Build.Storage
 		void AddGcCheckRecord(NamespaceId namespaceId, ObjectId blobInfoId)
 		{
 			double score = GetGcTimestamp();
-			_ = GetGcCheckSet(namespaceId).AddAsync(blobInfoId.ToByteArray(), score, flags: CommandFlags.FireAndForget);
+			_ = _redisService.GetDatabase().SortedSetAddAsync(GetGcCheckSet(namespaceId), blobInfoId.ToByteArray(), score, flags: CommandFlags.FireAndForget);
 		}
 
 		#endregion
diff --git a/Engine/Source/Programs/Shared/EpicGames.Redis/RedisChannel.cs b/Engine/Source/Programs/Shared/EpicGames.Redis/RedisChannel.cs
--- a/Engine/Source/Programs/Shared/EpicGames.Redis/RedisChannel.cs
+++ b/Engine/Source/Programs/Shared/EpicGames.Redis/RedisChannel.cs
@@ -1,13 +1,12 @@
 // Copyright Epic Games, Inc. All Rights Reserved.
 
+using EpicGames.Core;
 using StackExchange.Redis;
 using System;
 using System.Threading.Tasks;
 
 namespace EpicGames.Redis
 {
-	using RedisChannel = StackExchange.Redis.RedisChannel;
-
 	/// <summary>
 	/// Represents a typed pub/sub channel with a particular value
 	/// </summary>
@@ -25,6 +24,27 @@ namespace EpicGames.Redis
 		/// <param name="channel"></param>
 		public RedisChannel(RedisChannel channel) => Channel = channel;
 
+		/// <summary>
+		/// Constructor
+		/// </summary>
+		/// <param name="value">The channel name</param>
+		/// <param name="mode">Pattern mode</param>
+		public RedisChannel(Utf8String value, RedisChannel.PatternMode mode) => Channel = new RedisChannel(value.Memory.ToArray(), mode);
+
+		/// <summary>
+		/// Constructor
+		/// </summary>
+		/// <param name="value">The channel name</param>
+		/// <param name="mode">Pattern mode</param>
+		public RedisChannel(string value, RedisChannel.PatternMode mode) => Channel = new RedisChannel(value, mode);
+
+		/// <summary>
+		/// Constructor
+		/// </summary>
+		/// <param name="value">The channel name</param>
+		/// <param name="mode">Pattern mode</param>
+		public RedisChannel(byte[] value, RedisChannel.PatternMode mode) => Channel = new RedisChannel(value, mode);
+
 		/// <inheritdoc/>
 		public override bool Equals(object? obj) => obj is RedisChannel<T> list && Channel == list.Channel;
 
@@ -39,100 +59,32 @@ namespace EpicGames.Redis
 
 		/// <summary>Compares two instances for equality</summary>
 		public static bool operator !=(RedisChannel<T> left, RedisChannel<T> right) => left.Channel != right.Channel;
-	}
-
-	/// <summary>
-	/// Subscription to a <see cref="RedisChannel{Task}"/>
-	/// </summary>
-	/// <typeparam name="T"></typeparam>
-	public sealed class RedisChannelSubscription<T> : IDisposable, IAsyncDisposable
-	{
-		/// <summary>
-		/// The subscriber to register with
-		/// </summary>
-		ISubscriber? _subscriber;
 
 		/// <summary>
-		/// The channel to post on
+		/// Create a channel name from a <see cref="Utf8String"/>.
 		/// </summary>
-		public RedisChannel<T> Channel { get; }
+		/// <param name="key">The string to get a channel from.</param>
+		public static implicit operator RedisChannel<T>(Utf8String key) => new RedisChannel<T>(key, RedisChannel.PatternMode.Auto);
 
 		/// <summary>
-		/// The handler to call
+		/// Create a channel name from a byte array.
 		/// </summary>
-		readonly Action<RedisChannel<T>, T> _handler;
+		/// <param name="key">The byte array to get a channel from.</param>
+		public static implicit operator RedisChannel<T>(string key) => new RedisChannel<T>(key, RedisChannel.PatternMode.Auto);
 
 		/// <summary>
-		/// Constructor
+		/// Create a channel name from a byte array.
 		/// </summary>
-		/// <param name="subscriber"></param>
-		/// <param name="channel"></param>
-		/// <param name="handler"></param>
-		RedisChannelSubscription(ISubscriber subscriber, RedisChannel<T> channel, Action<RedisChannel<T>, T> handler)
-		{
-			_subscriber = subscriber;
-			Channel = channel;
-			_handler = handler;
-		}
-		
-		/// <summary>
-		/// Start the subscription
-		/// </summary>
-		internal static async Task<RedisChannelSubscription<T>> CreateAsync(ISubscriber subscriber, RedisChannel<T> channel, Action<RedisChannel<T>, T> handler, CommandFlags flags = CommandFlags.None)
-		{
-			RedisChannelSubscription<T> subscription = new RedisChannelSubscription<T>(subscriber, channel, handler);
-			await subscription._subscriber!.SubscribeAsync(channel.Channel, subscription.UntypedHandler, flags);
-			return subscription;
-		}
-
-		/// <summary>
-		/// Unsubscribe from the channel
-		/// </summary>
-		/// <param name="flags">Flags for the operation</param>
-		public async Task UnsubscribeAsync(CommandFlags flags = CommandFlags.None)
-		{
-			if (_subscriber != null)
-			{
-				await _subscriber.UnsubscribeAsync(Channel.Channel, UntypedHandler, flags);
-				_subscriber = null;
-			}
-		}
-
-		/// <summary>
-		/// The callback for messages being received
-		/// </summary>
-		/// <param name="_"></param>
-		/// <param name="message"></param>
-		void UntypedHandler(RedisChannel _, RedisValue message)
-		{
-			_handler(Channel, RedisSerializer.Deserialize<T>(message)!);
-		}
-
-		/// <inheritdoc/>
-		public async ValueTask DisposeAsync()
-		{
-			await UnsubscribeAsync();
-		}
-
-		/// <inheritdoc/>
-		public void Dispose()
-		{
-			UnsubscribeAsync().Wait();
-		}
+		/// <param name="key">The byte array to get a channel from.</param>
+		public static implicit operator RedisChannel<T>(byte[] key) => new RedisChannel<T>(key, RedisChannel.PatternMode.Auto);
 	}
 
 	/// <summary>
-	/// Extension methods for typed lists
+	/// Extension methods for channels
 	/// </summary>
 	public static class RedisChannelExtensions
 	{
-		/// <inheritdoc cref="ISubscriber.SubscribeAsync(StackExchange.Redis.RedisChannel, Action{StackExchange.Redis.RedisChannel, RedisValue}, CommandFlags)"/>
-		public static Task<RedisChannelSubscription<T>> SubscribeAsync<T>(this ISubscriber subscriber, RedisChannel<T> channel, Action<RedisChannel<T>, T> handler, CommandFlags flags = CommandFlags.None)
-		{
-			return RedisChannelSubscription<T>.CreateAsync(subscriber, channel, handler, flags);
-		}
-
-		/// <inheritdoc cref="IDatabaseAsync.PublishAsync(StackExchange.Redis.RedisChannel, RedisValue, CommandFlags)"/>
+		/// <inheritdoc cref="IDatabaseAsync.PublishAsync(RedisChannel, RedisValue, CommandFlags)"/>
 		public static Task PublishAsync<T>(this IDatabaseAsync database, RedisChannel<T> channel, T item, CommandFlags flags = CommandFlags.None)
 		{
 			RedisValue value = RedisSerializer.Serialize(item);
diff --git a/Engine/Source/Programs/Shared/EpicGames.Redis/RedisHash.cs b/Engine/Source/Programs/Shared/EpicGames.Redis/RedisHash.cs
--- a/Engine/Source/Programs/Shared/EpicGames.Redis/RedisHash.cs
+++ b/Engine/Source/Programs/Shared/EpicGames.Redis/RedisHash.cs
@@ -9,126 +9,37 @@ using System.Threading.Tasks;
 namespace EpicGames.Redis
 {
 	/// <summary>
-	/// Represents a typed Redis hash with a given key and value
+	/// Represents a typed Redis list with a given key
 	/// </summary>
-	/// <typeparam name="TName">The key type for the hash</typeparam>
-	/// <typeparam name="TValue">The value type for the hash</typeparam>
-	public readonly struct RedisHash<TName, TValue>
+	/// <typeparam name="TName">Type of the hash key</typeparam>
+	/// <typeparam name="TValue">Type of the hash value</typeparam>
+	public readonly struct RedisHashKey<TName, TValue>
 	{
-		internal readonly RedisConnectionPool? ConnectionPool = null;
-		internal readonly IDatabaseAsync? Database = null;
-
 		/// <summary>
 		/// The key for the list
 		/// </summary>
-		public readonly RedisKey Key { get; }
+		public readonly RedisKey Inner { get; }
 
 		/// <summary>
 		/// Constructor
 		/// </summary>
-		/// <param name="database">Explicitly set Redis database (overrides use of any connection pool)</param>
-		/// <param name="key">Redis key this type is using</param>
-		public RedisHash(IDatabaseAsync database, RedisKey key)
+		/// <param name="inner">Redis key this type is using</param>
+		public RedisHashKey(RedisKey inner)
 		{
-			Database = database;
-			Key = key;
+			Inner = inner;
 		}
-		
+
 		/// <summary>
-		/// Constructor
+		/// Implicit conversion to typed redis key.
 		/// </summary>
-		/// <param name="connectionPool">Connection pool for Redis databases</param>
-		/// <param name="key">Redis key this type is using</param>
-		public RedisHash(RedisConnectionPool connectionPool, RedisKey key)
-		{
-			ConnectionPool = connectionPool;
-			Key = key;
-		}
-		
+		/// <param name="key">Key to convert</param>
+		public static implicit operator RedisHashKey<TName, TValue>(string key) => new RedisHashKey<TName, TValue>(new RedisKey(key));
+
 		/// <summary>
-		/// Get the Redis database in use
+		/// Implicit conversion to untyped redis keys.
 		/// </summary>
-		/// <returns>A connection pool or explicitly set Redis database</returns>
-		/// <exception cref="InvalidOperationException">If neither are set</exception>
-		public IDatabaseAsync GetDatabase()
-		{
-			if (Database != null) return Database;
-			if (ConnectionPool != null) return ConnectionPool.GetDatabase();
-			throw new InvalidOperationException($"Neither {nameof(Database)} or {nameof(ConnectionPool)} has been set!");
-		}
-
-		/// <inheritdoc cref="IDatabaseAsync.HashDeleteAsync(RedisKey, RedisValue, CommandFlags)"/>
-		public Task<bool> DeleteAsync(TName name, CommandFlags flags = CommandFlags.None)
-		{
-			return GetDatabase().HashDeleteAsync(Key, RedisSerializer.Serialize(name), flags);
-		}
-
-		/// <inheritdoc cref="IDatabaseAsync.HashDeleteAsync(RedisKey, RedisValue[], CommandFlags)"/>
-		public Task<long> DeleteAsync(IEnumerable<TName> names, CommandFlags flags = CommandFlags.None)
-		{
-			RedisValue[] nameArray = names.Select(x => RedisSerializer.Serialize(x)).ToArray();
-			return GetDatabase().HashDeleteAsync(Key, nameArray, flags);
-		}
-
-		/// <inheritdoc cref="IDatabaseAsync.HashExistsAsync(RedisKey, RedisValue, CommandFlags)"/>
-		public Task<bool> ExistsAsync(TName name, CommandFlags flags = CommandFlags.None)
-		{
-			return GetDatabase().HashExistsAsync(Key, RedisSerializer.Serialize(name), flags);
-		}
-
-		/// <inheritdoc cref="IDatabaseAsync.HashGetAsync(RedisKey, RedisValue, CommandFlags)"/>
-		public async Task<TValue> GetAsync(TName name, CommandFlags flags = CommandFlags.None)
-		{
-			RedisValue value = await GetDatabase().HashGetAsync(Key, RedisSerializer.Serialize(name), flags);
-			return value.IsNull ? default! : RedisSerializer.Deserialize<TValue>(value)!;
-		}
-
-		/// <inheritdoc cref="IDatabaseAsync.HashGetAsync(RedisKey, RedisValue[], CommandFlags)"/>
-		public async Task<TValue[]> GetAsync(IEnumerable<TName> names, CommandFlags flags = CommandFlags.None)
-		{
-			RedisValue[] nameArray = names.Select(x => RedisSerializer.Serialize(x)).ToArray();
-			RedisValue[] valueArray = await GetDatabase().HashGetAsync(Key, nameArray, flags);
-			return Array.ConvertAll(valueArray, x => RedisSerializer.Deserialize<TValue>(x)!);
-		}
-
-		/// <inheritdoc cref="IDatabaseAsync.HashGetAllAsync(RedisKey, CommandFlags)"/>
-		public async Task<HashEntry<TName, TValue>[]> GetAllAsync(CommandFlags flags = CommandFlags.None)
-		{
-			HashEntry[] entries = await GetDatabase().HashGetAllAsync(Key, flags);
-			return Array.ConvertAll(entries, x => new HashEntry<TName, TValue>(RedisSerializer.Deserialize<TName>(x.Name)!, RedisSerializer.Deserialize<TValue>(x.Value)!));
-		}
-
-		/// <inheritdoc cref="IDatabaseAsync.HashKeysAsync(RedisKey, CommandFlags)"/>
-		public async Task<TName[]> KeysAsync(CommandFlags flags = CommandFlags.None)
-		{
-			RedisValue[] nameArray = await GetDatabase().HashKeysAsync(Key, flags);
-			return Array.ConvertAll(nameArray, x => RedisSerializer.Deserialize<TName>(x)!);
-		}
-
-		/// <inheritdoc cref="IDatabaseAsync.HashLengthAsync(RedisKey, CommandFlags)"/>
-		public Task<long> LengthAsync(CommandFlags flags = CommandFlags.None)
-		{
-			return GetDatabase().HashLengthAsync(Key, flags);
-		}
-
-		/// <inheritdoc cref="IDatabaseAsync.HashSetAsync(RedisKey, RedisValue, RedisValue, When, CommandFlags)"/>
-		public Task SetAsync(TName name, TValue value, When when = When.Always, CommandFlags flags = CommandFlags.None)
-		{
-			return GetDatabase().HashSetAsync(Key, RedisSerializer.Serialize(name), RedisSerializer.Serialize(value), when, flags);
-		}
-
-		/// <inheritdoc cref="IDatabaseAsync.HashSetAsync(RedisKey, HashEntry[], CommandFlags)"/>
-		public Task SetAsync(IEnumerable<HashEntry<TName, TValue>> entries, CommandFlags flags = CommandFlags.None)
-		{
-			return GetDatabase().HashSetAsync(Key, entries.Select(x => (HashEntry)x).ToArray(), flags);
-		}
-
-		/// <inheritdoc cref="IDatabaseAsync.HashValuesAsync(RedisKey, CommandFlags)"/>
-		public async Task<TValue[]> ValuesAsync(CommandFlags flags = CommandFlags.None)
-		{
-			RedisValue[] values = await GetDatabase().HashValuesAsync(Key, flags);
-			return Array.ConvertAll(values, x => RedisSerializer.Deserialize<TValue>(x)!);
-		}
+		/// <param name="key">Key to convert</param>
+		public static implicit operator TypedRedisKey(RedisHashKey<TName, TValue> key) => new TypedRedisKey(key.Inner);
 	}
 
 	/// <inheritdoc cref="HashEntry"/>
@@ -175,36 +86,137 @@ namespace EpicGames.Redis
 	/// </summary>
 	public static class RedisHashExtensions
 	{
+		#region HashDecrementAsync
+
 		/// <inheritdoc cref="IDatabaseAsync.HashDecrementAsync(RedisKey, RedisValue, Int64, CommandFlags)"/>
-		public static Task<long> DecrementAsync<TName>(this RedisHash<TName, long> hash, TName name, long value = 1L, CommandFlags flags = CommandFlags.None)
+		public static Task<long> HashDecrementAsync<TName>(this IDatabaseAsync target, RedisHashKey<TName, long> key, TName name, long value = 1L, CommandFlags flags = CommandFlags.None)
 		{
-			return hash.GetDatabase().HashDecrementAsync(hash.Key, RedisSerializer.Serialize<TName>(name), value, flags);
+			return target.HashDecrementAsync(key.Inner, RedisSerializer.Serialize<TName>(name), value, flags);
 		}
 
 		/// <inheritdoc cref="IDatabaseAsync.HashDecrementAsync(RedisKey, RedisValue, Double, CommandFlags)"/>
-		public static Task<double> DecrementAsync<TName>(this RedisHash<TName, long> hash, TName name, double value = 1.0, CommandFlags flags = CommandFlags.None)
+		public static Task<double> HashDecrementAsync<TName>(this IDatabaseAsync target, RedisHashKey<TName, long> key, TName name, double value = 1.0, CommandFlags flags = CommandFlags.None)
+		{
+			return target.HashDecrementAsync(key.Inner, RedisSerializer.Serialize<TName>(name), value, flags);
+		}
+
+		#endregion
+
+		#region HashDeleteAsync
+
+		/// <inheritdoc cref="IDatabaseAsync.HashDeleteAsync(RedisKey, RedisValue, CommandFlags)"/>
+		public static Task<bool> HashDeleteAsync<TName, TValue>(this IDatabaseAsync target, RedisHashKey<TName, TValue> key, TName name, CommandFlags flags = CommandFlags.None)
 		{
-			return hash.GetDatabase().HashDecrementAsync(hash.Key, RedisSerializer.Serialize<TName>(name), value, flags);
+			return target.HashDeleteAsync(key.Inner, RedisSerializer.Serialize(name), flags);
 		}
 
+		/// <inheritdoc cref="IDatabaseAsync.HashDeleteAsync(RedisKey, RedisValue[], CommandFlags)"/>
+		public static Task<long> HashDeleteAsync<TName, TValue>(this IDatabaseAsync target, RedisHashKey<TName, TValue> key, TName[] names, CommandFlags flags = CommandFlags.None)
+		{
+			return target.HashDeleteAsync(key.Inner, RedisSerializer.Serialize(names), flags);
+		}
+
+		#endregion
+
+		#region HashExistsAsync
+
+		/// <inheritdoc cref="IDatabaseAsync.HashExistsAsync(RedisKey, RedisValue, CommandFlags)"/>
+		public static Task<bool> HashExistsAsync<TName, TValue>(this IDatabaseAsync target, RedisHashKey<TName, TValue> key, TName name, CommandFlags flags = CommandFlags.None)
+		{
+			return target.HashExistsAsync(key.Inner, RedisSerializer.Serialize(name), flags);
+		}
+
+		#endregion
+
+		#region HashGetAsync
+
+		/// <inheritdoc cref="IDatabaseAsync.HashGetAsync(RedisKey, RedisValue, CommandFlags)"/>
+		public static Task<TValue> HashGetAsync<TName, TValue>(this IDatabaseAsync target, RedisHashKey<TName, TValue> key, TName name, CommandFlags flags = CommandFlags.None)
+		{
+			return target.HashGetAsync(key.Inner, RedisSerializer.Serialize(name), flags).DeserializeAsync<TValue>();
+		}
+
+		/// <inheritdoc cref="IDatabaseAsync.HashGetAsync(RedisKey, RedisValue[], CommandFlags)"/>
+		public static Task<TValue[]> HashGetAsync<TName, TValue>(this IDatabaseAsync target, RedisHashKey<TName, TValue> key, TName[] names, CommandFlags flags = CommandFlags.None)
+		{
+			return target.HashGetAsync(key.Inner, RedisSerializer.Serialize(names), flags).DeserializeAsync<TValue>();
+		}
+
+		#endregion
+
+		#region HashGetAllAsync
+
+		/// <inheritdoc cref="IDatabaseAsync.HashGetAllAsync(RedisKey, CommandFlags)"/>
+		public static async Task<HashEntry<TName, TValue>[]> HashGetAllAsync<TName, TValue>(this IDatabaseAsync target, RedisHashKey<TName, TValue> key, CommandFlags flags = CommandFlags.None)
+		{
+			HashEntry[] entries = await target.HashGetAllAsync(key.Inner, flags);
+			return Array.ConvertAll(entries, x => new HashEntry<TName, TValue>(RedisSerializer.Deserialize<TName>(x.Name)!, RedisSerializer.Deserialize<TValue>(x.Value)!));
+		}
+
+		#endregion
+
+		#region HashIncrementAsync
+
+
 		/// <inheritdoc cref="IDatabaseAsync.HashIncrementAsync(RedisKey, RedisValue, Int64, CommandFlags)"/>
-		public static Task<long> IncrementAsync<TName>(this RedisHash<TName, long> hash, TName name, long value = 1L, CommandFlags flags = CommandFlags.None)
+		public static Task<long> HashIncrementAsync<TName>(this IDatabaseAsync target, RedisHashKey<TName, long> key, TName name, long value = 1L, CommandFlags flags = CommandFlags.None)
 		{
-			return hash.GetDatabase().HashIncrementAsync(hash.Key, RedisSerializer.Serialize<TName>(name), value, flags);
+			return target.HashIncrementAsync(key.Inner, RedisSerializer.Serialize<TName>(name), value, flags);
 		}
 
 		/// <inheritdoc cref="IDatabaseAsync.HashIncrementAsync(RedisKey, RedisValue, Double, CommandFlags)"/>
-		public static Task<double> IncrementAsync<TName>(this RedisHash<TName, long> hash, TName name, double value = 1.0, CommandFlags flags = CommandFlags.None)
+		public static Task<double> HashIncrementAsync<TName>(this IDatabaseAsync target, RedisHashKey<TName, long> key, TName name, double value = 1.0, CommandFlags flags = CommandFlags.None)
 		{
-			return hash.GetDatabase().HashIncrementAsync(hash.Key, RedisSerializer.Serialize<TName>(name), value, flags);
+			return target.HashIncrementAsync(key.Inner, RedisSerializer.Serialize<TName>(name), value, flags);
 		}
 
-		/// <summary>
-		/// Creates a version of this set which modifies a transaction rather than the direct DB
-		/// </summary>
-		public static RedisHash<TName, TValue> With<TName, TValue>(this ITransaction transaction, RedisHash<TName, TValue> set)
+		#endregion
+
+		#region HashKeysAsync
+
+		/// <inheritdoc cref="IDatabaseAsync.HashKeysAsync(RedisKey, CommandFlags)"/>
+		public static Task<TName[]> HashKeysAsync<TName, TValue>(this IDatabaseAsync target, RedisHashKey<TName, TValue> key, CommandFlags flags = CommandFlags.None)
 		{
-			return new RedisHash<TName, TValue>(transaction, set.Key);
+			return target.HashKeysAsync(key.Inner, flags).DeserializeAsync<TName>();
 		}
+
+		#endregion
+
+		#region HashLengthAsync
+
+		/// <inheritdoc cref="IDatabaseAsync.HashLengthAsync(RedisKey, CommandFlags)"/>
+		public static Task<long> HashLengthAsync(this IDatabaseAsync target, TypedRedisKey key, CommandFlags flags = CommandFlags.None)
+		{
+			return target.HashLengthAsync(key, flags);
+		}
+
+		#endregion
+
+		#region HashSetAsync
+
+
+		/// <inheritdoc cref="IDatabaseAsync.HashSetAsync(RedisKey, RedisValue, RedisValue, When, CommandFlags)"/>
+		public static Task HashSetAsync<TName, TValue>(this IDatabaseAsync target, RedisHashKey<TName, TValue> key, TName name, TValue value, When when = When.Always, CommandFlags flags = CommandFlags.None)
+		{
+			return target.HashSetAsync(key.Inner, RedisSerializer.Serialize(name), RedisSerializer.Serialize(value), when, flags);
+		}
+
+		/// <inheritdoc cref="IDatabaseAsync.HashSetAsync(RedisKey, HashEntry[], CommandFlags)"/>
+		public static Task HashSetAsync<TName, TValue>(this IDatabaseAsync target, RedisHashKey<TName, TValue> key, IEnumerable<HashEntry<TName, TValue>> entries, CommandFlags flags = CommandFlags.None)
+		{
+			return target.HashSetAsync(key.Inner, entries.Select(x => (HashEntry)x).ToArray(), flags);
+		}
+
+		#endregion
+
+		#region HashValuesAsync
+
+		/// <inheritdoc cref="IDatabaseAsync.HashValuesAsync(RedisKey, CommandFlags)"/>
+		public static Task<TValue[]> HashValuesAsync<TName, TValue>(this IDatabaseAsync target, RedisHashKey<TName, TValue> key, CommandFlags flags = CommandFlags.None)
+		{
+			return target.HashValuesAsync(key.Inner, flags).DeserializeAsync<TValue>();
+		}
+
+		#endregion
 	}
 }
diff --git a/Engine/Source/Programs/Shared/EpicGames.Redis/RedisList.cs b/Engine/Source/Programs/Shared/EpicGames.Redis/RedisList.cs
--- a/Engine/Source/Programs/Shared/EpicGames.Redis/RedisList.cs
+++ b/Engine/Source/Programs/Shared/EpicGames.Redis/RedisList.cs
@@ -11,159 +11,173 @@ namespace EpicGames.Redis
 	/// <summary>
 	/// Represents a typed Redis list with a given key
 	/// </summary>
-	/// <typeparam name="TElement">The type of element stored in the list</typeparam>
-	public readonly struct RedisList<TElement>
+	/// <typeparam name="TElement">The type of element stored in the set</typeparam>
+	public readonly struct RedisListKey<TElement>
 	{
-		internal readonly RedisConnectionPool? ConnectionPool = null;
-		internal readonly IDatabaseAsync? Database = null;
-
 		/// <summary>
 		/// The key for the list
 		/// </summary>
-		public readonly RedisKey Key { get; }
+		public readonly RedisKey Inner { get; }
 
 		/// <summary>
 		/// Constructor
 		/// </summary>
-		/// <param name="database">Explicitly set Redis database (overrides use of any connection pool)</param>
-		/// <param name="key">Redis key this type is using</param>
-		public RedisList(IDatabaseAsync database, RedisKey key)
+		/// <param name="inner">Redis key this type is using</param>
+		public RedisListKey(RedisKey inner)
 		{
-			Database = database;
-			Key = key;
+			Inner = inner;
 		}
-		
+
 		/// <summary>
-		/// Constructor
+		/// Implicit conversion to typed redis key.
 		/// </summary>
-		/// <param name="connectionPool">Connection pool for Redis databases</param>
-		/// <param name="key">Redis key this type is using</param>
-		public RedisList(RedisConnectionPool connectionPool, RedisKey key)
-		{
-			ConnectionPool = connectionPool;
-			Key = key;
-		}
-		
+		/// <param name="key">Key to convert</param>
+		public static implicit operator RedisListKey<TElement>(string key) => new RedisListKey<TElement>(new RedisKey(key));
+
 		/// <summary>
-		/// Get the Redis database in use
+		/// Implicit conversion to untyped redis keys.
 		/// </summary>
-		/// <returns>A connection pool or explicitly set Redis database</returns>
-		/// <exception cref="InvalidOperationException">If neither are set</exception>
-		public IDatabaseAsync GetDatabase()
-		{
-			if (Database != null) return Database;
-			if (ConnectionPool != null) return ConnectionPool.GetDatabase();
-			throw new InvalidOperationException($"Neither {nameof(Database)} or {nameof(ConnectionPool)} has been set!");
-		}
+		/// <param name="key">Key to convert</param>
+		public static implicit operator TypedRedisKey(RedisListKey<TElement> key) => key.Inner;
+	}
+
+	/// <summary>
+	/// Extension methods for sets
+	/// </summary>
+	public static class RedisListExtensions
+	{
+		#region ListGetByIndexAsync
 
 		/// <inheritdoc cref="IDatabaseAsync.ListGetByIndexAsync(RedisKey, Int64, CommandFlags)"/>
-		public async Task<TElement> GetByIndexAsync(long index, CommandFlags flags = CommandFlags.None)
+		public static Task<TElement> ListGetByIndexAsync<TElement>(this IDatabaseAsync target, RedisListKey<TElement> key, long index, CommandFlags flags = CommandFlags.None)
 		{
-			RedisValue value = await GetDatabase().ListGetByIndexAsync(Key, index, flags);
-			return value.IsNull? default! : RedisSerializer.Deserialize<TElement>(value)!;
+			return target.ListGetByIndexAsync(key.Inner, index, flags).DeserializeAsync<TElement>();
 		}
 
+		#endregion
+
+		#region ListInsertAfterAsync
+
 		/// <inheritdoc cref="IDatabaseAsync.ListInsertAfterAsync(RedisKey, RedisValue, RedisValue, CommandFlags)"/>
-		public Task<long> InsertAfterAsync(TElement pivot, TElement item, CommandFlags flags = CommandFlags.None)
+		public static Task<long> ListInsertAfterAsync<TElement>(this IDatabaseAsync target, RedisListKey<TElement> key, TElement pivot, TElement item, CommandFlags flags = CommandFlags.None)
 		{
-			RedisValue pivotValue = RedisSerializer.Serialize(pivot);
-			RedisValue itemValue = RedisSerializer.Serialize(item);
-			return GetDatabase().ListInsertAfterAsync(Key, pivotValue, itemValue, flags);
+			return target.ListInsertAfterAsync(key.Inner, RedisSerializer.Serialize(pivot), RedisSerializer.Serialize(item), flags);
 		}
 
+		#endregion
+
+		#region ListInsertBeforeAsync
+
 		/// <inheritdoc cref="IDatabaseAsync.ListInsertBeforeAsync(RedisKey, RedisValue, RedisValue, CommandFlags)"/>
-		public Task<long> InsertBeforeAsync(TElement pivot, TElement item, CommandFlags flags = CommandFlags.None)
+		public static Task<long> ListInsertBeforeAsync<TElement>(this IDatabaseAsync target, RedisListKey<TElement> key, TElement pivot, TElement item, CommandFlags flags = CommandFlags.None)
 		{
 			RedisValue pivotValue = RedisSerializer.Serialize(pivot);
 			RedisValue itemValue = RedisSerializer.Serialize(item);
-			return GetDatabase().ListInsertBeforeAsync(Key, pivotValue, itemValue, flags);
+			return target.ListInsertBeforeAsync(key.Inner, pivotValue, itemValue, flags);
 		}
 
+		#endregion
+
+		#region ListLeftPopAsync
+
 		/// <inheritdoc cref="IDatabaseAsync.ListLeftPopAsync(RedisKey, CommandFlags)"/>
-		public async Task<TElement> LeftPopAsync(CommandFlags flags = CommandFlags.None)
+		public static Task<TElement> ListLeftPopAsync<TElement>(this IDatabaseAsync target, RedisListKey<TElement> key, CommandFlags flags = CommandFlags.None)
 		{
-			RedisValue value = await GetDatabase().ListLeftPopAsync(Key, flags);
-			return value.IsNull ? default! : RedisSerializer.Deserialize<TElement>(value)!;
+			return target.ListLeftPopAsync(key.Inner, flags).DeserializeAsync<TElement>();
 		}
 
+		#endregion
+
+		#region ListLeftPushAsync
+
 		/// <inheritdoc cref="IDatabaseAsync.ListLeftPushAsync(RedisKey, RedisValue, When, CommandFlags)"/>
-		public Task<long> LeftPushAsync(TElement item, When when = When.Always, CommandFlags flags = CommandFlags.None)
+		public static Task<long> ListLeftPushAsync<TElement>(this IDatabaseAsync target, RedisListKey<TElement> key, TElement item, When when = When.Always, CommandFlags flags = CommandFlags.None)
 		{
-			return GetDatabase().ListLeftPushAsync(Key, RedisSerializer.Serialize(item), when, flags);
+			return target.ListLeftPushAsync(key.Inner, RedisSerializer.Serialize(item), when, flags);
 		}
 
 		/// <inheritdoc cref="IDatabaseAsync.ListLeftPushAsync(RedisKey, RedisValue[], When, CommandFlags)"/>
-		public Task<long> LeftPushAsync(IEnumerable<TElement> items, When when = When.Always, CommandFlags flags = CommandFlags.None)
+		public static Task<long> ListLeftPushAsync<TElement>(this IDatabaseAsync target, RedisListKey<TElement> key, TElement[] values, When when = When.Always, CommandFlags flags = CommandFlags.None)
 		{
-			RedisValue[] values = items.Select(x => RedisSerializer.Serialize(x)).ToArray();
-			return GetDatabase().ListLeftPushAsync(Key, values, when, flags);
+			return target.ListLeftPushAsync(key.Inner, RedisSerializer.Serialize(values), when, flags);
 		}
 
+		#endregion
+
+		#region ListLengthAsync
+
 		/// <inheritdoc cref="IDatabaseAsync.ListLengthAsync(RedisKey, CommandFlags)"/>
-		public Task<long> LengthAsync()
+		public static Task<long> ListLengthAsync<TElement>(this IDatabaseAsync target, RedisListKey<TElement> key)
 		{
-			return GetDatabase().ListLengthAsync(Key);
+			return target.ListLengthAsync(key.Inner);
 		}
 
+		#endregion
+
+		#region ListRangeAsync
+
 		/// <inheritdoc cref="IDatabaseAsync.ListRangeAsync(RedisKey, Int64, Int64, CommandFlags)"/>
-		public async Task<TElement[]> RangeAsync(long start = 0, long stop = -1, CommandFlags flags = CommandFlags.None)
+		public static Task<TElement[]> ListRangeAsync<TElement>(this IDatabaseAsync target, RedisListKey<TElement> key, long start = 0, long stop = -1, CommandFlags flags = CommandFlags.None)
 		{
-			RedisValue[] values = await GetDatabase().ListRangeAsync(Key, start, stop, flags);
-			return Array.ConvertAll(values, (Converter<RedisValue, TElement>)(x => (TElement)RedisSerializer.Deserialize<TElement>(x)!));
+			return target.ListRangeAsync(key.Inner, start, stop, flags).DeserializeAsync<TElement>();
 		}
 
+		#endregion
+
+		#region ListRemoveAsync
+
 		/// <inheritdoc cref="IDatabaseAsync.ListRemoveAsync(RedisKey, RedisValue, Int64, CommandFlags)"/>
-		public Task<long> RemoveAsync(TElement item, long count = 0L, CommandFlags flags = CommandFlags.None)
+		public static Task<long> ListRemoveAsync<TElement>(this IDatabaseAsync target, RedisListKey<TElement> key, TElement value, long count = 0L, CommandFlags flags = CommandFlags.None)
 		{
-			RedisValue value = RedisSerializer.Serialize(item);
-			return GetDatabase().ListRemoveAsync(Key, value, count, flags);
+			return target.ListRemoveAsync(key.Inner, RedisSerializer.Serialize(value), count, flags);
 		}
 
+		#endregion
+
+		#region ListRightPopAsync
+
 		/// <inheritdoc cref="IDatabaseAsync.ListRightPopAsync(RedisKey, CommandFlags)"/>
-		public async Task<TElement> RightPopAsync(CommandFlags flags = CommandFlags.None)
+		public static Task<TElement> ListRightPopAsync<TElement>(this IDatabaseAsync target, RedisListKey<TElement> key, CommandFlags flags = CommandFlags.None)
 		{
-			RedisValue value = await GetDatabase().ListRightPopAsync(Key, flags);
-			return value.IsNull? default! : RedisSerializer.Deserialize<TElement>(value)!;
+			return target.ListRightPopAsync(key.Inner, flags).DeserializeAsync<TElement>();
 		}
 
+		#endregion
+
+		#region ListRightPushAsync
+
 		/// <inheritdoc cref="IDatabaseAsync.ListRightPushAsync(RedisKey, RedisValue, When, CommandFlags)"/>
-		public Task<long> RightPushAsync(TElement item, When when = When.Always, CommandFlags flags = CommandFlags.None)
+		public static Task<long> ListRightPushAsync<TElement>(this IDatabaseAsync target, RedisListKey<TElement> key, TElement item, When when = When.Always, CommandFlags flags = CommandFlags.None)
 		{
-			return GetDatabase().ListRightPushAsync(Key, RedisSerializer.Serialize(item), when, flags);
+			return target.ListRightPushAsync(key.Inner, RedisSerializer.Serialize(item), when, flags);
 		}
 
 		/// <inheritdoc cref="IDatabaseAsync.ListRightPushAsync(RedisKey, RedisValue[], When, CommandFlags)"/>
-		public Task<long> RightPushAsync(IEnumerable<TElement> items, When when = When.Always, CommandFlags flags = CommandFlags.None)
+		public static Task<long> ListRightPushAsync<TElement>(this IDatabaseAsync target, RedisListKey<TElement> key, IEnumerable<TElement> values, When when = When.Always, CommandFlags flags = CommandFlags.None)
 		{
-			RedisValue[] values = items.Select(x => RedisSerializer.Serialize(x)).ToArray();
-			return GetDatabase().ListRightPushAsync(Key, values, when, flags);
+			return target.ListRightPushAsync(key.Inner, RedisSerializer.Serialize(values), when, flags);
 		}
 
+		#endregion
+
+		#region ListSetByIndexAsync
+
 		/// <inheritdoc cref="IDatabaseAsync.ListSetByIndexAsync(RedisKey, Int64, RedisValue, CommandFlags)"/>
-		public Task SetByIndexAsync(long index, TElement item, CommandFlags flags = CommandFlags.None)
+		public static Task ListSetByIndexAsync<TElement>(IDatabaseAsync target, RedisListKey<TElement> key, long index, TElement value, CommandFlags flags = CommandFlags.None)
 		{
-			RedisValue value = RedisSerializer.Serialize(item);
-			return GetDatabase().ListSetByIndexAsync(Key, index, value, flags);
+			return target.ListSetByIndexAsync(key.Inner, index, RedisSerializer.Serialize(value), flags);
 		}
 
+		#endregion
+
+		#region ListTrimAsync
+
 		/// <inheritdoc cref="IDatabaseAsync.ListTrimAsync(RedisKey, Int64, Int64, CommandFlags)"/>
-		public Task TrimAsync(long start, long stop, CommandFlags flags = CommandFlags.None)
+		public static Task ListTrimAsync<TElement>(this IDatabaseAsync target, RedisListKey<TElement> key, long start, long stop, CommandFlags flags = CommandFlags.None)
 		{
-			return GetDatabase().ListTrimAsync(Key, start, stop, flags);
+			return target.ListTrimAsync(key.Inner, start, stop, flags);
 		}
-	}
 
-	/// <summary>
-	/// Extension methods for sets
-	/// </summary>
-	public static class RedisListExtensions
-	{
-		/// <summary>
-		/// Creates a version of this list which modifies a transaction rather than the direct DB
-		/// </summary>
-		public static RedisList<TElement> With<TElement>(this ITransaction transaction, RedisList<TElement> set)
-		{
-			return new RedisList<TElement>(transaction, set.Key);
-		}
+		#endregion
 	}
 }
diff --git a/Engine/Source/Programs/Shared/EpicGames.Redis/RedisSerializer.cs b/Engine/Source/Programs/Shared/EpicGames.Redis/RedisSerializer.cs
--- a/Engine/Source/Programs/Shared/EpicGames.Redis/RedisSerializer.cs
+++ b/Engine/Source/Programs/Shared/EpicGames.Redis/RedisSerializer.cs
@@ -1,11 +1,15 @@
 // Copyright Epic Games, Inc. All Rights Reserved.
 
+using EpicGames.Core;
 using EpicGames.Serialization;
 using StackExchange.Redis;
 using System;
 using System.Collections.Generic;
 using System.ComponentModel;
 using System.Reflection;
+using System.Reflection.Emit;
+using System.Text;
+using System.Threading.Tasks;
 
 namespace EpicGames.Redis
 {
@@ -47,7 +51,7 @@ namespace EpicGames.Redis
 		/// </summary>
 		/// <param name="value"></param>
 		/// <returns></returns>
-		T? FromRedisValue(RedisValue value);
+		T FromRedisValue(RedisValue value);
 	}
 
 	/// <summary>
@@ -63,7 +67,7 @@ namespace EpicGames.Redis
 		}
 
 		/// <inheritdoc/>
-		public T? FromRedisValue(RedisValue value)
+		public T FromRedisValue(RedisValue value)
 		{
 			return CbSerializer.Deserialize<T>(new CbField((byte[])value!));
 		}
@@ -84,7 +88,7 @@ namespace EpicGames.Redis
 			}
 
 			public RedisValue ToRedisValue(T value) => (string?)_typeConverter.ConvertTo(value, typeof(string));
-			public T? FromRedisValue(RedisValue value) => (T?)_typeConverter.ConvertFrom((string)value!);
+			public T FromRedisValue(RedisValue value) => (T)_typeConverter.ConvertFrom((string)value!)!;
 		}
 
 		class RedisNativeConverter<T> : IRedisConverter<T>
@@ -223,15 +227,87 @@ namespace EpicGames.Redis
 			return CachedConverter<T>.Converter.ToRedisValue(value);
 		}
 
+		/// <summary>
+		/// Serialize an object to a <see cref="RedisValue"/>
+		/// </summary>
+		/// <typeparam name="T"></typeparam>
+		/// <param name="inputs"></param>
+		/// <returns></returns>
+		public static RedisValue[] Serialize<T>(T[] inputs)
+		{
+			RedisValue[] outputs = new RedisValue[inputs.Length];
+
+			for (int idx = 0; idx < inputs.Length; idx++)
+			{
+				outputs[idx] = Serialize<T>(inputs[idx]);
+			}
+
+			return outputs;
+		}
+
 		/// <summary>
 		/// Deserialize a <see cref="RedisValue"/>
 		/// </summary>
 		/// <typeparam name="T"></typeparam>
 		/// <param name="value"></param>
 		/// <returns></returns>
-		public static T? Deserialize<T>(RedisValue value)
+		public static T Deserialize<T>(RedisValue value)
 		{
 			return CachedConverter<T>.Converter.FromRedisValue(value);
 		}
+
+		/// <summary>
+		/// Deserialize an array of <see cref="RedisValue"/> objects
+		/// </summary>
+		/// <typeparam name="T"></typeparam>
+		/// <param name="inputs"></param>
+		/// <returns></returns>
+		public static T[] Deserialize<T>(RedisValue[] inputs)
+		{
+			T[] outputs = new T[inputs.Length];
+
+			for (int idx = 0; idx < inputs.Length; idx++)
+			{
+				outputs[idx] = Deserialize<T>(inputs[idx]);
+			}
+
+			return outputs;
+		}
+	}
+
+	/// <summary>
+	/// Extension methods for serialization
+	/// </summary>
+	public static class RedisSerializerExtensions
+	{
+		/// <summary>
+		/// Deserialize a <see cref="RedisValue"/>
+		/// </summary>
+		/// <typeparam name="T"></typeparam>
+		/// <param name="value"></param>
+		/// <returns></returns>
+		public static async Task<T> DeserializeAsync<T>(this Task<RedisValue> value)
+		{
+			RedisValue result = await value;
+			if (result.IsNull)
+			{
+				return default!;
+			}
+			else
+			{
+				return RedisSerializer.Deserialize<T>(await value);
+			}
+		}
+
+		/// <summary>
+		/// Deserialize a <see cref="RedisValue"/>
+		/// </summary>
+		/// <typeparam name="T"></typeparam>
+		/// <param name="values"></param>
+		/// <returns></returns>
+		public static async Task<T[]> DeserializeAsync<T>(this Task<RedisValue[]> values)
+		{
+			return RedisSerializer.Deserialize<T>(await values);
+		}
 	}
 }
diff --git a/Engine/Source/Programs/Shared/EpicGames.Redis/RedisSet.cs b/Engine/Source/Programs/Shared/EpicGames.Redis/RedisSet.cs
--- a/Engine/Source/Programs/Shared/EpicGames.Redis/RedisSet.cs
+++ b/Engine/Source/Programs/Shared/EpicGames.Redis/RedisSet.cs
@@ -9,125 +9,119 @@ using System.Threading.Tasks;
 namespace EpicGames.Redis
 {
 	/// <summary>
-	/// Represents a typed Redis list with a given key
+	/// Represents a typed Redis set with a given key
 	/// </summary>
-	/// <typeparam name="TElement">The type of element stored in the list</typeparam>
-	public readonly struct RedisSet<TElement>
+	/// <typeparam name="TElement">The type of element stored in the set</typeparam>
+	public readonly struct RedisSetKey<TElement>
 	{
-		internal readonly RedisConnectionPool? ConnectionPool = null;
-		internal readonly IDatabaseAsync? Database = null;
-
 		/// <summary>
 		/// The key for the list
 		/// </summary>
-		public readonly RedisKey Key { get; }
+		public readonly RedisKey Inner { get; }
 
 		/// <summary>
 		/// Constructor
 		/// </summary>
-		/// <param name="database">Explicitly set Redis database (overrides use of any connection pool)</param>
-		/// <param name="key">Redis key this type is using</param>
-		public RedisSet(IDatabaseAsync database, RedisKey key)
+		/// <param name="inner">Redis key this type is using</param>
+		public RedisSetKey(RedisKey inner)
 		{
-			Database = database;
-			Key = key;
+			Inner = inner;
 		}
-		
+
 		/// <summary>
-		/// Constructor
+		/// Implicit conversion to typed redis key.
 		/// </summary>
-		/// <param name="connectionPool">Connection pool for Redis databases</param>
-		/// <param name="key">Redis key this type is using</param>
-		public RedisSet(RedisConnectionPool connectionPool, RedisKey key)
-		{
-			ConnectionPool = connectionPool;
-			Key = key;
-		}
-		
+		/// <param name="key">Key to convert</param>
+		public static implicit operator RedisSetKey<TElement>(string key) => new RedisSetKey<TElement>(new RedisKey(key));
+
 		/// <summary>
-		/// Get the Redis database in use
+		/// Implicit conversion to untyped redis keys.
 		/// </summary>
-		/// <returns>A connection pool or explicitly set Redis database</returns>
-		/// <exception cref="InvalidOperationException">If neither are set</exception>
-		public IDatabaseAsync GetDatabase()
-		{
-			if (Database != null) return Database;
-			if (ConnectionPool != null) return ConnectionPool.GetDatabase();
-			throw new InvalidOperationException($"Neither {nameof(Database)} or {nameof(ConnectionPool)} has been set!");
-		}
+		/// <param name="key">Key to convert</param>
+		public static implicit operator TypedRedisKey(RedisSetKey<TElement> key) => key.Inner;
+	}
+
+	/// <summary>
+	/// Extension methods for sets
+	/// </summary>
+	public static class RedisSetExtensions
+	{
+		#region SetAddAsync
 
 		/// <inheritdoc cref="IDatabaseAsync.SetAddAsync(RedisKey, RedisValue, CommandFlags)"/>
-		public Task<bool> AddAsync(TElement item, CommandFlags flags = CommandFlags.None)
+		public static Task<bool> SetAddAsync<TElement>(this IDatabaseAsync target, RedisSetKey<TElement> key, TElement item, CommandFlags flags = CommandFlags.None)
 		{
-			return GetDatabase().SetAddAsync(Key, RedisSerializer.Serialize(item), flags);
+			return target.SetAddAsync(key.Inner, RedisSerializer.Serialize(item), flags);
 		}
 
 		/// <inheritdoc cref="IDatabaseAsync.SetAddAsync(RedisKey, RedisValue[], CommandFlags)"/>
-		public Task<long> AddAsync(IEnumerable<TElement> items, CommandFlags flags = CommandFlags.None)
+		public static Task<long> SetAddAsync<TElement>(this IDatabaseAsync target, RedisSetKey<TElement> key, TElement[] values, CommandFlags flags = CommandFlags.None)
 		{
-			RedisValue[] values = items.Select(x => RedisSerializer.Serialize(x)).ToArray();
-			return GetDatabase().SetAddAsync(Key, values, flags);
+			return target.SetAddAsync(key.Inner, RedisSerializer.Serialize(values), flags);
 		}
 
+		#endregion
+
+		#region SetContainsAsync
+
 		/// <inheritdoc cref="IDatabaseAsync.SetContainsAsync(RedisKey, RedisValue, CommandFlags)"/>
-		public Task<bool> ContainsAsync(TElement item, CommandFlags flags = CommandFlags.None)
+		public static Task<bool> SetContainsAsync<TElement>(this IDatabaseAsync target, RedisSetKey<TElement> key, TElement value, CommandFlags flags = CommandFlags.None)
 		{
-			RedisValue value = RedisSerializer.Serialize(item);
-			return GetDatabase().SetContainsAsync(Key, value, flags);
+			return target.SetContainsAsync(key.Inner, RedisSerializer.Serialize(value), flags);
 		}
 
+		#endregion
+
+		#region SetLengthAsync
+
 		/// <inheritdoc cref="IDatabaseAsync.SetLengthAsync(RedisKey, CommandFlags)"/>
-		public Task<long> LengthAsync(CommandFlags flags = CommandFlags.None)
+		public static Task<long> SetLengthAsync<TElement>(this IDatabaseAsync target, RedisSetKey<TElement> key, CommandFlags flags = CommandFlags.None)
 		{
-			return GetDatabase().SetLengthAsync(Key, flags);
+			return target.SetLengthAsync(key.Inner, flags);
 		}
 
-		/// <inheritdoc cref="IDatabaseAsync.SetRemoveAsync(RedisKey, RedisValue, CommandFlags)"/>
-		public async Task<TElement[]> MembersAsync(CommandFlags flags = CommandFlags.None)
+		#endregion
+
+		#region SetMembersAsync
+
+		/// <inheritdoc cref="IDatabaseAsync.SetMembersAsync(RedisKey, CommandFlags)"/>
+		public static Task<TElement[]> SetMembersAsync<TElement>(this IDatabaseAsync target, RedisSetKey<TElement> key, CommandFlags flags = CommandFlags.None)
 		{
-			RedisValue[] values = await GetDatabase().SetMembersAsync(Key, flags);
-			return Array.ConvertAll(values, (Converter<RedisValue, TElement>)(x => RedisSerializer.Deserialize<TElement>(x)!));
+			return target.SetMembersAsync(key.Inner, flags).DeserializeAsync<TElement>();
 		}
 
+		#endregion
+
+		#region SetPopAsync
+
 		/// <inheritdoc cref="IDatabaseAsync.SetPopAsync(RedisKey, CommandFlags)"/>
-		public async Task<TElement> PopAsync(CommandFlags flags = CommandFlags.None)
+		public static Task<TElement> SetPopAsync<TElement>(this IDatabaseAsync target, RedisSetKey<TElement> key, CommandFlags flags = CommandFlags.None)
 		{
-			RedisValue value = await GetDatabase().SetPopAsync(Key, flags);
-			return value.IsNull ? default! : RedisSerializer.Deserialize<TElement>(value)!;
+			return target.SetPopAsync(key.Inner, flags).DeserializeAsync<TElement>();
 		}
 
 		/// <inheritdoc cref="IDatabaseAsync.SetPopAsync(RedisKey, Int64, CommandFlags)"/>
-		public async Task<TElement[]> PopAsync(long count, CommandFlags flags = CommandFlags.None)
+		public static Task<TElement[]> SetPopAsync<TElement>(this IDatabaseAsync target, RedisSetKey<TElement> key, long count, CommandFlags flags = CommandFlags.None)
 		{
-			RedisValue[] values = await GetDatabase().SetPopAsync(Key, count, flags);
-			return Array.ConvertAll<RedisValue, TElement>(values, x => RedisSerializer.Deserialize<TElement>(x)!);
+			return target.SetPopAsync(key.Inner, count, flags).DeserializeAsync<TElement>();
 		}
 
+		#endregion
+
+		#region SetRemoveAsync
+
 		/// <inheritdoc cref="IDatabaseAsync.SetRemoveAsync(RedisKey, RedisValue, CommandFlags)"/>
-		public Task<bool> RemoveAsync(TElement item, CommandFlags flags = CommandFlags.None)
+		public static Task<bool> SetRemoveAsync<TElement>(this IDatabaseAsync target, RedisSetKey<TElement> key, TElement item, CommandFlags flags = CommandFlags.None)
 		{
-			return GetDatabase().SetRemoveAsync(Key, RedisSerializer.Serialize(item), flags);
+			return target.SetRemoveAsync(key.Inner, RedisSerializer.Serialize(item), flags);
 		}
 
 		/// <inheritdoc cref="IDatabaseAsync.SetRemoveAsync(RedisKey, RedisValue[], CommandFlags)"/>
-		public Task<long> RemoveAsync(IEnumerable<TElement> items, CommandFlags flags = CommandFlags.None)
+		public static Task<long> SetRemoveAsync<TElement>(this IDatabaseAsync target, RedisSetKey<TElement> key, TElement[] values, CommandFlags flags = CommandFlags.None)
 		{
-			RedisValue[] values = items.Select(x => RedisSerializer.Serialize(x)).ToArray();
-			return GetDatabase().SetRemoveAsync(Key, values, flags);
+			return target.SetRemoveAsync(key.Inner, RedisSerializer.Serialize(values), flags);
 		}
-	}
 
-	/// <summary>
-	/// Extension methods for sets
-	/// </summary>
-	public static class RedisSetExtensions
-	{
-		/// <summary>
-		/// Creates a version of this set which modifies a transaction rather than the direct DB
-		/// </summary>
-		public static RedisSet<TElement> With<TElement>(this ITransaction transaction, RedisSet<TElement> set)
-		{
-			return new RedisSet<TElement>(transaction, set.Key);
-		}
+		#endregion
 	}
 }
diff --git a/Engine/Source/Programs/Shared/EpicGames.Redis/RedisSortedSet.cs b/Engine/Source/Programs/Shared/EpicGames.Redis/RedisSortedSet.cs
--- a/Engine/Source/Programs/Shared/EpicGames.Redis/RedisSortedSet.cs
+++ b/Engine/Source/Programs/Shared/EpicGames.Redis/RedisSortedSet.cs
@@ -1,5 +1,6 @@
 // Copyright Epic Games, Inc. All Rights Reserved.
 
+using EpicGames.Core;
 using StackExchange.Redis;
 using System;
 using System.Collections.Generic;
@@ -8,6 +9,39 @@ using System.Threading.Tasks;
 
 namespace EpicGames.Redis
 {
+	/// <summary>
+	/// Represents a typed Redis sorted set with a given key
+	/// </summary>
+	/// <typeparam name="TElement">The type of element stored in the set</typeparam>
+	public readonly struct RedisSortedSetKey<TElement>
+	{
+		/// <summary>
+		/// The key for the list
+		/// </summary>
+		public readonly RedisKey Inner { get; }
+
+		/// <summary>
+		/// Constructor
+		/// </summary>
+		/// <param name="inner">Redis key this type is using</param>
+		public RedisSortedSetKey(RedisKey inner)
+		{
+			Inner = inner;
+		}
+
+		/// <summary>
+		/// Implicit conversion to typed redis key.
+		/// </summary>
+		/// <param name="key">Key to convert</param>
+		public static implicit operator RedisSortedSetKey<TElement>(string key) => new RedisSortedSetKey<TElement>(new RedisKey(key));
+
+		/// <summary>
+		/// Implicit conversion to untyped redis keys.
+		/// </summary>
+		/// <param name="key">Key to convert</param>
+		public static implicit operator TypedRedisKey(RedisSortedSetKey<TElement> key) => key.Inner;
+	}
+
 	/// <summary>
 	/// Typed implementation of <see cref="SortedSetEntry"/>
 	/// </summary>
@@ -65,185 +99,175 @@ namespace EpicGames.Redis
 	}
 
 	/// <summary>
-	/// Represents a typed Redis list with a given key
+	/// Extension methods for sets
 	/// </summary>
-	/// <typeparam name="TElement">The type of element stored in the list</typeparam>
-	public readonly struct RedisSortedSet<TElement>
+	public static class RedisSortedSetExtensions
 	{
-		internal readonly RedisConnectionPool? ConnectionPool = null;
-		internal readonly IDatabaseAsync? Database = null;
-
-		/// <summary>
-		/// The key for the list
-		/// </summary>
-		public readonly RedisKey Key { get; }
-
-		/// <summary>
-		/// Constructor
-		/// </summary>
-		/// <param name="database">Explicitly set Redis database (overrides use of any connection pool)</param>
-		/// <param name="key">Redis key this type is using</param>
-		public RedisSortedSet(IDatabaseAsync database, RedisKey key)
-		{
-			Database = database;
-			Key = key;
-		}
-		
-		/// <summary>
-		/// Constructor
-		/// </summary>
-		/// <param name="connectionPool">Connection pool for Redis databases</param>
-		/// <param name="key">Redis key this type is using</param>
-		public RedisSortedSet(RedisConnectionPool connectionPool, RedisKey key)
-		{
-			ConnectionPool = connectionPool;
-			Key = key;
-		}
-		
-		/// <summary>
-		/// Get the Redis database in use
-		/// </summary>
-		/// <returns>A connection pool or explicitly set Redis database</returns>
-		/// <exception cref="InvalidOperationException">If neither are set</exception>
-		public IDatabaseAsync GetDatabase()
-		{
-			if (Database != null) return Database;
-			if (ConnectionPool != null) return ConnectionPool.GetDatabase();
-			throw new InvalidOperationException($"Neither {nameof(Database)} or {nameof(ConnectionPool)} has been set!");
-		}
+		#region SortedSetAddAsync
 
 		/// <inheritdoc cref="IDatabaseAsync.SortedSetAddAsync(RedisKey, RedisValue, Double, CommandFlags)"/>
-		public Task<bool> AddAsync(TElement item, double score, When when = When.Always, CommandFlags flags = CommandFlags.None)
+		public static Task<bool> SortedSetAddAsync<TElement>(this IDatabaseAsync target, RedisSortedSetKey<TElement> key, TElement value, double score, When when = When.Always, CommandFlags flags = CommandFlags.None)
 		{
-			RedisValue value = RedisSerializer.Serialize(item);
-			return GetDatabase().SortedSetAddAsync(Key, value, score, when, flags);
+			return target.SortedSetAddAsync(key.Inner, RedisSerializer.Serialize(value), score, when, flags);
 		}
 
 		/// <inheritdoc cref="IDatabaseAsync.SortedSetAddAsync(RedisKey, RedisValue, Double, When, CommandFlags)"/>
-		public Task<long> AddAsync(SortedSetEntry<TElement>[] values, When when = When.Always, CommandFlags flags = CommandFlags.None)
+		public static Task<long> SortedSetAddAsync<TElement>(this IDatabaseAsync target, RedisSortedSetKey<TElement> key, SortedSetEntry<TElement>[] values, When when = When.Always, CommandFlags flags = CommandFlags.None)
 		{
-			SortedSetEntry[] untyped = Array.ConvertAll(values, x => new SortedSetEntry(x.ElementValue, x.Score));
-			return GetDatabase().SortedSetAddAsync(Key, untyped, when, flags);
+			return target.SortedSetAddAsync(key.Inner, values.ConvertAll(x => new SortedSetEntry(x.ElementValue, x.Score)), when, flags);
 		}
 
+		#endregion
+
+		#region SortedSetLengthAsync
+
 		/// <inheritdoc cref="IDatabaseAsync.SortedSetLengthAsync(RedisKey, Double, Double, Exclude, CommandFlags)"/>
-		public Task<long> LengthAsync(double min = Double.NegativeInfinity, double max = Double.PositiveInfinity, Exclude exclude = Exclude.None, CommandFlags flags = CommandFlags.None)
+		public static Task<long> SortedSetLengthAsync<TElement>(this IDatabaseAsync target, RedisSortedSetKey<TElement> key, double min = Double.NegativeInfinity, double max = Double.PositiveInfinity, Exclude exclude = Exclude.None, CommandFlags flags = CommandFlags.None)
 		{
-			return GetDatabase().SortedSetLengthAsync(Key, min, max, exclude, flags);
+			return target.SortedSetLengthAsync(key.Inner, min, max, exclude, flags);
 		}
 
+		#endregion
+
+		#region SortedSetLengthByValueAsync
+
 		/// <inheritdoc cref="IDatabaseAsync.SortedSetLengthByValueAsync(RedisKey, RedisValue, RedisValue, Exclude, CommandFlags)"/>
-		public Task<long> LengthByValueAsync(TElement min, TElement max, Exclude exclude = Exclude.None, CommandFlags flags = CommandFlags.None)
+		public static Task<long> SortedSetLengthByValueAsync<TElement>(this IDatabaseAsync target, RedisSortedSetKey<TElement> key, TElement min, TElement max, Exclude exclude = Exclude.None, CommandFlags flags = CommandFlags.None)
 		{
-			RedisValue minValue = RedisSerializer.Serialize(min);
-			RedisValue maxValue = RedisSerializer.Serialize(max);
-			return GetDatabase().SortedSetLengthByValueAsync(Key, minValue, maxValue, exclude, flags);
+			return target.SortedSetLengthByValueAsync(key.Inner, RedisSerializer.Serialize(min), RedisSerializer.Serialize(max), exclude, flags);
 		}
 
+		#endregion
+
+		#region SortedSetRangeByRankAsync
+
 		/// <inheritdoc cref="IDatabaseAsync.SortedSetRangeByRankAsync(RedisKey, Int64, Int64, Order, CommandFlags)"/>
-		public async Task<TElement[]> RangeByRankAsync(long start, long stop = -1, Order order = Order.Ascending, CommandFlags flags = CommandFlags.None)
+		public static Task<TElement[]> SortedSetRangeByRankAsync<TElement>(this IDatabaseAsync target, RedisSortedSetKey<TElement> key, long start, long stop = -1, Order order = Order.Ascending, CommandFlags flags = CommandFlags.None)
 		{
-			RedisValue[] values = await GetDatabase().SortedSetRangeByRankAsync(Key, start, stop, order, flags);
-			return Array.ConvertAll(values, x => RedisSerializer.Deserialize<TElement>(x)!);
+			return target.SortedSetRangeByRankAsync(key.Inner, start, stop, order, flags).DeserializeAsync<TElement>();
 		}
 
+		#endregion
+
+		#region SortedSetRangeByRankWithScoresAsync
+
 		/// <inheritdoc cref="IDatabaseAsync.SortedSetRangeByRankWithScoresAsync(RedisKey, Int64, Int64, Order, CommandFlags)"/>
-		public async Task<SortedSetEntry<TElement>[]> RangeByRankWithScoresAsync(long start, long stop = -1, Order order = Order.Ascending, CommandFlags flags = CommandFlags.None)
+		public static async Task<SortedSetEntry<TElement>[]> SortedSetRangeByRankWithScoresAsync<TElement>(this IDatabaseAsync target, RedisSortedSetKey<TElement> key, long start, long stop = -1, Order order = Order.Ascending, CommandFlags flags = CommandFlags.None)
 		{
-			SortedSetEntry[] values = await GetDatabase().SortedSetRangeByRankWithScoresAsync(Key, start, stop, order, flags);
+			SortedSetEntry[] values = await target.SortedSetRangeByRankWithScoresAsync(key.Inner, start, stop, order, flags);
 			return Array.ConvertAll(values, x => new SortedSetEntry<TElement>(x));
 		}
 
+		#endregion
+
+		#region SortedSetRangeByScoreAsync
+
 		/// <inheritdoc cref="IDatabaseAsync.SortedSetRangeByScoreAsync(RedisKey, Double, Double, Exclude, Order, Int64, Int64, CommandFlags)"/>
-		public async Task<TElement[]> RangeByScoreAsync(double start = Double.NegativeInfinity, double stop = Double.PositiveInfinity, Exclude exclude = Exclude.None, Order order = Order.Ascending, long skip = 0L, long take = -1L, CommandFlags flags = CommandFlags.None)
+		public static Task<TElement[]> SortedSetRangeByScoreAsync<TElement>(this IDatabaseAsync target, RedisSortedSetKey<TElement> key, double start = Double.NegativeInfinity, double stop = Double.PositiveInfinity, Exclude exclude = Exclude.None, Order order = Order.Ascending, long skip = 0L, long take = -1L, CommandFlags flags = CommandFlags.None)
 		{
-			RedisValue[] values = await GetDatabase().SortedSetRangeByScoreAsync(Key, start, stop, exclude, order, skip, take, flags);
-			return Array.ConvertAll(values, x => RedisSerializer.Deserialize<TElement>(x)!);
+			return target.SortedSetRangeByScoreAsync(key.Inner, start, stop, exclude, order, skip, take, flags).DeserializeAsync<TElement>();
 		}
 
+		#endregion
+
+		#region SortedSetRangeByScoreWithScoresAsync
+
 		/// <inheritdoc cref="IDatabaseAsync.SortedSetRangeByScoreWithScoresAsync(RedisKey, Double, Double, Exclude, Order, Int64, Int64, CommandFlags)"/>
-		public async Task<SortedSetEntry<TElement>[]> RangeByScoreWithScoresAsync(double start = Double.NegativeInfinity, double stop = Double.PositiveInfinity, Exclude exclude = Exclude.None, Order order = Order.Ascending, long skip = 0L, long take = -1L, CommandFlags flags = CommandFlags.None)
+		public static async Task<SortedSetEntry<TElement>[]> SortedSetRangeByScoreWithScoresAsync<TElement>(this IDatabaseAsync target, RedisSortedSetKey<TElement> key, double start = Double.NegativeInfinity, double stop = Double.PositiveInfinity, Exclude exclude = Exclude.None, Order order = Order.Ascending, long skip = 0L, long take = -1L, CommandFlags flags = CommandFlags.None)
 		{
-			SortedSetEntry[] values = await GetDatabase().SortedSetRangeByScoreWithScoresAsync(Key, start, stop, exclude, order, skip, take, flags);
+			SortedSetEntry[] values = await target.SortedSetRangeByScoreWithScoresAsync(key.Inner, start, stop, exclude, order, skip, take, flags);
 			return Array.ConvertAll(values, x => new SortedSetEntry<TElement>(x));
 		}
 
+		#endregion
+
+		#region SortedSetRangeByValueAsync
+
 		/// <inheritdoc cref="IDatabaseAsync.SortedSetRangeByValueAsync(RedisKey, RedisValue, RedisValue, Exclude, Order, Int64, Int64, CommandFlags)"/>
-		public async Task<TElement[]> RangeByValueAsync(TElement min, TElement max, Exclude exclude = Exclude.None, Order order = Order.Ascending, long skip = 0L, long take = -1L, CommandFlags flags = CommandFlags.None)
+		public static Task<TElement[]> SortedSetRangeByValueAsync<TElement>(this IDatabaseAsync target, RedisSortedSetKey<TElement> key, TElement min, TElement max, Exclude exclude = Exclude.None, Order order = Order.Ascending, long skip = 0L, long take = -1L, CommandFlags flags = CommandFlags.None)
 		{
-			RedisValue minValue = RedisSerializer.Serialize(min);
-			RedisValue maxValue = RedisSerializer.Serialize(max);
-			RedisValue[] values = await GetDatabase().SortedSetRangeByValueAsync(Key, minValue, maxValue, exclude, order, skip, take, flags);
-			return Array.ConvertAll(values, x => RedisSerializer.Deserialize<TElement>(x)!);
+			return target.SortedSetRangeByValueAsync(key.Inner, RedisSerializer.Serialize(min), RedisSerializer.Serialize(max), exclude, order, skip, take, flags).DeserializeAsync<TElement>();
 		}
 
+		#endregion
+
+		#region SortedSetRankAsync
+
 		/// <inheritdoc cref="IDatabaseAsync.SortedSetRankAsync(RedisKey, RedisValue, Order, CommandFlags)"/>
-		public Task<long?> RankAsync(TElement item, Order order = Order.Ascending, CommandFlags flags = CommandFlags.None)
+		public static Task<long?> SortedSetRankAsync<TElement>(this IDatabaseAsync target, RedisSortedSetKey<TElement> key, TElement item, Order order = Order.Ascending, CommandFlags flags = CommandFlags.None)
 		{
-			return GetDatabase().SortedSetRankAsync(Key, RedisSerializer.Serialize(item), order, flags);
+			return target.SortedSetRankAsync(key.Inner, RedisSerializer.Serialize(item), order, flags);
 		}
 
+		#endregion
+
+		#region 
+
 		/// <inheritdoc cref="IDatabaseAsync.SortedSetRemoveAsync(RedisKey, RedisValue, CommandFlags)"/>
-		public Task<bool> RemoveAsync(TElement value, CommandFlags flags = CommandFlags.None)
+		public static Task<bool> SortedSetRemoveAsync<TElement>(this IDatabaseAsync target, RedisSortedSetKey<TElement> key, TElement value, CommandFlags flags = CommandFlags.None)
 		{
-			return GetDatabase().SortedSetRemoveAsync(Key, RedisSerializer.Serialize(value), flags);
+			return target.SortedSetRemoveAsync(key.Inner, RedisSerializer.Serialize(value), flags);
 		}
 
 		/// <inheritdoc cref="IDatabaseAsync.SortedSetRemoveAsync(RedisKey, RedisValue[], CommandFlags)"/>
-		public Task<long> RemoveAsync(IEnumerable<TElement> items, CommandFlags flags = CommandFlags.None)
+		public static Task<long> SortedSetRemoveAsync<TElement>(this IDatabaseAsync target, RedisSortedSetKey<TElement> key, TElement[] values, CommandFlags flags = CommandFlags.None)
 		{
-			RedisValue[] values = items.Select(x => RedisSerializer.Serialize(x)).ToArray();
-			return GetDatabase().SortedSetRemoveAsync(Key, values, flags);
+			return target.SortedSetRemoveAsync(key.Inner, RedisSerializer.Serialize(values), flags);
 		}
 
+		#endregion
+
+		#region SortedSetRemoveRangeByRankAsync
+
 		/// <inheritdoc cref="IDatabaseAsync.SortedSetRemoveRangeByRankAsync(RedisKey, Int64, Int64, CommandFlags)"/>
-		public Task<long> RemoveRangeByRankAsync(long start, long stop, CommandFlags flags = CommandFlags.None)
+		public static Task<long> SortedSetRemoveRangeByRankAsync<TElement>(this IDatabaseAsync target, RedisSortedSetKey<TElement> key, long start, long stop, CommandFlags flags = CommandFlags.None)
 		{
-			return GetDatabase().SortedSetRemoveRangeByRankAsync(Key, start, stop, flags);
+			return target.SortedSetRemoveRangeByRankAsync(key.Inner, start, stop, flags);
 		}
 
+		#endregion
+
+		#region SortedSetRemoveRangeByScoreAsync
+
 		/// <inheritdoc cref="IDatabaseAsync.SortedSetRemoveRangeByScoreAsync(RedisKey, Double, Double, Exclude, CommandFlags)"/>
-		public Task<long> RemoveRangeByScoreAsync(double start, double stop, Exclude exclude = Exclude.None, CommandFlags flags = CommandFlags.None)
+		public static Task<long> SortedSetRemoveRangeByScoreAsync<TElement>(this IDatabaseAsync target, RedisSortedSetKey<TElement> key, double start, double stop, Exclude exclude = Exclude.None, CommandFlags flags = CommandFlags.None)
 		{
-			return GetDatabase().SortedSetRemoveRangeByScoreAsync(Key, start, stop, exclude, flags);
+			return target.SortedSetRemoveRangeByScoreAsync(key.Inner, start, stop, exclude, flags);
 		}
 
+		#endregion
+
+		#region SortedSetRemoveRangeByValueAsync
+
 		/// <inheritdoc cref="IDatabaseAsync.SortedSetRemoveRangeByValueAsync(RedisKey, RedisValue, RedisValue, Exclude, CommandFlags)"/>
-		public async Task<long> RemoveRangeByValueAsync(TElement min, TElement max, Exclude exclude = Exclude.None, CommandFlags flags = CommandFlags.None)
+		public static Task<long> SortedSetRemoveRangeByValueAsync<TElement>(this IDatabaseAsync target, RedisSortedSetKey<TElement> key, TElement min, TElement max, Exclude exclude = Exclude.None, CommandFlags flags = CommandFlags.None)
 		{
-			RedisValue minValue = RedisSerializer.Serialize(min);
-			RedisValue maxValue = RedisSerializer.Serialize(max);
-			return await GetDatabase().SortedSetRemoveRangeByValueAsync(Key, minValue, maxValue, exclude, flags);
+			return target.SortedSetRemoveRangeByValueAsync(key.Inner, RedisSerializer.Serialize(min), RedisSerializer.Serialize(max), exclude, flags);
 		}
 
-		/// <inheritdoc cref="IDatabaseAsync.SortedSetAddAsync(RedisKey, RedisValue, Double, CommandFlags)"/>
-		public async IAsyncEnumerable<SortedSetEntry<TElement>> ScanAsync(RedisValue pattern = default, int pageSize = 250, long cursor = 0, int pageOffset = 0, CommandFlags flags = CommandFlags.None)
+		#endregion
+
+		#region SortedSetScanAsync
+
+		/// <inheritdoc cref="IDatabaseAsync.SortedSetScanAsync(RedisKey, RedisValue, Int32, Int64, Int32, CommandFlags)"/>
+		public static async IAsyncEnumerable<SortedSetEntry<TElement>> SortedSetScanAsync<TElement>(this IDatabaseAsync target, RedisSortedSetKey<TElement> key, RedisValue pattern = default, int pageSize = 250, long cursor = 0, int pageOffset = 0, CommandFlags flags = CommandFlags.None)
 		{
-			await foreach (SortedSetEntry entry in GetDatabase().SortedSetScanAsync(Key, pattern, pageSize, cursor, pageOffset, flags))
+			await foreach (SortedSetEntry entry in target.SortedSetScanAsync(key.Inner, pattern, pageSize, cursor, pageOffset, flags))
 			{
 				yield return new SortedSetEntry<TElement>(entry);
 			}
 		}
 
+		#endregion
+
+		#region SortedSetScoreAsync
+
 		/// <inheritdoc cref="IDatabaseAsync.SortedSetScoreAsync(RedisKey, RedisValue, CommandFlags)"/>
-		public Task<double?> ScoreAsync(TElement member, CommandFlags flags = CommandFlags.None)
+		public static Task<double?> SortedSetScoreAsync<TElement>(this IDatabaseAsync target, RedisSortedSetKey<TElement> key, TElement member, CommandFlags flags = CommandFlags.None)
 		{
-			RedisValue value = RedisSerializer.Serialize(member);
-			return GetDatabase().SortedSetScoreAsync(Key, value, flags);
+			return target.SortedSetScoreAsync(key.Inner, RedisSerializer.Serialize(member), flags);
 		}
-	}
 
-	/// <summary>
-	/// Extension methods for sets
-	/// </summary>
-	public static class RedisSortedSetExtensions
-	{
-		/// <summary>
-		/// Creates a version of this set which modifies a transaction rather than the direct DB
-		/// </summary>
-		public static RedisSortedSet<TElement> With<TElement>(this ITransaction transaction, RedisSortedSet<TElement> set)
-		{
-			return new RedisSortedSet<TElement>(transaction, set.Key);
-		}
+		#endregion
 	}
 }
diff --git a/Engine/Source/Programs/Shared/EpicGames.Redis/RedisString.cs b/Engine/Source/Programs/Shared/EpicGames.Redis/RedisString.cs
--- a/Engine/Source/Programs/Shared/EpicGames.Redis/RedisString.cs
+++ b/Engine/Source/Programs/Shared/EpicGames.Redis/RedisString.cs
@@ -1,70 +1,44 @@
 // Copyright Epic Games, Inc. All Rights Reserved.
 
+using EpicGames.Core;
 using StackExchange.Redis;
 using System;
+using System.Collections.Generic;
 using System.Threading.Tasks;
 
 namespace EpicGames.Redis
 {
 	/// <summary>
-	/// Represents a typed Redis list with a given key
+	/// Represents a typed Redis string with a given value type
 	/// </summary>
-	/// <typeparam name="TElement">The type of element stored in the list</typeparam>
-	public readonly struct RedisString<TElement>
+	/// <typeparam name="TValue">The type of element stored in the list</typeparam>
+	public readonly struct RedisStringKey<TValue>
 	{
-		internal readonly RedisConnectionPool? ConnectionPool = null;
-		internal readonly IDatabaseAsync? Database = null;
-
 		/// <summary>
-		/// The key for the list
+		/// The untyped key for the string
 		/// </summary>
-		public readonly RedisKey Key { get; }
+		public readonly RedisKey Inner { get; }
 
 		/// <summary>
 		/// Constructor
 		/// </summary>
-		/// <param name="database"></param>
-		/// <param name="key"></param>
-		public RedisString(IDatabaseAsync database, RedisKey key)
+		/// <param name="inner"></param>
+		public RedisStringKey(RedisKey inner)
 		{
-			Database = database;
-			Key = key;
+			Inner = inner;
 		}
-		
+
 		/// <summary>
-		/// Constructor
+		/// Implicit conversion to typed redis key.
 		/// </summary>
-		/// <param name="connectionPool"></param>
-		/// <param name="key"></param>
-		public RedisString(RedisConnectionPool connectionPool, RedisKey key)
-		{
-			ConnectionPool = connectionPool;
-			Key = key;
-		}
+		/// <param name="key">Key to convert</param>
+		public static implicit operator RedisStringKey<TValue>(string key) => new RedisStringKey<TValue>(new RedisKey(key));
 
 		/// <summary>
-		/// Get the Redis database in use
+		/// Implicit conversion to a regular RedisKey
 		/// </summary>
-		/// <returns>A connection pool or explicitly set Redis database</returns>
-		/// <exception cref="InvalidOperationException">If neither are set</exception>
-		public IDatabaseAsync GetDatabase()
-		{
-			if (Database != null) return Database;
-			if (ConnectionPool != null) return ConnectionPool.GetDatabase();
-			throw new InvalidOperationException($"Neither {nameof(Database)} or {nameof(ConnectionPool)} has been set!");
-		}
-
-		/// <inheritdoc cref="IDatabaseAsync.StringLengthAsync(RedisKey, CommandFlags)"/>
-		public Task<long> LengthAsync(CommandFlags flags = CommandFlags.None)
-		{
-			return GetDatabase().StringLengthAsync(Key, flags);
-		}
-
-		/// <inheritdoc cref="IDatabaseAsync.StringSetAsync(RedisKey, RedisValue, TimeSpan?, When, CommandFlags)"/>
-		public Task<bool> SetAsync(TElement value, TimeSpan? expiry = null, When when = When.Always, CommandFlags flags = CommandFlags.None)
-		{
-			return GetDatabase().StringSetAsync(Key, RedisSerializer.Serialize(value), expiry, when, flags);
-		}
+		/// <param name="key">The key to convert</param>
+		public static implicit operator TypedRedisKey(RedisStringKey<TValue> key) => key.Inner;
 	}
 
 	/// <summary>
@@ -72,64 +46,106 @@ namespace EpicGames.Redis
 	/// </summary>
 	public static class RedisStringExtensions
 	{
-		/// <inheritdoc cref="IDatabaseAsync.KeyDeleteAsync(RedisKey, CommandFlags)"/>
-		public static Task<bool> DeleteAsync<TElement>(this RedisString<TElement> str, CommandFlags flags = CommandFlags.None)
+		#region StringDecrementAsync
+
+		/// <inheritdoc cref="IDatabaseAsync.StringDecrementAsync(RedisKey, Int64, CommandFlags)"/>
+		public static Task<long> StringDecrementAsync(this IDatabaseAsync target, RedisKey key, long value = 1L, CommandFlags flags = CommandFlags.None)
+		{
+			return target.StringDecrementAsync(key, value, flags);
+		}
+
+		/// <inheritdoc cref="IDatabaseAsync.StringDecrementAsync(RedisKey, Int64, CommandFlags)"/>
+		public static Task<long> StringDecrementAsync(this IDatabaseAsync target, RedisStringKey<long> key, long value = 1L, CommandFlags flags = CommandFlags.None)
+		{
+			return target.StringDecrementAsync(key.Inner, value, flags);
+		}
+
+		/// <inheritdoc cref="IDatabaseAsync.StringDecrementAsync(RedisKey, Double, CommandFlags)"/>
+		public static Task<double> StringDecrementAsync(this IDatabaseAsync target, RedisStringKey<double> key, double value = 1.0, CommandFlags flags = CommandFlags.None)
 		{
-			return str.GetDatabase().KeyDeleteAsync(str.Key, flags);
+			return target.StringDecrementAsync(key.Inner, value, flags);
 		}
 
+		#endregion
+
+		#region StringGetAsync
+
 		/// <inheritdoc cref="IDatabaseAsync.StringGetAsync(RedisKey, CommandFlags)"/>
-		public static async Task<TElement?> GetAsync<TElement>(this RedisString<TElement> str, CommandFlags flags = CommandFlags.None) where TElement : class
+		public static async Task<TValue?> StringGetAsync<TValue>(this IDatabaseAsync target, RedisStringKey<TValue> key, CommandFlags flags = CommandFlags.None)
 		{
-			RedisValue value = await str.GetDatabase().StringGetAsync(str.Key, flags);
+			RedisValue value = await target.StringGetAsync(key.Inner, flags);
 			if (value.IsNullOrEmpty)
 			{
-				return null;
+				return default(TValue);
 			}
-			return RedisSerializer.Deserialize<TElement>(value);
+			return RedisSerializer.Deserialize<TValue>(value)!;
 		}
 
 		/// <inheritdoc cref="IDatabaseAsync.StringGetAsync(RedisKey, CommandFlags)"/>
-		public static async Task<TElement?> GetValueAsync<TElement>(this RedisString<TElement> str, CommandFlags flags = CommandFlags.None) where TElement : struct
+		public static async Task<TValue> StringGetAsync<TValue>(this IDatabaseAsync target, RedisStringKey<TValue> key, TValue defaultValue, CommandFlags flags = CommandFlags.None)
 		{
-			RedisValue value = await str.GetDatabase().StringGetAsync(str.Key, flags);
+			RedisValue value = await target.StringGetAsync(key.Inner, flags);
 			if (value.IsNullOrEmpty)
 			{
-				return null;
+				return defaultValue;
 			}
-			return RedisSerializer.Deserialize<TElement>(value);
+			return RedisSerializer.Deserialize<TValue>(value)!;
 		}
 
-		/// <inheritdoc cref="IDatabaseAsync.StringDecrementAsync(RedisKey, Int64, CommandFlags)"/>
-		public static Task<long> DecrementAsync(this RedisString<long> str, long value = 1L, CommandFlags flags = CommandFlags.None)
+		#endregion
+
+		#region StringIncrementAsync
+
+		/// <inheritdoc cref="IDatabaseAsync.StringIncrementAsync(RedisKey, Double, CommandFlags)"/>
+		public static Task<long> StringIncrementAsync(this IDatabaseAsync target, RedisKey key, long value = 1L, CommandFlags flags = CommandFlags.None)
 		{
-			return str.GetDatabase().StringDecrementAsync(str.Key, value, flags);
+			return target.StringIncrementAsync(key, value, flags);
 		}
 
-		/// <inheritdoc cref="IDatabaseAsync.StringDecrementAsync(RedisKey, Double, CommandFlags)"/>
-		public static Task<double> DecrementAsync(this RedisString<double> str, double value = 1.0, CommandFlags flags = CommandFlags.None)
+		/// <inheritdoc cref="IDatabaseAsync.StringIncrementAsync(RedisKey, Double, CommandFlags)"/>
+		public static Task<long> StringIncrementAsync(this IDatabaseAsync target, RedisStringKey<long> key, long value = 1L, CommandFlags flags = CommandFlags.None)
 		{
-			return str.GetDatabase().StringDecrementAsync(str.Key, value, flags);
+			return target.StringIncrementAsync(key.Inner, value, flags);
 		}
 
-		/// <inheritdoc cref="IDatabaseAsync.StringDecrementAsync(RedisKey, Double, CommandFlags)"/>
-		public static Task<long> IncrementAsync(this RedisString<long> str, long value = 1L, CommandFlags flags = CommandFlags.None)
+		/// <inheritdoc cref="IDatabaseAsync.StringIncrementAsync(RedisKey, Double, CommandFlags)"/>
+		public static Task<double> StringIncrementAsync(this IDatabaseAsync target, RedisKey key, double value, CommandFlags flags = CommandFlags.None)
 		{
-			return str.GetDatabase().StringIncrementAsync(str.Key, value, flags);
+			return target.StringIncrementAsync(key, value, flags);
 		}
 
-		/// <inheritdoc cref="IDatabaseAsync.StringDecrementAsync(RedisKey, Double, CommandFlags)"/>
-		public static Task<double> IncrementAsync(this RedisString<double> str, double value = 1.0, CommandFlags flags = CommandFlags.None)
+		/// <inheritdoc cref="IDatabaseAsync.StringIncrementAsync(RedisKey, Double, CommandFlags)"/>
+		public static Task<double> StringIncrementAsync(this IDatabaseAsync target, RedisStringKey<double> key, double value = 1.0, CommandFlags flags = CommandFlags.None)
 		{
-			return str.GetDatabase().StringIncrementAsync(str.Key, value, flags);
+			return target.StringIncrementAsync(key.Inner, value, flags);
 		}
 
-		/// <summary>
-		/// Creates a version of this string which modifies a transaction rather than the direct DB
-		/// </summary>
-		public static RedisString<TElement> With<TElement>(this ITransaction transaction, RedisString<TElement> str)
+		#endregion
+
+		#region StringLengthAsync
+
+		/// <inheritdoc cref="IDatabaseAsync.StringLengthAsync(RedisKey, CommandFlags)"/>
+		public static Task<long> StringLengthAsync<TValue>(this IDatabaseAsync target, RedisStringKey<TValue> key, CommandFlags flags = CommandFlags.None)
+		{
+			return target.StringLengthAsync(key.Inner, flags);
+		}
+
+		#endregion
+
+		#region StringSetAsync
+
+		/// <inheritdoc cref="IDatabaseAsync.StringSetAsync(RedisKey, RedisValue, TimeSpan?, When, CommandFlags)"/>
+		public static Task<bool> StringSetAsync<TValue>(this IDatabaseAsync target, RedisStringKey<TValue> key, TValue value, TimeSpan? expiry = null, When when = When.Always, CommandFlags flags = CommandFlags.None)
+		{
+			return target.StringSetAsync(key.Inner, RedisSerializer.Serialize(value), expiry, when, flags);
+		}
+
+		/// <inheritdoc cref="IDatabaseAsync.StringSetAsync(RedisKey, RedisValue, TimeSpan?, When, CommandFlags)"/>
+		public static Task<bool> StringSetAsync<TValue>(this IDatabaseAsync target, KeyValuePair<RedisStringKey<TValue>, TValue>[] pairs, When when = When.Always, CommandFlags flags = CommandFlags.None)
 		{
-			return new RedisString<TElement>(transaction, str.Key);
+			return target.StringSetAsync(pairs.ConvertAll(x => new KeyValuePair<RedisKey, RedisValue>(x.Key.Inner, RedisSerializer.Serialize(x.Value))), when, flags);
 		}
+
+		#endregion
 	}
 }
diff --git a/Engine/Source/Programs/Shared/EpicGames.Redis/RedisSubscription.cs b/Engine/Source/Programs/Shared/EpicGames.Redis/RedisSubscription.cs
new file mode 100644
--- /dev/null
+++ b/Engine/Source/Programs/Shared/EpicGames.Redis/RedisSubscription.cs
@@ -0,0 +1,114 @@
+// Copyright Epic Games, Inc. All Rights Reserved.
+
+using StackExchange.Redis;
+using System;
+using System.Threading.Channels;
+using System.Threading.Tasks;
+
+namespace EpicGames.Redis
+{
+	/// <summary>
+	/// Subscription to a <see cref="RedisChannel"/>
+	/// </summary>
+	public sealed class RedisSubscription : IDisposable, IAsyncDisposable
+	{
+		ISubscriber? _subscriber;
+		readonly RedisChannel _channel;
+		readonly Action<RedisChannel, RedisValue> _handler;
+
+		/// <summary>
+		/// Accessor for the channel
+		/// </summary>
+		public RedisChannel Channel => _channel;
+
+		/// <summary>
+		/// Constructor
+		/// </summary>
+		/// <param name="channel">The channel that is subscribed to</param>
+		/// <param name="handler">Callback method</param>
+		public RedisSubscription(RedisChannel channel, Action<RedisChannel, RedisValue> handler)
+		{
+			_channel = channel;
+			_handler = handler;
+		}
+
+		/// <inheritdoc/>
+		public void Dispose()
+		{
+			_ = UnsubscribeAsync(CommandFlags.FireAndForget);
+		}
+
+		/// <inheritdoc/>
+		public async ValueTask DisposeAsync()
+		{
+			await UnsubscribeAsync();
+		}
+
+		/// <summary>
+		/// Subscribes to the channel
+		/// </summary>
+		/// <param name="subscriber">Connection to use for the subscription</param>
+		/// <param name="flags">Flags for the operation</param>
+		public async Task SubscribeAsync(ISubscriber subscriber, CommandFlags flags = CommandFlags.None)
+		{
+			if (_subscriber != null)
+			{
+				throw new InvalidOperationException("Cannot subscribe to a connection while a subscription is already active");
+			}
+
+			_subscriber = subscriber;
+			await subscriber.SubscribeAsync(_channel, _handler, flags);
+		}
+
+		/// <summary>
+		/// Unsubscribe from the channel
+		/// </summary>
+		/// <param name="flags">Flags for the operation</param>
+		public async Task UnsubscribeAsync(CommandFlags flags = CommandFlags.None)
+		{
+			if (_subscriber != null)
+			{
+				await _subscriber.UnsubscribeAsync(_channel, _handler, flags);
+				_subscriber = null;
+			}
+		}
+	}
+
+	/// <summary>
+	/// Extension methods for typed lists
+	/// </summary>
+	public static class RedisSubscriptionExtensions
+	{
+		#region SubscribeAsync
+
+		/// <inheritdoc cref="ISubscriber.SubscribeAsync(RedisChannel, Action{RedisChannel, RedisValue}, CommandFlags)"/>
+		public static Task<RedisSubscription> SubscribeAsync(this IConnectionMultiplexer connection, RedisChannel channel, Action<RedisValue> handler, CommandFlags flags = CommandFlags.None)
+		{
+			return SubscribeAsync(connection, channel, (_, v) => handler(v), flags);
+		}
+
+		/// <inheritdoc cref="ISubscriber.SubscribeAsync(RedisChannel, Action{RedisChannel, RedisValue}, CommandFlags)"/>
+		public static async Task<RedisSubscription> SubscribeAsync(this IConnectionMultiplexer connection, RedisChannel channel, Action<RedisChannel, RedisValue> handler, CommandFlags flags = CommandFlags.None)
+		{
+			RedisSubscription subscription = new RedisSubscription(channel, handler);
+			await subscription.SubscribeAsync(connection.GetSubscriber(), flags);
+			return subscription;
+		}
+
+		/// <inheritdoc cref="ISubscriber.SubscribeAsync(RedisChannel, Action{RedisChannel, RedisValue}, CommandFlags)"/>
+		public static Task<RedisSubscription> SubscribeAsync<T>(this IConnectionMultiplexer connection, RedisChannel<T> channel, Action<T> handler, CommandFlags flags = CommandFlags.None)
+		{
+			Action<RedisChannel, RedisValue> action = (_, v) => handler(RedisSerializer.Deserialize<T>(v));
+			return SubscribeAsync(connection, channel.Channel, action, flags);
+		}
+
+		/// <inheritdoc cref="ISubscriber.SubscribeAsync(RedisChannel, Action{RedisChannel, RedisValue}, CommandFlags)"/>
+		public static Task<RedisSubscription> SubscribeAsync<T>(this IConnectionMultiplexer connection, RedisChannel<T> channel, Action<RedisChannel<T>, T> handler, CommandFlags flags = CommandFlags.None)
+		{
+			Action<RedisChannel, RedisValue> action = (_, v) => handler(channel, RedisSerializer.Deserialize<T>(v));
+			return SubscribeAsync(connection, channel.Channel, action, flags);
+		}
+
+		#endregion
+	}
+}
