Horde: Support populating a ManagedWorkspace without using the have table

diff --git a/Engine/Source/Programs/Horde/Horde.Agent/Commands/Workspace/PopulateCacheCommand.cs b/Engine/Source/Programs/Horde/Horde.Agent/Commands/Workspace/PopulateCacheCommand.cs
--- a/Engine/Source/Programs/Horde/Horde.Agent/Commands/Workspace/PopulateCacheCommand.cs
+++ b/Engine/Source/Programs/Horde/Horde.Agent/Commands/Workspace/PopulateCacheCommand.cs
@@ -26,9 +26,15 @@ namespace Horde.Agent.Commands.Workspace
 		[CommandLine("-FakeSync")]
 		[Description("Simulates the sync without actually fetching any files")]
 		bool FakeSync { get; set; } = false;
+		
+		[CommandLine("-UseHaveTable")]
+		[Description("Update the have-table")]
+		string UseHaveTable { get; set; } = "true";
 
 		protected override async Task ExecuteAsync(IPerforceConnection perforce, ManagedWorkspace repo, ILogger logger)
 		{
+			bool useHaveTable = UseHaveTable.Equals("true", StringComparison.Ordinal);
+
 			List<string> expandedFilters = ExpandFilters(Filters);
 
 			List<PopulateRequest> populateRequests = new List<PopulateRequest>();
@@ -44,7 +50,7 @@ namespace Horde.Agent.Commands.Workspace
 				populateRequests.Add(new PopulateRequest(perforceClient, clientAndStreamParam.Substring(idx + 1), expandedFilters));
 			}
 
-			await repo.PopulateAsync(populateRequests, FakeSync, CancellationToken.None);
+			await repo.PopulateAsync(populateRequests, FakeSync, useHaveTable, CancellationToken.None);
 		}
 	}
 }
diff --git a/Engine/Source/Programs/Horde/Horde.Agent/Execution/PerforceExecutor.cs b/Engine/Source/Programs/Horde/Horde.Agent/Execution/PerforceExecutor.cs
--- a/Engine/Source/Programs/Horde/Horde.Agent/Execution/PerforceExecutor.cs
+++ b/Engine/Source/Programs/Horde/Horde.Agent/Execution/PerforceExecutor.cs
@@ -233,7 +233,7 @@ namespace Horde.Agent.Execution
 			await _workspace.CleanAsync(cancellationToken);
 		}
 
-		public static async Task ConformAsync(DirectoryReference rootDir, IList<AgentWorkspace> pendingWorkspaces, bool removeUntrackedFiles, ILogger logger, CancellationToken cancellationToken)
+		public static async Task ConformAsync(DirectoryReference rootDir, IList<AgentWorkspace> pendingWorkspaces, bool removeUntrackedFiles, bool useHaveTable, ILogger logger, CancellationToken cancellationToken)
 		{
 			using IScope scope = GlobalTracer.Instance.BuildSpan("Conform").StartActive();
 			scope.Span.SetTag("workspaces", String.Join(',', pendingWorkspaces.Select(x => x.Identifier)));
@@ -383,13 +383,13 @@ namespace Horde.Agent.Execution
 					if (populateRequests.Count == 1 && !firstWorkspace.RemoveUntrackedFiles && !removeUntrackedFiles)
 					{
 						await firstWorkspace.CleanAsync(cancellationToken);
-						syncFuncs.Add(() => firstWorkspace.SyncAsync(-1, -1, null, true, cancellationToken));
+						syncFuncs.Add(() => firstWorkspace.SyncAsync(-1, -1, null, useHaveTable, cancellationToken));
 					}
 					else
 					{
 						ManagedWorkspace repository = firstWorkspace.Repository;
-						Tuple<int, StreamSnapshot>[] streamStates = await repository.PopulateCleanAsync(populateRequests, cancellationToken);
-						syncFuncs.Add(() => repository.PopulateSyncAsync(populateRequests, streamStates, false, cancellationToken));
+						Tuple<int, StreamSnapshot>[] streamStates = await repository.PopulateCleanAsync(populateRequests, useHaveTable, cancellationToken);
+						syncFuncs.Add(() => repository.PopulateSyncAsync(populateRequests, streamStates, false, useHaveTable, cancellationToken));
 					}
 				}
 				foreach (Func<Task> syncFunc in syncFuncs)
diff --git a/Engine/Source/Programs/Horde/Horde.Agent/Leases/Handlers/ConformHandler.cs b/Engine/Source/Programs/Horde/Horde.Agent/Leases/Handlers/ConformHandler.cs
--- a/Engine/Source/Programs/Horde/Horde.Agent/Leases/Handlers/ConformHandler.cs
+++ b/Engine/Source/Programs/Horde/Horde.Agent/Leases/Handlers/ConformHandler.cs
@@ -46,7 +46,7 @@ namespace Horde.Agent.Leases.Handlers
 				// Run the conform task
 				if (_settings.Executor.Equals(PerforceExecutor.Name, StringComparison.OrdinalIgnoreCase) && _settings.PerforceExecutor.RunConform)
 				{
-					await PerforceExecutor.ConformAsync(session.WorkingDir, pendingWorkspaces, removeUntrackedFiles, conformLogger, cancellationToken);
+					await PerforceExecutor.ConformAsync(session.WorkingDir, pendingWorkspaces, removeUntrackedFiles, false, conformLogger, cancellationToken);
 				}
 				else
 				{
diff --git a/Engine/Source/Programs/Shared/EpicGames.Perforce.Fixture/PerforceFixture.cs b/Engine/Source/Programs/Shared/EpicGames.Perforce.Fixture/PerforceFixture.cs
--- a/Engine/Source/Programs/Shared/EpicGames.Perforce.Fixture/PerforceFixture.cs
+++ b/Engine/Source/Programs/Shared/EpicGames.Perforce.Fixture/PerforceFixture.cs
@@ -91,11 +91,17 @@ public class ChangelistFixture
 	/// Assert files in stream for this changelist matches the client's have table
 	/// </summary>
 	/// <param name="perforce">Perforce connection</param>
-	public async Task AssertHaveTableAsync(IPerforceConnection perforce)
+	/// <param name="useHaveTable">When set to false, have table is expected to be empty</param>
+	public async Task AssertHaveTableAsync(IPerforceConnection perforce, bool useHaveTable = true)
 	{
 		List<HaveRecord> haveRecords = await perforce.HaveAsync(new FileSpecList(), CancellationToken.None).ToListAsync();
 		HashSet<(string depotFile, int rev)> actual = new(haveRecords.Select((x) => (x.DepotFile, x.HaveRev)));
 		HashSet<(string depotFile, int rev)> expected = new(StreamFiles.Select((x) => (x.DepotFile, x.Revision)));
+
+		if (!useHaveTable)
+		{
+			expected.Clear();
+		}
 		
 		if (!expected.SetEquals(actual))
 		{
@@ -172,6 +178,8 @@ public class StreamFixture
 		Changelists = changelists;
 	}
 
+	public ChangelistFixture LatestChangelist => Changelists.Last();
+
 	public ChangelistFixture GetChangelist(int changeNum)
 	{
 		foreach (ChangelistFixture changelist in Changelists)
diff --git a/Engine/Source/Programs/Shared/EpicGames.Perforce.Managed.Tests/ManagedWorkspaceTest.cs b/Engine/Source/Programs/Shared/EpicGames.Perforce.Managed.Tests/ManagedWorkspaceTest.cs
--- a/Engine/Source/Programs/Shared/EpicGames.Perforce.Managed.Tests/ManagedWorkspaceTest.cs
+++ b/Engine/Source/Programs/Shared/EpicGames.Perforce.Managed.Tests/ManagedWorkspaceTest.cs
@@ -49,6 +49,24 @@ public class ManagedWorkspaceTest : BasePerforceFixtureTest
 		await AssertHaveTableFileCount(0);
 	}
 	
+	[TestMethod]
+	public async Task SyncBackwardsToOlderChangelistRemoveUntracked()
+	{
+		ManagedWorkspace ws = await GetManagedWorkspace();
+
+		await SyncAsync(ws, 6, removeUntracked: false);
+		Stream.GetChangelist(6).AssertDepotFiles(SyncDir);
+		await Stream.GetChangelist(6).AssertHaveTableAsync(PerforceConnection);
+		
+		await SyncAsync(ws, 7, removeUntracked: false);
+		Stream.GetChangelist(7).AssertDepotFiles(SyncDir);
+		await Stream.GetChangelist(7).AssertHaveTableAsync(PerforceConnection);
+
+		await SyncAsync(ws, 6, removeUntracked: false);
+		Stream.GetChangelist(6).AssertDepotFiles(SyncDir);
+		await Stream.GetChangelist(6).AssertHaveTableAsync(PerforceConnection);
+	}
+	
 	[TestMethod]
 	public async Task SyncBackwardsToOlderChangelist()
 	{
@@ -131,6 +149,22 @@ public class ManagedWorkspaceTest : BasePerforceFixtureTest
 			await cl.AssertHaveTableAsync(PerforceConnection);
 		}
 	}
+	
+	[DataTestMethod]
+	[DataRow(true, DisplayName = "With have-table")]
+	[DataRow(false, DisplayName = "Without have-table")]
+	public async Task Populate(bool useHaveTable)
+	{
+		ManagedWorkspace ws = await GetManagedWorkspace();
+
+		List<PopulateRequest> populateRequests = new ()
+		{
+			new PopulateRequest(PerforceConnection, StreamName, new List<string>())
+		};
+		await ws.PopulateAsync(populateRequests, false, useHaveTable, CancellationToken.None);
+		Stream.LatestChangelist.AssertDepotFiles(SyncDir);
+		await Stream.LatestChangelist.AssertHaveTableAsync(PerforceConnection, useHaveTable);
+	}
 
 	private async Task<ManagedWorkspace> GetManagedWorkspace()
 	{
@@ -139,9 +173,9 @@ public class ManagedWorkspaceTest : BasePerforceFixtureTest
 		return ws;
 	}
 	
-	private async Task SyncAsync(ManagedWorkspace managedWorkspace, int changeNumber, bool useHaveTable = true, FileReference? cacheFile = null)
+	private async Task SyncAsync(ManagedWorkspace managedWorkspace, int changeNumber, bool useHaveTable = true, FileReference? cacheFile = null, bool removeUntracked = true)
 	{
-		await managedWorkspace.SyncAsync(PerforceConnection, StreamName, changeNumber, Array.Empty<string>(), true, false, useHaveTable, cacheFile, CancellationToken.None);
+		await managedWorkspace.SyncAsync(PerforceConnection, StreamName, changeNumber, Array.Empty<string>(), removeUntracked, false, useHaveTable, cacheFile, CancellationToken.None);
 	}
 	
 	private async Task AssertHaveTableFileCount(int expected)
@@ -158,4 +192,13 @@ public class ManagedWorkspaceTest : BasePerforceFixtureTest
 			Assert.Fail($"Actual have table file count does not match expected count. Actual={haveRecords.Count} Expected={expected}");
 		}
 	}
+
+	private void DumpMetaDir()
+	{
+		Console.WriteLine("Meta dir: --------------------------");
+		foreach (string path in Directory.GetFiles(TempDir.FullName, "*", SearchOption.AllDirectories))
+		{
+			Console.WriteLine(Path.GetRelativePath(TempDir.FullName, path));
+		}
+	}
 }
\ No newline at end of file
diff --git a/Engine/Source/Programs/Shared/EpicGames.Perforce.Managed/ManagedWorkspace.cs b/Engine/Source/Programs/Shared/EpicGames.Perforce.Managed/ManagedWorkspace.cs
--- a/Engine/Source/Programs/Shared/EpicGames.Perforce.Managed/ManagedWorkspace.cs
+++ b/Engine/Source/Programs/Shared/EpicGames.Perforce.Managed/ManagedWorkspace.cs
@@ -955,20 +955,20 @@ namespace EpicGames.Perforce.Managed
 		/// <summary>
 		/// Populates the cache with the head revision of the given streams.
 		/// </summary>
-		public async Task PopulateAsync(List<PopulateRequest> requests, bool fakeSync, CancellationToken cancellationToken)
+		public async Task PopulateAsync(List<PopulateRequest> requests, bool fakeSync, bool useHaveTable, CancellationToken cancellationToken)
 		{
 			_logger.LogInformation("Populating with {NumStreams} streams", requests.Count);
 			using (_logger.BeginIndentScope("  "))
 			{
-				Tuple<int, StreamSnapshot>[] streamState = await PopulateCleanAsync(requests, cancellationToken);
-				await PopulateSyncAsync(requests, streamState, fakeSync, cancellationToken);
+				Tuple<int, StreamSnapshot>[] streamState = await PopulateCleanAsync(requests, useHaveTable, cancellationToken);
+				await PopulateSyncAsync(requests, streamState, fakeSync, useHaveTable, cancellationToken);
 			}
 		}
 
 		/// <summary>
 		/// Perform the clean part of a populate command
 		/// </summary>
-		public async Task<Tuple<int, StreamSnapshot>[]> PopulateCleanAsync(List<PopulateRequest> requests, CancellationToken cancellationToken)
+		public async Task<Tuple<int, StreamSnapshot>[]> PopulateCleanAsync(List<PopulateRequest> requests, bool useHaveTable, CancellationToken cancellationToken)
 		{
 			// Revert all changes in each of the unique clients
 			foreach (PopulateRequest request in requests)
@@ -1001,10 +1001,18 @@ namespace EpicGames.Perforce.Managed
 					int changeNumber = await GetLatestClientChangeAsync(request.PerforceClient, cancellationToken);
 					_logger.LogInformation("Latest change is CL {CL}", changeNumber);
 
-					await UpdateClientHaveTableAsync(request.PerforceClient, changeNumber, request.View, cancellationToken);
+					if (useHaveTable)
+					{
+						await UpdateClientHaveTableAsync(request.PerforceClient, changeNumber, request.View, cancellationToken);
 
-					StreamSnapshot contents = await FindClientContentsAsync(request.PerforceClient, changeNumber, cancellationToken);
-					streamState[idx] = Tuple.Create(changeNumber, contents);
+						StreamSnapshot contents = await FindClientContentsAsync(request.PerforceClient, changeNumber, cancellationToken);
+						streamState[idx] = Tuple.Create(changeNumber, contents);
+					}
+					else
+					{
+						StreamSnapshot contents = await FindClientContentsWithoutHaveTableAsync(request.PerforceClient, streamName, changeNumber, cancellationToken);
+						streamState[idx] = Tuple.Create(changeNumber, contents);
+					}
 
 					GC.Collect();
 				}
@@ -1063,7 +1071,7 @@ namespace EpicGames.Perforce.Managed
 		/// <summary>
 		/// Perform the sync part of a populate command
 		/// </summary>
-		public async Task PopulateSyncAsync(List<PopulateRequest> requests, Tuple<int, StreamSnapshot>[] streamState, bool fakeSync, CancellationToken cancellationToken)
+		public async Task PopulateSyncAsync(List<PopulateRequest> requests, Tuple<int, StreamSnapshot>[] streamState, bool fakeSync, bool useHaveTable, CancellationToken cancellationToken)
 		{
 			// Sync all the new files
 			for (int idx = 0; idx < requests.Count; idx++)
@@ -1074,15 +1082,18 @@ namespace EpicGames.Perforce.Managed
 
 				using (_logger.BeginIndentScope("  "))
 				{
-					await DeleteClientAsync(request.PerforceClient, cancellationToken);
-					await UpdateClientAsync(request.PerforceClient, streamName, cancellationToken);
+					if (useHaveTable)
+					{
+						await DeleteClientAsync(request.PerforceClient, cancellationToken);
+						await UpdateClientAsync(request.PerforceClient, streamName, cancellationToken);
 
-					int changeNumber = streamState[idx].Item1;
-					await UpdateClientHaveTableAsync(request.PerforceClient, changeNumber, requests[idx].View, cancellationToken);
+						int changeNumber = streamState[idx].Item1;
+						await UpdateClientHaveTableAsync(request.PerforceClient, changeNumber, requests[idx].View, cancellationToken);
+					}
 
 					StreamSnapshot contents = streamState[idx].Item2;
 					await RemoveFilesFromWorkspaceAsync(contents, cancellationToken);
-					await AddFilesToWorkspaceAsync(request.PerforceClient, contents, fakeSync, true, cancellationToken);
+					await AddFilesToWorkspaceAsync(request.PerforceClient, contents, fakeSync, useHaveTable, cancellationToken);
 				}
 			}
 
