Horde: Missing file.

diff --git a/Engine/Source/Programs/Shared/EpicGames.Redis/TypedRedisKey.cs b/Engine/Source/Programs/Shared/EpicGames.Redis/TypedRedisKey.cs
new file mode 100644
--- /dev/null
+++ b/Engine/Source/Programs/Shared/EpicGames.Redis/TypedRedisKey.cs
@@ -0,0 +1,95 @@
+// Copyright Epic Games, Inc. All Rights Reserved.
+
+using System.Collections.Generic;
+using System.Threading.Tasks;
+using System;
+using StackExchange.Redis;
+
+namespace EpicGames.Redis
+{
+	/// <summary>
+	/// Wrapper struct that supports conversion from a regular <see cref="RedisKey"/> as well as the key types in EpicGames.Redis.
+	/// </summary>
+	public readonly struct TypedRedisKey
+	{
+		/// <summary>
+		/// The key for the list
+		/// </summary>
+		public readonly RedisKey Inner { get; }
+
+		/// <summary>
+		/// Constructor
+		/// </summary>
+		/// <param name="inner">Redis key this type is using</param>
+		public TypedRedisKey(RedisKey inner)
+		{
+			Inner = inner;
+		}
+
+		/// <summary>
+		/// Implicit conversion to typed redis key.
+		/// </summary>
+		/// <param name="key">Key to convert</param>
+		public static implicit operator TypedRedisKey(RedisKey key) => new TypedRedisKey(key);
+
+		/// <summary>
+		/// Implicit conversion to typed redis key.
+		/// </summary>
+		/// <param name="key">Key to convert</param>
+		public static implicit operator RedisKey(TypedRedisKey key) => key.Inner;
+	}
+
+	/// <summary>
+	/// Extension methods for <see cref="TypedRedisKey"/>
+	/// </summary>
+	public static class TypedRedisKeyExtensions
+	{
+		#region KeyDeleteAsync
+
+		/// <inheritdoc cref="IDatabaseAsync.KeyDeleteAsync(RedisKey, CommandFlags)"/>
+		public static Task KeyDeleteAsync(this IDatabaseAsync target, TypedRedisKey key, CommandFlags flags = CommandFlags.None)
+		{
+			return target.KeyDeleteAsync(key.Inner, flags);
+		}
+
+		/// <inheritdoc cref="IDatabaseAsync.KeyDeleteAsync(RedisKey[], CommandFlags)"/>
+		public static Task KeyDeleteAsync(this IDatabaseAsync target, RedisKey[] keys, CommandFlags flags = CommandFlags.None)
+		{
+			return target.KeyDeleteAsync(keys, flags);
+		}
+
+		#endregion
+
+		#region KeyExistsAsync
+
+		/// <inheritdoc cref="IDatabaseAsync.KeyExistsAsync(RedisKey, CommandFlags)"/>
+		public static Task<bool> KeyExistsAsync(this IDatabaseAsync target, TypedRedisKey key, CommandFlags flags = CommandFlags.None)
+		{
+			return target.KeyExistsAsync(key.Inner, flags);
+		}
+/*
+		/// <inheritdoc cref="IDatabaseAsync.KeyExistsAsync(RedisKey[], CommandFlags)"/>
+		public static Task<long> KeyExistsAsync(this IDatabaseAsync target, IEnumerable<RedisKey> keys, CommandFlags flags = CommandFlags.None)
+		{
+			return target.KeyExistsAsync(keys.ToArray(), flags);
+		}
+*/
+		#endregion
+
+		#region KeyExpireAsync
+
+		/// <inheritdoc cref="IDatabaseAsync.KeyExpireAsync(RedisKey, DateTime?, ExpireWhen, CommandFlags)"/>
+		public static Task KeyExpireAsync(this IDatabaseAsync target, TypedRedisKey key, DateTime? expiry = null, ExpireWhen when = ExpireWhen.Always, CommandFlags flags = CommandFlags.None)
+		{
+			return target.KeyExpireAsync(key.Inner, expiry, when, flags);
+		}
+
+		/// <inheritdoc cref="IDatabaseAsync.KeyExpireAsync(RedisKey, TimeSpan?, ExpireWhen, CommandFlags)"/>
+		public static Task KeyExpireAsync(this IDatabaseAsync target, TypedRedisKey key, TimeSpan? expiry, ExpireWhen when = ExpireWhen.Always, CommandFlags flags = CommandFlags.None)
+		{
+			return target.KeyExpireAsync(key.Inner, expiry, when, flags);
+		}
+
+		#endregion
+	}
+}
