EpicGames.Perforce: Use more parameterized tests for ManagedWorkspace

diff --git a/Engine/Source/Programs/Shared/EpicGames.Perforce.Managed.Tests/ManagedWorkspaceTest.cs b/Engine/Source/Programs/Shared/EpicGames.Perforce.Managed.Tests/ManagedWorkspaceTest.cs
--- a/Engine/Source/Programs/Shared/EpicGames.Perforce.Managed.Tests/ManagedWorkspaceTest.cs
+++ b/Engine/Source/Programs/Shared/EpicGames.Perforce.Managed.Tests/ManagedWorkspaceTest.cs
@@ -28,25 +28,16 @@ public class ManagedWorkspaceTest : BasePerforceFixtureTest
 	}
 
 	[TestMethod]
-	public async Task SyncSingleChangelist()
-	{
-		ManagedWorkspace ws = await GetManagedWorkspace();
-		await AssertHaveTableFileCount(0);
-
-		await SyncAsync(ws, 6);
-		Stream.GetChangelist(6).AssertDepotFiles(SyncDir);
-		await Stream.GetChangelist(6).AssertHaveTableAsync(PerforceConnection);
-	}
-	
-	[TestMethod]
-	public async Task SyncSingleChangelistWithoutHaveTable()
+	[DataRow(true, DisplayName = "With have-table")]
+	[DataRow(false, DisplayName = "Without have-table")]
+	public async Task SyncSingleChangelist(bool useHaveTable)
 	{
 		ManagedWorkspace ws = await GetManagedWorkspace();
 		await AssertHaveTableFileCount(0);
 
-		await SyncAsync(ws, 6, useHaveTable: false);
+		await SyncAsync(ws, 6, useHaveTable);
 		Stream.GetChangelist(6).AssertDepotFiles(SyncDir);
-		await AssertHaveTableFileCount(0);
+		await Stream.GetChangelist(6).AssertHaveTableAsync(PerforceConnection, useHaveTable);
 	}
 	
 	[TestMethod]
@@ -68,41 +59,24 @@ public class ManagedWorkspaceTest : BasePerforceFixtureTest
 	}
 	
 	[TestMethod]
-	public async Task SyncBackwardsToOlderChangelist()
+	[DataRow(true, DisplayName = "With have-table")]
+	[DataRow(false, DisplayName = "Without have-table")]
+	public async Task SyncBackwardsToOlderChangelist(bool useHaveTable)
 	{
 		ManagedWorkspace ws = await GetManagedWorkspace();
 
-		await SyncAsync(ws, 6);
+		await SyncAsync(ws, 6, useHaveTable);
 		Stream.GetChangelist(6).AssertDepotFiles(SyncDir);
-		await Stream.GetChangelist(6).AssertHaveTableAsync(PerforceConnection);
+		await Stream.GetChangelist(6).AssertHaveTableAsync(PerforceConnection, useHaveTable);
 		
-		await SyncAsync(ws, 7);
+		await SyncAsync(ws, 7, useHaveTable);
 		Stream.GetChangelist(7).AssertDepotFiles(SyncDir);
-		await Stream.GetChangelist(7).AssertHaveTableAsync(PerforceConnection);
+		await Stream.GetChangelist(7).AssertHaveTableAsync(PerforceConnection, useHaveTable);
 		
 		// Go back one changelist
-		await SyncAsync(ws, 6);
+		await SyncAsync(ws, 6, useHaveTable);
 		Stream.GetChangelist(6).AssertDepotFiles(SyncDir);
-		await Stream.GetChangelist(6).AssertHaveTableAsync(PerforceConnection);
-	}
-	
-	[TestMethod]
-	public async Task SyncBackwardsToOlderChangelistWithoutHaveTable()
-	{
-		ManagedWorkspace ws = await GetManagedWorkspace();
-		
-		await SyncAsync(ws, 6, useHaveTable: false);
-		Stream.GetChangelist(6).AssertDepotFiles(SyncDir);
-		await AssertHaveTableFileCount(0);
-
-		await SyncAsync(ws, 7, useHaveTable: false);
-		Stream.GetChangelist(7).AssertDepotFiles(SyncDir);
-		await AssertHaveTableFileCount(0);
-		
-		// Go back one changelist
-		await SyncAsync(ws, 6, useHaveTable: false);
-		Stream.GetChangelist(6).AssertDepotFiles(SyncDir);
-		await AssertHaveTableFileCount(0);
+		await Stream.GetChangelist(6).AssertHaveTableAsync(PerforceConnection, useHaveTable);
 	}
 
 	[TestMethod]
@@ -124,7 +98,9 @@ public class ManagedWorkspaceTest : BasePerforceFixtureTest
 	}
 
 	[TestMethod]
-	public async Task SyncUsingCacheFiles()
+	[DataRow(true, DisplayName = "With have-table")]
+	[DataRow(false, DisplayName = "Without have-table")]
+	public async Task SyncUsingCacheFiles(bool useHaveTable)
 	{
 		ManagedWorkspace ws = await GetManagedWorkspace();
 
@@ -136,21 +112,21 @@ public class ManagedWorkspaceTest : BasePerforceFixtureTest
 		// Sync and create a new cache file per change number
 		foreach (ChangelistFixture cl in Stream.Changelists)
 		{
-			await ws.SyncAsync(PerforceConnection, StreamName, cl.Number, Array.Empty<string>(), true, false, true, GetCacheFilePath(cl.Number), CancellationToken.None);
+			await SyncAsync(ws, cl.Number, useHaveTable: useHaveTable, cacheFile: GetCacheFilePath(cl.Number));
 			cl.AssertDepotFiles(SyncDir);
-			await cl.AssertHaveTableAsync(PerforceConnection);
+			await cl.AssertHaveTableAsync(PerforceConnection, useHaveTable);
 		}
 		
 		// Sync again but using the cache files created above
 		foreach (ChangelistFixture cl in Stream.Changelists.Reverse())
 		{
-			await ws.SyncAsync(PerforceConnection, StreamName, cl.Number, Array.Empty<string>(), true, false, true, GetCacheFilePath(cl.Number), CancellationToken.None);
+			await SyncAsync(ws, cl.Number, useHaveTable: useHaveTable, cacheFile: GetCacheFilePath(cl.Number));
 			cl.AssertDepotFiles(SyncDir);
-			await cl.AssertHaveTableAsync(PerforceConnection);
+			await cl.AssertHaveTableAsync(PerforceConnection, useHaveTable);
 		}
 	}
 	
-	[DataTestMethod]
+	[TestMethod]
 	[DataRow(true, DisplayName = "With have-table")]
 	[DataRow(false, DisplayName = "Without have-table")]
 	public async Task Populate(bool useHaveTable)
