EpicGames.Core: Change Utf8StringBuilder from struct to class.

diff --git a/Engine/Source/Programs/Shared/EpicGames.Core/Utf8StringBuilder.cs b/Engine/Source/Programs/Shared/EpicGames.Core/Utf8StringBuilder.cs
--- a/Engine/Source/Programs/Shared/EpicGames.Core/Utf8StringBuilder.cs
+++ b/Engine/Source/Programs/Shared/EpicGames.Core/Utf8StringBuilder.cs
@@ -1,15 +1,13 @@
 // Copyright Epic Games, Inc. All Rights Reserved.
 
 using System;
-using System.Collections.Generic;
-using System.Text;
 
 namespace EpicGames.Core
 {
 	/// <summary>
 	/// Represents a memory region which can be treated as a utf-8 string.
 	/// </summary>
-	public struct Utf8StringBuilder
+	public class Utf8StringBuilder
 	{
 		readonly ArrayMemoryWriter _writer;
 
