Improve C# UHT's file reading time.

diff --git a/Engine/Source/Programs/Shared/EpicGames.UHT/Utils/UhtBuffer.cs b/Engine/Source/Programs/Shared/EpicGames.UHT/Utils/UhtBuffer.cs
--- a/Engine/Source/Programs/Shared/EpicGames.UHT/Utils/UhtBuffer.cs
+++ b/Engine/Source/Programs/Shared/EpicGames.UHT/Utils/UhtBuffer.cs
@@ -219,4 +219,171 @@ namespace EpicGames.UHT.Utils
 			UhtBuffer.Return(Buffer);
 		}
 	}
+
+	/// <summary>
+	/// Cached character buffer system.
+	/// 
+	/// Invoke UhtBuffer.Borrow method to get a buffer of the given size.
+	/// Invoke UhtBuffer.Return to return the buffer to the cache.
+	/// </summary>
+	// TODO: Refactor to share implementation and possibly support ArrayPool<T>.
+	public class UhtByteBuffer
+	{
+		/// <summary>
+		/// Any requests of the given size or smaller will be placed in bucket zero with the given size.
+		/// </summary>
+		private const int MinSize = 1024 * 16;
+
+		/// <summary>
+		/// Adjustment to the bucket index to account for the minimum bucket size
+		/// </summary>
+		private static readonly int s_buckedAdjustment = BitOperations.Log2((uint)UhtByteBuffer.MinSize);
+
+		/// <summary>
+		/// Total number of supported buckets
+		/// </summary>
+		private static readonly int s_bucketCount = 32 - UhtByteBuffer.s_buckedAdjustment;
+
+		/// <summary>
+		/// Bucket lookaside list
+		/// </summary>
+		private static readonly UhtByteBuffer?[] s_lookAsideArray = new UhtByteBuffer?[UhtByteBuffer.s_bucketCount];
+
+		/// <summary>
+		/// The bucket index associated with the buffer
+		/// </summary>
+		private int Bucket { get; }
+
+		/// <summary>
+		/// Single list link to the next cached buffer
+		/// </summary>
+		private UhtByteBuffer? NextBuffer { get; set; } = null;
+
+		/// <summary>
+		/// The backing character block.  The size of the array will normally be larger than the 
+		/// requested size.
+		/// </summary>
+		[System.Diagnostics.CodeAnalysis.SuppressMessage("Design", "CA1051:Do not declare visible instance fields", Justification = "<Pending>")]
+		[System.Diagnostics.CodeAnalysis.SuppressMessage("Style", "IDE1006:Naming Styles", Justification = "<Pending>")]
+		public byte[] Block;
+
+		/// <summary>
+		/// Memory region sized to the requested size
+		/// </summary>
+		public Memory<byte> Memory { get; set; }
+
+		/// <summary>
+		/// Construct a new buffer
+		/// </summary>
+		/// <param name="size">The initial size of the buffer</param>
+		/// <param name="bucket">The bucket associated with the buffer</param>
+		/// <param name="bucketSize">The size all blocks in this bucket</param>
+		private UhtByteBuffer(int size, int bucket, int bucketSize)
+		{
+			Block = new byte[bucketSize];
+			Bucket = bucket;
+			Reset(size);
+		}
+
+		/// <summary>
+		/// Reset the memory region to the given size
+		/// </summary>
+		/// <param name="size"></param>
+		public void Reset(int size)
+		{
+			Memory = new Memory<byte>(Block, 0, size);
+		}
+
+		/// <summary>
+		/// Borrow a new buffer of the given size
+		/// </summary>
+		/// <param name="size">Size of the buffer</param>
+		/// <returns>Buffer that should be returned with a call to Return</returns>
+		public static UhtByteBuffer Borrow(int size)
+		{
+			if (size <= UhtByteBuffer.MinSize)
+			{
+				return BorrowInternal(size, 0, UhtByteBuffer.MinSize);
+			}
+			else
+			{
+
+				// Round up the size to the next larger power of two if it isn't a power of two
+				uint usize = (uint)size;
+				--usize;
+				usize |= usize >> 1;
+				usize |= usize >> 2;
+				usize |= usize >> 4;
+				usize |= usize >> 8;
+				usize |= usize >> 16;
+				++usize;
+				int bucket = BitOperations.Log2(usize) - UhtByteBuffer.s_buckedAdjustment;
+				return BorrowInternal(size, bucket, (int)usize);
+			}
+		}
+
+		/// <summary>
+		/// Return the buffer to the cache.  The buffer should no longer be accessed.
+		/// </summary>
+		/// <param name="buffer">The buffer to be returned.</param>
+		public static void Return(UhtByteBuffer buffer)
+		{
+			lock (UhtByteBuffer.s_lookAsideArray)
+			{
+				buffer.NextBuffer = UhtByteBuffer.s_lookAsideArray[buffer.Bucket];
+				UhtByteBuffer.s_lookAsideArray[buffer.Bucket] = buffer;
+			}
+		}
+
+		/// <summary>
+		/// Internal helper to allocate a buffer
+		/// </summary>
+		/// <param name="size">The initial size of the buffer</param>
+		/// <param name="bucket">The bucket associated with the buffer</param>
+		/// <param name="bucketSize">The size all blocks in this bucket</param>
+		/// <returns>The allocated buffer</returns>
+		private static UhtByteBuffer BorrowInternal(int size, int bucket, int bucketSize)
+		{
+			lock (UhtByteBuffer.s_lookAsideArray)
+			{
+				if (UhtByteBuffer.s_lookAsideArray[bucket] != null)
+				{
+					UhtByteBuffer buffer = UhtByteBuffer.s_lookAsideArray[bucket]!;
+					UhtByteBuffer.s_lookAsideArray[bucket] = buffer.NextBuffer;
+					buffer.Reset(size);
+					return buffer;
+				}
+			}
+			return new UhtByteBuffer(size, bucket, bucketSize);
+		}
+	}
+
+	/// <summary>
+	/// Helper class for using pattern to borrow and return a buffer.
+	/// </summary>
+	public struct UhtBorrowByteBuffer : IDisposable
+	{
+
+		/// <summary>
+		/// The borrowed buffer
+		/// </summary>
+		public UhtByteBuffer Buffer { get; set; }
+
+		/// <summary>
+		/// Borrow a buffer with the given size
+		/// </summary>
+		/// <param name="size">The size to borrow</param>
+		public UhtBorrowByteBuffer(int size)
+		{
+			Buffer = UhtByteBuffer.Borrow(size);
+		}
+
+		/// <summary>
+		/// Return the borrowed buffer to the cache
+		/// </summary>
+		public void Dispose()
+		{
+			UhtByteBuffer.Return(Buffer);
+		}
+	}
 }
diff --git a/Engine/Source/Programs/Shared/EpicGames.UHT/Utils/UhtFileManager.cs b/Engine/Source/Programs/Shared/EpicGames.UHT/Utils/UhtFileManager.cs
--- a/Engine/Source/Programs/Shared/EpicGames.UHT/Utils/UhtFileManager.cs
+++ b/Engine/Source/Programs/Shared/EpicGames.UHT/Utils/UhtFileManager.cs
@@ -2,7 +2,6 @@
 
 using System;
 using System.IO;
-using System.Linq;
 using System.Text;
 using EpicGames.Core;
 
@@ -58,7 +57,6 @@ namespace EpicGames.UHT.Utils
 	/// </summary>
 	public class UhtStdFileManager : IUhtFileManager
 	{
-
 		/// <summary>
 		/// Construct a new file manager
 		/// </summary>
@@ -96,28 +94,18 @@ namespace EpicGames.UHT.Utils
 			try
 			{
 				using FileStream fs = new(filePath, FileMode.Open, FileAccess.Read, FileShare.Read, 4 * 1024, FileOptions.SequentialScan);
-				using StreamReader sr = new(fs, Encoding.UTF8, true);
-
-				// Try to read the whole file into a buffer created by hand.  This avoids a LOT of memory allocations which in turn reduces the
-				// GC stress on the system.  Removing the StreamReader would be nice in the future.
-				long rawFileLength = fs.Length;
-				UhtBuffer initialBuffer = UhtBuffer.Borrow((int)rawFileLength);
-				int readLength = sr.Read(initialBuffer.Memory.Span);
-				if (sr.EndOfStream)
-				{
-					initialBuffer.Reset(readLength);
-					return initialBuffer;
-				}
-				else
-				{
-					string remaining = sr.ReadToEnd();
-					long totalSize = readLength + remaining.Length;
-					UhtBuffer combined = UhtBuffer.Borrow((int)totalSize);
-					Buffer.BlockCopy(initialBuffer.Block, 0, combined.Block, 0, readLength * sizeof(char));
-					Buffer.BlockCopy(remaining.ToArray(), 0, combined.Block, readLength * sizeof(char), remaining.Length * sizeof(char));
-					UhtBuffer.Return(initialBuffer);
-					return combined;
-				}
+				using UhtBorrowByteBuffer byteBuffer = new((int)fs.Length);
+				Span<byte> byteSpan = byteBuffer.Buffer.Memory.Span;
+				int readLength = fs.Read(byteSpan);
+				byteSpan = byteSpan.Slice(0, readLength);
+
+				Encoding encoding = GetEncoding(byteSpan, out int skipBytes);
+				byteSpan = byteSpan.Slice(skipBytes);
+
+				int charCount = encoding.GetCharCount(byteBuffer.Buffer.Memory.Span);
+				UhtBuffer initialBuffer = UhtBuffer.Borrow(charCount);
+				encoding.GetChars(byteBuffer.Buffer.Memory.Span, initialBuffer.Memory.Span);
+				return initialBuffer;
 			}
 			catch (IOException)
 			{
@@ -177,29 +165,83 @@ namespace EpicGames.UHT.Utils
 			try
 			{
 				using FileStream fs = new(filePath, FileMode.Open, FileAccess.Read, FileShare.Read, 4 * 1024, FileOptions.SequentialScan);
-				using StreamReader sr = new(fs, Encoding.UTF8, true);
-
-				// Try to read the whole file into a buffer created by hand.  This avoids a LOT of memory allocations which in turn reduces the
-				// GC stress on the system.  Removing the StreamReader would be nice in the future.
-				long rawFileLength = fs.Length;
-				char[] initialBuffer = new char[rawFileLength];
-				int readLength = sr.Read(initialBuffer, 0, (int)rawFileLength);
-				if (sr.EndOfStream)
+				using UhtBorrowByteBuffer byteBuffer = new((int)fs.Length);
+				Span<byte> byteSpan = byteBuffer.Buffer.Memory.Span;
+				int readLength = fs.Read(byteSpan);
+				byteSpan = byteSpan.Slice(0, readLength);
+
+				Encoding encoding = GetEncoding(byteSpan, out int skipBytes);
+				byteSpan = byteSpan.Slice(skipBytes);
+
+				int charCount = encoding.GetCharCount(byteSpan);
+				char[] initialBuffer = new char[charCount];
+				encoding.GetChars(byteSpan, initialBuffer);
+				contents = new StringView(new ReadOnlyMemory<char>(initialBuffer, 0, charCount));
+				return true;
+			}
+			catch (IOException)
+			{
+				contents = new StringView();
+				return false;
+			}
+		}
+
+		/// <summary>
+		/// Get the encoding and number of bytes to skip at the start of the buffer
+		/// </summary>
+		/// <param name="bytes">Bytes to scan for a BOM</param>
+		/// <param name="skipBytes">Number of bytes to skip</param>
+		/// <returns>The encoding</returns>
+		// https://simple.wikipedia.org/wiki/Byte_order_mark
+		private static Encoding GetEncoding(ReadOnlySpan<byte> bytes, out int skipBytes)
+		{
+			skipBytes = 0;
+			int length = bytes.Length;
+			Encoding encoding = Encoding.UTF8;
+
+			if (length < 2)
+			{
+				return encoding;
+			}
+
+			if (bytes[0] == 0xfe && bytes[1] == 0xff)
+			{
+				encoding = Encoding.BigEndianUnicode; // Big 16
+				skipBytes = 2;
+			}
+			else if (bytes[0] == 0xff && bytes[1] == 0xfe)
+			{
+				// Check for FFFE but not FFFE0000
+				if (length < 4 || bytes[2] != 0 || bytes[3] != 0)
 				{
-					contents = new StringView(new ReadOnlyMemory<char>(initialBuffer, 0, readLength));
+					encoding = Encoding.Unicode; // Little 16
+					skipBytes = 2;
 				}
 				else
 				{
-					string remaining = sr.ReadToEnd();
-					contents = new StringView(String.Concat(new ReadOnlySpan<char>(initialBuffer, 0, readLength), remaining));
+					encoding = Encoding.UTF32; // Little 32
+					skipBytes = 4;
 				}
-				return true;
 			}
-			catch (IOException)
+			else if (length >= 3 && bytes[0] == 0xef && bytes[1] == 0xbb && bytes[2] == 0xbf)
 			{
-				contents = new StringView();
-				return false;
+				encoding = Encoding.UTF8; // 8
+				skipBytes = 3;
+			}
+			else if (length >= 4 && bytes[0] == 0 && bytes[1] == 0 && bytes[2] == 0xfe && bytes[3] == 0xff)
+			{
+				encoding = new UTF32Encoding(bigEndian: true, byteOrderMark: true); // Big 32
+				skipBytes = 4;
+			}
+
+			if (encoding.Preamble.Length > 0)
+			{
+				if (bytes.Slice(skipBytes).StartsWith(encoding.Preamble))
+				{
+					skipBytes += encoding.Preamble.Length;
+				}
 			}
+			return encoding;
 		}
 	}
 }
