Horde: Add a default converter for compound types to/from Redis values, serializing each property as a string.

diff --git a/Engine/Source/Programs/Shared/EpicGames.Core/Utf8StringBuilder.cs b/Engine/Source/Programs/Shared/EpicGames.Core/Utf8StringBuilder.cs
--- a/Engine/Source/Programs/Shared/EpicGames.Core/Utf8StringBuilder.cs
+++ b/Engine/Source/Programs/Shared/EpicGames.Core/Utf8StringBuilder.cs
@@ -1,6 +1,7 @@
 // Copyright Epic Games, Inc. All Rights Reserved.
 
 using System;
+using System.Text;
 
 namespace EpicGames.Core
 {
@@ -73,6 +74,69 @@ namespace EpicGames.Core
 			_writer.WriteFixedLengthBytes(text.Span);
 		}
 
+		/// <summary>
+		/// Appends an integer to the end of this builder
+		/// </summary>
+		/// <param name="value">Value to append</param>
+		public void Append(int value) => Append((long)value);
+
+		/// <summary>
+		/// Appends an integer to the end of this builder
+		/// </summary>
+		/// <param name="value">Value to append</param>
+		public void Append(uint value) => Append((ulong)value);
+
+		/// <summary>
+		/// Appends an unsigned long to the end of this builder
+		/// </summary>
+		/// <param name="value">Value to append</param>
+		public void Append(long value)
+		{
+			if (value < 0)
+			{
+				_writer.WriteUInt8((byte)'-');
+				Append((ulong)-value);
+			}
+			else
+			{
+				Append((ulong)value);
+			}
+		}
+
+		/// <summary>
+		/// Appends an unsigned long to the end of this builder
+		/// </summary>
+		/// <param name="value">Value to append</param>
+		public void Append(ulong value)
+		{
+			Span<byte> span = _writer.GetSpan(32);
+
+			int length = 0;
+			for (ulong remaining = value; ; remaining /= 10)
+			{
+				ulong digit = remaining % 10U;
+				span[length++] = (byte)('0' + digit);
+
+				if (remaining == digit)
+				{
+					break;
+				}
+			}
+
+			span.Slice(0, length).Reverse();
+			_writer.Advance(length);
+		}
+
+		/// <summary>
+		/// Appends a string to the end of this builder
+		/// </summary>
+		/// <param name="text">Text to append</param>
+		public void Append(string text)
+		{
+			Span<byte> span = _writer.GetSpanAndAdvance(Encoding.UTF8.GetByteCount(text));
+			Encoding.UTF8.GetBytes(text.AsSpan(), span);
+		}
+
 		/// <summary>
 		/// Converts the written memory to a utf8 string
 		/// </summary>
diff --git a/Engine/Source/Programs/Shared/EpicGames.Redis.Tests/EpicGames.Redis.Tests.csproj b/Engine/Source/Programs/Shared/EpicGames.Redis.Tests/EpicGames.Redis.Tests.csproj
new file mode 100644
--- /dev/null
+++ b/Engine/Source/Programs/Shared/EpicGames.Redis.Tests/EpicGames.Redis.Tests.csproj
@@ -0,0 +1,37 @@
+<Project Sdk="Microsoft.NET.Sdk">
+
+  <PropertyGroup>
+    <TargetFramework>net6.0</TargetFramework>
+    <Configuration Condition=" '$(Configuration)' == '' ">Development</Configuration>
+    <OutputType>Library</OutputType>
+    <Deterministic>true</Deterministic>
+    <RootNamespace>EpicGames.Serialization.Tests</RootNamespace>
+    <GenerateAssemblyCompanyAttribute>false</GenerateAssemblyCompanyAttribute>
+    <GenerateAssemblyInformationalVersionAttribute>false</GenerateAssemblyInformationalVersionAttribute>
+    <GenerateAssemblyProductAttribute>false</GenerateAssemblyProductAttribute>
+    <GenerateAssemblyVersionAttribute>false</GenerateAssemblyVersionAttribute>
+    <Configurations>Debug;Release;Development</Configurations>
+    <IsPackable>false</IsPackable>
+    <Nullable>enable</Nullable>
+  </PropertyGroup>
+  <ItemGroup>
+    <Compile Include="..\..\Shared\MetaData.cs">
+      <Link>Properties\MetaData.cs</Link>
+    </Compile>
+  </ItemGroup>
+
+  <ItemGroup>
+    <PackageReference Include="Microsoft.NET.Test.Sdk" Version="17.3.0" />
+    <PackageReference Include="MSTest.TestAdapter" Version="2.2.10" />
+    <PackageReference Include="MSTest.TestFramework" Version="2.2.10" />
+    <PackageReference Include="coverlet.collector" Version="3.1.2">
+      <PrivateAssets>all</PrivateAssets>
+      <IncludeAssets>runtime; build; native; contentfiles; analyzers; buildtransitive</IncludeAssets>
+    </PackageReference>
+  </ItemGroup>
+
+  <ItemGroup>
+    <ProjectReference Include="..\EpicGames.Redis\EpicGames.Redis.csproj" />
+  </ItemGroup>
+
+</Project>
diff --git a/Engine/Source/Programs/Shared/EpicGames.Redis.Tests/SerializationTests.cs b/Engine/Source/Programs/Shared/EpicGames.Redis.Tests/SerializationTests.cs
new file mode 100644
--- /dev/null
+++ b/Engine/Source/Programs/Shared/EpicGames.Redis.Tests/SerializationTests.cs
@@ -0,0 +1,119 @@
+// Copyright Epic Games, Inc. All Rights Reserved.
+
+using System;
+using System.ComponentModel;
+using System.Globalization;
+using System.IO;
+using System.Text;
+using System.Threading.Tasks;
+using EpicGames.Core;
+using EpicGames.Redis.Converters;
+using Microsoft.VisualStudio.TestTools.UnitTesting;
+using StackExchange.Redis;
+
+namespace EpicGames.Redis.Tests
+{
+	[TestClass]
+	public class SerializationTests
+	{
+		record class TestRecord(int Number, string String, Utf8String StringU8);
+
+		[TestMethod]
+		public void RoundTrip()
+		{
+			TestRecord input = new TestRecord(123, "hello", "world");
+
+			RedisValue value = RedisSerializer.Serialize(input);
+
+			Assert.AreEqual("123|hello|world", (string?)value);
+
+			TestRecord output = RedisSerializer.Deserialize<TestRecord>(value);
+			Assert.AreEqual(123, output.Number);
+			Assert.AreEqual("hello", output.String);
+			Assert.AreEqual("world", output.StringU8);
+		}
+
+		[TestMethod]
+		public void EscapedCharacters()
+		{
+			TestRecord input = new TestRecord(123, "|||", "\\");
+
+			RedisValue value = RedisSerializer.Serialize(input);
+
+			Assert.AreEqual("123|\\|\\|\\||\\\\", (string?)value);
+
+			TestRecord output = RedisSerializer.Deserialize<TestRecord>(value);
+			Assert.AreEqual(123, output.Number);
+			Assert.AreEqual("|||", output.String);
+			Assert.AreEqual("\\", output.StringU8);
+		}
+
+		[TypeConverter(typeof(TestStringTypeConverter))]
+		class TestStringType
+		{
+			public int Number { get; init; }
+		}
+
+		class TestStringTypeConverter : TypeConverter
+		{
+			public override bool CanConvertFrom(ITypeDescriptorContext? context, Type sourceType) => sourceType == typeof(string);
+
+			public override object? ConvertFrom(ITypeDescriptorContext? context, CultureInfo? culture, object value) => new TestStringType { Number = 123 };
+
+			public override bool CanConvertTo(ITypeDescriptorContext? context, Type? destinationType) => destinationType == typeof(string);
+
+			public override object? ConvertTo(ITypeDescriptorContext? context, CultureInfo? culture, object? value, Type destinationType) => "hello";
+		}
+
+		record class TestStringRecord(int Number, TestStringType Instance);
+
+		[TestMethod]
+		public void CustomStringConverter()
+		{
+			TestStringRecord input = new TestStringRecord(123, new TestStringType { Number = 0 });
+
+			RedisClassConverter<TestStringRecord> converter = new RedisClassConverter<TestStringRecord>();
+			RedisValue value = converter.ToRedisValue(input);
+
+			Assert.AreEqual("123|hello", (string?)value);
+
+			TestStringRecord output = converter.FromRedisValue(value);
+			Assert.AreEqual(123, output.Number);
+			Assert.AreEqual(123, output.Instance.Number);
+		}
+
+		[TypeConverter(typeof(TestUtf8StringTypeConverter))]
+		class TestUtf8StringType
+		{
+			public int Number { get; init; }
+		}
+
+		class TestUtf8StringTypeConverter : TypeConverter
+		{
+			public override bool CanConvertFrom(ITypeDescriptorContext? context, Type sourceType) => sourceType == typeof(Utf8String);
+
+			public override object? ConvertFrom(ITypeDescriptorContext? context, CultureInfo? culture, object value) => new TestUtf8StringType { Number = 123 };
+
+			public override bool CanConvertTo(ITypeDescriptorContext? context, Type? destinationType) => destinationType == typeof(Utf8String);
+
+			public override object? ConvertTo(ITypeDescriptorContext? context, CultureInfo? culture, object? value, Type destinationType) => new Utf8String("hello");
+		}
+
+		record class TestUtf8StringRecord(int Number, TestUtf8StringType Instance);
+
+		[TestMethod]
+		public void CustomUtf8StringConverter()
+		{
+			TestUtf8StringRecord input = new TestUtf8StringRecord(123, new TestUtf8StringType { Number = 0 });
+
+			RedisClassConverter<TestUtf8StringRecord> converter = new RedisClassConverter<TestUtf8StringRecord>();
+			RedisValue value = converter.ToRedisValue(input);
+
+			Assert.AreEqual("123|hello", (string?)value);
+
+			TestUtf8StringRecord output = converter.FromRedisValue(value);
+			Assert.AreEqual(123, output.Number);
+			Assert.AreEqual(123, output.Instance.Number);
+		}
+	}
+}
diff --git a/Engine/Source/Programs/Shared/EpicGames.Redis/Converters/RedisClassConverter.cs b/Engine/Source/Programs/Shared/EpicGames.Redis/Converters/RedisClassConverter.cs
new file mode 100644
--- /dev/null
+++ b/Engine/Source/Programs/Shared/EpicGames.Redis/Converters/RedisClassConverter.cs
@@ -0,0 +1,177 @@
+// Copyright Epic Games, Inc. All Rights Reserved.
+
+using EpicGames.Core;
+using StackExchange.Redis;
+using System;
+using System.ComponentModel;
+using System.Diagnostics.CodeAnalysis;
+using System.Reflection;
+
+namespace EpicGames.Redis.Converters
+{
+	/// <summary>
+	/// Converter for records to Redis values.
+	/// </summary>
+	/// <typeparam name="T">The record type</typeparam>
+	public class RedisClassConverter<T> : IRedisConverter<T>
+	{
+		readonly PropertyInfo[] _properties;
+		readonly Func<Utf8String, object?>[] _typeReaders; // Utf8String -> object
+		readonly Action<object?, Utf8StringBuilder>[] _typeWriters;
+		readonly ConstructorInfo _constructor;
+
+		/// <summary>
+		/// Constructor
+		/// </summary>
+		public RedisClassConverter()
+		{
+			Type type = typeof(T);
+			_properties = type.GetProperties(BindingFlags.Public | BindingFlags.Instance);
+
+			_typeReaders = new Func<Utf8String, object?>[_properties.Length];
+			_typeWriters = new Action<object?, Utf8StringBuilder>[_properties.Length];
+
+			for (int idx = 0; idx < _properties.Length; idx++)
+			{
+				Type propertyType = _properties[idx].PropertyType;
+				if (propertyType == typeof(Utf8String))
+				{
+					_typeReaders[idx] = str => str;
+					_typeWriters[idx] = (obj, builder) => WriteEscapedString((Utf8String)obj!, builder);
+				}
+				else if (TryGetUtf8StringTypeConverter(propertyType, out TypeConverter? converter))
+				{
+					_typeReaders[idx] = str => converter.ConvertFrom(str);
+					_typeWriters[idx] = (obj, builder) => WriteEscapedString((Utf8String)converter.ConvertTo(obj, typeof(Utf8String))!, builder);
+				}
+				else if (TryGetStringTypeConverter(propertyType, out converter))
+				{
+					_typeReaders[idx] = str => converter.ConvertFromInvariantString(str.ToString());
+					_typeWriters[idx] = (obj, builder) => WriteEscapedString(converter.ConvertToInvariantString(obj) ?? String.Empty, builder);
+				}
+				else
+				{
+					throw new InvalidOperationException($"Unable to create converter to/from strings for {propertyType.Name}");
+				}
+			}
+
+			ConstructorInfo? constructor = type.GetConstructor(BindingFlags.Instance | BindingFlags.Public, _properties.ConvertAll(x => x.PropertyType));
+			if (constructor == null)
+			{
+				throw new InvalidOperationException("No constructor found with same arguments as properties. Try using a record type.");
+			}
+			_constructor = constructor;
+		}
+
+		static bool TryGetUtf8StringTypeConverter(Type propertyType, [NotNullWhen(true)] out TypeConverter? converter)
+		{
+			converter = TypeDescriptor.GetConverter(propertyType);
+			return converter != null && converter.CanConvertFrom(typeof(Utf8String)) && converter.CanConvertTo(typeof(Utf8String));
+		}
+
+		static bool TryGetStringTypeConverter(Type propertyType, [NotNullWhen(true)] out TypeConverter? converter)
+		{
+			converter = TypeDescriptor.GetConverter(propertyType);
+			return converter != null && converter.CanConvertFrom(typeof(string)) && converter.CanConvertTo(typeof(string));
+		}
+
+		static Utf8String ReadEscapedString(byte[] data, ref int idx)
+		{
+			int startIdx = idx;
+			int endIdx = startIdx;
+			int outIdx = startIdx;
+
+			for (; endIdx < data.Length; endIdx++, outIdx++)
+			{
+				if (data[endIdx] == '|')
+				{
+					endIdx++;
+					break;
+				}
+				else if (data[endIdx] == '\\')
+				{
+					for (; endIdx < data.Length; endIdx++, outIdx++)
+					{
+						if (data[endIdx] == '|')
+						{
+							endIdx++;
+							break;
+						}
+						if (data[endIdx] == '\\')
+						{
+							endIdx++;
+						}
+						data[outIdx] = data[endIdx];
+					}
+					break;
+				}
+			}
+
+			idx = endIdx;
+			return new Utf8String(data.AsMemory(startIdx, outIdx - startIdx));
+		}
+
+		static void WriteEscapedString(Utf8String str, Utf8StringBuilder builder)
+		{
+			if (str.Length > 0)
+			{
+				int minIdx = 0;
+				for (int idx = 0; idx < str.Length; idx++)
+				{
+					if (str[idx] == '\\' || str[idx] == '|')
+					{
+						builder.Append(str.Substring(minIdx, idx - minIdx));
+						builder.Append((byte)'\\');
+						minIdx = idx;
+					}
+				}
+				builder.Append(str.Substring(minIdx));
+			}
+		}
+
+		void AppendProperty(T value, int propertyIdx, Utf8StringBuilder builder)
+		{
+			object? arg = _properties[propertyIdx].GetValue(value);
+			_typeWriters[propertyIdx](arg, builder);
+		}
+
+		/// <inheritdoc/>
+		public RedisValue ToRedisValue(T value)
+		{
+			Utf8StringBuilder builder = new Utf8StringBuilder();
+			if (_properties.Length > 0)
+			{
+				AppendProperty(value, 0, builder);
+				for (int idx = 1; idx < _properties.Length; idx++)
+				{
+					builder.Append((byte)'|');
+					AppendProperty(value, idx, builder);
+				}
+			}
+			return builder.ToString();
+		}
+
+		/// <inheritdoc/>
+		public T FromRedisValue(RedisValue value) => FromBytes((byte[]?)value);
+
+		T FromBytes(byte[]? data)
+		{
+			if (data == null)
+			{
+				return default(T)!;
+			}
+
+			object?[] arguments = new object?[_properties.Length];
+			int argumentIdx = 0;
+
+			for (int idx = 0; idx < data.Length;)
+			{
+				Utf8String str = ReadEscapedString(data, ref idx);
+				arguments[argumentIdx] = _typeReaders[argumentIdx].Invoke(str);
+				argumentIdx++;
+			}
+
+			return (T)_constructor.Invoke(arguments);
+		}
+	}
+}
diff --git a/Engine/Source/Programs/Shared/EpicGames.Redis/RedisSerializer.cs b/Engine/Source/Programs/Shared/EpicGames.Redis/RedisSerializer.cs
--- a/Engine/Source/Programs/Shared/EpicGames.Redis/RedisSerializer.cs
+++ b/Engine/Source/Programs/Shared/EpicGames.Redis/RedisSerializer.cs
@@ -1,6 +1,7 @@
 // Copyright Epic Games, Inc. All Rights Reserved.
 
 using EpicGames.Core;
+using EpicGames.Redis.Converters;
 using EpicGames.Serialization;
 using StackExchange.Redis;
 using System;
@@ -91,6 +92,19 @@ namespace EpicGames.Redis
 			public T FromRedisValue(RedisValue value) => (T)_typeConverter.ConvertFrom((string)value!)!;
 		}
 
+		class RedisUtf8StringConverter<T> : IRedisConverter<T>
+		{
+			readonly TypeConverter _typeConverter;
+
+			public RedisUtf8StringConverter(TypeConverter typeConverter)
+			{
+				_typeConverter = typeConverter;
+			}
+
+			public RedisValue ToRedisValue(T value) => ((Utf8String)_typeConverter.ConvertTo(value, typeof(Utf8String))!).Memory;
+			public T FromRedisValue(RedisValue value) => (T)_typeConverter.ConvertFrom(new Utf8String((ReadOnlyMemory<byte>)value!))!;
+		}
+
 		class RedisNativeConverter<T> : IRedisConverter<T>
 		{
 			readonly Func<RedisValue, T> _fromRedisValueFunc;
@@ -189,11 +203,25 @@ namespace EpicGames.Redis
 
 			// Check if there's a regular converter we can use to convert to/from a string
 			TypeConverter? converter = TypeDescriptor.GetConverter(type);
-			if (converter != null && converter.CanConvertFrom(typeof(string)) && converter.CanConvertTo(typeof(string)))
+			if (converter != null)
+			{
+				if (converter.CanConvertFrom(typeof(Utf8String)) && converter.CanConvertTo(typeof(Utf8String)))
+				{
+					return new RedisUtf8StringConverter<T>(converter);
+				}
+				if (converter.CanConvertFrom(typeof(string)) && converter.CanConvertTo(typeof(string)))
+				{
+					return new RedisStringConverter<T>(converter);
+				}
+			}
+
+			// If it's a compound type, try to create a class converter
+			if (type.IsClass)
 			{
-				return new RedisStringConverter<T>(converter);
+				return new RedisClassConverter<T>();
 			}
 
+			// Otherwise fail
 			throw new Exception($"Unable to find Redis converter for {type.Name}");
 		}
 
