Missing cleanup of encoding detection code

#rnx

diff --git a/Engine/Source/Programs/Shared/EpicGames.UHT/Utils/UhtFileManager.cs b/Engine/Source/Programs/Shared/EpicGames.UHT/Utils/UhtFileManager.cs
--- a/Engine/Source/Programs/Shared/EpicGames.UHT/Utils/UhtFileManager.cs
+++ b/Engine/Source/Programs/Shared/EpicGames.UHT/Utils/UhtFileManager.cs
@@ -196,49 +196,38 @@ namespace EpicGames.UHT.Utils
 		private static Encoding GetEncoding(ReadOnlySpan<byte> bytes, out int skipBytes)
 		{
 			skipBytes = 0;
-			int length = bytes.Length;
 			Encoding encoding = Encoding.UTF8;
 
-			if (length < 2)
-			{
-				return encoding;
-			}
-
-			if (bytes[0] == 0xfe && bytes[1] == 0xff)
-			{
-				encoding = Encoding.BigEndianUnicode; // Big 16
-				skipBytes = 2;
-			}
-			else if (bytes[0] == 0xff && bytes[1] == 0xfe)
+			int length = bytes.Length;
+			if (length >= 2)
 			{
-				// Check for FFFE but not FFFE0000
-				if (length < 4 || bytes[2] != 0 || bytes[3] != 0)
+				if (bytes[0] == 0xfe && bytes[1] == 0xff)
 				{
-					encoding = Encoding.Unicode; // Little 16
+					encoding = Encoding.BigEndianUnicode; // Big 16
 					skipBytes = 2;
 				}
-				else
+				else if (bytes[0] == 0xff && bytes[1] == 0xfe)
 				{
-					encoding = Encoding.UTF32; // Little 32
-					skipBytes = 4;
+					if (length >= 4 && bytes[2] == 0 && bytes[3] == 0)
+					{
+						encoding = Encoding.UTF32; // Little 32
+						skipBytes = 4;
+					}
+					else
+					{
+						encoding = Encoding.Unicode; // Little 16
+						skipBytes = 2;
+					}
 				}
-			}
-			else if (length >= 3 && bytes[0] == 0xef && bytes[1] == 0xbb && bytes[2] == 0xbf)
-			{
-				encoding = Encoding.UTF8; // 8
-				skipBytes = 3;
-			}
-			else if (length >= 4 && bytes[0] == 0 && bytes[1] == 0 && bytes[2] == 0xfe && bytes[3] == 0xff)
-			{
-				encoding = new UTF32Encoding(bigEndian: true, byteOrderMark: true); // Big 32
-				skipBytes = 4;
-			}
-
-			if (encoding.Preamble.Length > 0)
-			{
-				if (bytes.Slice(skipBytes).StartsWith(encoding.Preamble))
+				else if (length >= 3 && bytes[0] == 0xef && bytes[1] == 0xbb && bytes[2] == 0xbf)
+				{
+					encoding = Encoding.UTF8;
+					skipBytes = 3;
+				}
+				else if (length >= 4 && bytes[0] == 0 && bytes[1] == 0 && bytes[2] == 0xfe && bytes[3] == 0xff)
 				{
-					skipBytes += encoding.Preamble.Length;
+					encoding = new UTF32Encoding(bigEndian: true, byteOrderMark: true); // Big 32
+					skipBytes = 4;
 				}
 			}
 			return encoding;
