Horde: Include the hash of the target node in refs, allowing TreeNodeRef objects to be constructed directly.

diff --git a/Engine/Source/Programs/Horde/Horde.Agent.Tests/JsonRpcLoggerTests.cs b/Engine/Source/Programs/Horde/Horde.Agent.Tests/JsonRpcLoggerTests.cs
--- a/Engine/Source/Programs/Horde/Horde.Agent.Tests/JsonRpcLoggerTests.cs
+++ b/Engine/Source/Programs/Horde/Horde.Agent.Tests/JsonRpcLoggerTests.cs
@@ -104,16 +104,16 @@ namespace Horde.Agent.Tests
 
 		class FakeJsonRpcLoggerBackend : JsonRpcAndStorageLogSink
 		{
-			public NodeLocator Locator { get; private set; }
+			public RefTarget? Target { get; private set; }
 
 			public FakeJsonRpcLoggerBackend(IRpcConnection connection, string logId, IJsonRpcLogSink inner, IStorageClient store, ILogger logger)
 				: base(connection, logId, inner, store, logger)
 			{
 			}
 
-			protected override Task UpdateLogAsync(NodeLocator locator, int lineCount, CancellationToken cancellationToken)
+			protected override Task UpdateLogAsync(RefTarget target, int lineCount, CancellationToken cancellationToken)
 			{
-				Locator = locator;
+				Target = target;
 				return Task.CompletedTask;
 			}
 
@@ -157,7 +157,7 @@ namespace Horde.Agent.Tests
 						logger.LogInformation("Testing {Number}", idx);
 					}
 				}
-				file = await reader.ReadNodeAsync<LogNode>(sink.Locator);
+				file = await reader.ReadNodeAsync<LogNode>(sink.Target!);
 			}
 
 			// Check the index text
diff --git a/Engine/Source/Programs/Horde/Horde.Agent/Commands/Bundles/DumpCommand.cs b/Engine/Source/Programs/Horde/Horde.Agent/Commands/Bundles/DumpCommand.cs
--- a/Engine/Source/Programs/Horde/Horde.Agent/Commands/Bundles/DumpCommand.cs
+++ b/Engine/Source/Programs/Horde/Horde.Agent/Commands/Bundles/DumpCommand.cs
@@ -27,8 +27,8 @@ namespace Horde.Agent.Commands.Bundles
 
 			if (BlobId == null)
 			{
-				NodeLocator locator = await store.ReadRefTargetAsync(RefName);
-				BlobId = locator.Blob;
+				RefTarget locator = await store.ReadRefTargetAsync(RefName);
+				BlobId = locator.Locator.Blob;
 			}
 
 			logger.LogInformation("Summary for blob {BlobId}", BlobId.Value);
diff --git a/Engine/Source/Programs/Horde/Horde.Agent/Utility/JsonRpcLogger.cs b/Engine/Source/Programs/Horde/Horde.Agent/Utility/JsonRpcLogger.cs
--- a/Engine/Source/Programs/Horde/Horde.Agent/Utility/JsonRpcLogger.cs
+++ b/Engine/Source/Programs/Horde/Horde.Agent/Utility/JsonRpcLogger.cs
@@ -213,8 +213,8 @@ namespace Horde.Agent.Parser
 
 			if (request.Flush || _bufferLength > FlushLength)
 			{
-				NodeLocator locator = await _builder.FlushAsync(_writer, request.Flush, cancellationToken);
-				await UpdateLogAsync(locator, _builder.LineCount, cancellationToken);
+				RefTarget target = await _builder.FlushAsync(_writer, request.Flush, cancellationToken);
+				await UpdateLogAsync(target, _builder.LineCount, cancellationToken);
 				_bufferLength = 0;
 			}
 
@@ -223,14 +223,14 @@ namespace Horde.Agent.Parser
 
 		#region RPC calls
 
-		protected virtual async Task UpdateLogAsync(NodeLocator locator, int lineCount, CancellationToken cancellationToken)
+		protected virtual async Task UpdateLogAsync(RefTarget target, int lineCount, CancellationToken cancellationToken)
 		{
-			_logger.LogDebug("Updating log {LogId} to line {LineCount}, blob {Locator}", _logId, lineCount, locator);
+			_logger.LogDebug("Updating log {LogId} to line {LineCount}, blob {Locator}", _logId, lineCount, target.Locator);
 
 			UpdateLogRequest request = new UpdateLogRequest();
 			request.LogId = _logId;
 			request.LineCount = lineCount;
-			request.BlobLocator = locator.ToString();
+			request.Target = target.ToString();
 			await _connection.InvokeAsync((LogRpcClient client) => client.UpdateLogAsync(request, cancellationToken: cancellationToken), cancellationToken);
 		}
 
diff --git a/Engine/Source/Programs/Horde/Horde.Build.Tests/BlobStoreTests.cs b/Engine/Source/Programs/Horde/Horde.Build.Tests/BlobStoreTests.cs
--- a/Engine/Source/Programs/Horde/Horde.Build.Tests/BlobStoreTests.cs
+++ b/Engine/Source/Programs/Horde/Horde.Build.Tests/BlobStoreTests.cs
@@ -106,19 +106,22 @@ namespace Horde.Build.Tests
 			IStorageClient store = await CreateStorageClientAsync();
 
 			byte[] input1 = CreateTestData(256, 1);
-			BlobLocator locator1 = await store.WriteBundleAsync(CreateTestBundle(input1, Array.Empty<BlobLocator>()));
+			Bundle bundle1 = CreateTestBundle(input1, Array.Empty<BlobLocator>());
+			BlobLocator locator1 = await store.WriteBundleAsync(bundle1);
 			Blob blob1 = await ReadBlobAsync(store, locator1);
 			Assert.IsTrue(blob1.Data.Span.SequenceEqual(input1));
 			Assert.IsTrue(blob1.References.SequenceEqual(Array.Empty<BlobLocator>()));
 
 			byte[] input2 = CreateTestData(256, 2);
-			BlobLocator locator2 = await store.WriteBundleAsync(CreateTestBundle(input2, new BlobLocator[] { locator1 }));
+			Bundle bundle2 = CreateTestBundle(input2, new BlobLocator[] { locator1 });
+			BlobLocator locator2 = await store.WriteBundleAsync(bundle2);
 			Blob blob2 = await ReadBlobAsync(store, locator2);
 			Assert.IsTrue(blob2.Data.Span.SequenceEqual(input2));
 			Assert.IsTrue(blob2.References.SequenceEqual(new BlobLocator[] { locator1 }));
 
 			byte[] input3 = CreateTestData(256, 3);
-			BlobLocator locator3 = await store.WriteBundleAsync(CreateTestBundle(input3, new BlobLocator[] { locator1, locator2, locator1 }));
+			Bundle bundle3 = CreateTestBundle(input3, new BlobLocator[] { locator1, locator2, locator1 });
+			BlobLocator locator3 = await store.WriteBundleAsync(bundle3);
 			Blob blob3 = await ReadBlobAsync(store, locator3);
 			Assert.IsTrue(blob3.Data.Span.SequenceEqual(input3));
 			Assert.IsTrue(blob3.References.SequenceEqual(new BlobLocator[] { locator1, locator2, locator1 }));
@@ -126,15 +129,15 @@ namespace Horde.Build.Tests
 			for(int idx = 0; idx < 2; idx++)
 			{
 				RefName refName = new RefName("hello");
-				await store.WriteRefTargetAsync(refName, new NodeLocator(locator3, 0));
-				NodeLocator refTarget = await store.ReadRefTargetAsync(refName);
-				Assert.AreEqual(locator3, refTarget.Blob);
+				await store.WriteRefTargetAsync(refName, new RefTarget(bundle3.Header.Exports[0].Hash, new NodeLocator(locator3, 0)));
+				RefTarget refTarget = await store.ReadRefTargetAsync(refName);
+				Assert.AreEqual(locator3, refTarget.Locator.Blob);
 			}
 
 			RefName refName2 = new RefName("hello2");
 
-			NodeLocator refTargetId2 = await store.WriteRefAsync(refName2, CreateTestBundle(input3, new BlobLocator[] { locator1, locator2 }), 0);
-			Blob refTarget2 = await ReadBlobAsync(store, refTargetId2.Blob);
+			RefTarget refTargetId2 = await store.WriteRefAsync(refName2, CreateTestBundle(input3, new BlobLocator[] { locator1, locator2 }), 0);
+			Blob refTarget2 = await ReadBlobAsync(store, refTargetId2.Locator.Blob);
 			Assert.IsTrue(refTarget2.Data.Span.SequenceEqual(input3));
 			Assert.IsTrue(refTarget2.References.SequenceEqual(new BlobLocator[] { locator1, locator2 }));
 		}
@@ -144,8 +147,9 @@ namespace Horde.Build.Tests
 		{
 			IStorageClient store = await CreateStorageClientAsync();
 
-			BlobLocator locator1 = await store.WriteBundleAsync(CreateTestBundle(new byte[] { 1, 2, 3 }, Array.Empty<BlobLocator>()));
-			NodeLocator target = new NodeLocator(locator1, 0);
+			Bundle bundle1 = CreateTestBundle(new byte[] { 1, 2, 3 }, Array.Empty<BlobLocator>());
+			BlobLocator locator1 = await store.WriteBundleAsync(bundle1);
+			RefTarget target = new RefTarget(bundle1.Header.Exports[0].Hash, new NodeLocator(locator1, 0));
 
 			await store.WriteRefTargetAsync("test-ref-1", target);
 			await store.WriteRefTargetAsync("test-ref-2", target, new RefOptions { Lifetime = TimeSpan.FromMinutes(30.0), Extend = true });
diff --git a/Engine/Source/Programs/Horde/Horde.Build.Tests/GcServiceTests.cs b/Engine/Source/Programs/Horde/Horde.Build.Tests/GcServiceTests.cs
--- a/Engine/Source/Programs/Horde/Horde.Build.Tests/GcServiceTests.cs
+++ b/Engine/Source/Programs/Horde/Horde.Build.Tests/GcServiceTests.cs
@@ -46,7 +46,7 @@ namespace Horde.Build.Tests
 				int blobIdx = (int)(random.NextDouble() * blobs.Length);
 				if (roots.Add(blobs[blobIdx]))
 				{
-					await store.WriteRefTargetAsync(new RefName($"ref-{idx}"), new NodeLocator(blobs[blobIdx], 0));
+					await store.WriteRefTargetAsync(new RefName($"ref-{idx}"), new RefTarget(IoHash.Zero, new NodeLocator(blobs[blobIdx], 0)));
 				}
 			}
 
diff --git a/Engine/Source/Programs/Horde/Horde.Build/Logs/LogRpcService.cs b/Engine/Source/Programs/Horde/Horde.Build/Logs/LogRpcService.cs
--- a/Engine/Source/Programs/Horde/Horde.Build/Logs/LogRpcService.cs
+++ b/Engine/Source/Programs/Horde/Horde.Build/Logs/LogRpcService.cs
@@ -52,8 +52,8 @@ namespace Horde.Build.Logs
 			}
 
 			IStorageClient store = await _storageService.GetClientAsync(Namespace.Logs, context.CancellationToken);
-			_logger.LogInformation("Updating {LogId} to node {Locator}", request.LogId, request.BlobLocator);
-			await store.WriteRefTargetAsync(new RefName(request.LogId), NodeLocator.Parse(request.BlobLocator));
+			_logger.LogInformation("Updating {LogId} to node {RefTarget}", request.LogId, request.Target);
+			await store.WriteRefTargetAsync(new RefName(request.LogId), RefTarget.Parse(request.Target));
 
 			await _logFileCollection.UpdateLineCountAsync(logFile, request.LineCount, CancellationToken.None);
 
diff --git a/Engine/Source/Programs/Horde/Horde.Build/Storage/StorageController.cs b/Engine/Source/Programs/Horde/Horde.Build/Storage/StorageController.cs
--- a/Engine/Source/Programs/Horde/Horde.Build/Storage/StorageController.cs
+++ b/Engine/Source/Programs/Horde/Horde.Build/Storage/StorageController.cs
@@ -26,7 +26,7 @@ namespace Horde.Build.Storage
 		/// <summary>
 		/// Locator for the uploaded bundle
 		/// </summary>
-		public BlobLocator Locator { get; set; }
+		public BlobLocator Blob { get; set; }
 
 		/// <summary>
 		/// URL to upload the blob to.
@@ -47,13 +47,18 @@ namespace Horde.Build.Storage
 		/// <summary>
 		/// Locator for the target blob
 		/// </summary>
-		public BlobLocator Locator { get; set; }
+		public BlobLocator Blob { get; set; }
 
 		/// <summary>
 		/// Export index for the ref
 		/// </summary>
 		public int ExportIdx { get; set; }
 
+		/// <summary>
+		/// Hash of the target node
+		/// </summary>
+		public IoHash Hash { get; set; }
+
 		/// <summary>
 		/// Options for the ref
 		/// </summary>
@@ -75,13 +80,19 @@ namespace Horde.Build.Storage
 		/// </summary>
 		public int ExportIdx { get; set; }
 
+		/// <summary>
+		/// Hash of the target node
+		/// </summary>
+		public IoHash Hash { get; set; }
+
 		/// <summary>
 		/// Constructor
 		/// </summary>
-		public ReadRefResponse(NodeLocator target)
+		public ReadRefResponse(RefTarget target)
 		{
-			Blob = target.Blob;
-			ExportIdx = target.ExportIdx;
+			Blob = target.Locator.Blob;
+			ExportIdx = target.Locator.ExportIdx;
+			Hash = target.Hash;
 		}
 	}
 
@@ -130,7 +141,7 @@ namespace Horde.Build.Storage
 				}
 				else
 				{
-					return new WriteBlobResponse { Locator = result.Value.Locator, UploadUrl = result.Value.UploadUrl };
+					return new WriteBlobResponse { Blob = result.Value.Locator, UploadUrl = result.Value.UploadUrl };
 				}
 			}
 			else
@@ -138,7 +149,7 @@ namespace Horde.Build.Storage
 				using (Stream stream = file.OpenReadStream())
 				{
 					BlobLocator locator = await client.WriteBlobAsync(stream, prefix: (prefix == null) ? Utf8String.Empty : new Utf8String(prefix), cancellationToken: cancellationToken);
-					return new WriteBlobResponse { Locator = locator, SupportsRedirects = client.SupportsRedirects? (bool?)true : null };
+					return new WriteBlobResponse { Blob = locator, SupportsRedirects = client.SupportsRedirects? (bool?)true : null };
 				}
 			}
 		}
@@ -164,7 +175,7 @@ namespace Horde.Build.Storage
 				}
 				else
 				{
-					return new WriteBlobResponse { Locator = result.Value.Locator, UploadUrl = result.Value.UploadUrl };
+					return new WriteBlobResponse { Blob = result.Value.Locator, UploadUrl = result.Value.UploadUrl };
 				}
 			}
 			else
@@ -172,7 +183,7 @@ namespace Horde.Build.Storage
 				using (Stream stream = file.OpenReadStream())
 				{
 					BlobLocator locator = await client.WriteBlobAsync(stream, prefix: (prefix == null) ? Utf8String.Empty : new Utf8String(prefix), cancellationToken: cancellationToken);
-					return new WriteBlobResponse { Locator = locator, SupportsRedirects = client.SupportsRedirects ? (bool?)true : null };
+					return new WriteBlobResponse { Blob = locator, SupportsRedirects = client.SupportsRedirects ? (bool?)true : null };
 				}
 			}
 		}
@@ -238,7 +249,7 @@ namespace Horde.Build.Storage
 			}
 
 			IStorageClient client = await _storageService.GetClientAsync(namespaceId, cancellationToken);
-			NodeLocator target = new NodeLocator(request.Locator, request.ExportIdx);
+			RefTarget target = new RefTarget(request.Hash, new NodeLocator(request.Blob, request.ExportIdx));
 			await client.WriteRefTargetAsync(refName, target, request.Options, cancellationToken);
 
 			return Ok();
@@ -261,8 +272,8 @@ namespace Horde.Build.Storage
 
 			IStorageClient client = await _storageService.GetClientAsync(namespaceId, cancellationToken);
 
-			NodeLocator target = await client.TryReadRefTargetAsync(refName, cancellationToken: cancellationToken);
-			if (!target.IsValid())
+			RefTarget? target = await client.TryReadRefTargetAsync(refName, cancellationToken: cancellationToken);
+			if (target == null)
 			{
 				return NotFound();
 			}
diff --git a/Engine/Source/Programs/Horde/Horde.Build/Storage/StorageService.cs b/Engine/Source/Programs/Horde/Horde.Build/Storage/StorageService.cs
--- a/Engine/Source/Programs/Horde/Horde.Build/Storage/StorageService.cs
+++ b/Engine/Source/Programs/Horde/Horde.Build/Storage/StorageService.cs
@@ -212,10 +212,10 @@ namespace Horde.Build.Storage
 			#region Refs
 
 			/// <inheritdoc/>
-			public override Task<NodeLocator> TryReadRefTargetAsync(RefName name, DateTime cacheTime = default, CancellationToken cancellationToken = default) => _outer.TryReadRefTargetAsync(NamespaceId, name, cacheTime, cancellationToken);
+			public override Task<RefTarget?> TryReadRefTargetAsync(RefName name, DateTime cacheTime = default, CancellationToken cancellationToken = default) => _outer.TryReadRefTargetAsync(NamespaceId, name, cacheTime, cancellationToken);
 
 			/// <inheritdoc/>
-			public override Task WriteRefTargetAsync(RefName name, NodeLocator target, RefOptions? options = null, CancellationToken cancellationToken = default) => _outer.WriteRefTargetAsync(NamespaceId, name, target, options, cancellationToken);
+			public override Task WriteRefTargetAsync(RefName name, RefTarget target, RefOptions? options = null, CancellationToken cancellationToken = default) => _outer.WriteRefTargetAsync(NamespaceId, name, target, options, cancellationToken);
 
 			/// <inheritdoc/>
 			public override Task DeleteRefAsync(RefName name, CancellationToken cancellationToken = default) => _outer.DeleteRefAsync(NamespaceId, name, cancellationToken);
@@ -310,6 +310,9 @@ namespace Horde.Build.Storage
 			[BsonElement("idx")]
 			public int ExportIdx { get; set; }
 
+			[BsonElement("hash")]
+			public IoHash Hash { get; set; }
+
 			[BsonElement("xa"), BsonIgnoreIfDefault]
 			public DateTime? ExpiresAtUtc { get; set; }
 
@@ -317,7 +320,7 @@ namespace Horde.Build.Storage
 			public TimeSpan? Lifetime { get; set; }
 
 			[BsonIgnore]
-			public NodeLocator Target => new NodeLocator(BlobLocator, ExportIdx);
+			public RefTarget Target => new RefTarget(Hash, new NodeLocator(BlobLocator, ExportIdx));
 
 			public RefInfo()
 			{
@@ -325,13 +328,14 @@ namespace Horde.Build.Storage
 				BlobLocator = BlobLocator.Empty;
 			}
 
-			public RefInfo(NamespaceId namespaceId, RefName name, NodeLocator target, ObjectId blobInfoId)
+			public RefInfo(NamespaceId namespaceId, RefName name, RefTarget target, ObjectId blobInfoId)
 			{
 				NamespaceId = namespaceId;
 				Name = name;
-				BlobLocator = target.Blob;
+				BlobLocator = target.Locator.Blob;
 				BlobInfoId = blobInfoId;
-				ExportIdx = target.ExportIdx;
+				ExportIdx = target.Locator.ExportIdx;
+				Hash = target.Hash;
 			}
 
 			public bool HasExpired(DateTime utcNow) => ExpiresAtUtc.HasValue && utcNow >= ExpiresAtUtc.Value;
@@ -852,7 +856,7 @@ namespace Horde.Build.Storage
 		}
 
 		/// <inheritdoc/>
-		async Task<NodeLocator> TryReadRefTargetAsync(NamespaceId namespaceId, RefName name, DateTime cacheTime = default, CancellationToken cancellationToken = default)
+		async Task<RefTarget?> TryReadRefTargetAsync(NamespaceId namespaceId, RefName name, DateTime cacheTime = default, CancellationToken cancellationToken = default)
 		{
 			RefCacheValue entry;
 			if (!_cache.TryGetValue(name, out entry) || entry.Time < cacheTime)
@@ -863,7 +867,7 @@ namespace Horde.Build.Storage
 
 			if (entry.Value == null)
 			{
-				return default;
+				return null;
 			}
 
 			if (entry.Value.ExpiresAtUtc != null)
@@ -880,16 +884,16 @@ namespace Horde.Build.Storage
 				}
 			}
 
-			return new NodeLocator(entry.Value.BlobLocator, entry.Value.ExportIdx);
+			return entry.Value.Target;
 		}
 
 		/// <inheritdoc/>
-		async Task WriteRefTargetAsync(NamespaceId namespaceId, RefName name, NodeLocator target, RefOptions? options = null, CancellationToken cancellationToken = default)
+		async Task WriteRefTargetAsync(NamespaceId namespaceId, RefName name, RefTarget target, RefOptions? options = null, CancellationToken cancellationToken = default)
 		{
-			BlobInfo? newBlobInfo = await _blobCollection.Find(x => x.NamespaceId == namespaceId && x.BlobId == target.Blob.BlobId).FirstOrDefaultAsync(cancellationToken);
+			BlobInfo? newBlobInfo = await _blobCollection.Find(x => x.NamespaceId == namespaceId && x.BlobId == target.Locator.Blob.BlobId).FirstOrDefaultAsync(cancellationToken);
 			if (newBlobInfo == null)
 			{
-				throw new Exception($"Invalid/unknown blob identifier '{target.Blob.BlobId}' in namespace {namespaceId}");
+				throw new Exception($"Invalid/unknown blob identifier '{target.Locator.Blob.BlobId}' in namespace {namespaceId}");
 			}
 
 			RefInfo newRefInfo = new RefInfo(namespaceId, name, target, newBlobInfo.Id);
diff --git a/Engine/Source/Programs/Shared/EpicGames.Core/StringUtils.cs b/Engine/Source/Programs/Shared/EpicGames.Core/StringUtils.cs
--- a/Engine/Source/Programs/Shared/EpicGames.Core/StringUtils.cs
+++ b/Engine/Source/Programs/Shared/EpicGames.Core/StringUtils.cs
@@ -327,12 +327,18 @@ namespace EpicGames.Core
 		/// Parses a hexadecimal string into an array of bytes
 		/// </summary>
 		/// <returns>Array of bytes</returns>
-		public static byte[] ParseHexString(string text)
+		public static byte[] ParseHexString(string text) => ParseHexString(text.AsSpan());
+
+		/// <summary>
+		/// Parses a hexadecimal string into an array of bytes
+		/// </summary>
+		/// <returns>Array of bytes</returns>
+		public static byte[] ParseHexString(ReadOnlySpan<char> text)
 		{
 			byte[]? bytes;
 			if(!TryParseHexString(text, out bytes))
 			{
-				throw new FormatException(String.Format("Invalid hex string: '{0}'", text));
+				throw new FormatException(String.Format("Invalid hex string: '{0}'", text.ToString()));
 			}
 			return bytes;
 		}
@@ -357,7 +363,7 @@ namespace EpicGames.Core
 		/// <param name="text">Text to parse</param>
 		/// <param name="outBytes">Receives the parsed string</param>
 		/// <returns></returns>
-		public static bool TryParseHexString(string text, [NotNullWhen(true)] out byte[]? outBytes)
+		public static bool TryParseHexString(ReadOnlySpan<char> text, [NotNullWhen(true)] out byte[]? outBytes)
 		{
 			if((text.Length & 1) != 0)
 			{
diff --git a/Engine/Source/Programs/Shared/EpicGames.Horde.Tests/BundleTests.cs b/Engine/Source/Programs/Shared/EpicGames.Horde.Tests/BundleTests.cs
--- a/Engine/Source/Programs/Shared/EpicGames.Horde.Tests/BundleTests.cs
+++ b/Engine/Source/Programs/Shared/EpicGames.Horde.Tests/BundleTests.cs
@@ -185,8 +185,8 @@ namespace EpicGames.Horde.Tests
 			Assert.AreEqual(1, store.Blobs.Count);
 
 			// Check the ref
-			NodeLocator refTarget = await store.ReadRefTargetAsync(refName);
-			Bundle bundle = await store.ReadBundleAsync(refTarget.Blob);
+			RefTarget refTarget = await store.ReadRefTargetAsync(refName);
+			Bundle bundle = await store.ReadBundleAsync(refTarget.Locator.Blob);
 			Assert.AreEqual(0, bundle.Header.Imports.Count);
 			Assert.AreEqual(3, bundle.Header.Exports.Count);
 
diff --git a/Engine/Source/Programs/Shared/EpicGames.Horde.Tests/LogTests.cs b/Engine/Source/Programs/Shared/EpicGames.Horde.Tests/LogTests.cs
--- a/Engine/Source/Programs/Shared/EpicGames.Horde.Tests/LogTests.cs
+++ b/Engine/Source/Programs/Shared/EpicGames.Horde.Tests/LogTests.cs
@@ -87,8 +87,8 @@ namespace EpicGames.Horde.Tests
 			}
 
 			// Flush it to storage, and read the finished log node
-			NodeLocator locator = await builder.FlushAsync(new TreeWriter(_store), true, CancellationToken.None);
-			LogNode log = await _reader.ReadNodeAsync<LogNode>(locator);
+			RefTarget target = await builder.FlushAsync(new TreeWriter(_store), true, CancellationToken.None);
+			LogNode log = await _reader.ReadNodeAsync<LogNode>(target.Locator);
 
 			// Read the data back out and check it's the same
 			byte[] readData = new byte[_data.Length];
@@ -129,10 +129,10 @@ namespace EpicGames.Horde.Tests
 			{
 				builder.WriteData(Encoding.UTF8.GetBytes(lines[lineIdx]));
 			}
-			NodeLocator rootNodeLocator = await builder.FlushAsync(new TreeWriter(_store), true, CancellationToken.None);
+			RefTarget rootNodeLocator = await builder.FlushAsync(new TreeWriter(_store), true, CancellationToken.None);
 
 			// Read it back in and test the index
-			LogNode rootNode = await _reader.ReadNodeAsync<LogNode>(rootNodeLocator);
+			LogNode rootNode = await _reader.ReadNodeAsync<LogNode>(rootNodeLocator.Locator);
 			LogIndexNode index = await rootNode.IndexRef.ExpandAsync(_reader);
 		
 			for (int lineIdx = 0; lineIdx < lines.Length; lineIdx++)
diff --git a/Engine/Source/Programs/Shared/EpicGames.Horde/Logs/LogNode.cs b/Engine/Source/Programs/Shared/EpicGames.Horde/Logs/LogNode.cs
--- a/Engine/Source/Programs/Shared/EpicGames.Horde/Logs/LogNode.cs
+++ b/Engine/Source/Programs/Shared/EpicGames.Horde/Logs/LogNode.cs
@@ -269,7 +269,7 @@ namespace EpicGames.Horde.Logs
 		/// <param name="writer">Writer for the output nodes</param>
 		/// <param name="complete">Whether the log is complete</param>
 		/// <param name="cancellationToken">Cancellation token for the operation</param>
-		public async Task<NodeLocator> FlushAsync(TreeWriter writer, bool complete, CancellationToken cancellationToken)
+		public async Task<RefTarget> FlushAsync(TreeWriter writer, bool complete, CancellationToken cancellationToken)
 		{
 			// Capture the new data that needs to be written
 			IReadOnlyList<LogChunkNode> writeTextChunks;
@@ -313,7 +313,7 @@ namespace EpicGames.Horde.Logs
 				_indexTextBuilder.Remove(writeIndexTextChunks.Count);
 			}
 
-			return newRootLocator;
+			return new RefTarget(newRootHash, newRootLocator);
 		}
 	}
 
diff --git a/Engine/Source/Programs/Shared/EpicGames.Horde/Storage/Backends/FileStorageClient.cs b/Engine/Source/Programs/Shared/EpicGames.Horde/Storage/Backends/FileStorageClient.cs
--- a/Engine/Source/Programs/Shared/EpicGames.Horde/Storage/Backends/FileStorageClient.cs
+++ b/Engine/Source/Programs/Shared/EpicGames.Horde/Storage/Backends/FileStorageClient.cs
@@ -82,7 +82,7 @@ namespace EpicGames.Horde.Storage.Backends
 		}
 
 		/// <inheritdoc/>
-		public override async Task<NodeLocator> TryReadRefTargetAsync(RefName name, DateTime cacheTime = default, CancellationToken cancellationToken = default)
+		public override async Task<RefTarget?> TryReadRefTargetAsync(RefName name, DateTime cacheTime = default, CancellationToken cancellationToken = default)
 		{
 			FileReference file = GetRefFile(name);
 			if (!FileReference.Exists(file))
@@ -92,21 +92,16 @@ namespace EpicGames.Horde.Storage.Backends
 
 			_logger.LogInformation("Reading {File}", file);
 			string text = await FileReference.ReadAllTextAsync(file, cancellationToken);
-
-			int hashIdx = text.IndexOf('#', StringComparison.Ordinal);
-			BlobLocator locator = new BlobLocator(text.Substring(0, hashIdx));
-			int exportIdx = Int32.Parse(text.Substring(hashIdx + 1), CultureInfo.InvariantCulture);
-
-			return new NodeLocator(locator, exportIdx);
+			return RefTarget.Parse(text);
 		}
 
 		/// <inheritdoc/>
-		public override async Task WriteRefTargetAsync(RefName name, NodeLocator target, RefOptions? options = null, CancellationToken cancellationToken = default)
+		public override async Task WriteRefTargetAsync(RefName name, RefTarget target, RefOptions? options = null, CancellationToken cancellationToken = default)
 		{
 			FileReference file = GetRefFile(name);
 			DirectoryReference.CreateDirectory(file.Directory);
 			_logger.LogInformation("Writing {File}", file);
-			await FileReference.WriteAllTextAsync(file, $"{target.Blob}#{target.ExportIdx}");
+			await FileReference.WriteAllTextAsync(file, target.ToString());
 		}
 
 		#endregion
diff --git a/Engine/Source/Programs/Shared/EpicGames.Horde/Storage/Backends/HttpStorageClient.cs b/Engine/Source/Programs/Shared/EpicGames.Horde/Storage/Backends/HttpStorageClient.cs
--- a/Engine/Source/Programs/Shared/EpicGames.Horde/Storage/Backends/HttpStorageClient.cs
+++ b/Engine/Source/Programs/Shared/EpicGames.Horde/Storage/Backends/HttpStorageClient.cs
@@ -37,6 +37,7 @@ namespace EpicGames.Horde.Storage.Backends
 		{
 			public BlobLocator Blob { get; set; }
 			public int ExportIdx { get; set; }
+			public IoHash Hash { get; set; }
 		}
 
 		readonly Func<HttpClient> _createClient;
@@ -195,7 +196,7 @@ namespace EpicGames.Horde.Storage.Backends
 		}
 
 		/// <inheritdoc/>
-		public override async Task<NodeLocator> TryReadRefTargetAsync(RefName name, DateTime cacheTime = default, CancellationToken cancellationToken = default)
+		public override async Task<RefTarget?> TryReadRefTargetAsync(RefName name, DateTime cacheTime = default, CancellationToken cancellationToken = default)
 		{
 			using (HttpClient httpClient = _createClient())
 			{
@@ -206,7 +207,7 @@ namespace EpicGames.Horde.Storage.Backends
 						if (response.StatusCode == HttpStatusCode.NotFound)
 						{
 							_logger.LogDebug("Read ref {RefName} -> None", name);
-							return default;
+							return null;
 						}
 						else if (!response.IsSuccessStatusCode)
 						{
@@ -218,7 +219,7 @@ namespace EpicGames.Horde.Storage.Backends
 							response.EnsureSuccessStatusCode();
 							ReadRefResponse? data = await response.Content.ReadFromJsonAsync<ReadRefResponse>(cancellationToken: cancellationToken);
 							_logger.LogDebug("Read ref {RefName} -> {Blob}#{ExportIdx}", name, data!.Blob, data!.ExportIdx);
-							return new NodeLocator(data!.Blob, data!.ExportIdx);
+							return new(data.Hash, new NodeLocator(data!.Blob, data!.ExportIdx));
 						}
 					}
 				}
@@ -226,12 +227,12 @@ namespace EpicGames.Horde.Storage.Backends
 		}
 
 		/// <inheritdoc/>
-		public override async Task WriteRefTargetAsync(RefName name, NodeLocator target, RefOptions? options = null, CancellationToken cancellationToken = default)
+		public override async Task WriteRefTargetAsync(RefName name, RefTarget target, RefOptions? options = null, CancellationToken cancellationToken = default)
 		{
-			_logger.LogDebug("Writing ref {RefName} -> {Blob}#{ExportIdx}", name, target.Blob, target.ExportIdx);
+			_logger.LogDebug("Writing ref {RefName} -> {RefTarget}", name, target);
 			using (HttpClient httpClient = _createClient())
 			{
-				using (HttpResponseMessage response = await httpClient.PutAsync($"refs/{name}", new { locator = target.Blob, exportIdx = target.ExportIdx, options }, cancellationToken))
+				using (HttpResponseMessage response = await httpClient.PutAsync($"refs/{name}", new { blob = target.Locator.Blob, exportIdx = target.Locator.ExportIdx, hash = target.Hash, options }, cancellationToken))
 				{
 					response.EnsureSuccessStatusCode();
 				}
diff --git a/Engine/Source/Programs/Shared/EpicGames.Horde/Storage/Backends/MemoryStorageClient.cs b/Engine/Source/Programs/Shared/EpicGames.Horde/Storage/Backends/MemoryStorageClient.cs
--- a/Engine/Source/Programs/Shared/EpicGames.Horde/Storage/Backends/MemoryStorageClient.cs
+++ b/Engine/Source/Programs/Shared/EpicGames.Horde/Storage/Backends/MemoryStorageClient.cs
@@ -26,13 +26,13 @@ namespace EpicGames.Horde.Storage.Backends
 		/// <summary>
 		/// Map of ref name to ref data
 		/// </summary>
-		readonly ConcurrentDictionary<RefName, NodeLocator> _refs = new ConcurrentDictionary<RefName, NodeLocator>();
+		readonly ConcurrentDictionary<RefName, RefTarget> _refs = new ConcurrentDictionary<RefName, RefTarget>();
 
 		/// <inheritdoc cref="_blobs"/>
 		public IReadOnlyDictionary<BlobLocator, Bundle> Blobs => _blobs;
 
 		/// <inheritdoc cref="_refs"/>
-		public IReadOnlyDictionary<RefName, NodeLocator> Refs => _refs;
+		public IReadOnlyDictionary<RefName, RefTarget> Refs => _refs;
 
 		/// <summary>
 		/// Constructor
@@ -73,15 +73,15 @@ namespace EpicGames.Horde.Storage.Backends
 		public override Task DeleteRefAsync(RefName name, CancellationToken cancellationToken) => Task.FromResult(_refs.TryRemove(name, out _));
 
 		/// <inheritdoc/>
-		public override Task<NodeLocator> TryReadRefTargetAsync(RefName name, DateTime cacheTime = default, CancellationToken cancellationToken = default)
+		public override Task<RefTarget?> TryReadRefTargetAsync(RefName name, DateTime cacheTime = default, CancellationToken cancellationToken = default)
 		{
-			NodeLocator refTarget;
+			RefTarget? refTarget;
 			_refs.TryGetValue(name, out refTarget);
 			return Task.FromResult(refTarget);
 		}
 
 		/// <inheritdoc/>
-		public override Task WriteRefTargetAsync(RefName name, NodeLocator target, RefOptions? options = null, CancellationToken cancellationToken = default)
+		public override Task WriteRefTargetAsync(RefName name, RefTarget target, RefOptions? options = null, CancellationToken cancellationToken = default)
 		{
 			_refs[name] = target;
 			return Task.CompletedTask;
diff --git a/Engine/Source/Programs/Shared/EpicGames.Horde/Storage/IStorageClient.cs b/Engine/Source/Programs/Shared/EpicGames.Horde/Storage/IStorageClient.cs
--- a/Engine/Source/Programs/Shared/EpicGames.Horde/Storage/IStorageClient.cs
+++ b/Engine/Source/Programs/Shared/EpicGames.Horde/Storage/IStorageClient.cs
@@ -36,32 +36,6 @@ namespace EpicGames.Horde.Storage
 		}
 	}
 
-	/// <summary>
-	/// Value for a ref
-	/// </summary>
-	[DebuggerDisplay("{Target}")]
-	public class RefValue
-	{
-		/// <summary>
-		/// Target for the ref
-		/// </summary>
-		public NodeLocator Target { get; }
-
-		/// <summary>
-		/// Bundle referenced as the target
-		/// </summary>
-		public Bundle Bundle { get; }
-
-		/// <summary>
-		/// Constructor
-		/// </summary>
-		public RefValue(NodeLocator target, Bundle bundle)
-		{
-			Target = target;
-			Bundle = bundle;
-		}
-	}
-
 	/// <summary>
 	/// Options for a new ref
 	/// </summary>
@@ -112,11 +86,11 @@ namespace EpicGames.Horde.Storage
 		/// </summary>
 		/// <param name="text">Text to parse</param>
 		/// <returns></returns>
-		public static NodeLocator Parse(string text)
+		public static NodeLocator Parse(ReadOnlySpan<char> text)
 		{
-			int hashIdx = text.IndexOf('#', StringComparison.Ordinal);
-			int exportIdx = Int32.Parse(text.AsSpan(hashIdx + 1), NumberStyles.None, CultureInfo.InvariantCulture);
-			BlobLocator blobLocator = new BlobLocator(new Utf8String(text.AsSpan(0, hashIdx)));
+			int hashIdx = text.IndexOf('#');
+			int exportIdx = Int32.Parse(text.Slice(hashIdx + 1), NumberStyles.None, CultureInfo.InvariantCulture);
+			BlobLocator blobLocator = new BlobLocator(new Utf8String(text.Slice(0, hashIdx)));
 			return new NodeLocator(blobLocator, exportIdx);
 		}
 
@@ -139,6 +113,28 @@ namespace EpicGames.Horde.Storage
 		public static bool operator !=(NodeLocator left, NodeLocator right) => !left.Equals(right);
 	}
 
+	/// <summary>
+	/// Target for a ref
+	/// </summary>
+	/// <param name="Hash">Hash of the target node</param>
+	/// <param name="Locator">Locator for the target node</param>
+	public record class RefTarget(IoHash Hash, NodeLocator Locator)
+	{
+		/// <summary>
+		/// Parse a ref target from the given string
+		/// </summary>
+		/// <param name="text"></param>
+		/// <returns></returns>
+		public static RefTarget Parse(ReadOnlySpan<char> text)
+		{
+			int idx = text.IndexOf('@');
+			return new RefTarget(IoHash.Parse(text.Slice(0, idx)), NodeLocator.Parse(text.Slice(idx + 1)));
+		}
+
+		/// <inheritdoc/>
+		public override string ToString() => $"{Hash}@{Locator}";
+	}
+
 	/// <summary>
 	/// Base interface for a low-level storage backend. Blobs added to this store are not content addressed, but referenced by <see cref="BlobLocator"/>.
 	/// </summary>
@@ -183,7 +179,7 @@ namespace EpicGames.Horde.Storage
 		/// <param name="cacheTime">Minimum coherency for any cached value to be returned</param>
 		/// <param name="cancellationToken">Cancellation token for the operation</param>
 		/// <returns>Node pointed to by the ref</returns>
-		Task<NodeLocator> TryReadRefTargetAsync(RefName name, DateTime cacheTime = default, CancellationToken cancellationToken = default);
+		Task<RefTarget?> TryReadRefTargetAsync(RefName name, DateTime cacheTime = default, CancellationToken cancellationToken = default);
 
 		/// <summary>
 		/// Writes a new ref to the store which points to a new blob
@@ -195,7 +191,7 @@ namespace EpicGames.Horde.Storage
 		/// <param name="options">Options for the new ref</param>
 		/// <param name="cancellationToken">Cancellation token for the operation</param>
 		/// <returns>Unique identifier for the blob</returns>
-		Task<NodeLocator> WriteRefAsync(RefName name, Bundle bundle, int exportIdx, Utf8String prefix = default, RefOptions? options = null, CancellationToken cancellationToken = default);
+		Task<RefTarget> WriteRefAsync(RefName name, Bundle bundle, int exportIdx, Utf8String prefix = default, RefOptions? options = null, CancellationToken cancellationToken = default);
 
 		/// <summary>
 		/// Writes a new ref to the store
@@ -205,7 +201,7 @@ namespace EpicGames.Horde.Storage
 		/// <param name="options">Options for the new ref</param>
 		/// <param name="cancellationToken">Cancellation token for the operation</param>
 		/// <returns>Unique identifier for the blob</returns>
-		Task WriteRefTargetAsync(RefName name, NodeLocator target, RefOptions? options = null, CancellationToken cancellationToken = default);
+		Task WriteRefTargetAsync(RefName name, RefTarget target, RefOptions? options = null, CancellationToken cancellationToken = default);
 
 		/// <summary>
 		/// Reads data for a ref from the store
@@ -254,13 +250,13 @@ namespace EpicGames.Horde.Storage
 			public Task DeleteRefAsync(RefName name, CancellationToken cancellationToken = default) => _inner.DeleteRefAsync(name, cancellationToken);
 
 			/// <inheritdoc/>
-			public Task<NodeLocator> TryReadRefTargetAsync(RefName name, DateTime cacheTime = default, CancellationToken cancellationToken = default) => _inner.TryReadRefTargetAsync(name, cacheTime, cancellationToken);
+			public Task<RefTarget?> TryReadRefTargetAsync(RefName name, DateTime cacheTime = default, CancellationToken cancellationToken = default) => _inner.TryReadRefTargetAsync(name, cacheTime, cancellationToken);
 
 			/// <inheritdoc/>
-			public Task<NodeLocator> WriteRefAsync(RefName name, Bundle bundle, int exportIdx = 0, Utf8String prefix = default, RefOptions? options = null, CancellationToken cancellationToken = default) => _inner.WriteRefAsync(name, bundle, exportIdx, prefix, options, cancellationToken);
+			public Task<RefTarget> WriteRefAsync(RefName name, Bundle bundle, int exportIdx = 0, Utf8String prefix = default, RefOptions? options = null, CancellationToken cancellationToken = default) => _inner.WriteRefAsync(name, bundle, exportIdx, prefix, options, cancellationToken);
 
 			/// <inheritdoc/>
-			public Task WriteRefTargetAsync(RefName name, NodeLocator target, RefOptions? options = null, CancellationToken cancellationToken = default) => _inner.WriteRefTargetAsync(name, target, options, cancellationToken);
+			public Task WriteRefTargetAsync(RefName name, RefTarget target, RefOptions? options = null, CancellationToken cancellationToken = default) => _inner.WriteRefTargetAsync(name, target, options, cancellationToken);
 
 			#endregion
 		}
@@ -348,8 +344,8 @@ namespace EpicGames.Horde.Storage
 		/// <returns>True if the ref exists, false if it did not exist</returns>
 		public static async Task<bool> HasRefAsync(this IStorageClient store, RefName name, DateTime cacheTime = default, CancellationToken cancellationToken = default)
 		{
-			NodeLocator locator = await store.TryReadRefTargetAsync(name, cacheTime, cancellationToken);
-			return locator.IsValid();
+			RefTarget? target = await store.TryReadRefTargetAsync(name, cacheTime, cancellationToken);
+			return target != null;
 		}
 
 		/// <summary>
@@ -373,7 +369,7 @@ namespace EpicGames.Horde.Storage
 		/// <param name="maxAge">Maximum age of any cached ref</param>
 		/// <param name="cancellationToken">Cancellation token for the operation</param>
 		/// <returns>The ref target</returns>
-		public static Task<NodeLocator> TryReadRefTargetAsync(this IStorageClient store, RefName name, TimeSpan maxAge, CancellationToken cancellationToken = default)
+		public static Task<RefTarget?> TryReadRefTargetAsync(this IStorageClient store, RefName name, TimeSpan maxAge, CancellationToken cancellationToken = default)
 		{
 			return store.TryReadRefTargetAsync(name, DateTime.UtcNow - maxAge, cancellationToken);
 		}
@@ -386,10 +382,10 @@ namespace EpicGames.Horde.Storage
 		/// <param name="cacheTime">Minimum coherency of any cached result</param>
 		/// <param name="cancellationToken">Cancellation token for the operation</param>
 		/// <returns>The ref target</returns>
-		public static async Task<NodeLocator> ReadRefTargetAsync(this IStorageClient store, RefName name, DateTime cacheTime = default, CancellationToken cancellationToken = default)
+		public static async Task<RefTarget> ReadRefTargetAsync(this IStorageClient store, RefName name, DateTime cacheTime = default, CancellationToken cancellationToken = default)
 		{
-			NodeLocator refTarget = await store.TryReadRefTargetAsync(name, cacheTime, cancellationToken);
-			if (!refTarget.IsValid())
+			RefTarget? refTarget = await store.TryReadRefTargetAsync(name, cacheTime, cancellationToken);
+			if (refTarget == null)
 			{
 				throw new RefNameNotFoundException(name);
 			}
@@ -404,7 +400,7 @@ namespace EpicGames.Horde.Storage
 		/// <param name="maxAge">Maximum age for any cached result</param>
 		/// <param name="cancellationToken">Cancellation token for the operation</param>
 		/// <returns>The blob instance</returns>
-		public static Task<NodeLocator> ReadRefTargetAsync(this IStorageClient store, RefName name, TimeSpan maxAge, CancellationToken cancellationToken = default)
+		public static Task<RefTarget> ReadRefTargetAsync(this IStorageClient store, RefName name, TimeSpan maxAge, CancellationToken cancellationToken = default)
 		{
 			return ReadRefTargetAsync(store, name, DateTime.UtcNow - maxAge, cancellationToken);
 		}
diff --git a/Engine/Source/Programs/Shared/EpicGames.Horde/Storage/Nodes/CbNode.cs b/Engine/Source/Programs/Shared/EpicGames.Horde/Storage/Nodes/CbNode.cs
--- a/Engine/Source/Programs/Shared/EpicGames.Horde/Storage/Nodes/CbNode.cs
+++ b/Engine/Source/Programs/Shared/EpicGames.Horde/Storage/Nodes/CbNode.cs
@@ -31,7 +31,7 @@ namespace EpicGames.Horde.Storage.Nodes
 		public CbNode(CbObject obj, IReadOnlyDictionary<IoHash, NodeLocator> references)
 		{
 			Object = obj;
-			References = references.ToDictionary(x => x.Key, x => new TreeNodeRef(x.Key, x.Value));
+			References = references.ToDictionary(x => x.Key, x => new TreeNodeRef(new RefTarget(x.Key, x.Value)));
 		}
 
 		/// <summary>
@@ -41,7 +41,7 @@ namespace EpicGames.Horde.Storage.Nodes
 		public CbNode(ITreeNodeReader reader)
 		{
 			Object = new CbObject(reader.ReadFixedLengthBytes(reader.Length));
-			References = reader.References.ToDictionary(x => x.Key, x => new TreeNodeRef(x.Key, x.Value));
+			References = reader.References.ToDictionary(x => x.Key, x => new TreeNodeRef(new RefTarget(x.Key, x.Value)));
 		}
 
 		/// <inheritdoc/>
diff --git a/Engine/Source/Programs/Shared/EpicGames.Horde/Storage/StorageClientBase.cs b/Engine/Source/Programs/Shared/EpicGames.Horde/Storage/StorageClientBase.cs
--- a/Engine/Source/Programs/Shared/EpicGames.Horde/Storage/StorageClientBase.cs
+++ b/Engine/Source/Programs/Shared/EpicGames.Horde/Storage/StorageClientBase.cs
@@ -39,19 +39,19 @@ namespace EpicGames.Horde.Storage
 		public abstract Task DeleteRefAsync(RefName name, CancellationToken cancellationToken = default);
 
 		/// <inheritdoc/>
-		public abstract Task<NodeLocator> TryReadRefTargetAsync(RefName name, DateTime cacheTime = default, CancellationToken cancellationToken = default);
+		public abstract Task<RefTarget?> TryReadRefTargetAsync(RefName name, DateTime cacheTime = default, CancellationToken cancellationToken = default);
 
 		/// <inheritdoc/>
-		public virtual async Task<NodeLocator> WriteRefAsync(RefName name, Bundle bundle, int exportIdx = 0, Utf8String prefix = default, RefOptions? options = null, CancellationToken cancellationToken = default)
+		public virtual async Task<RefTarget> WriteRefAsync(RefName name, Bundle bundle, int exportIdx = 0, Utf8String prefix = default, RefOptions? options = null, CancellationToken cancellationToken = default)
 		{
 			BlobLocator locator = await this.WriteBundleAsync(bundle, prefix, cancellationToken);
-			NodeLocator target = new NodeLocator(locator, exportIdx);
+			RefTarget target = new RefTarget(bundle.Header.Exports[exportIdx].Hash, new NodeLocator(locator, exportIdx));
 			await WriteRefTargetAsync(name, target, options, cancellationToken);
 			return target;
 		}
 
 		/// <inheritdoc/>
-		public abstract Task WriteRefTargetAsync(RefName name, NodeLocator target, RefOptions? options = null, CancellationToken cancellationToken = default);
+		public abstract Task WriteRefTargetAsync(RefName name, RefTarget target, RefOptions? options = null, CancellationToken cancellationToken = default);
 
 		#endregion
 	}
diff --git a/Engine/Source/Programs/Shared/EpicGames.Horde/Storage/TreeNodeRef.cs b/Engine/Source/Programs/Shared/EpicGames.Horde/Storage/TreeNodeRef.cs
--- a/Engine/Source/Programs/Shared/EpicGames.Horde/Storage/TreeNodeRef.cs
+++ b/Engine/Source/Programs/Shared/EpicGames.Horde/Storage/TreeNodeRef.cs
@@ -54,26 +54,21 @@ namespace EpicGames.Horde.Storage
 		}
 
 		/// <summary>
-		/// Creates a reference to a node with the given hash
+		/// Creates a reference to a node in memory.
 		/// </summary>
-		/// <param name="hash">Hash of the referenced node</param>
-		/// <param name="locator">Locator for the node</param>
-		public TreeNodeRef(IoHash hash, NodeLocator locator)
+		/// <param name="target">Node to reference</param>
+		public TreeNodeRef(RefTarget target)
 		{
-			Hash = hash;
-			Locator = locator;
+			Hash = target.Hash;
+			Locator = target.Locator;
 		}
 
 		/// <summary>
 		/// Deserialization constructor
 		/// </summary>
 		/// <param name="reader"></param>
-		public TreeNodeRef(ITreeNodeReader reader)
+		public TreeNodeRef(ITreeNodeReader reader) : this(reader.ReadRefTarget())
 		{
-			TreeNodeRefData data = reader.ReadRefData();
-			Hash = data.Hash;
-			Locator = data.Locator;
-			Debug.Assert(Hash != IoHash.Zero);
 		}
 
 		/// <summary>
@@ -196,19 +191,16 @@ namespace EpicGames.Horde.Storage
 		/// <summary>
 		/// Constructor
 		/// </summary>
-		/// <param name="hash">Hash of the referenced node</param>
-		/// <param name="locator">Locator for the node</param>
-		internal TreeNodeRef(IoHash hash, NodeLocator locator) 
-			: base(hash, locator)
+		/// <param name="target">The referenced node</param>
+		public TreeNodeRef(RefTarget target) : base(target)
 		{
 		}
 
 		/// <summary>
 		/// Constructor
 		/// </summary>
-		/// <param name="reader"></param>
-		public TreeNodeRef(ITreeNodeReader reader)
-			: base(reader)
+		/// <param name="reader">The reader to deserialize from</param>
+		public TreeNodeRef(ITreeNodeReader reader) : base(reader)
 		{
 		}
 
@@ -234,9 +226,4 @@ namespace EpicGames.Horde.Storage
 			return await reader.ReadNodeAsync<T>(Locator, cancellationToken);
 		}
 	}
-
-	/// <summary>
-	/// Deserialized ref data
-	/// </summary>
-	public record struct TreeNodeRefData(IoHash Hash, NodeLocator Locator);
 }
diff --git a/Engine/Source/Programs/Shared/EpicGames.Horde/Storage/TreeReader.cs b/Engine/Source/Programs/Shared/EpicGames.Horde/Storage/TreeReader.cs
--- a/Engine/Source/Programs/Shared/EpicGames.Horde/Storage/TreeReader.cs
+++ b/Engine/Source/Programs/Shared/EpicGames.Horde/Storage/TreeReader.cs
@@ -39,7 +39,7 @@ namespace EpicGames.Horde.Storage
 		/// <summary>
 		/// Reads a reference to another node
 		/// </summary>
-		TreeNodeRefData ReadRefData();
+		RefTarget ReadRefTarget();
 	}
 
 	/// <summary>
@@ -139,7 +139,7 @@ namespace EpicGames.Horde.Storage
 			public readonly BundleHeader Header;
 			public readonly int[] PacketOffsets;
 			public readonly ExportInfo[] Exports;
-			public readonly TreeNodeRefData[] References;
+			public readonly RefTarget[] References;
 			public readonly TypeInfo?[] Types;
 
 			public BundleInfo(BlobLocator locator, BundleHeader header, int headerLength, TypeInfo?[] types)
@@ -167,7 +167,7 @@ namespace EpicGames.Horde.Storage
 					packetOffset += packet.EncodedLength;
 				}
 
-				References = new TreeNodeRefData[header.Imports.Sum(x => x.Exports.Count) + header.Exports.Count];
+				References = new RefTarget[header.Imports.Sum(x => x.Exports.Count) + header.Exports.Count];
 
 				int referenceIdx = 0;
 				foreach (BundleImport import in header.Imports)
@@ -175,12 +175,12 @@ namespace EpicGames.Horde.Storage
 					foreach ((int importExportIdx, IoHash importExportHash) in import.Exports)
 					{
 						NodeLocator importLocator = new NodeLocator(import.Locator, importExportIdx);
-						References[referenceIdx++] = new TreeNodeRefData(importExportHash, importLocator);
+						References[referenceIdx++] = new RefTarget(importExportHash, importLocator);
 					}
 				}
 				for (int idx = 0; idx < header.Exports.Count; idx++)
 				{
-					References[referenceIdx++] = new TreeNodeRefData(header.Exports[idx].Hash, new NodeLocator(locator, idx));
+					References[referenceIdx++] = new RefTarget(header.Exports[idx].Hash, new NodeLocator(locator, idx));
 				}
 
 				Types = types;
@@ -244,10 +244,10 @@ namespace EpicGames.Horde.Storage
 
 			public IReadOnlyDictionary<IoHash, NodeLocator> References => _refs;
 
-			public TreeNodeRefData ReadRefData()
+			public RefTarget ReadRefTarget()
 			{
 				IoHash hash = this.ReadIoHash();
-				return new TreeNodeRefData(hash, _refs[hash]);
+				return new RefTarget(hash, _refs[hash]);
 			}
 		}
 
@@ -690,7 +690,7 @@ namespace EpicGames.Horde.Storage
 			Dictionary<IoHash, NodeLocator> refs = new Dictionary<IoHash, NodeLocator>(export.References.Count);
 			for (int idx = 0; idx < export.References.Count; idx++)
 			{
-				TreeNodeRefData reference = bundleInfo.References[export.References[idx]];
+				RefTarget reference = bundleInfo.References[export.References[idx]];
 				refs.Add(reference.Hash, reference.Locator);
 			}
 
@@ -716,6 +716,17 @@ namespace EpicGames.Horde.Storage
 			return (TNode)await ReadNodeAsync(locator, cancellationToken);
 		}
 
+		/// <summary>
+		/// Reads a node from a bundle
+		/// </summary>
+		/// <param name="target">Locator for the node</param>
+		/// <param name="cancellationToken">Cancellation token for the operation</param>
+		/// <returns>Node data read from the given bundle</returns>
+		public async ValueTask<TNode> ReadNodeAsync<TNode>(RefTarget target, CancellationToken cancellationToken = default) where TNode : TreeNode
+		{
+			return (TNode)await ReadNodeAsync(target.Locator, cancellationToken);
+		}
+
 		/// <summary>
 		/// Reads data for a ref from the store, along with the node's contents.
 		/// </summary>
@@ -725,12 +736,29 @@ namespace EpicGames.Horde.Storage
 		/// <returns>Node for the given ref, or null if it does not exist</returns>
 		public async Task<TNode?> TryReadNodeAsync<TNode>(RefName name, DateTime cacheTime = default, CancellationToken cancellationToken = default) where TNode : TreeNode
 		{
-			NodeLocator refTarget = await _store.TryReadRefTargetAsync(name, cacheTime, cancellationToken);
-			if (!refTarget.IsValid())
+			RefTarget? refTarget = await _store.TryReadRefTargetAsync(name, cacheTime, cancellationToken);
+			if (refTarget == null)
+			{
+				return null;
+			}
+			return await ReadNodeAsync<TNode>(refTarget.Locator, cancellationToken);
+		}
+
+		/// <summary>
+		/// Reads data for a ref from the store, along with the node's contents.
+		/// </summary>
+		/// <param name="name">The ref name</param>
+		/// <param name="cacheTime">Minimum coherency for any cached value to be returned</param>
+		/// <param name="cancellationToken">Cancellation token for the operation</param>
+		/// <returns>Node for the given ref, or null if it does not exist</returns>
+		public async Task<TreeNodeRef<TNode>?> TryReadNodeRefAsync<TNode>(RefName name, DateTime cacheTime = default, CancellationToken cancellationToken = default) where TNode : TreeNode
+		{
+			RefTarget? refTarget = await _store.TryReadRefTargetAsync(name, cacheTime, cancellationToken);
+			if (refTarget == null)
 			{
 				return null;
 			}
-			return await ReadNodeAsync<TNode>(refTarget, cancellationToken);
+			return new TreeNodeRef<TNode>(refTarget);
 		}
 
 		/// <summary>
@@ -749,5 +777,22 @@ namespace EpicGames.Horde.Storage
 			}
 			return refValue;
 		}
+
+		/// <summary>
+		/// Reads a ref from the store, throwing an exception if it does not exist
+		/// </summary>
+		/// <param name="name">Id for the ref</param>
+		/// <param name="cacheTime">Minimum coherency of any cached result</param>
+		/// <param name="cancellationToken">Cancellation token for the operation</param>
+		/// <returns>The blob instance</returns>
+		public async Task<TreeNodeRef<TNode>> ReadNodeRefAsync<TNode>(RefName name, DateTime cacheTime = default, CancellationToken cancellationToken = default) where TNode : TreeNode
+		{
+			TreeNodeRef<TNode>? refValue = await TryReadNodeRefAsync<TNode>(name, cacheTime, cancellationToken);
+			if (refValue == null)
+			{
+				throw new RefNameNotFoundException(name);
+			}
+			return refValue;
+		}
 	}
 }
diff --git a/Engine/Source/Programs/Shared/EpicGames.Horde/Storage/TreeWriter.cs b/Engine/Source/Programs/Shared/EpicGames.Horde/Storage/TreeWriter.cs
--- a/Engine/Source/Programs/Shared/EpicGames.Horde/Storage/TreeWriter.cs
+++ b/Engine/Source/Programs/Shared/EpicGames.Horde/Storage/TreeWriter.cs
@@ -259,7 +259,7 @@ namespace EpicGames.Horde.Storage
 			await FlushAsync(cancellationToken);
 
 			Debug.Assert(nodeRef.Locator.IsValid());
-			await _store.WriteRefTargetAsync(name, nodeRef.Locator, options, cancellationToken);
+			await _store.WriteRefTargetAsync(name, new RefTarget(nodeRef.Hash, nodeRef.Locator), options, cancellationToken);
 			return nodeRef.Locator;
 		}
 
@@ -274,7 +274,7 @@ namespace EpicGames.Horde.Storage
 		{
 			await FlushAsync(cancellationToken);
 			NodeLocator locator = GetLocator(root);
-			await _store.WriteRefTargetAsync(name, locator, options, cancellationToken);
+			await _store.WriteRefTargetAsync(name, new RefTarget(root, locator), options, cancellationToken);
 			return locator;
 		}
 
diff --git a/Engine/Source/Programs/Shared/EpicGames.IoHash/IoHash.cs b/Engine/Source/Programs/Shared/EpicGames.IoHash/IoHash.cs
--- a/Engine/Source/Programs/Shared/EpicGames.IoHash/IoHash.cs
+++ b/Engine/Source/Programs/Shared/EpicGames.IoHash/IoHash.cs
@@ -151,6 +151,16 @@ namespace EpicGames.Core
 			return new IoHash(StringUtils.ParseHexString(text));
 		}
 
+		/// <summary>
+		/// Parses a digest from the given hex string
+		/// </summary>
+		/// <param name="text"></param>
+		/// <returns></returns>
+		public static IoHash Parse(ReadOnlySpan<char> text)
+		{
+			return new IoHash(StringUtils.ParseHexString(text));
+		}
+
 		/// <summary>
 		/// Parses a digest from the given hex string
 		/// </summary>
