Horde: Rename TreeWriter.WriteRefAsync() -> WriteAsync() for consistency with other overloads.

diff --git a/Engine/Source/Programs/Horde/Horde.Agent/Commands/Bundles/CreateCommand.cs b/Engine/Source/Programs/Horde/Horde.Agent/Commands/Bundles/CreateCommand.cs
--- a/Engine/Source/Programs/Horde/Horde.Agent/Commands/Bundles/CreateCommand.cs
+++ b/Engine/Source/Programs/Horde/Horde.Agent/Commands/Bundles/CreateCommand.cs
@@ -106,7 +106,7 @@ namespace Horde.Agent.Commands.Bundles
 			ChunkingOptions options = new ChunkingOptions();
 			await node.CopyFromDirectoryAsync(InputDir.ToDirectoryInfo(), options, writer, CancellationToken.None);
 
-			await writer.WriteRefAsync(RefName, node);
+			await writer.WriteAsync(RefName, node);
 
 			logger.LogInformation("Time: {Time}", timer.Elapsed.TotalSeconds);
 			return 0;
diff --git a/Engine/Source/Programs/Horde/Horde.Agent/Execution/JobExecutor.cs b/Engine/Source/Programs/Horde/Horde.Agent/Execution/JobExecutor.cs
--- a/Engine/Source/Programs/Horde/Horde.Agent/Execution/JobExecutor.cs
+++ b/Engine/Source/Programs/Horde/Horde.Agent/Execution/JobExecutor.cs
@@ -903,7 +903,7 @@ namespace Horde.Agent.Execution
 				}
 
 				// Write the final node
-				await treeWriter.WriteRefAsync(refName, outputNode, new RefOptions { Lifetime = TimeSpan.FromDays(3.0) }, cancellationToken);
+				await treeWriter.WriteAsync(refName, outputNode, new RefOptions { Lifetime = TimeSpan.FromDays(3.0) }, cancellationToken);
 				logger.LogInformation("Upload took {Time:n1}s", timer.Elapsed.TotalSeconds);
 			}
 
@@ -1293,7 +1293,7 @@ namespace Horde.Agent.Execution
 				ChunkingOptions options = new ChunkingOptions();
 				await directoryNode.CopyFromDirectoryAsync(logDir.ToDirectoryInfo(), options, treeWriter, cancellationToken);
 
-				await treeWriter.WriteRefAsync(refName, directoryNode, cancellationToken: cancellationToken);
+				await treeWriter.WriteAsync(refName, directoryNode, cancellationToken: cancellationToken);
 			}
 			else
 			{
diff --git a/Engine/Source/Programs/Horde/Horde.Build/Commands/Bundles/CreateCommand.cs b/Engine/Source/Programs/Horde/Horde.Build/Commands/Bundles/CreateCommand.cs
--- a/Engine/Source/Programs/Horde/Horde.Build/Commands/Bundles/CreateCommand.cs
+++ b/Engine/Source/Programs/Horde/Horde.Build/Commands/Bundles/CreateCommand.cs
@@ -40,7 +40,7 @@ namespace Horde.Build.Commands.Bundles
 			DirectoryNode node = new DirectoryNode(DirectoryFlags.None);
 			await node.CopyFromDirectoryAsync(InputDir.ToDirectoryInfo(), new ChunkingOptions(), writer, CancellationToken.None);
 
-			await writer.WriteRefAsync(RefName, node);
+			await writer.WriteAsync(RefName, node);
 			return 0;
 		}
 	}
diff --git a/Engine/Source/Programs/Shared/EpicGames.Horde.Tests/BundleStoreTests.cs b/Engine/Source/Programs/Shared/EpicGames.Horde.Tests/BundleStoreTests.cs
--- a/Engine/Source/Programs/Shared/EpicGames.Horde.Tests/BundleStoreTests.cs
+++ b/Engine/Source/Programs/Shared/EpicGames.Horde.Tests/BundleStoreTests.cs
@@ -88,7 +88,7 @@ namespace EpicGames.Horde.Tests
 
 				SimpleNode root = new SimpleNode(new ReadOnlySequence<byte>(new byte[] { 5 }), new[] { new TreeNodeRef<SimpleNode>(node4), new TreeNodeRef<SimpleNode>(node3) });
 
-				await writer.WriteRefAsync(new RefName("test"), root);
+				await writer.WriteAsync(new RefName("test"), root);
 
 				TreeReader reader = new TreeReader(store, null, s_readOptions, NullLogger.Instance);
 				await CheckTree(reader, root);
@@ -188,7 +188,7 @@ namespace EpicGames.Horde.Tests
 				FileNode world = await FileNode.CreateAsync(Encoding.UTF8.GetBytes("world"), new ChunkingOptions(), writer, CancellationToken.None);
 				hello.AddFile("world", FileEntryFlags.None, world);
 
-				await writer.WriteRefAsync(new RefName("test"), root);
+				await writer.WriteAsync(new RefName("test"), root);
 
 				await CheckFileTreeAsync(reader, root);
 			}
@@ -326,7 +326,7 @@ namespace EpicGames.Horde.Tests
 				FileNode file = await FileNode.CreateAsync(data, options, writer, CancellationToken.None);
 				root.AddFile("test", FileEntryFlags.None, file);
 
-				await writer.WriteRefAsync(new RefName("test"), root);
+				await writer.WriteAsync(new RefName("test"), root);
 
 				TreeReader reader = new TreeReader(store, null, NullLogger.Instance);
 				await CheckLargeFileTreeAsync(reader, root, data);
diff --git a/Engine/Source/Programs/Shared/EpicGames.Horde/Storage/TreeNode.cs b/Engine/Source/Programs/Shared/EpicGames.Horde/Storage/TreeNode.cs
--- a/Engine/Source/Programs/Shared/EpicGames.Horde/Storage/TreeNode.cs
+++ b/Engine/Source/Programs/Shared/EpicGames.Horde/Storage/TreeNode.cs
@@ -188,7 +188,7 @@ namespace EpicGames.Horde.Storage
 		public static async Task<NodeLocator> WriteNodeAsync(this IStorageClient store, RefName name, TreeNode node, TreeOptions? options = null, Utf8String prefix = default, RefOptions? refOptions = null, CancellationToken cancellationToken = default)
 		{
 			TreeWriter writer = new TreeWriter(store, options, prefix.IsEmpty ? name.Text : prefix);
-			return await writer.WriteRefAsync(name, node, refOptions, cancellationToken);
+			return await writer.WriteAsync(name, node, refOptions, cancellationToken);
 		}
 
 		/// <summary>
diff --git a/Engine/Source/Programs/Shared/EpicGames.Horde/Storage/TreeWriter.cs b/Engine/Source/Programs/Shared/EpicGames.Horde/Storage/TreeWriter.cs
--- a/Engine/Source/Programs/Shared/EpicGames.Horde/Storage/TreeWriter.cs
+++ b/Engine/Source/Programs/Shared/EpicGames.Horde/Storage/TreeWriter.cs
@@ -133,6 +133,17 @@ namespace EpicGames.Horde.Storage
 			_packetWriter = new ArrayMemoryWriter(_options.MinCompressionPacketSize);
 		}
 
+		/// <summary>
+		/// Constructor
+		/// </summary>
+		/// <param name="store">Store to write data to</param>
+		/// <param name="refName">Ref being written. Will be used as a prefix for storing blobs.</param>
+		/// <param name="options">Options for the writer</param>
+		public TreeWriter(IStorageClient store, RefName refName, TreeOptions? options = null)
+			: this(store, options, refName.Text)
+		{
+		}
+
 		/// <summary>
 		/// Copies settings from another tree writer
 		/// </summary>
@@ -244,6 +255,21 @@ namespace EpicGames.Horde.Storage
 			return hash;
 		}
 
+		/// <summary>
+		/// Writes a ref to the node with the given hash
+		/// </summary>
+		/// <param name="name">Name of the ref to write</param>
+		/// <param name="root">Hash of the root node</param>
+		/// <param name="options"></param>
+		/// <param name="cancellationToken">Cancellation token for the operation</param>
+		public async Task<NodeLocator> WriteAsync(RefName name, IoHash root, RefOptions? options = null, CancellationToken cancellationToken = default)
+		{
+			await FlushAsync(cancellationToken);
+			NodeLocator locator = GetLocator(root);
+			await _store.WriteRefTargetAsync(name, new RefTarget(root, locator), options, cancellationToken);
+			return locator;
+		}
+
 		/// <summary>
 		/// Writes a node to the given ref
 		/// </summary>
@@ -252,7 +278,7 @@ namespace EpicGames.Horde.Storage
 		/// <param name="options"></param>
 		/// <param name="cancellationToken"></param>
 		/// <returns></returns>
-		public async Task<NodeLocator> WriteRefAsync(RefName name, TreeNode node, RefOptions? options = null, CancellationToken cancellationToken = default)
+		public async Task<NodeLocator> WriteAsync(RefName name, TreeNode node, RefOptions? options = null, CancellationToken cancellationToken = default)
 		{
 			TreeNodeRef nodeRef = new TreeNodeRef(node);
 			await WriteAsync(nodeRef, cancellationToken);
@@ -263,21 +289,6 @@ namespace EpicGames.Horde.Storage
 			return nodeRef.Locator;
 		}
 
-		/// <summary>
-		/// Writes a ref to the node with the given hash
-		/// </summary>
-		/// <param name="name">Name of the ref to write</param>
-		/// <param name="root">Hash of the root node</param>
-		/// <param name="options"></param>
-		/// <param name="cancellationToken">Cancellation token for the operation</param>
-		public async Task<NodeLocator> WriteAsync(RefName name, IoHash root, RefOptions? options = null, CancellationToken cancellationToken = default)
-		{
-			await FlushAsync(cancellationToken);
-			NodeLocator locator = GetLocator(root);
-			await _store.WriteRefTargetAsync(name, new RefTarget(root, locator), options, cancellationToken);
-			return locator;
-		}
-
 		void FlushPacket()
 		{
 			if (_packetWriter.Length > 0)
