EpicGames.Core: Fix corrupt output when writing an optional string to IMemoryWriter.

diff --git a/Engine/Source/Programs/Shared/EpicGames.Core/MemoryWriter.cs b/Engine/Source/Programs/Shared/EpicGames.Core/MemoryWriter.cs
--- a/Engine/Source/Programs/Shared/EpicGames.Core/MemoryWriter.cs
+++ b/Engine/Source/Programs/Shared/EpicGames.Core/MemoryWriter.cs
@@ -482,6 +482,8 @@ namespace EpicGames.Core
 			{
 				Span<byte> span = writer.GetSpan(1);
 				span[0] = 0;
+
+				writer.Advance(1);
 			}
 			else
 			{
