Added support to set the underlying type for regular and namespaced enums.

diff --git a/Engine/Source/Programs/Shared/EpicGames.UHT/Exporters/CodeGen/UhtHeaderCodeGeneratorHFile.cs b/Engine/Source/Programs/Shared/EpicGames.UHT/Exporters/CodeGen/UhtHeaderCodeGeneratorHFile.cs
--- a/Engine/Source/Programs/Shared/EpicGames.UHT/Exporters/CodeGen/UhtHeaderCodeGeneratorHFile.cs
+++ b/Engine/Source/Programs/Shared/EpicGames.UHT/Exporters/CodeGen/UhtHeaderCodeGeneratorHFile.cs
@@ -338,13 +338,31 @@ namespace EpicGames.UHT.Exporters.CodeGen
 				builder.Append("enum class ").Append(enumObj.CppType);
 				if (enumObj.UnderlyingType != UhtEnumUnderlyingType.Unspecified)
 				{
-					builder.Append(" : ").Append(enumObj.UnderlyingType.ToString());
+					builder.Append(" : ").Append(enumObj.UnderlyingType.ToString().ToLower());
 				}
 				builder.Append(";\r\n");
-				
+
 				// Add TIsUEnumClass typetraits
 				builder.Append("template<> struct TIsUEnumClass<").Append(enumObj.CppType).Append("> { enum { Value = true }; };\r\n");
+			}
+			else if (enumObj.CppForm == UhtEnumCppForm.Regular && enumObj.UnderlyingType != UhtEnumUnderlyingType.Unspecified)
+			{
+				builder.Append("\r\n");
+				builder.Append("enum ").Append(enumObj.CppType);
+				builder.Append(" : ").Append(enumObj.UnderlyingType.ToString().ToLower());
+				builder.Append(";\r\n");
+			}
+			else if (enumObj.CppForm == UhtEnumCppForm.Namespaced && enumObj.UnderlyingType != UhtEnumUnderlyingType.Unspecified)
+			{
+				string[] SplitName = enumObj.CppType.Split("::");
+				builder.Append("\r\n");
+				builder.Append("namespace ").Append(SplitName[0]).Append(" { enum ").Append(SplitName[1]);
+				builder.Append(" : ").Append(enumObj.UnderlyingType.ToString().ToLower());
+				builder.Append("; }\r\n");
+			}
 
+			if (enumObj.CppForm == UhtEnumCppForm.EnumClass || enumObj.UnderlyingType != UhtEnumUnderlyingType.Unspecified)
+			{	
 				// Forward declare the StaticEnum<> specialization for enum classes
 				builder.Append("template<> ").Append(PackageApi).Append("UEnum* StaticEnum<").Append(enumObj.CppType).Append(">();\r\n");
 				builder.Append("\r\n");
diff --git a/Engine/Source/Programs/Shared/EpicGames.UHT/Parsers/UhtEnumParser.cs b/Engine/Source/Programs/Shared/EpicGames.UHT/Parsers/UhtEnumParser.cs
--- a/Engine/Source/Programs/Shared/EpicGames.UHT/Parsers/UhtEnumParser.cs
+++ b/Engine/Source/Programs/Shared/EpicGames.UHT/Parsers/UhtEnumParser.cs
@@ -80,7 +80,7 @@ namespace EpicGames.UHT.Parsers
 						{
 							UhtToken enumType = topScope.TokenReader.GetIdentifier("enumeration base");
 
-							if (!System.Enum.TryParse<UhtEnumUnderlyingType>(enumType.Value.ToString(), out UhtEnumUnderlyingType underlyingType) || underlyingType == UhtEnumUnderlyingType.Unspecified)
+							if (!System.Enum.TryParse<UhtEnumUnderlyingType>(enumType.Value.ToString(), true, out UhtEnumUnderlyingType underlyingType) || underlyingType == UhtEnumUnderlyingType.Unspecified)
 							{
 								topScope.TokenReader.LogError(enumType.InputLine, $"Unsupported enum class base type '{enumType.Value}'");
 							}
@@ -90,6 +90,11 @@ namespace EpicGames.UHT.Parsers
 						{
 							enumObject.UnderlyingType = UhtEnumUnderlyingType.Unspecified;
 						}
+
+						if (enumObject.UnderlyingType != UhtEnumUnderlyingType.Unspecified && enumObject.UnderlyingType != UhtEnumUnderlyingType.Uint8 && enumObject.MetaData.ContainsKey("BlueprintType"))
+						{
+							topScope.TokenReader.LogError("Invalid BlueprintType enum base - currently only uint8 supported");
+						}
 					}
 					else
 					{
@@ -99,10 +104,15 @@ namespace EpicGames.UHT.Parsers
 							{
 								UhtToken enumType = topScope.TokenReader.GetIdentifier("enumeration base");
 
-								if (enumType.Value != "int")
+								if (!System.Enum.TryParse<UhtEnumUnderlyingType>(enumType.Value.ToString(), true, out UhtEnumUnderlyingType underlyingType) || underlyingType == UhtEnumUnderlyingType.Unspecified)
 								{
-									topScope.TokenReader.LogError($"Regular enums only support 'int' as the value size");
+									topScope.TokenReader.LogError(enumType.InputLine, $"Unsupported regular enum type '{enumType.Value}'");
 								}
+								enumObject.UnderlyingType = underlyingType;
+							}
+							else
+							{
+								enumObject.UnderlyingType = UhtEnumUnderlyingType.Unspecified;
 							}
 						}
 						if ((enumObject.EnumFlags & EEnumFlags.Flags) != 0)
@@ -111,11 +121,6 @@ namespace EpicGames.UHT.Parsers
 						}
 					}
 
-					if (enumObject.UnderlyingType != UhtEnumUnderlyingType.uint8 && enumObject.MetaData.ContainsKey("BlueprintType"))
-					{
-						topScope.TokenReader.LogError("Invalid BlueprintType enum base - currently only uint8 supported");
-					}
-
 					//EnumDef.GetDefinitionRange().Start = &Input[InputPos];
 
 					// Get the opening brace
@@ -134,10 +139,15 @@ namespace EpicGames.UHT.Parsers
 							{
 								UhtToken enumType = topScope.TokenReader.GetIdentifier("enumeration base");
 
-								if (enumType.Value != "int")
+								if (!System.Enum.TryParse<UhtEnumUnderlyingType>(enumType.Value.ToString(), true, out UhtEnumUnderlyingType underlyingType) || underlyingType == UhtEnumUnderlyingType.Unspecified)
 								{
-									topScope.TokenReader.LogError($"Namespace enums only support 'int' as the value size");
+									topScope.TokenReader.LogError(enumType.InputLine, $"Unsupported namespaced enum type '{enumType.Value}'");
 								}
+								enumObject.UnderlyingType = underlyingType;
+							}
+							else
+							{
+								enumObject.UnderlyingType = UhtEnumUnderlyingType.Unspecified;
 							}
 
 							topScope.TokenReader.Require('{');
diff --git a/Engine/Source/Programs/Shared/EpicGames.UHT/Types/Properties/UhtByteProperty.cs b/Engine/Source/Programs/Shared/EpicGames.UHT/Types/Properties/UhtByteProperty.cs
--- a/Engine/Source/Programs/Shared/EpicGames.UHT/Types/Properties/UhtByteProperty.cs
+++ b/Engine/Source/Programs/Shared/EpicGames.UHT/Types/Properties/UhtByteProperty.cs
@@ -151,7 +151,7 @@ namespace EpicGames.UHT.Types
 
 			if (function.FunctionFlags.HasAnyFlags(EFunctionFlags.BlueprintEvent | EFunctionFlags.BlueprintCallable))
 			{
-				if (Enum != null && Enum.UnderlyingType != UhtEnumUnderlyingType.uint8 && Enum.UnderlyingType != UhtEnumUnderlyingType.Unspecified)
+				if (Enum != null && Enum.CppForm == UhtEnumCppForm.EnumClass && Enum.UnderlyingType != UhtEnumUnderlyingType.Uint8 && Enum.UnderlyingType != UhtEnumUnderlyingType.Unspecified)
 				{
 					this.LogError("Invalid enum param for Blueprints - currently only uint8 supported");
 				}
diff --git a/Engine/Source/Programs/Shared/EpicGames.UHT/Types/Properties/UhtEnumProperty.cs b/Engine/Source/Programs/Shared/EpicGames.UHT/Types/Properties/UhtEnumProperty.cs
--- a/Engine/Source/Programs/Shared/EpicGames.UHT/Types/Properties/UhtEnumProperty.cs
+++ b/Engine/Source/Programs/Shared/EpicGames.UHT/Types/Properties/UhtEnumProperty.cs
@@ -44,10 +44,10 @@ namespace EpicGames.UHT.Types
 		public UhtProperty? UnderlyingProperty { get; set; }
 
 		/// <summary>
-		/// Underlying type which defaults to int32 if the referenced enum doesn't have an underlying type
+		/// Underlying type which defaults to Int32 if the referenced enum doesn't have an underlying type
 		/// </summary>
 		[JsonConverter(typeof(JsonStringEnumConverter))]
-		public UhtEnumUnderlyingType UnderlyingType => Enum.UnderlyingType != UhtEnumUnderlyingType.Unspecified ? Enum.UnderlyingType : UhtEnumUnderlyingType.int32;
+		public UhtEnumUnderlyingType UnderlyingType => Enum.UnderlyingType != UhtEnumUnderlyingType.Unspecified ? Enum.UnderlyingType : UhtEnumUnderlyingType.Int32;
 
 		/// <summary>
 		/// Underlying type size.  Defaults to unsized if the referenced enum doesn't have an underlying type
@@ -89,7 +89,7 @@ namespace EpicGames.UHT.Types
 			{
 				return null;
 			}
-			return $"enum class {Enum.SourceName} : {UnderlyingType};";
+			return System.String.Format("enum class {0} : {1};", Enum.SourceName, UnderlyingType.ToString().ToLower());
 		}
 
 		/// <inheritdoc/>
@@ -315,21 +315,21 @@ namespace EpicGames.UHT.Types
 			propertySettings.SourceName = "UnderlyingType";
 			switch (UnderlyingType)
 			{
-				case UhtEnumUnderlyingType.int8:
+				case UhtEnumUnderlyingType.Int8:
 					return new UhtInt8Property(propertySettings, UnderlyingTypeSize);
-				case UhtEnumUnderlyingType.int16:
+				case UhtEnumUnderlyingType.Int16:
 					return new UhtInt16Property(propertySettings, UnderlyingTypeSize);
-				case UhtEnumUnderlyingType.int32:
+				case UhtEnumUnderlyingType.Int32:
 					return new UhtIntProperty(propertySettings, UnderlyingTypeSize);
-				case UhtEnumUnderlyingType.int64:
+				case UhtEnumUnderlyingType.Int64:
 					return new UhtInt64Property(propertySettings, UnderlyingTypeSize);
-				case UhtEnumUnderlyingType.uint8:
+				case UhtEnumUnderlyingType.Uint8:
 					return new UhtByteProperty(propertySettings, UnderlyingTypeSize);
-				case UhtEnumUnderlyingType.uint16:
+				case UhtEnumUnderlyingType.Uint16:
 					return new UhtUInt16Property(propertySettings, UnderlyingTypeSize);
-				case UhtEnumUnderlyingType.uint32:
+				case UhtEnumUnderlyingType.Uint32:
 					return new UhtUInt32Property(propertySettings, UnderlyingTypeSize);
-				case UhtEnumUnderlyingType.uint64:
+				case UhtEnumUnderlyingType.Uint64:
 					return new UhtUInt64Property(propertySettings, UnderlyingTypeSize);
 				default:
 					throw new UhtIceException("Unexpected underlying enum type");
diff --git a/Engine/Source/Programs/Shared/EpicGames.UHT/Types/UhtEnum.cs b/Engine/Source/Programs/Shared/EpicGames.UHT/Types/UhtEnum.cs
--- a/Engine/Source/Programs/Shared/EpicGames.UHT/Types/UhtEnum.cs
+++ b/Engine/Source/Programs/Shared/EpicGames.UHT/Types/UhtEnum.cs
@@ -43,44 +43,49 @@ namespace EpicGames.UHT.Types
 		Unspecified,
 
 		/// <summary>
-		/// uint8
+		/// Uint8
 		/// </summary>
-		uint8,
+		Uint8,
 
 		/// <summary>
-		/// uint16
+		/// Uint16
 		/// </summary>
-		uint16,
+		Uint16,
 
 		/// <summary>
-		/// uint32
+		/// Uint32
 		/// </summary>
-		uint32,
+		Uint32,
 
 		/// <summary>
-		/// uint64
+		/// Uint64
 		/// </summary>
-		uint64,
+		Uint64,
 
 		/// <summary>
-		/// int8
+		/// Int8
 		/// </summary>
-		int8,
+		Int8,
 
 		/// <summary>
-		/// int16
+		/// Int16
 		/// </summary>
-		int16,
+		Int16,
 
 		/// <summary>
-		/// int32
+		/// Int32
 		/// </summary>
-		int32,
+		Int32,
 
 		/// <summary>
-		/// int64
+		/// Int64
 		/// </summary>
-		int64
+		Int64,
+
+		/// <summary>
+		/// Int
+		/// </summary>
+		Int,
 	}
 
 	/// <summary>
@@ -121,7 +126,7 @@ namespace EpicGames.UHT.Types
 		/// Underlying integer enumeration type
 		/// </summary>
 		[JsonConverter(typeof(JsonStringEnumConverter))]
-		public UhtEnumUnderlyingType UnderlyingType { get; set; } = UhtEnumUnderlyingType.uint8;
+		public UhtEnumUnderlyingType UnderlyingType { get; set; } = UhtEnumUnderlyingType.Uint8;
 
 		/// <summary>
 		/// Full enumeration type.  For namespace enumerations, this includes the namespace name and the enum type name
