Cleaned up some cut/copy/pasted code.

#rnx

diff --git a/Engine/Source/Programs/Shared/EpicGames.UHT/Parsers/UhtEnumParser.cs b/Engine/Source/Programs/Shared/EpicGames.UHT/Parsers/UhtEnumParser.cs
--- a/Engine/Source/Programs/Shared/EpicGames.UHT/Parsers/UhtEnumParser.cs
+++ b/Engine/Source/Programs/Shared/EpicGames.UHT/Parsers/UhtEnumParser.cs
@@ -76,20 +76,7 @@ namespace EpicGames.UHT.Parsers
 					// Read base for enum class
 					if (enumObject.CppForm == UhtEnumCppForm.EnumClass)
 					{
-						if (topScope.TokenReader.TryOptional(':'))
-						{
-							UhtToken enumType = topScope.TokenReader.GetIdentifier("enumeration base");
-
-							if (!System.Enum.TryParse<UhtEnumUnderlyingType>(enumType.Value.ToString(), true, out UhtEnumUnderlyingType underlyingType) || underlyingType == UhtEnumUnderlyingType.Unspecified)
-							{
-								topScope.TokenReader.LogError(enumType.InputLine, $"Unsupported enum class base type '{enumType.Value}'");
-							}
-							enumObject.UnderlyingType = underlyingType;
-						}
-						else
-						{
-							enumObject.UnderlyingType = UhtEnumUnderlyingType.Unspecified;
-						}
+						ParseUnderlyingType(topScope, enumObject);
 
 						if (enumObject.UnderlyingType != UhtEnumUnderlyingType.Unspecified && enumObject.UnderlyingType != UhtEnumUnderlyingType.Uint8 && enumObject.MetaData.ContainsKey("BlueprintType"))
 						{
@@ -100,20 +87,7 @@ namespace EpicGames.UHT.Parsers
 					{
 						if (enumObject.CppForm == UhtEnumCppForm.Regular)
 						{
-							if (topScope.TokenReader.TryOptional(":"))
-							{
-								UhtToken enumType = topScope.TokenReader.GetIdentifier("enumeration base");
-
-								if (!System.Enum.TryParse<UhtEnumUnderlyingType>(enumType.Value.ToString(), true, out UhtEnumUnderlyingType underlyingType) || underlyingType == UhtEnumUnderlyingType.Unspecified)
-								{
-									topScope.TokenReader.LogError(enumType.InputLine, $"Unsupported regular enum type '{enumType.Value}'");
-								}
-								enumObject.UnderlyingType = underlyingType;
-							}
-							else
-							{
-								enumObject.UnderlyingType = UhtEnumUnderlyingType.Unspecified;
-							}
+							ParseUnderlyingType(topScope, enumObject);
 						}
 						if ((enumObject.EnumFlags & EEnumFlags.Flags) != 0)
 						{
@@ -135,20 +109,7 @@ namespace EpicGames.UHT.Parsers
 
 							UhtToken innerEnumToken = topScope.TokenReader.GetIdentifier("enumeration type name");
 
-							if (topScope.TokenReader.TryOptional(":"))
-							{
-								UhtToken enumType = topScope.TokenReader.GetIdentifier("enumeration base");
-
-								if (!System.Enum.TryParse<UhtEnumUnderlyingType>(enumType.Value.ToString(), true, out UhtEnumUnderlyingType underlyingType) || underlyingType == UhtEnumUnderlyingType.Unspecified)
-								{
-									topScope.TokenReader.LogError(enumType.InputLine, $"Unsupported namespaced enum type '{enumType.Value}'");
-								}
-								enumObject.UnderlyingType = underlyingType;
-							}
-							else
-							{
-								enumObject.UnderlyingType = UhtEnumUnderlyingType.Unspecified;
-							}
+							ParseUnderlyingType(topScope, enumObject);
 
 							topScope.TokenReader.Require('{');
 							enumObject.CppType = $"{enumObject.SourceName}::{innerEnumToken.Value}";
@@ -263,5 +224,23 @@ namespace EpicGames.UHT.Parsers
 				return UhtParseResult.Handled;
 			}
 		}
+
+		private static void ParseUnderlyingType(UhtParsingScope topScope, UhtEnum enumObject)
+		{
+			if (topScope.TokenReader.TryOptional(':'))
+			{
+				UhtToken enumType = topScope.TokenReader.GetIdentifier("enumeration base");
+
+				if (!System.Enum.TryParse<UhtEnumUnderlyingType>(enumType.Value.ToString(), true, out UhtEnumUnderlyingType underlyingType) || underlyingType == UhtEnumUnderlyingType.Unspecified)
+				{
+					topScope.TokenReader.LogError(enumType.InputLine, $"Unsupported enum underlying base type '{enumType.Value}'");
+				}
+				enumObject.UnderlyingType = underlyingType;
+			}
+			else
+			{
+				enumObject.UnderlyingType = UhtEnumUnderlyingType.Unspecified;
+			}
+		}
 	}
 }
