EpicGames.Perforce: Add test for shelved changelists

diff --git a/Engine/Source/Programs/Shared/EpicGames.Perforce.Fixture/PerforceFixture.cs b/Engine/Source/Programs/Shared/EpicGames.Perforce.Fixture/PerforceFixture.cs
--- a/Engine/Source/Programs/Shared/EpicGames.Perforce.Fixture/PerforceFixture.cs
+++ b/Engine/Source/Programs/Shared/EpicGames.Perforce.Fixture/PerforceFixture.cs
@@ -40,6 +40,7 @@ public class ChangelistFixture
 	
 	public int Number { get; }
 	public string Description { get; }
+	public bool IsShelved { get; }
 	
 	/// <summary>
 	/// List of files in stream as how they would appear locally on disk, when synced to this changelist
@@ -47,11 +48,12 @@ public class ChangelistFixture
 	/// </summary>
 	public IReadOnlyList<DepotFileFixture> StreamFiles { get; }
 
-	public ChangelistFixture(int number, string description, List<DepotFileFixture> streamFiles)
+	public ChangelistFixture(int number, string description, List<DepotFileFixture> streamFiles, bool isShelved = false)
 	{
 		Number = number;
 		Description = description;
 		StreamFiles = streamFiles;
+		IsShelved = isShelved;
 	}
 
 	/// <summary>
@@ -114,14 +116,14 @@ public class ChangelistFixture
 			Console.WriteLine("Expected ------------------------------------------------------");
 			foreach ((string depotFile, int rev) in expectedList)
 			{
-				Console.WriteLine($"{depotFile,-20} | {rev,5}");
+				Console.WriteLine($"{depotFile,-30} | {rev,5}");
 			}
 			
 			Console.WriteLine("");
 			Console.WriteLine("Actual --------------------------------------------------------");
 			foreach ((string depotFile, int rev) in actualList)
 			{
-				Console.WriteLine($"{depotFile,-20} | {rev,5}");
+				Console.WriteLine($"{depotFile,-30} | {rev,5}");
 			}
 
 			Assert.Fail("Files in stream does not match files in client's have table");
@@ -170,19 +172,20 @@ public class ChangelistFixture
 public class StreamFixture
 {
 	public string Root { get; }
-	public IEnumerable<ChangelistFixture> Changelists { get; }
+	public IEnumerable<ChangelistFixture> Changelists => _changelists.Where(x => !x.IsShelved);
+	private readonly IReadOnlyList<ChangelistFixture> _changelists;
 
 	public StreamFixture(string root, IReadOnlyList<ChangelistFixture> changelists)
 	{
 		Root = root;
-		Changelists = changelists;
+		_changelists = changelists;
 	}
 
 	public ChangelistFixture LatestChangelist => Changelists.Last();
 
 	public ChangelistFixture GetChangelist(int changeNum)
 	{
-		foreach (ChangelistFixture changelist in Changelists)
+		foreach (ChangelistFixture changelist in _changelists)
 		{
 			if (changelist.Number == changeNum)
 			{
@@ -253,6 +256,15 @@ public class PerforceFixture
 					new("//Foo/Main/Data/data.txt", "Data/data.txt", "This is change data.txt #1", 1),
 					new("//Foo/Main/Data/moredata.txt", "Data/moredata.txt", "This is change moredata.txt #1", 1),
 				}),
+			new(8, "A shelved CL", // Assumes base CL is 7
+				new List<DepotFileFixture>()
+				{
+					new("//Foo/Main/main.h", "main.h", "This is change main.h #3", 3),
+					new("//Foo/Main/shared.h", "shared.h", "This is change common.h #1", 1),
+					new("//Foo/Main/shelved.cpp", "shelved.cpp", "This is change shelved.cpp #1", 1),
+					new("//Foo/Main/Data/data.txt", "Data/data.txt", "This is change data.txt #1", 1),
+					new("//Foo/Main/Data/moredata.txt", "Data/moredata.txt", "This is change moredata.txt #1", 1),
+				}, true),
 		});
 
 	public static string CalcMd5(string content)
diff --git a/Engine/Source/Programs/Shared/EpicGames.Perforce.Fixture/setup.sh b/Engine/Source/Programs/Shared/EpicGames.Perforce.Fixture/setup.sh
--- a/Engine/Source/Programs/Shared/EpicGames.Perforce.Fixture/setup.sh
+++ b/Engine/Source/Programs/Shared/EpicGames.Perforce.Fixture/setup.sh
@@ -102,5 +102,15 @@ p4 submit -d "Some updates to main"
 echo "This is change moredata.txt #1" > $CLIENT_ROOT/Data/moredata.txt && p4 add $CLIENT_ROOT/Data/moredata.txt
 p4 submit -d "Add more data"
 
+# Changelist #8
+p4 change -i << EOF
+Change: new
+Description: A shelved CL
+EOF
+echo "This is change shelved.cpp #1" > $CLIENT_ROOT/shelved.cpp && p4 add -c 8 $CLIENT_ROOT/shelved.cpp
+p4 edit -c 8 main.h && echo "This is change main.h #3" > $CLIENT_ROOT/main.h
+p4 delete -c 8 main.cpp
+p4 shelve -c 8
+
 # Stop server running in background
 kill "$(pidof p4d)"
\ No newline at end of file
diff --git a/Engine/Source/Programs/Shared/EpicGames.Perforce.Managed.Tests/ManagedWorkspaceTest.cs b/Engine/Source/Programs/Shared/EpicGames.Perforce.Managed.Tests/ManagedWorkspaceTest.cs
--- a/Engine/Source/Programs/Shared/EpicGames.Perforce.Managed.Tests/ManagedWorkspaceTest.cs
+++ b/Engine/Source/Programs/Shared/EpicGames.Perforce.Managed.Tests/ManagedWorkspaceTest.cs
@@ -141,6 +141,20 @@ public class ManagedWorkspaceTest : BasePerforceFixtureTest
 		Stream.LatestChangelist.AssertDepotFiles(SyncDir);
 		await Stream.LatestChangelist.AssertHaveTableAsync(PerforceConnection, useHaveTable);
 	}
+	
+	[TestMethod]
+	[DataRow(true, DisplayName = "With have-table")]
+	[DataRow(false, DisplayName = "Without have-table")]
+	public async Task Unshelve(bool useHaveTable)
+	{
+		ManagedWorkspace ws = await GetManagedWorkspace();
+		await SyncAsync(ws, 7, useHaveTable);
+		await ws.UnshelveAsync(PerforceConnection, 8, CancellationToken.None);
+		Stream.GetChangelist(8).AssertDepotFiles(SyncDir);
+		
+		// Have-table still correspond to CL 7 as CL 8 is shelved, only p4 printed to workspace
+		await Stream.GetChangelist(7).AssertHaveTableAsync(PerforceConnection, useHaveTable);
+	}
 
 	private async Task<ManagedWorkspace> GetManagedWorkspace()
 	{
