Horde: Remove hashes from TreeNodeRef objects, and distinguish between identically serialized nodes with different types. This prevents cascading the requirement to propagate metadata into imports and refs, and simplifies a lot of bookkeeping around node states.

diff --git a/Engine/Source/Programs/Horde/Horde.Agent.Tests/JsonRpcLoggerTests.cs b/Engine/Source/Programs/Horde/Horde.Agent.Tests/JsonRpcLoggerTests.cs
--- a/Engine/Source/Programs/Horde/Horde.Agent.Tests/JsonRpcLoggerTests.cs
+++ b/Engine/Source/Programs/Horde/Horde.Agent.Tests/JsonRpcLoggerTests.cs
@@ -104,14 +104,14 @@ namespace Horde.Agent.Tests
 
 		class FakeJsonRpcLoggerBackend : JsonRpcAndStorageLogSink
 		{
-			public RefTarget? Target { get; private set; }
+			public NodeLocator Target { get; private set; }
 
 			public FakeJsonRpcLoggerBackend(IRpcConnection connection, string logId, IJsonRpcLogSink inner, IStorageClient store, ILogger logger)
 				: base(connection, logId, inner, store, logger)
 			{
 			}
 
-			protected override Task UpdateLogAsync(RefTarget target, int lineCount, CancellationToken cancellationToken)
+			protected override Task UpdateLogAsync(NodeLocator target, int lineCount, CancellationToken cancellationToken)
 			{
 				Target = target;
 				return Task.CompletedTask;
diff --git a/Engine/Source/Programs/Horde/Horde.Agent/Commands/Bundles/DumpCommand.cs b/Engine/Source/Programs/Horde/Horde.Agent/Commands/Bundles/DumpCommand.cs
--- a/Engine/Source/Programs/Horde/Horde.Agent/Commands/Bundles/DumpCommand.cs
+++ b/Engine/Source/Programs/Horde/Horde.Agent/Commands/Bundles/DumpCommand.cs
@@ -6,7 +6,6 @@ using System.Linq;
 using System.Threading.Tasks;
 using EpicGames.Core;
 using EpicGames.Horde.Storage;
-using Microsoft.Extensions.Caching.Memory;
 using Microsoft.Extensions.Logging;
 
 namespace Horde.Agent.Commands.Bundles
@@ -27,8 +26,8 @@ namespace Horde.Agent.Commands.Bundles
 
 			if (BlobId == null)
 			{
-				RefTarget locator = await store.ReadRefTargetAsync(RefName);
-				BlobId = locator.Locator.Blob;
+				NodeLocator locator = await store.ReadRefTargetAsync(RefName);
+				BlobId = locator.Blob;
 			}
 
 			logger.LogInformation("Summary for blob {BlobId}", BlobId.Value);
@@ -44,9 +43,9 @@ namespace Horde.Agent.Commands.Bundles
 			foreach (BundleImport import in header.Imports)
 			{
 				logger.LogInformation("  From blob {BlobId} ({NumExports} nodes)", import.Locator, import.Exports.Count);
-				foreach ((int exportIdx, IoHash exportHash) in import.Exports)
+				foreach (int exportIdx in import.Exports)
 				{
-					logger.LogInformation("    [{Index}] IMP {BlobId}:{ExportIdx} = {ExportHash}", refIdx, import.Locator, exportIdx, exportHash);
+					logger.LogInformation("    [{Index}] IMP {BlobId}:{ExportIdx}", refIdx, import.Locator, exportIdx);
 					refIdx++;
 				}
 			}
diff --git a/Engine/Source/Programs/Horde/Horde.Agent/Commands/Vcs/VcsCommand.cs b/Engine/Source/Programs/Horde/Horde.Agent/Commands/Vcs/VcsCommand.cs
--- a/Engine/Source/Programs/Horde/Horde.Agent/Commands/Vcs/VcsCommand.cs
+++ b/Engine/Source/Programs/Horde/Horde.Agent/Commands/Vcs/VcsCommand.cs
@@ -482,7 +482,7 @@ namespace Horde.Agent.Commands.Bundles
 			return 0;
 		}
 
-		static async Task<DirectoryState> RealizeAsync(TreeReader reader, TreeNodeRef<DirectoryNode> directoryRef, DirectoryReference dirPath, DirectoryState? directoryState, bool clean, ILogger logger)
+		static async Task<DirectoryState> RealizeAsync(TreeReader reader, DirectoryNodeRef directoryRef, DirectoryReference dirPath, DirectoryState? directoryState, bool clean, ILogger logger)
 		{
 			DirectoryReference.CreateDirectory(dirPath);
 
@@ -529,7 +529,7 @@ namespace Horde.Agent.Commands.Bundles
 			return newState;
 		}
 
-		static async Task<FileState> CheckoutFileAsync(TreeReader reader, TreeNodeRef<FileNode> fileRef, FileInfo fileInfo, ILogger logger)
+		static async Task<FileState> CheckoutFileAsync(TreeReader reader, FileEntry fileRef, FileInfo fileInfo, ILogger logger)
 		{
 			logger.LogInformation("Updating {File} to {Hash}", fileInfo, fileRef.Hash);
 			FileNode fileNode = await fileRef.ExpandAsync(reader);
@@ -588,10 +588,10 @@ namespace Horde.Agent.Commands.Bundles
 				rootNode = await tip.Contents.ExpandCopyAsync(reader);
 			}
 
-			TreeNodeRef<DirectoryNode> rootRef = new TreeNodeRef<DirectoryNode>(rootNode);
+			DirectoryNodeRef rootRef = new DirectoryNodeRef(rootNode);
 
 			List<(DirectoryNode, FileInfo, FileState)> files = new List<(DirectoryNode, FileInfo, FileState)>();
-			List<(TreeNodeRef<DirectoryNode>, DirectoryState)> directories = new List<(TreeNodeRef<DirectoryNode>, DirectoryState)>();
+			List<(DirectoryNodeRef, DirectoryState)> directories = new List<(DirectoryNodeRef, DirectoryState)>();
 			await UpdateTreeAsync(reader, rootRef, rootDir, oldState, newState, files, directories);
 
 			TreeWriter writer = new TreeWriter(store);
@@ -608,7 +608,7 @@ namespace Horde.Agent.Commands.Bundles
 			}
 			await store.WriteNodeAsync(workspaceState.Branch, newTip);
 
-			foreach ((TreeNodeRef<DirectoryNode> directoryRef, DirectoryState directoryState) in directories)
+			foreach ((DirectoryNodeRef directoryRef, DirectoryState directoryState) in directories)
 			{
 				directoryState.Hash = directoryRef.Hash;
 			}
@@ -627,7 +627,7 @@ namespace Horde.Agent.Commands.Bundles
 			return 0;
 		}
 
-		private async Task UpdateTreeAsync(TreeReader reader, TreeNodeRef<DirectoryNode> rootRef, DirectoryReference rootDir, DirectoryState? oldState, DirectoryState newState, List<(DirectoryNode, FileInfo, FileState)> files, List<(TreeNodeRef<DirectoryNode>, DirectoryState)> directories)
+		private async Task UpdateTreeAsync(TreeReader reader, DirectoryNodeRef rootRef, DirectoryReference rootDir, DirectoryState? oldState, DirectoryState newState, List<(DirectoryNode, FileInfo, FileState)> files, List<(DirectoryNodeRef, DirectoryState)> directories)
 		{
 			directories.Add((rootRef, newState));
 
diff --git a/Engine/Source/Programs/Horde/Horde.Agent/Utility/JsonRpcLogger.cs b/Engine/Source/Programs/Horde/Horde.Agent/Utility/JsonRpcLogger.cs
--- a/Engine/Source/Programs/Horde/Horde.Agent/Utility/JsonRpcLogger.cs
+++ b/Engine/Source/Programs/Horde/Horde.Agent/Utility/JsonRpcLogger.cs
@@ -213,7 +213,7 @@ namespace Horde.Agent.Parser
 
 			if (request.Flush || _bufferLength > FlushLength)
 			{
-				RefTarget target = await _builder.FlushAsync(_writer, request.Flush, cancellationToken);
+				NodeLocator target = await _builder.FlushAsync(_writer, request.Flush, cancellationToken);
 				await UpdateLogAsync(target, _builder.LineCount, cancellationToken);
 				_bufferLength = 0;
 			}
@@ -223,9 +223,9 @@ namespace Horde.Agent.Parser
 
 		#region RPC calls
 
-		protected virtual async Task UpdateLogAsync(RefTarget target, int lineCount, CancellationToken cancellationToken)
+		protected virtual async Task UpdateLogAsync(NodeLocator target, int lineCount, CancellationToken cancellationToken)
 		{
-			_logger.LogDebug("Updating log {LogId} to line {LineCount}, blob {Locator}", _logId, lineCount, target.Locator);
+			_logger.LogDebug("Updating log {LogId} to line {LineCount}, target {Locator}", _logId, lineCount, target);
 
 			UpdateLogRequest request = new UpdateLogRequest();
 			request.LogId = _logId;
diff --git a/Engine/Source/Programs/Horde/Horde.Build.Tests/BlobStoreTests.cs b/Engine/Source/Programs/Horde/Horde.Build.Tests/BlobStoreTests.cs
--- a/Engine/Source/Programs/Horde/Horde.Build.Tests/BlobStoreTests.cs
+++ b/Engine/Source/Programs/Horde/Horde.Build.Tests/BlobStoreTests.cs
@@ -58,7 +58,7 @@ namespace Horde.Build.Tests
 				int index;
 				if (!locatorToIndex.TryGetValue(locator, out index))
 				{
-					BundleImport import = new BundleImport(locator, new List<(int, IoHash)> { (0, IoHash.Zero) });
+					BundleImport import = new BundleImport(locator, new int[] { 0 });
 					imports.Add(import);
 					locatorToIndex.Add(locator, imports.Count - 1);
 				}
@@ -129,15 +129,15 @@ namespace Horde.Build.Tests
 			for(int idx = 0; idx < 2; idx++)
 			{
 				RefName refName = new RefName("hello");
-				await store.WriteRefTargetAsync(refName, new RefTarget(bundle3.Header.Exports[0].Hash, new NodeLocator(locator3, 0)));
-				RefTarget refTarget = await store.ReadRefTargetAsync(refName);
-				Assert.AreEqual(locator3, refTarget.Locator.Blob);
+				await store.WriteRefTargetAsync(refName, new NodeLocator(locator3, 0));
+				NodeLocator refTarget = await store.ReadRefTargetAsync(refName);
+				Assert.AreEqual(locator3, refTarget.Blob);
 			}
 
 			RefName refName2 = new RefName("hello2");
 
-			RefTarget refTargetId2 = await store.WriteRefAsync(refName2, CreateTestBundle(input3, new BlobLocator[] { locator1, locator2 }), 0);
-			Blob refTarget2 = await ReadBlobAsync(store, refTargetId2.Locator.Blob);
+			NodeLocator refTargetId2 = await store.WriteRefAsync(refName2, CreateTestBundle(input3, new BlobLocator[] { locator1, locator2 }), 0);
+			Blob refTarget2 = await ReadBlobAsync(store, refTargetId2.Blob);
 			Assert.IsTrue(refTarget2.Data.Span.SequenceEqual(input3));
 			Assert.IsTrue(refTarget2.References.SequenceEqual(new BlobLocator[] { locator1, locator2 }));
 		}
@@ -149,7 +149,7 @@ namespace Horde.Build.Tests
 
 			Bundle bundle1 = CreateTestBundle(new byte[] { 1, 2, 3 }, Array.Empty<BlobLocator>());
 			BlobLocator locator1 = await store.WriteBundleAsync(bundle1);
-			RefTarget target = new RefTarget(bundle1.Header.Exports[0].Hash, new NodeLocator(locator1, 0));
+			NodeLocator target = new NodeLocator(locator1, 0);
 
 			await store.WriteRefTargetAsync("test-ref-1", target);
 			await store.WriteRefTargetAsync("test-ref-2", target, new RefOptions { Lifetime = TimeSpan.FromMinutes(30.0), Extend = true });
diff --git a/Engine/Source/Programs/Horde/Horde.Build.Tests/GcServiceTests.cs b/Engine/Source/Programs/Horde/Horde.Build.Tests/GcServiceTests.cs
--- a/Engine/Source/Programs/Horde/Horde.Build.Tests/GcServiceTests.cs
+++ b/Engine/Source/Programs/Horde/Horde.Build.Tests/GcServiceTests.cs
@@ -46,7 +46,7 @@ namespace Horde.Build.Tests
 				int blobIdx = (int)(random.NextDouble() * blobs.Length);
 				if (roots.Add(blobs[blobIdx]))
 				{
-					await store.WriteRefTargetAsync(new RefName($"ref-{idx}"), new RefTarget(IoHash.Zero, new NodeLocator(blobs[blobIdx], 0)));
+					await store.WriteRefTargetAsync(new RefName($"ref-{idx}"), new NodeLocator(blobs[blobIdx], 0));
 				}
 			}
 
@@ -106,7 +106,7 @@ namespace Horde.Build.Tests
 			for (int idx = numNodes - 1; idx >= 0; idx--)
 			{
 				List<BundleType> types = new List<BundleType> { new BundleType(Guid.Parse("{AFDF76A7-5333-4DEE-B837-B5F5CA511245}"), 0) };
-				List<BundleImport> imports = children[idx].ConvertAll(x => new BundleImport(locators[x], Array.Empty<(int, IoHash)>()));
+				List<BundleImport> imports = children[idx].ConvertAll(x => new BundleImport(locators[x], Array.Empty<int>()));
 				BundleHeader header = new BundleHeader(BundleCompressionFormat.None, types, imports, Array.Empty<BundleExport>(), Array.Empty<BundlePacket>());
 				Bundle bundle = new Bundle(header, Array.Empty<ReadOnlyMemory<byte>>());
 				locators[idx] = await store.WriteBundleAsync(bundle, prefix: "gctest");
diff --git a/Engine/Source/Programs/Horde/Horde.Build/Logs/LogRpcService.cs b/Engine/Source/Programs/Horde/Horde.Build/Logs/LogRpcService.cs
--- a/Engine/Source/Programs/Horde/Horde.Build/Logs/LogRpcService.cs
+++ b/Engine/Source/Programs/Horde/Horde.Build/Logs/LogRpcService.cs
@@ -53,7 +53,7 @@ namespace Horde.Build.Logs
 
 			IStorageClient store = await _storageService.GetClientAsync(Namespace.Logs, context.CancellationToken);
 			_logger.LogInformation("Updating {LogId} to node {RefTarget}", request.LogId, request.Target);
-			await store.WriteRefTargetAsync(new RefName(request.LogId), RefTarget.Parse(request.Target));
+			await store.WriteRefTargetAsync(new RefName(request.LogId), NodeLocator.Parse(request.Target));
 
 			await _logFileCollection.UpdateLineCountAsync(logFile, request.LineCount, CancellationToken.None);
 
diff --git a/Engine/Source/Programs/Horde/Horde.Build/Perforce/PerforceReplicator.cs b/Engine/Source/Programs/Horde/Horde.Build/Perforce/PerforceReplicator.cs
--- a/Engine/Source/Programs/Horde/Horde.Build/Perforce/PerforceReplicator.cs
+++ b/Engine/Source/Programs/Horde/Horde.Build/Perforce/PerforceReplicator.cs
@@ -420,7 +420,7 @@ namespace Horde.Build.Perforce
 
 			// Create the commit node
 			ChangeRecord changeRecord = await perforce.GetChangeAsync(GetChangeOptions.None, change, cancellationToken);
-			CommitNode commitNode = new CommitNode(change, parentRef, changeRecord.User ?? "Unknown", changeRecord.Description ?? String.Empty, changeRecord.Date, new TreeNodeRef<DirectoryNode>(root));
+			CommitNode commitNode = new CommitNode(change, parentRef, changeRecord.User ?? "Unknown", changeRecord.Description ?? String.Empty, changeRecord.Date, new DirectoryNodeRef(root));
 			await writer.WriteAsync(refName, commitNode, options.RefOptions, cancellationToken); 
 		}
 
diff --git a/Engine/Source/Programs/Horde/Horde.Build/Storage/StorageController.cs b/Engine/Source/Programs/Horde/Horde.Build/Storage/StorageController.cs
--- a/Engine/Source/Programs/Horde/Horde.Build/Storage/StorageController.cs
+++ b/Engine/Source/Programs/Horde/Horde.Build/Storage/StorageController.cs
@@ -54,11 +54,6 @@ namespace Horde.Build.Storage
 		/// </summary>
 		public int ExportIdx { get; set; }
 
-		/// <summary>
-		/// Hash of the target node
-		/// </summary>
-		public IoHash Hash { get; set; }
-
 		/// <summary>
 		/// Options for the ref
 		/// </summary>
@@ -80,19 +75,13 @@ namespace Horde.Build.Storage
 		/// </summary>
 		public int ExportIdx { get; set; }
 
-		/// <summary>
-		/// Hash of the target node
-		/// </summary>
-		public IoHash Hash { get; set; }
-
 		/// <summary>
 		/// Constructor
 		/// </summary>
-		public ReadRefResponse(RefTarget target)
+		public ReadRefResponse(NodeLocator target)
 		{
-			Blob = target.Locator.Blob;
-			ExportIdx = target.Locator.ExportIdx;
-			Hash = target.Hash;
+			Blob = target.Blob;
+			ExportIdx = target.ExportIdx;
 		}
 	}
 
@@ -249,7 +238,7 @@ namespace Horde.Build.Storage
 			}
 
 			IStorageClient client = await _storageService.GetClientAsync(namespaceId, cancellationToken);
-			RefTarget target = new RefTarget(request.Hash, new NodeLocator(request.Blob, request.ExportIdx));
+			NodeLocator target = new NodeLocator(request.Blob, request.ExportIdx);
 			await client.WriteRefTargetAsync(refName, target, request.Options, cancellationToken);
 
 			return Ok();
@@ -272,8 +261,8 @@ namespace Horde.Build.Storage
 
 			IStorageClient client = await _storageService.GetClientAsync(namespaceId, cancellationToken);
 
-			RefTarget? target = await client.TryReadRefTargetAsync(refName, cancellationToken: cancellationToken);
-			if (target == null)
+			NodeLocator target = await client.TryReadRefTargetAsync(refName, cancellationToken: cancellationToken);
+			if (!target.IsValid())
 			{
 				return NotFound();
 			}
diff --git a/Engine/Source/Programs/Shared/EpicGames.Horde.Tests/BundleTests.cs b/Engine/Source/Programs/Shared/EpicGames.Horde.Tests/BundleTests.cs
--- a/Engine/Source/Programs/Shared/EpicGames.Horde.Tests/BundleTests.cs
+++ b/Engine/Source/Programs/Shared/EpicGames.Horde.Tests/BundleTests.cs
@@ -107,8 +107,8 @@ namespace EpicGames.Horde.Tests
 				types.Add(new BundleType(Guid.NewGuid(), 0));
 
 				List<BundleImport> imports = new List<BundleImport>();
-				imports.Add(new BundleImport(new BlobLocator("import1"), new (int, IoHash)[] { (5, IoHash.Compute(Encoding.UTF8.GetBytes("blob1"))) }));
-				imports.Add(new BundleImport(new BlobLocator("import2"), new (int, IoHash)[] { (6, IoHash.Compute(Encoding.UTF8.GetBytes("blob2"))) }));
+				imports.Add(new BundleImport(new BlobLocator("import1"), new[] { 5 }));
+				imports.Add(new BundleImport(new BlobLocator("import2"), new[] { 6 }));
 
 				List<BundleExport> exports = new List<BundleExport>();
 				exports.Add(new BundleExport(0, IoHash.Compute(Encoding.UTF8.GetBytes("export1")), 2, new int[] { 1, 2 }));
@@ -138,11 +138,10 @@ namespace EpicGames.Horde.Tests
 
 				for (int importIdx = 0; importIdx < oldImport.Exports.Count; importIdx++)
 				{
-					(int Index, IoHash Hash) oldExport = oldImport.Exports[importIdx];
-					(int Index, IoHash Hash) newExport = newImport.Exports[importIdx];
+					int oldExportIdx = oldImport.Exports[importIdx];
+					int newExportIdx = newImport.Exports[importIdx];
 
-					Assert.AreEqual(oldExport.Hash, newExport.Hash);
-					Assert.AreEqual(oldExport.Index, newExport.Index);
+					Assert.AreEqual(oldExportIdx, newExportIdx);
 				}
 			}
 
@@ -185,8 +184,8 @@ namespace EpicGames.Horde.Tests
 			Assert.AreEqual(1, store.Blobs.Count);
 
 			// Check the ref
-			RefTarget refTarget = await store.ReadRefTargetAsync(refName);
-			Bundle bundle = await store.ReadBundleAsync(refTarget.Locator.Blob);
+			NodeLocator refTarget = await store.ReadRefTargetAsync(refName);
+			Bundle bundle = await store.ReadBundleAsync(refTarget.Blob);
 			Assert.AreEqual(0, bundle.Header.Imports.Count);
 			Assert.AreEqual(3, bundle.Header.Exports.Count);
 
diff --git a/Engine/Source/Programs/Shared/EpicGames.Horde.Tests/LogTests.cs b/Engine/Source/Programs/Shared/EpicGames.Horde.Tests/LogTests.cs
--- a/Engine/Source/Programs/Shared/EpicGames.Horde.Tests/LogTests.cs
+++ b/Engine/Source/Programs/Shared/EpicGames.Horde.Tests/LogTests.cs
@@ -87,8 +87,8 @@ namespace EpicGames.Horde.Tests
 			}
 
 			// Flush it to storage, and read the finished log node
-			RefTarget target = await builder.FlushAsync(new TreeWriter(_store), true, CancellationToken.None);
-			LogNode log = await _reader.ReadNodeAsync<LogNode>(target.Locator);
+			NodeLocator target = await builder.FlushAsync(new TreeWriter(_store), true, CancellationToken.None);
+			LogNode log = await _reader.ReadNodeAsync<LogNode>(target);
 
 			// Read the data back out and check it's the same
 			byte[] readData = new byte[_data.Length];
@@ -129,10 +129,10 @@ namespace EpicGames.Horde.Tests
 			{
 				builder.WriteData(Encoding.UTF8.GetBytes(lines[lineIdx]));
 			}
-			RefTarget rootNodeLocator = await builder.FlushAsync(new TreeWriter(_store), true, CancellationToken.None);
+			NodeLocator rootNodeLocator = await builder.FlushAsync(new TreeWriter(_store), true, CancellationToken.None);
 
 			// Read it back in and test the index
-			LogNode rootNode = await _reader.ReadNodeAsync<LogNode>(rootNodeLocator.Locator);
+			LogNode rootNode = await _reader.ReadNodeAsync<LogNode>(rootNodeLocator);
 			LogIndexNode index = await rootNode.IndexRef.ExpandAsync(_reader);
 		
 			for (int lineIdx = 0; lineIdx < lines.Length; lineIdx++)
diff --git a/Engine/Source/Programs/Shared/EpicGames.Horde/Logs/LogChunkNode.cs b/Engine/Source/Programs/Shared/EpicGames.Horde/Logs/LogChunkNode.cs
--- a/Engine/Source/Programs/Shared/EpicGames.Horde/Logs/LogChunkNode.cs
+++ b/Engine/Source/Programs/Shared/EpicGames.Horde/Logs/LogChunkNode.cs
@@ -244,7 +244,7 @@ namespace EpicGames.Horde.Logs
 		}
 
 		/// <inheritdoc/>
-		public new void Serialize(ITreeNodeWriter writer)
+		protected override void Serialize(IMemoryWriter writer)
 		{
 			base.Serialize(writer);
 
diff --git a/Engine/Source/Programs/Shared/EpicGames.Horde/Logs/LogIndexNode.cs b/Engine/Source/Programs/Shared/EpicGames.Horde/Logs/LogIndexNode.cs
--- a/Engine/Source/Programs/Shared/EpicGames.Horde/Logs/LogIndexNode.cs
+++ b/Engine/Source/Programs/Shared/EpicGames.Horde/Logs/LogIndexNode.cs
@@ -106,7 +106,7 @@ namespace EpicGames.Horde.Logs
 			writer.WriteUnsignedVarInt(CurrentVersion);
 			writer.WriteNgramSet(_ngramSet);
 			writer.WriteUnsignedVarInt(_numChunkBits);
-			writer.WriteVariableLengthArray(_plainTextChunkRefs, x => x.Serialize(writer));
+			writer.WriteVariableLengthArray(_plainTextChunkRefs, x => writer.WriteRef(x));
 		}
 
 		/// <inheritdoc/>
diff --git a/Engine/Source/Programs/Shared/EpicGames.Horde/Logs/LogNode.cs b/Engine/Source/Programs/Shared/EpicGames.Horde/Logs/LogNode.cs
--- a/Engine/Source/Programs/Shared/EpicGames.Horde/Logs/LogNode.cs
+++ b/Engine/Source/Programs/Shared/EpicGames.Horde/Logs/LogNode.cs
@@ -107,7 +107,7 @@ namespace EpicGames.Horde.Logs
 			writer.WriteUnsignedVarInt(LineCount);
 			writer.WriteUnsignedVarInt((ulong)Length);
 			writer.WriteRef(IndexRef);
-			writer.WriteList(TextChunkRefs, x => x.Serialize(writer));
+			writer.WriteList(TextChunkRefs, x => writer.WriteRef(x));
 			writer.WriteBoolean(Complete);
 		}
 
@@ -269,7 +269,7 @@ namespace EpicGames.Horde.Logs
 		/// <param name="writer">Writer for the output nodes</param>
 		/// <param name="complete">Whether the log is complete</param>
 		/// <param name="cancellationToken">Cancellation token for the operation</param>
-		public async Task<RefTarget> FlushAsync(TreeWriter writer, bool complete, CancellationToken cancellationToken)
+		public async Task<NodeLocator> FlushAsync(TreeWriter writer, bool complete, CancellationToken cancellationToken)
 		{
 			// Capture the new data that needs to be written
 			IReadOnlyList<LogChunkNode> writeTextChunks;
@@ -299,10 +299,7 @@ namespace EpicGames.Horde.Logs
 
 			LogNode newRoot = new LogNode(_format, lineCount, length, newJsonChunkRefs, new TreeNodeRef<LogIndexNode>(newIndex), complete);
 
-			IoHash newRootHash = await writer.WriteAsync(new TreeNodeRef(newRoot), cancellationToken);
-			await writer.FlushAsync(cancellationToken);
-
-			NodeLocator newRootLocator = writer.GetLocator(newRootHash);
+			NodeLocator newRootLocator = await writer.FlushAsync(newRoot, cancellationToken);
 
 			// Update the new state
 			lock (_lockObject)
@@ -313,7 +310,7 @@ namespace EpicGames.Horde.Logs
 				_indexTextBuilder.Remove(writeIndexTextChunks.Count);
 			}
 
-			return new RefTarget(newRootHash, newRootLocator);
+			return newRootLocator;
 		}
 	}
 
diff --git a/Engine/Source/Programs/Shared/EpicGames.Horde/Storage/Backends/FileStorageClient.cs b/Engine/Source/Programs/Shared/EpicGames.Horde/Storage/Backends/FileStorageClient.cs
--- a/Engine/Source/Programs/Shared/EpicGames.Horde/Storage/Backends/FileStorageClient.cs
+++ b/Engine/Source/Programs/Shared/EpicGames.Horde/Storage/Backends/FileStorageClient.cs
@@ -82,7 +82,7 @@ namespace EpicGames.Horde.Storage.Backends
 		}
 
 		/// <inheritdoc/>
-		public override async Task<RefTarget?> TryReadRefTargetAsync(RefName name, DateTime cacheTime = default, CancellationToken cancellationToken = default)
+		public override async Task<NodeLocator> TryReadRefTargetAsync(RefName name, DateTime cacheTime = default, CancellationToken cancellationToken = default)
 		{
 			FileReference file = GetRefFile(name);
 			if (!FileReference.Exists(file))
@@ -92,11 +92,11 @@ namespace EpicGames.Horde.Storage.Backends
 
 			_logger.LogInformation("Reading {File}", file);
 			string text = await FileReference.ReadAllTextAsync(file, cancellationToken);
-			return RefTarget.Parse(text);
+			return NodeLocator.Parse(text);
 		}
 
 		/// <inheritdoc/>
-		public override async Task WriteRefTargetAsync(RefName name, RefTarget target, RefOptions? options = null, CancellationToken cancellationToken = default)
+		public override async Task WriteRefTargetAsync(RefName name, NodeLocator target, RefOptions? options = null, CancellationToken cancellationToken = default)
 		{
 			FileReference file = GetRefFile(name);
 			DirectoryReference.CreateDirectory(file.Directory);
diff --git a/Engine/Source/Programs/Shared/EpicGames.Horde/Storage/Backends/MemoryStorageClient.cs b/Engine/Source/Programs/Shared/EpicGames.Horde/Storage/Backends/MemoryStorageClient.cs
--- a/Engine/Source/Programs/Shared/EpicGames.Horde/Storage/Backends/MemoryStorageClient.cs
+++ b/Engine/Source/Programs/Shared/EpicGames.Horde/Storage/Backends/MemoryStorageClient.cs
@@ -26,13 +26,13 @@ namespace EpicGames.Horde.Storage.Backends
 		/// <summary>
 		/// Map of ref name to ref data
 		/// </summary>
-		readonly ConcurrentDictionary<RefName, RefTarget> _refs = new ConcurrentDictionary<RefName, RefTarget>();
+		readonly ConcurrentDictionary<RefName, NodeLocator> _refs = new ConcurrentDictionary<RefName, NodeLocator>();
 
 		/// <inheritdoc cref="_blobs"/>
 		public IReadOnlyDictionary<BlobLocator, Bundle> Blobs => _blobs;
 
 		/// <inheritdoc cref="_refs"/>
-		public IReadOnlyDictionary<RefName, RefTarget> Refs => _refs;
+		public IReadOnlyDictionary<RefName, NodeLocator> Refs => _refs;
 
 		/// <summary>
 		/// Constructor
@@ -73,15 +73,15 @@ namespace EpicGames.Horde.Storage.Backends
 		public override Task DeleteRefAsync(RefName name, CancellationToken cancellationToken) => Task.FromResult(_refs.TryRemove(name, out _));
 
 		/// <inheritdoc/>
-		public override Task<RefTarget?> TryReadRefTargetAsync(RefName name, DateTime cacheTime = default, CancellationToken cancellationToken = default)
+		public override Task<NodeLocator> TryReadRefTargetAsync(RefName name, DateTime cacheTime = default, CancellationToken cancellationToken = default)
 		{
-			RefTarget? refTarget;
+			NodeLocator refTarget;
 			_refs.TryGetValue(name, out refTarget);
 			return Task.FromResult(refTarget);
 		}
 
 		/// <inheritdoc/>
-		public override Task WriteRefTargetAsync(RefName name, RefTarget target, RefOptions? options = null, CancellationToken cancellationToken = default)
+		public override Task WriteRefTargetAsync(RefName name, NodeLocator target, RefOptions? options = null, CancellationToken cancellationToken = default)
 		{
 			_refs[name] = target;
 			return Task.CompletedTask;
diff --git a/Engine/Source/Programs/Shared/EpicGames.Horde/Storage/Bundle.cs b/Engine/Source/Programs/Shared/EpicGames.Horde/Storage/Bundle.cs
--- a/Engine/Source/Programs/Shared/EpicGames.Horde/Storage/Bundle.cs
+++ b/Engine/Source/Programs/Shared/EpicGames.Horde/Storage/Bundle.cs
@@ -7,6 +7,7 @@ using System.Buffers;
 using System.Buffers.Binary;
 using System.Collections.Generic;
 using System.Diagnostics;
+using System.Globalization;
 using System.IO;
 using System.IO.Compression;
 using System.Linq;
@@ -366,6 +367,19 @@ namespace EpicGames.Horde.Storage
 			writer.WriteUnsignedVarInt(Version);
 		}
 
+		/// <summary>
+		/// Parse a span of characters as a bundle type
+		/// </summary>
+		/// <param name="text">Text to parse</param>
+		/// <returns>The parsed bundle type</returns>
+		public static BundleType Parse(ReadOnlySpan<char> text)
+		{
+			int hashIdx = text.IndexOf('#');
+			Guid guid = Guid.Parse(text.Slice(0, hashIdx));
+			int version = Int32.Parse(text.Slice(hashIdx + 1), NumberStyles.None, CultureInfo.InvariantCulture);
+			return new BundleType(guid, version);
+		}
+
 		/// <inheritdoc/>
 		public override bool Equals(object? obj) => obj is BundleType type && Equals(type);
 
@@ -375,6 +389,9 @@ namespace EpicGames.Horde.Storage
 		/// <inheritdoc/>
 		public bool Equals(BundleType? type) => type is object && Guid == type.Guid && Version == type.Version;
 
+		/// <inheritdoc/>
+		public override string ToString() => $"{Guid}#{Version}";
+
 		/// <inheritdoc/>
 		public static bool operator ==(BundleType lhs, BundleType rhs) => lhs.Equals(rhs);
 
@@ -383,7 +400,7 @@ namespace EpicGames.Horde.Storage
 	}
 
 	/// <summary>
-	/// Reference to another tree pack object
+	/// Reference to nodes in another bundle
 	/// </summary>
 	public class BundleImport
 	{
@@ -393,14 +410,14 @@ namespace EpicGames.Horde.Storage
 		public BlobLocator Locator { get; }
 
 		/// <summary>
-		/// Indexes of referenced nodes exported from this bundle
+		/// Entries for each node
 		/// </summary>
-		public IReadOnlyList<(int, IoHash)> Exports { get; }
+		public IReadOnlyList<int> Exports { get; }
 
 		/// <summary>
 		/// Constructor
 		/// </summary>
-		public BundleImport(BlobLocator locator, IReadOnlyList<(int, IoHash)> exports)
+		public BundleImport(BlobLocator locator, IReadOnlyList<int> exports)
 		{
 			Locator = locator;
 			Exports = exports;
@@ -413,17 +430,7 @@ namespace EpicGames.Horde.Storage
 		internal BundleImport(IMemoryReader reader)
 		{
 			Locator = reader.ReadBlobLocator();
-
-			int count = (int)reader.ReadUnsignedVarInt();
-
-			(int, IoHash)[] exports = new (int, IoHash)[count];
-			for (int idx = 0; idx < count; idx++)
-			{
-				int index = (int)reader.ReadUnsignedVarInt();
-				IoHash hash = reader.ReadIoHash();
-				exports[idx] = (index, hash);
-			}
-			Exports = exports;
+			Exports = reader.ReadVariableLengthArray(() => (int)reader.ReadUnsignedVarInt());
 		}
 
 		/// <summary>
@@ -433,14 +440,7 @@ namespace EpicGames.Horde.Storage
 		public void Write(IMemoryWriter writer)
 		{
 			writer.WriteBlobLocator(Locator);
-
-			writer.WriteUnsignedVarInt(Exports.Count);
-			for(int idx = 0; idx < Exports.Count; idx++)
-			{
-				(int index, IoHash hash) = Exports[idx];
-				writer.WriteUnsignedVarInt(index);
-				writer.WriteIoHash(hash);
-			}
+			writer.WriteVariableLengthArray(Exports, x => writer.WriteUnsignedVarInt(x));
 		}
 	}
 
diff --git a/Engine/Source/Programs/Shared/EpicGames.Horde/Storage/IStorageClient.cs b/Engine/Source/Programs/Shared/EpicGames.Horde/Storage/IStorageClient.cs
--- a/Engine/Source/Programs/Shared/EpicGames.Horde/Storage/IStorageClient.cs
+++ b/Engine/Source/Programs/Shared/EpicGames.Horde/Storage/IStorageClient.cs
@@ -113,28 +113,6 @@ namespace EpicGames.Horde.Storage
 		public static bool operator !=(NodeLocator left, NodeLocator right) => !left.Equals(right);
 	}
 
-	/// <summary>
-	/// Target for a ref
-	/// </summary>
-	/// <param name="Hash">Hash of the target node</param>
-	/// <param name="Locator">Locator for the target node</param>
-	public record class RefTarget(IoHash Hash, NodeLocator Locator)
-	{
-		/// <summary>
-		/// Parse a ref target from the given string
-		/// </summary>
-		/// <param name="text"></param>
-		/// <returns></returns>
-		public static RefTarget Parse(ReadOnlySpan<char> text)
-		{
-			int idx = text.IndexOf('@');
-			return new RefTarget(IoHash.Parse(text.Slice(0, idx)), NodeLocator.Parse(text.Slice(idx + 1)));
-		}
-
-		/// <inheritdoc/>
-		public override string ToString() => $"{Hash}@{Locator}";
-	}
-
 	/// <summary>
 	/// Base interface for a low-level storage backend. Blobs added to this store are not content addressed, but referenced by <see cref="BlobLocator"/>.
 	/// </summary>
@@ -179,7 +157,7 @@ namespace EpicGames.Horde.Storage
 		/// <param name="cacheTime">Minimum coherency for any cached value to be returned</param>
 		/// <param name="cancellationToken">Cancellation token for the operation</param>
 		/// <returns>Node pointed to by the ref</returns>
-		Task<RefTarget?> TryReadRefTargetAsync(RefName name, DateTime cacheTime = default, CancellationToken cancellationToken = default);
+		Task<NodeLocator> TryReadRefTargetAsync(RefName name, DateTime cacheTime = default, CancellationToken cancellationToken = default);
 
 		/// <summary>
 		/// Writes a new ref to the store which points to a new blob
@@ -191,7 +169,7 @@ namespace EpicGames.Horde.Storage
 		/// <param name="options">Options for the new ref</param>
 		/// <param name="cancellationToken">Cancellation token for the operation</param>
 		/// <returns>Unique identifier for the blob</returns>
-		Task<RefTarget> WriteRefAsync(RefName name, Bundle bundle, int exportIdx, Utf8String prefix = default, RefOptions? options = null, CancellationToken cancellationToken = default);
+		Task<NodeLocator> WriteRefAsync(RefName name, Bundle bundle, int exportIdx, Utf8String prefix = default, RefOptions? options = null, CancellationToken cancellationToken = default);
 
 		/// <summary>
 		/// Writes a new ref to the store
@@ -201,7 +179,7 @@ namespace EpicGames.Horde.Storage
 		/// <param name="options">Options for the new ref</param>
 		/// <param name="cancellationToken">Cancellation token for the operation</param>
 		/// <returns>Unique identifier for the blob</returns>
-		Task WriteRefTargetAsync(RefName name, RefTarget target, RefOptions? options = null, CancellationToken cancellationToken = default);
+		Task WriteRefTargetAsync(RefName name, NodeLocator target, RefOptions? options = null, CancellationToken cancellationToken = default);
 
 		/// <summary>
 		/// Reads data for a ref from the store
@@ -250,13 +228,13 @@ namespace EpicGames.Horde.Storage
 			public Task DeleteRefAsync(RefName name, CancellationToken cancellationToken = default) => _inner.DeleteRefAsync(name, cancellationToken);
 
 			/// <inheritdoc/>
-			public Task<RefTarget?> TryReadRefTargetAsync(RefName name, DateTime cacheTime = default, CancellationToken cancellationToken = default) => _inner.TryReadRefTargetAsync(name, cacheTime, cancellationToken);
+			public Task<NodeLocator> TryReadRefTargetAsync(RefName name, DateTime cacheTime = default, CancellationToken cancellationToken = default) => _inner.TryReadRefTargetAsync(name, cacheTime, cancellationToken);
 
 			/// <inheritdoc/>
-			public Task<RefTarget> WriteRefAsync(RefName name, Bundle bundle, int exportIdx = 0, Utf8String prefix = default, RefOptions? options = null, CancellationToken cancellationToken = default) => _inner.WriteRefAsync(name, bundle, exportIdx, prefix, options, cancellationToken);
+			public Task<NodeLocator> WriteRefAsync(RefName name, Bundle bundle, int exportIdx = 0, Utf8String prefix = default, RefOptions? options = null, CancellationToken cancellationToken = default) => _inner.WriteRefAsync(name, bundle, exportIdx, prefix, options, cancellationToken);
 
 			/// <inheritdoc/>
-			public Task WriteRefTargetAsync(RefName name, RefTarget target, RefOptions? options = null, CancellationToken cancellationToken = default) => _inner.WriteRefTargetAsync(name, target, options, cancellationToken);
+			public Task WriteRefTargetAsync(RefName name, NodeLocator target, RefOptions? options = null, CancellationToken cancellationToken = default) => _inner.WriteRefTargetAsync(name, target, options, cancellationToken);
 
 			#endregion
 		}
@@ -344,8 +322,8 @@ namespace EpicGames.Horde.Storage
 		/// <returns>True if the ref exists, false if it did not exist</returns>
 		public static async Task<bool> HasRefAsync(this IStorageClient store, RefName name, DateTime cacheTime = default, CancellationToken cancellationToken = default)
 		{
-			RefTarget? target = await store.TryReadRefTargetAsync(name, cacheTime, cancellationToken);
-			return target != null;
+			NodeLocator target = await store.TryReadRefTargetAsync(name, cacheTime, cancellationToken);
+			return target.IsValid();
 		}
 
 		/// <summary>
@@ -369,7 +347,7 @@ namespace EpicGames.Horde.Storage
 		/// <param name="maxAge">Maximum age of any cached ref</param>
 		/// <param name="cancellationToken">Cancellation token for the operation</param>
 		/// <returns>The ref target</returns>
-		public static Task<RefTarget?> TryReadRefTargetAsync(this IStorageClient store, RefName name, TimeSpan maxAge, CancellationToken cancellationToken = default)
+		public static Task<NodeLocator> TryReadRefTargetAsync(this IStorageClient store, RefName name, TimeSpan maxAge, CancellationToken cancellationToken = default)
 		{
 			return store.TryReadRefTargetAsync(name, DateTime.UtcNow - maxAge, cancellationToken);
 		}
@@ -382,10 +360,10 @@ namespace EpicGames.Horde.Storage
 		/// <param name="cacheTime">Minimum coherency of any cached result</param>
 		/// <param name="cancellationToken">Cancellation token for the operation</param>
 		/// <returns>The ref target</returns>
-		public static async Task<RefTarget> ReadRefTargetAsync(this IStorageClient store, RefName name, DateTime cacheTime = default, CancellationToken cancellationToken = default)
+		public static async Task<NodeLocator> ReadRefTargetAsync(this IStorageClient store, RefName name, DateTime cacheTime = default, CancellationToken cancellationToken = default)
 		{
-			RefTarget? refTarget = await store.TryReadRefTargetAsync(name, cacheTime, cancellationToken);
-			if (refTarget == null)
+			NodeLocator refTarget = await store.TryReadRefTargetAsync(name, cacheTime, cancellationToken);
+			if (!refTarget.IsValid())
 			{
 				throw new RefNameNotFoundException(name);
 			}
@@ -400,7 +378,7 @@ namespace EpicGames.Horde.Storage
 		/// <param name="maxAge">Maximum age for any cached result</param>
 		/// <param name="cancellationToken">Cancellation token for the operation</param>
 		/// <returns>The blob instance</returns>
-		public static Task<RefTarget> ReadRefTargetAsync(this IStorageClient store, RefName name, TimeSpan maxAge, CancellationToken cancellationToken = default)
+		public static Task<NodeLocator> ReadRefTargetAsync(this IStorageClient store, RefName name, TimeSpan maxAge, CancellationToken cancellationToken = default)
 		{
 			return ReadRefTargetAsync(store, name, DateTime.UtcNow - maxAge, cancellationToken);
 		}
diff --git a/Engine/Source/Programs/Shared/EpicGames.Horde/Storage/Nodes/CbNode.cs b/Engine/Source/Programs/Shared/EpicGames.Horde/Storage/Nodes/CbNode.cs
--- a/Engine/Source/Programs/Shared/EpicGames.Horde/Storage/Nodes/CbNode.cs
+++ b/Engine/Source/Programs/Shared/EpicGames.Horde/Storage/Nodes/CbNode.cs
@@ -21,17 +21,17 @@ namespace EpicGames.Horde.Storage.Nodes
 		/// <summary>
 		/// Imported nodes
 		/// </summary>
-		public IReadOnlyDictionary<IoHash, TreeNodeRef> References { get; }
+		public IReadOnlyList<TreeNodeRef> References { get; }
 
 		/// <summary>
 		/// Constructor
 		/// </summary>
 		/// <param name="obj">The compact binary object</param>
-		/// <param name="references">Map of attachment hash to node locator</param>
-		public CbNode(CbObject obj, IReadOnlyDictionary<IoHash, NodeLocator> references)
+		/// <param name="references">List of references to attachments</param>
+		public CbNode(CbObject obj, IReadOnlyList<TreeNodeRef> references)
 		{
 			Object = obj;
-			References = references.ToDictionary(x => x.Key, x => new TreeNodeRef(new RefTarget(x.Key, x.Value)));
+			References = references;
 		}
 
 		/// <summary>
@@ -41,7 +41,7 @@ namespace EpicGames.Horde.Storage.Nodes
 		public CbNode(ITreeNodeReader reader)
 		{
 			Object = new CbObject(reader.ReadFixedLengthBytes(reader.Length));
-			References = reader.References.ToDictionary(x => x.Key, x => new TreeNodeRef(new RefTarget(x.Key, x.Value)));
+			References = reader.References.ConvertAll(x => new TreeNodeRef(x));
 		}
 
 		/// <inheritdoc/>
@@ -51,6 +51,6 @@ namespace EpicGames.Horde.Storage.Nodes
 		}
 
 		/// <inheritdoc/>
-		public override IEnumerable<TreeNodeRef> EnumerateRefs() => References.Values;
+		public override IEnumerable<TreeNodeRef> EnumerateRefs() => References;
 	}
 }
diff --git a/Engine/Source/Programs/Shared/EpicGames.Horde/Storage/Nodes/CommitNode.cs b/Engine/Source/Programs/Shared/EpicGames.Horde/Storage/Nodes/CommitNode.cs
--- a/Engine/Source/Programs/Shared/EpicGames.Horde/Storage/Nodes/CommitNode.cs
+++ b/Engine/Source/Programs/Shared/EpicGames.Horde/Storage/Nodes/CommitNode.cs
@@ -55,7 +55,7 @@ namespace EpicGames.Horde.Storage.Nodes
 		/// <summary>
 		/// Contents of the tree at this commit
 		/// </summary>
-		public TreeNodeRef<DirectoryNode> Contents { get; set; }
+		public DirectoryNodeRef Contents { get; set; }
 
 		/// <summary>
 		/// Metadata for this commit, keyed by arbitrary GUID
@@ -71,7 +71,7 @@ namespace EpicGames.Horde.Storage.Nodes
 		/// <param name="message">Message for the commit</param>
 		/// <param name="time">The commit time</param>
 		/// <param name="contents">Contents of the tree at this commit</param>
-		public CommitNode(int number, TreeNodeRef<CommitNode>? parent, string author, string message, DateTime time, TreeNodeRef<DirectoryNode> contents)
+		public CommitNode(int number, TreeNodeRef<CommitNode>? parent, string author, string message, DateTime time, DirectoryNodeRef contents)
 		{
 			Number = number;
 			Parent = parent;
@@ -95,7 +95,7 @@ namespace EpicGames.Horde.Storage.Nodes
 			CommitterId = reader.ReadOptionalString();
 			Message = reader.ReadString();
 			Time = reader.ReadDateTime();
-			Contents = reader.ReadRef<DirectoryNode>();
+			Contents = new DirectoryNodeRef(reader);
 			Metadata = reader.ReadDictionary(() => reader.ReadGuid(), () => reader.ReadRef());
 		}
 
diff --git a/Engine/Source/Programs/Shared/EpicGames.Horde/Storage/Nodes/DirectoryNode.cs b/Engine/Source/Programs/Shared/EpicGames.Horde/Storage/Nodes/DirectoryNode.cs
--- a/Engine/Source/Programs/Shared/EpicGames.Horde/Storage/Nodes/DirectoryNode.cs
+++ b/Engine/Source/Programs/Shared/EpicGames.Horde/Storage/Nodes/DirectoryNode.cs
@@ -65,9 +65,9 @@ namespace EpicGames.Horde.Storage.Nodes
 		public long Length => (Target == null) ? _cachedLength : Target.Length;
 
 		/// <summary>
-		/// SHA1 hash of this file, with Git prefix
+		/// Hash of the target node
 		/// </summary>
-		public Sha1Hash GitHash { get; set; }
+		public IoHash Hash { get; private set; }
 
 		/// <summary>
 		/// Cached length of this node
@@ -101,7 +101,7 @@ namespace EpicGames.Horde.Storage.Nodes
 		/// Serialize this entry
 		/// </summary>
 		/// <param name="writer"></param>
-		public new void Serialize(ITreeNodeWriter writer)
+		protected override void Serialize(IMemoryWriter writer)
 		{
 			base.Serialize(writer);
 
@@ -141,6 +141,7 @@ namespace EpicGames.Horde.Storage.Nodes
 			base.OnCollapse();
 
 			_cachedLength = Target!.Length;
+			Hash = Target!.Hash;
 		}
 
 		/// <inheritdoc/>
@@ -148,24 +149,72 @@ namespace EpicGames.Horde.Storage.Nodes
 	}
 
 	/// <summary>
-	/// Entry for a directory within a directory node
+	/// Reference to a directory node, including the target hash and length
 	/// </summary>
-	public class DirectoryEntry : TreeNodeRef<DirectoryNode>
+	public class DirectoryNodeRef : TreeNodeRef<DirectoryNode>
 	{
 		/// <summary>
-		/// Name of this directory
+		/// Length of this directory tree
 		/// </summary>
-		public Utf8String Name { get; }
+		public long Length => (Target == null) ? _cachedLength : Target.Length;
 
 		/// <summary>
-		/// Length of this directory tree
+		/// Hash of the target node
 		/// </summary>
-		public long Length => (Target == null) ? _cachedLength : Target.Length;
+		public IoHash Hash { get; private set; }
 
 		/// <summary>
 		/// Cached value for the length of this tree
 		/// </summary>
-		readonly long _cachedLength;
+		long _cachedLength;
+
+		/// <summary>
+		/// Constructor
+		/// </summary>
+		public DirectoryNodeRef(DirectoryNode node)
+			: base(node)
+		{
+		}
+
+		/// <summary>
+		/// Constructor
+		/// </summary>
+		/// <param name="reader"></param>
+		public DirectoryNodeRef(ITreeNodeReader reader)
+			: base(reader)
+		{
+			_cachedLength = (long)reader.ReadUnsignedVarInt();
+		}
+
+		/// <summary>
+		/// Serialize this directory entry to disk
+		/// </summary>
+		/// <param name="writer"></param>
+		protected override void Serialize(IMemoryWriter writer)
+		{
+			base.Serialize(writer);
+			writer.WriteUnsignedVarInt((ulong)Length);
+		}
+
+		/// <inheritdoc/>
+		protected override void OnCollapse()
+		{
+			base.OnCollapse();
+
+			Hash = Target!.Hash;
+			_cachedLength = Target!.Length;
+		}
+	}
+
+	/// <summary>
+	/// Entry for a directory within a directory node
+	/// </summary>
+	public class DirectoryEntry : DirectoryNodeRef
+	{
+		/// <summary>
+		/// Name of this directory
+		/// </summary>
+		public Utf8String Name { get; }
 
 		/// <summary>
 		/// Constructor
@@ -185,26 +234,23 @@ namespace EpicGames.Horde.Storage.Nodes
 		}
 
 		/// <summary>
-		/// 
+		/// Deserializing constructor
 		/// </summary>
 		/// <param name="reader"></param>
 		public DirectoryEntry(ITreeNodeReader reader)
 			: base(reader)
 		{
 			Name = reader.ReadUtf8String();
-			_cachedLength = (long)reader.ReadUnsignedVarInt();
 		}
 
 		/// <summary>
 		/// Serialize this directory entry to disk
 		/// </summary>
 		/// <param name="writer"></param>
-		public new void Serialize(ITreeNodeWriter writer)
+		protected override void Serialize(IMemoryWriter writer)
 		{
 			base.Serialize(writer);
-
 			writer.WriteUtf8String(Name);
-			writer.WriteUnsignedVarInt((ulong)Length);
 		}
 
 		/// <inheritdoc/>
@@ -305,13 +351,13 @@ namespace EpicGames.Horde.Storage.Nodes
 			writer.WriteUnsignedVarInt(Files.Count);
 			foreach (FileEntry fileEntry in _nameToFileEntry.Values)
 			{
-				fileEntry.Serialize(writer);
+				writer.WriteRef(fileEntry);
 			}
 
 			writer.WriteUnsignedVarInt(Directories.Count);
 			foreach (DirectoryEntry directoryEntry in _nameToDirectoryEntry.Values)
 			{
-				directoryEntry.Serialize(writer);
+				writer.WriteRef(directoryEntry);
 			}
 		}
 
diff --git a/Engine/Source/Programs/Shared/EpicGames.Horde/Storage/TreeNodeRef.cs b/Engine/Source/Programs/Shared/EpicGames.Horde/Storage/TreeNodeRef.cs
--- a/Engine/Source/Programs/Shared/EpicGames.Horde/Storage/TreeNodeRef.cs
+++ b/Engine/Source/Programs/Shared/EpicGames.Horde/Storage/TreeNodeRef.cs
@@ -1,10 +1,10 @@
 // Copyright Epic Games, Inc. All Rights Reserved.
 
 using System;
-using System.Diagnostics;
 using System.Threading;
 using System.Threading.Tasks;
 using EpicGames.Core;
+using EpicGames.Serialization;
 
 namespace EpicGames.Horde.Storage
 {
@@ -13,27 +13,22 @@ namespace EpicGames.Horde.Storage
 	/// 
 	/// Refs may be in the following states:
 	///  * In storage 
-	///      * Hash and Locator are valid, Target is null, IsDirty() returns false.
+	///      * Locator is valid, Target is null, IsDirty() returns false.
 	///      * Writing or calling Collapse() is a no-op.
 	///  * In memory and target was expanded, has not been modified 
-	///      * Hash and Locator are valid, Target is valid, IsDirty() returns false.
+	///      * Locator is valid, Target is valid, IsDirty() returns false.
 	///      * Writing or calling Collapse() on a ref transitions it to the "in storage" state.
 	///  * In memory and target is new
-	///      * Hash and Locator are invalid, Target is valid, IsDirty() returns true.
+	///      * Locator is invalid, Target is valid, IsDirty() returns true.
 	///      * Writing a ref transitions it to the "in storage" state. Calling Collapse is a no-op.
 	///  * In memory but target has been modified
-	///      * Hash and Locator are set but may not reflect the current node state, Target is valid, IsDirty() returns true. 
+	///      * Locator is set but may not reflect the current node state, Target is valid, IsDirty() returns true. 
 	///      * Writing a ref transitions it to the "in storage" state. Calling Collapse is a no-op.
 	///      
 	/// The <see cref="OnCollapse"/> and <see cref="OnExpand"/> methods allow overriden implementations to cache information about the target.
 	/// </summary>
 	public class TreeNodeRef
 	{
-		/// <summary>
-		/// Hash of the referenced node. Invalid for nodes in memory.
-		/// </summary>
-		public IoHash Hash { get; private set; }
-
 		/// <summary>
 		/// Locator for the blob containing this node. Invalid for nodes in memory.
 		/// </summary>
@@ -54,20 +49,19 @@ namespace EpicGames.Horde.Storage
 		}
 
 		/// <summary>
-		/// Creates a reference to a node in memory.
+		/// Creates a reference to a node in storage.
 		/// </summary>
-		/// <param name="target">Node to reference</param>
-		public TreeNodeRef(RefTarget target)
+		/// <param name="locator">Node to reference</param>
+		public TreeNodeRef(NodeLocator locator)
 		{
-			Hash = target.Hash;
-			Locator = target.Locator;
+			Locator = locator;
 		}
 
 		/// <summary>
 		/// Deserialization constructor
 		/// </summary>
 		/// <param name="reader"></param>
-		public TreeNodeRef(ITreeNodeReader reader) : this(reader.ReadRefTarget())
+		public TreeNodeRef(ITreeNodeReader reader) : this(reader.ReadNodeLocator())
 		{
 		}
 
@@ -87,44 +81,45 @@ namespace EpicGames.Horde.Storage
 				throw new InvalidOperationException("Cannot mark a ref as dirty without having expanded it.");
 			}
 
-			Hash = default;
 			Locator = default;
 		}
 
 		/// <summary>
 		/// Updates the hash and revision number for the ref.
 		/// </summary>
-		/// <param name="hash">Hash of the node</param>
 		/// <returns></returns>
-		internal void MarkAsPendingWrite(IoHash hash)
+		internal void MarkAsPendingWrite()
 		{
-			Hash = hash;
 			OnCollapse();
 		}
 
 		/// <summary>
 		/// Update the reference to refer to a location in storage.
 		/// </summary>
-		/// <param name="hash">Hash of the node</param>
 		/// <param name="locator">Location of the node</param>
 		/// <param name="revision">Revision number that was written</param>
-		internal void MarkAsWritten(IoHash hash, NodeLocator locator, uint revision)
+		internal void MarkAsWritten(NodeLocator locator, uint revision)
 		{
-			if (Hash == hash)
+			Locator = locator;
+			if (Target != null && Target.Revision == revision)
 			{
-				Locator = locator;
-				if (Target != null && Target.Revision == revision)
-				{
-					Target = null;
-				}
+				Target = null;
 			}
 		}
 
+		/// <summary>
+		/// Internal helper method to allow implementations of <see cref="ITreeNodeWriter"/> to call <see cref="Serialize(IMemoryWriter)"/>.
+		/// </summary>
+		/// <param name="writer"></param>
+		internal void SerializeInternal(IMemoryWriter writer) => Serialize(writer);
+
 		/// <summary>
 		/// Serialize the node to the given writer
 		/// </summary>
 		/// <param name="writer"></param>
-		public void Serialize(ITreeNodeWriter writer) => writer.WriteRef(this);
+		protected virtual void Serialize(IMemoryWriter writer)
+		{
+		}
 
 		/// <summary>
 		/// Resolve this reference to a concrete node
@@ -192,7 +187,7 @@ namespace EpicGames.Horde.Storage
 		/// Constructor
 		/// </summary>
 		/// <param name="target">The referenced node</param>
-		public TreeNodeRef(RefTarget target) : base(target)
+		public TreeNodeRef(NodeLocator target) : base(target)
 		{
 		}
 
diff --git a/Engine/Source/Programs/Shared/EpicGames.Horde/Storage/TreeReader.cs b/Engine/Source/Programs/Shared/EpicGames.Horde/Storage/TreeReader.cs
--- a/Engine/Source/Programs/Shared/EpicGames.Horde/Storage/TreeReader.cs
+++ b/Engine/Source/Programs/Shared/EpicGames.Horde/Storage/TreeReader.cs
@@ -31,15 +31,20 @@ namespace EpicGames.Horde.Storage
 		/// </summary>
 		int Length { get; }
 
+		/// <summary>
+		/// Hash of the node being deserialized
+		/// </summary>
+		IoHash Hash { get; }
+
 		/// <summary>
 		/// References to other nodes
 		/// </summary>
-		IReadOnlyDictionary<IoHash, NodeLocator> References { get; }
+		IReadOnlyList<NodeLocator> References { get; }
 
 		/// <summary>
 		/// Reads a reference to another node
 		/// </summary>
-		RefTarget ReadRefTarget();
+		NodeLocator ReadNodeLocator();
 	}
 
 	/// <summary>
@@ -140,7 +145,7 @@ namespace EpicGames.Horde.Storage
 			public readonly BundleHeader Header;
 			public readonly int[] PacketOffsets;
 			public readonly ExportInfo[] Exports;
-			public readonly RefTarget[] References;
+			public readonly NodeLocator[] References;
 			public readonly TypeInfo?[] Types;
 
 			public BundleInfo(BlobLocator locator, BundleHeader header, int headerLength, TypeInfo?[] types)
@@ -168,20 +173,20 @@ namespace EpicGames.Horde.Storage
 					packetOffset += packet.EncodedLength;
 				}
 
-				References = new RefTarget[header.Imports.Sum(x => x.Exports.Count) + header.Exports.Count];
+				References = new NodeLocator[header.Imports.Sum(x => x.Exports.Count) + header.Exports.Count];
 
 				int referenceIdx = 0;
 				foreach (BundleImport import in header.Imports)
 				{
-					foreach ((int importExportIdx, IoHash importExportHash) in import.Exports)
+					foreach (int importExportIdx in import.Exports)
 					{
 						NodeLocator importLocator = new NodeLocator(import.Locator, importExportIdx);
-						References[referenceIdx++] = new RefTarget(importExportHash, importLocator);
+						References[referenceIdx++] = importLocator;
 					}
 				}
 				for (int idx = 0; idx < header.Exports.Count; idx++)
 				{
-					References[referenceIdx++] = new RefTarget(header.Exports[idx].Hash, new NodeLocator(locator, idx));
+					References[referenceIdx++] = new NodeLocator(locator, idx);
 				}
 
 				Types = types;
@@ -227,14 +232,16 @@ namespace EpicGames.Horde.Storage
 		/// </summary>
 		class NodeReader : MemoryReader, ITreeNodeReader
 		{
-			readonly IReadOnlyDictionary<IoHash, NodeLocator> _refs;
+			readonly IReadOnlyList<NodeLocator> _refs;
+			readonly IoHash _hash;
 			readonly BundleType _type;
 			readonly int _length;
 
-			public NodeReader(ReadOnlyMemory<byte> data, IReadOnlyDictionary<IoHash, NodeLocator> refs, BundleType type)
+			public NodeReader(ReadOnlyMemory<byte> data, IoHash hash, IReadOnlyList<NodeLocator> refs, BundleType type)
 				: base(data)
 			{
 				_refs = refs;
+				_hash = hash;
 				_type = type;
 				_length = data.Length;
 			}
@@ -243,12 +250,14 @@ namespace EpicGames.Horde.Storage
 
 			public int Length => _length;
 
-			public IReadOnlyDictionary<IoHash, NodeLocator> References => _refs;
+			public IoHash Hash => _hash;
 
-			public RefTarget ReadRefTarget()
+			public IReadOnlyList<NodeLocator> References => _refs;
+
+			public NodeLocator ReadNodeLocator()
 			{
-				IoHash hash = this.ReadIoHash();
-				return new RefTarget(hash, _refs[hash]);
+				int idx = (int)this.ReadUnsignedVarInt();
+				return _refs[idx];
 			}
 		}
 
@@ -536,7 +545,7 @@ namespace EpicGames.Horde.Storage
 		/// <param name="locator"></param>
 		/// <param name="cancellationToken"></param>
 		/// <returns></returns>
-		async Task<BundleHeader> ReadBundleHeaderAsync(BlobLocator locator, CancellationToken cancellationToken = default)
+		public async Task<BundleHeader> ReadBundleHeaderAsync(BlobLocator locator, CancellationToken cancellationToken = default)
 		{
 			BundleInfo info = await GetBundleInfoAsync(locator, cancellationToken);
 			return info.Header;
@@ -688,11 +697,11 @@ namespace EpicGames.Horde.Storage
 			ExportInfo exportInfo = bundleInfo.Exports[locator.ExportIdx];
 			ReadOnlyMemory<byte> packetData = await ReadBundlePacketAsync(bundleInfo, exportInfo.PacketIdx, cancellationToken);
 
-			Dictionary<IoHash, NodeLocator> refs = new Dictionary<IoHash, NodeLocator>(export.References.Count);
+			List<NodeLocator> refs = new List<NodeLocator>(export.References.Count);
 			for (int idx = 0; idx < export.References.Count; idx++)
 			{
-				RefTarget reference = bundleInfo.References[export.References[idx]];
-				refs.Add(reference.Hash, reference.Locator);
+				NodeLocator reference = bundleInfo.References[export.References[idx]];
+				refs.Add(reference);
 			}
 
 			ReadOnlyMemory<byte> nodeData = packetData.Slice(exportInfo.Offset, export.Length);
@@ -703,7 +712,7 @@ namespace EpicGames.Horde.Storage
 				throw new InvalidOperationException($"No registered serializer for type {bundleInfo.Header.Types[export.TypeIdx].Guid}");
 			}
 
-			return typeInfo.Deserialize(new NodeReader(nodeData, refs, typeInfo.BundleType));
+			return typeInfo.Deserialize(new NodeReader(nodeData, export.Hash, refs, typeInfo.BundleType));
 		}
 
 		/// <summary>
@@ -717,17 +726,6 @@ namespace EpicGames.Horde.Storage
 			return (TNode)await ReadNodeAsync(locator, cancellationToken);
 		}
 
-		/// <summary>
-		/// Reads a node from a bundle
-		/// </summary>
-		/// <param name="target">Locator for the node</param>
-		/// <param name="cancellationToken">Cancellation token for the operation</param>
-		/// <returns>Node data read from the given bundle</returns>
-		public async ValueTask<TNode> ReadNodeAsync<TNode>(RefTarget target, CancellationToken cancellationToken = default) where TNode : TreeNode
-		{
-			return (TNode)await ReadNodeAsync(target.Locator, cancellationToken);
-		}
-
 		/// <summary>
 		/// Reads data for a ref from the store, along with the node's contents.
 		/// </summary>
@@ -737,12 +735,12 @@ namespace EpicGames.Horde.Storage
 		/// <returns>Node for the given ref, or null if it does not exist</returns>
 		public async Task<TNode?> TryReadNodeAsync<TNode>(RefName name, DateTime cacheTime = default, CancellationToken cancellationToken = default) where TNode : TreeNode
 		{
-			RefTarget? refTarget = await _store.TryReadRefTargetAsync(name, cacheTime, cancellationToken);
-			if (refTarget == null)
+			NodeLocator refTarget = await _store.TryReadRefTargetAsync(name, cacheTime, cancellationToken);
+			if (!refTarget.IsValid())
 			{
 				return null;
 			}
-			return await ReadNodeAsync<TNode>(refTarget.Locator, cancellationToken);
+			return await ReadNodeAsync<TNode>(refTarget, cancellationToken);
 		}
 
 		/// <summary>
@@ -754,8 +752,8 @@ namespace EpicGames.Horde.Storage
 		/// <returns>Node for the given ref, or null if it does not exist</returns>
 		public async Task<TreeNodeRef<TNode>?> TryReadNodeRefAsync<TNode>(RefName name, DateTime cacheTime = default, CancellationToken cancellationToken = default) where TNode : TreeNode
 		{
-			RefTarget? refTarget = await _store.TryReadRefTargetAsync(name, cacheTime, cancellationToken);
-			if (refTarget == null)
+			NodeLocator refTarget = await _store.TryReadRefTargetAsync(name, cacheTime, cancellationToken);
+			if (!refTarget.IsValid())
 			{
 				return null;
 			}
diff --git a/Engine/Source/Programs/Shared/EpicGames.Horde/Storage/TreeWriter.cs b/Engine/Source/Programs/Shared/EpicGames.Horde/Storage/TreeWriter.cs
--- a/Engine/Source/Programs/Shared/EpicGames.Horde/Storage/TreeWriter.cs
+++ b/Engine/Source/Programs/Shared/EpicGames.Horde/Storage/TreeWriter.cs
@@ -2,14 +2,12 @@
 
 using System;
 using System.Buffers;
-using System.Collections.Concurrent;
 using System.Collections.Generic;
 using System.Diagnostics;
 using System.Linq;
 using System.Threading;
 using System.Threading.Tasks;
 using EpicGames.Core;
-using EpicGames.Horde.Compute.Impl;
 
 namespace EpicGames.Horde.Storage
 {
@@ -39,47 +37,63 @@ namespace EpicGames.Horde.Storage
 	/// </summary>
 	public class TreeWriter
 	{
-		class NodeInfo
+		record NodeKey(IoHash Hash, BundleType Type);
+		record struct NodeRevision(TreeNodeRef NodeRef, uint Revision);
+
+		// Information about a TreeNodeRef that was written, and needs to be fixed up after flushing
+		class NodeRefInfo
 		{
-			public readonly BundleType Type;
-			public readonly IoHash Hash;
-			public readonly int Length;
-			public readonly IReadOnlyList<IoHash> RefHashes;
 			public readonly TreeNodeRef NodeRef;
 			public readonly uint Revision;
+			public NodeRefInfo? Next;
 
-			public NodeInfo(BundleType type, IoHash hash, int length, IReadOnlyList<IoHash> refs, TreeNodeRef nodeRef)
+			public NodeRefInfo(TreeNodeRef nodeRef, NodeRefInfo? next)
 			{
-				Type = type;
-				Hash = hash;
-				Length = length;
-				RefHashes = refs;
 				NodeRef = nodeRef;
 				Revision = nodeRef.Target!.Revision;
+				Next = next;
 			}
 		}
 
+		// Information about a unique output node. Note that multiple node refs may de-duplicate to the same output node.
+		class OutputNodeInfo : NodeRefInfo
+		{
+			public readonly NodeKey Key;
+			public readonly int Length;
+			public readonly IReadOnlyList<TreeNodeRef> Refs;
+
+			public OutputNodeInfo(NodeKey key, int length, IReadOnlyList<TreeNodeRef> refs, TreeNodeRef nodeRef)
+				: base(nodeRef, null)
+			{
+				Key = key;
+				Length = length;
+				Refs = refs;
+			}
+
+			public void Add(TreeNodeRef nodeRef) => Next = new NodeRefInfo(nodeRef, Next);
+		}
+
 		class NodeWriter : ITreeNodeWriter
 		{
 			public ByteArrayBuilder Data = new ByteArrayBuilder();
 
-			readonly IReadOnlyDictionary<TreeNodeRef, IoHash> _refToHash;
+			readonly IReadOnlyDictionary<TreeNodeRef, int> _refToIndex;
 
 			public int Length => Data.Length;
 
-			public NodeWriter(IReadOnlyDictionary<TreeNodeRef, IoHash> refToHash)
+			public NodeWriter(IReadOnlyList<TreeNodeRef> refs)
 			{
-				_refToHash = refToHash;
+				_refToIndex = new Dictionary<TreeNodeRef, int>(refs.Distinct().Select((x, i) => KeyValuePair.Create(x, i)));
 			}
 
 			public void WriteRef(TreeNodeRef target)
 			{
-				IoHash hash;
-				if (!_refToHash.TryGetValue(target, out hash))
+				if (!_refToIndex.TryGetValue(target, out int index))
 				{
 					throw new InvalidOperationException($"Cannot serialize value; was not returned by owner's {nameof(target.Target.EnumerateRefs)} method.");
 				}
-				Data.WriteIoHash(hash);
+				Data.WriteUnsignedVarInt(index);
+				target.SerializeInternal(this);
 			}
 
 			public Memory<byte> GetMemory(int minSize = 1) => Data.GetMemory(minSize);
@@ -96,14 +110,17 @@ namespace EpicGames.Horde.Storage
 		// Writer used to create packets
 		readonly ArrayMemoryWriter _packetWriter;
 
+		// Map of keys to nodes in the queue
+		readonly Dictionary<NodeKey, OutputNodeInfo> _nodeKeyToInfo = new Dictionary<NodeKey, OutputNodeInfo>();
+
 		// Map of hashes to existing nodes. Empty/invalid locators are used for items that are queued for writing, but not available yet.
-		readonly Dictionary<IoHash, NodeLocator> _hashToNode = new Dictionary<IoHash, NodeLocator>(); // TODO: this needs to include some additional state from external sources.
+		readonly Dictionary<NodeKey, NodeLocator> _nodeKeyToLocator = new Dictionary<NodeKey, NodeLocator>(); // TODO: this needs to include some additional state from external sources.
 
 		// Queue of nodes for the current bundle
-		readonly List<NodeInfo> _queue = new List<NodeInfo>();
+		readonly List<OutputNodeInfo> _queue = new List<OutputNodeInfo>();
 
-		// Queue of nodes that are already scheduled to be written elsewhere
-		readonly List<NodeInfo> _secondaryQueue = new List<NodeInfo>();
+		// Map of node refs to their serialized state
+		readonly HashSet<(TreeNodeRef, uint)> _queuedRefs = new HashSet<(TreeNodeRef, uint)>();
 
 		// List of packets in the current bundle
 		readonly List<BundlePacket> _packets = new List<BundlePacket>();
@@ -153,44 +170,29 @@ namespace EpicGames.Horde.Storage
 		{
 		}
 
-		/// <summary>
-		/// Attempt to get the locator for the node with a given hash
-		/// </summary>
-		/// <param name="hash">Hash of the node</param>
-		/// <returns>True if the locator was found</returns>
-		public NodeLocator GetLocator(IoHash hash) => _hashToNode[hash];
-
-		/// <summary>
-		/// Attempt to get the locator for the node with a given hash
-		/// </summary>
-		/// <param name="hash">Hash of the node</param>
-		/// <param name="locator">Receives the locator for the node</param>
-		/// <returns>True if the locator was found</returns>
-		public bool TryGetLocator(IoHash hash, out NodeLocator locator) => _hashToNode.TryGetValue(hash, out locator);
-
 		/// <inheritdoc/>
-		public async Task<IoHash> WriteAsync(TreeNodeRef nodeRef, CancellationToken cancellationToken = default)
+		public async Task<bool> WriteAsync(TreeNodeRef nodeRef, CancellationToken cancellationToken = default)
 		{
 			// Check we actually have a target node. If we don't, we don't need to write anything.
-			TreeNode? node = nodeRef.Target;
-			if (node == null)
+			TreeNode? target = nodeRef.Target;
+			if (target == null)
 			{
-				Debug.Assert(nodeRef.Hash != IoHash.Zero);
-				_hashToNode.TryAdd(nodeRef.Hash, nodeRef.Locator);
-				return nodeRef.Hash;
+				Debug.Assert(nodeRef.Locator.IsValid());
+				return false;
 			}
 
-			// Write all the nodes it references, and mark the ref as dirty if any of them change.
-			Dictionary<TreeNodeRef, IoHash> nextRefToHash = new Dictionary<TreeNodeRef, IoHash>();
-			foreach (TreeNodeRef nextRef in node.EnumerateRefs())
+			// Early-out if this ref is already queued for writing.
+			uint revision = target.Revision;
+			if (!_queuedRefs.Add((nodeRef, revision)))
 			{
-				IoHash prevHash = nextRef.Hash;
-				IoHash nextHash = await WriteAsync(nextRef, cancellationToken);
-				Debug.Assert(nextHash != IoHash.Zero);
-
-				nextRefToHash.Add(nextRef, nextHash);
+				return true;
+			}
 
-				if (prevHash != nextHash)
+			// Write all the nodes it references, and mark the ref as dirty if any of them change.
+			List<TreeNodeRef> nextRefs = target.EnumerateRefs().ToList();
+			foreach (TreeNodeRef nextRef in nextRefs)
+			{
+				if (await WriteAsync(nextRef, cancellationToken))
 				{
 					nodeRef.MarkAsDirty();
 				}
@@ -200,46 +202,42 @@ namespace EpicGames.Horde.Storage
 			if (!nodeRef.IsDirty())
 			{
 				// Make sure the locator is valid. The node may be queued for writing but not flushed to disk yet.
-				if (nodeRef.Locator.IsValid())
-				{
-					_hashToNode.TryAdd(nodeRef.Hash, nodeRef.Locator);
-					nodeRef.Collapse();
-				}
-				return nodeRef.Hash;
+				Debug.Assert(nodeRef.Locator.IsValid());
+				nodeRef.Collapse();
+				return false;
 			}
 
 			// Serialize the node
-			NodeWriter nodeWriter = new NodeWriter(nextRefToHash);
-			nodeRef.Target!.Serialize(nodeWriter);
-			BundleType bundleType = nodeRef.Target.GetBundleType();
+			NodeWriter nodeWriter = new NodeWriter(nextRefs);
+			target.Serialize(nodeWriter);
 
 			// Get the hash for the new blob
 			ReadOnlySequence<byte> data = nodeWriter.Data.AsSequence();
-			IoHash hash = IoHash.Compute(data);
-			nodeRef.MarkAsPendingWrite(hash);
+			target.Hash = IoHash.Compute(data);
+			nodeRef.MarkAsPendingWrite();
+
+			NodeKey nodeKey = new NodeKey(target.Hash, target.GetBundleType());
 
 			// Check if we're already tracking a node with the same hash
-			NodeInfo nodeInfo = new NodeInfo(bundleType, hash, (int)data.Length, nextRefToHash.Values.Distinct().ToArray(), nodeRef);
-			if (_hashToNode.TryGetValue(hash, out NodeLocator locator))
+			OutputNodeInfo? nodeInfo;
+			if (_nodeKeyToInfo.TryGetValue(nodeKey, out nodeInfo))
 			{
-				// If the node is in the lookup but not currently valid, it's already queued for writing. Add this ref to the list of refs that needs fixing up,
-				// so we can update it after flushing.
-				if (locator.IsValid())
-				{
-					nodeRef.MarkAsWritten(hash, locator, nodeInfo.Revision);
-				}
-				else
-				{
-					_secondaryQueue.Add(nodeInfo);
-				}
+				// This node is already queued for writing to storage
+				nodeInfo.Add(nodeRef);
+			}
+			else if (_nodeKeyToLocator.TryGetValue(nodeKey, out NodeLocator locator))
+			{
+				// This node has already been written to storage
+				nodeRef.MarkAsWritten(locator, nodeRef.Target!.Revision);
 			}
 			else
 			{
 				// Write the node if we don't already have it
 				_packetWriter.WriteFixedLengthBytes(data);
 
+				nodeInfo = new OutputNodeInfo(nodeKey, (int)data.Length, nextRefs, nodeRef);
 				_queue.Add(nodeInfo);
-				_hashToNode.Add(hash, default);
+				_nodeKeyToInfo.Add(nodeKey, nodeInfo);
 
 				if (_packetWriter.Length > Math.Min(_options.MinCompressionPacketSize, _options.MaxBlobSize))
 				{
@@ -252,22 +250,7 @@ namespace EpicGames.Horde.Storage
 				}
 			}
 
-			return hash;
-		}
-
-		/// <summary>
-		/// Writes a ref to the node with the given hash
-		/// </summary>
-		/// <param name="name">Name of the ref to write</param>
-		/// <param name="root">Hash of the root node</param>
-		/// <param name="options"></param>
-		/// <param name="cancellationToken">Cancellation token for the operation</param>
-		public async Task<NodeLocator> WriteAsync(RefName name, IoHash root, RefOptions? options = null, CancellationToken cancellationToken = default)
-		{
-			await FlushAsync(cancellationToken);
-			NodeLocator locator = GetLocator(root);
-			await _store.WriteRefTargetAsync(name, new RefTarget(root, locator), options, cancellationToken);
-			return locator;
+			return true;
 		}
 
 		/// <summary>
@@ -285,7 +268,7 @@ namespace EpicGames.Horde.Storage
 			await FlushAsync(cancellationToken);
 
 			Debug.Assert(nodeRef.Locator.IsValid());
-			await _store.WriteRefTargetAsync(name, new RefTarget(nodeRef.Hash, nodeRef.Locator), options, cancellationToken);
+			await _store.WriteRefTargetAsync(name, nodeRef.Locator, options, cancellationToken);
 			return nodeRef.Locator;
 		}
 
@@ -304,6 +287,20 @@ namespace EpicGames.Horde.Storage
 			}
 		}
 
+		/// <summary>
+		/// Flushes all the current nodes to storage
+		/// </summary>
+		/// <param name="root">Root for the tree</param>
+		/// <param name="cancellationToken">Cancellation token for the operation</param>
+		/// <returns></returns>
+		public async Task<NodeLocator> FlushAsync(TreeNode root, CancellationToken cancellationToken)
+		{
+			TreeNodeRef rootRef = new TreeNodeRef(root);
+			await WriteAsync(rootRef, cancellationToken);
+			await FlushAsync(cancellationToken);
+			return rootRef.Locator;
+		}
+
 		/// <summary>
 		/// Flushes all the current nodes to storage
 		/// </summary>
@@ -316,45 +313,57 @@ namespace EpicGames.Horde.Storage
 			if (_queue.Count > 0)
 			{
 				// Create a set from the nodes to be written. We use this to determine references that are imported.
-				HashSet<IoHash> nodeSet = new HashSet<IoHash>(_queue.Select(x => x.Hash));
-
+				HashSet<TreeNodeRef> nodeSet = new HashSet<TreeNodeRef>();
+				foreach (OutputNodeInfo nodeInfo in _queue)
+				{
+					for (NodeRefInfo? nodeRefInfo = nodeInfo; nodeRefInfo != null; nodeRefInfo = nodeRefInfo.Next)
+					{
+						nodeSet.Add(nodeRefInfo.NodeRef);
+					}
+				}
+				
 				// Find all the imported nodes by bundle
-				Dictionary<BlobLocator, List<(int, IoHash)>> bundleToImportedNodes = new Dictionary<BlobLocator, List<(int, IoHash)>>();
-				foreach (NodeInfo nodeInfo in _queue)
+				Dictionary<BlobLocator, List<(int, TreeNodeRef)>> bundleToImports = new Dictionary<BlobLocator, List<(int, TreeNodeRef)>>();
+				foreach (OutputNodeInfo nodeInfo in _queue)
 				{
-					foreach (IoHash refHash in nodeInfo.RefHashes)
+					foreach (TreeNodeRef refKey in nodeInfo.Refs)
 					{
-						if (nodeSet.Add(refHash))
+						if (nodeSet.Add(refKey))
 						{
-							NodeLocator refLocator = _hashToNode[refHash];
+							Debug.Assert(refKey.Locator.IsValid());
+							NodeLocator refLocator = refKey.Locator;
 
-							List<(int, IoHash)>? importedNodes;
-							if (!bundleToImportedNodes.TryGetValue(refLocator.Blob, out importedNodes))
+							List<(int, TreeNodeRef)>? importedNodes;
+							if (!bundleToImports.TryGetValue(refLocator.Blob, out importedNodes))
 							{
-								importedNodes = new List<(int, IoHash)>();
-								bundleToImportedNodes.Add(refLocator.Blob, importedNodes);
+								importedNodes = new List<(int, TreeNodeRef)>();
+								bundleToImports.Add(refLocator.Blob, importedNodes);
 							}
 
-							importedNodes.Add((refLocator.ExportIdx, refHash));
+							importedNodes.Add((refLocator.ExportIdx, refKey));
 						}
 					}
 				}
 
 				// Map from node hash to index, with imported nodes first, ordered by blob, and exported nodes second.
-				Dictionary<IoHash, int> nodeToIndex = new Dictionary<IoHash, int>();
+				int nodeIdx = 0;
+				Dictionary<TreeNodeRef, int> nodeToIndex = new Dictionary<TreeNodeRef, int>();
 
 				// Add all the imports and assign them identifiers
 				List<BundleImport> imports = new List<BundleImport>();
-				foreach ((BlobLocator blobLocator, List<(int, IoHash)> importedNodes) in bundleToImportedNodes)
+				foreach ((BlobLocator blobLocator, List<(int, TreeNodeRef)> importEntries) in bundleToImports)
 				{
-					importedNodes.SortBy(x => x.Item1);
+					importEntries.SortBy(x => x.Item1);
 
-					foreach ((_, IoHash hash) in importedNodes)
+					int[] entries = new int[importEntries.Count];
+					for (int idx = 0; idx < importEntries.Count; idx++)
 					{
-						nodeToIndex[hash] = nodeToIndex.Count;
+						TreeNodeRef key = importEntries[idx].Item2;
+						nodeToIndex.Add(key, nodeIdx++);
+						entries[idx] = importEntries[idx].Item1;
 					}
 
-					imports.Add(new BundleImport(blobLocator, importedNodes));
+					imports.Add(new BundleImport(blobLocator, entries));
 				}
 
 				// List of types in the bundle
@@ -363,20 +372,17 @@ namespace EpicGames.Horde.Storage
 
 				// Create the export list
 				List<BundleExport> exports = new List<BundleExport>(_queue.Count);
-				foreach (NodeInfo nodeInfo in _queue)
+				foreach (OutputNodeInfo nodeInfo in _queue)
 				{
-					int typeIdx;
-					if (!typeToIndex.TryGetValue(nodeInfo.Type, out typeIdx))
+					int typeIdx = FindOrAddType(types, typeToIndex, nodeInfo.Key.Type);
+					int[] references = nodeInfo.Refs.Select(x => nodeToIndex[x]).ToArray();
+					BundleExport export = new BundleExport(typeIdx, nodeInfo.Key.Hash, nodeInfo.Length, references);
+					for (NodeRefInfo? nodeRefInfo = nodeInfo; nodeRefInfo != null; nodeRefInfo = nodeRefInfo.Next)
 					{
-						typeIdx = types.Count;
-						typeToIndex.Add(nodeInfo.Type, typeIdx);
-						types.Add(nodeInfo.Type);
+						nodeToIndex[nodeRefInfo.NodeRef] = nodeIdx;
 					}
-
-					int[] references = nodeInfo.RefHashes.Select(x => nodeToIndex[x]).ToArray();
-					BundleExport export = new BundleExport(typeIdx, nodeInfo.Hash, nodeInfo.Length, references);
-					nodeToIndex.Add(nodeInfo.Hash, nodeToIndex.Count);
 					exports.Add(export);
+					nodeIdx++;
 				}
 
 				// Create the bundle
@@ -389,30 +395,35 @@ namespace EpicGames.Horde.Storage
 				{
 					NodeLocator nodeLocator = new NodeLocator(locator, idx);
 
-					NodeInfo nodeInfo = _queue[idx];
-					nodeInfo.NodeRef.MarkAsWritten(nodeInfo.Hash, nodeLocator, nodeInfo.Revision);
-
-					_hashToNode[nodeInfo.Hash] = nodeLocator;
-				}
-
-				// Update any refs with their target locator
-				int refIdx = 0;
-				for (; refIdx < _secondaryQueue.Count; refIdx++)
-				{
-					NodeInfo nodeInfo = _secondaryQueue[refIdx];
-					if (_hashToNode.TryGetValue(nodeInfo.Hash, out NodeLocator refLocator))
+					OutputNodeInfo nodeInfo = _queue[idx];
+					for (NodeRefInfo? nodeRefInfo = nodeInfo; nodeRefInfo != null; nodeRefInfo = nodeRefInfo.Next)
 					{
-						nodeInfo.NodeRef.MarkAsWritten(nodeInfo.Hash, refLocator, nodeInfo.Revision);
+						nodeRefInfo.NodeRef.MarkAsWritten(nodeLocator, nodeRefInfo.Revision);
 					}
+
+					_nodeKeyToLocator[nodeInfo.Key] = nodeLocator;
 				}
-				_secondaryQueue.RemoveRange(0, refIdx);
 
 				// Reset the output state
 				_packets.Clear();
 				_encodedPackets.Clear();
 				_queue.Clear();
+				_queuedRefs.Clear();
+				_nodeKeyToInfo.Clear();
 				_currentBundleLength = 0;
 			}
 		}
+
+		static int FindOrAddType(List<BundleType> types, Dictionary<BundleType, int> typeToIndex, BundleType type)
+		{
+			int typeIdx;
+			if (!typeToIndex.TryGetValue(type, out typeIdx))
+			{
+				typeIdx = types.Count;
+				typeToIndex.Add(type, typeIdx);
+				types.Add(type);
+			}
+			return typeIdx;
+		}
 	}
 }
