Horde: Tidy up some naming/comments.

diff --git a/Engine/Source/Programs/Shared/EpicGames.Horde/Storage/TreeWriter.cs b/Engine/Source/Programs/Shared/EpicGames.Horde/Storage/TreeWriter.cs
--- a/Engine/Source/Programs/Shared/EpicGames.Horde/Storage/TreeWriter.cs
+++ b/Engine/Source/Programs/Shared/EpicGames.Horde/Storage/TreeWriter.cs
@@ -37,9 +37,6 @@ namespace EpicGames.Horde.Storage
 	/// </summary>
 	public class TreeWriter
 	{
-		record NodeKey(IoHash Hash, BundleType Type);
-		record struct NodeRevision(TreeNodeRef NodeRef, uint Revision);
-
 		// Information about a TreeNodeRef that was written, and needs to be fixed up after flushing
 		class NodeRefInfo
 		{
@@ -55,14 +52,17 @@ namespace EpicGames.Horde.Storage
 			}
 		}
 
+		// Unique identifier for an output node
+		record OutputNodeKey(IoHash Hash, BundleType Type);
+
 		// Information about a unique output node. Note that multiple node refs may de-duplicate to the same output node.
 		class OutputNodeInfo : NodeRefInfo
 		{
-			public readonly NodeKey Key;
+			public readonly OutputNodeKey Key;
 			public readonly int Length;
 			public readonly IReadOnlyList<TreeNodeRef> Refs;
 
-			public OutputNodeInfo(NodeKey key, int length, IReadOnlyList<TreeNodeRef> refs, TreeNodeRef nodeRef)
+			public OutputNodeInfo(OutputNodeKey key, int length, IReadOnlyList<TreeNodeRef> refs, TreeNodeRef nodeRef)
 				: base(nodeRef, null)
 			{
 				Key = key;
@@ -111,10 +111,10 @@ namespace EpicGames.Horde.Storage
 		readonly ArrayMemoryWriter _packetWriter;
 
 		// Map of keys to nodes in the queue
-		readonly Dictionary<NodeKey, OutputNodeInfo> _nodeKeyToInfo = new Dictionary<NodeKey, OutputNodeInfo>();
+		readonly Dictionary<OutputNodeKey, OutputNodeInfo> _nodeKeyToInfo = new Dictionary<OutputNodeKey, OutputNodeInfo>();
 
 		// Map of hashes to existing nodes. Empty/invalid locators are used for items that are queued for writing, but not available yet.
-		readonly Dictionary<NodeKey, NodeLocator> _nodeKeyToLocator = new Dictionary<NodeKey, NodeLocator>(); // TODO: this needs to include some additional state from external sources.
+		readonly Dictionary<OutputNodeKey, NodeLocator> _nodeKeyToLocator = new Dictionary<OutputNodeKey, NodeLocator>(); // TODO: this needs to include some additional state from external sources.
 
 		// Queue of nodes for the current bundle
 		readonly List<OutputNodeInfo> _queue = new List<OutputNodeInfo>();
@@ -216,7 +216,7 @@ namespace EpicGames.Horde.Storage
 			target.Hash = IoHash.Compute(data);
 			nodeRef.MarkAsPendingWrite();
 
-			NodeKey nodeKey = new NodeKey(target.Hash, target.GetBundleType());
+			OutputNodeKey nodeKey = new OutputNodeKey(target.Hash, target.GetBundleType());
 
 			// Check if we're already tracking a node with the same hash
 			OutputNodeInfo? nodeInfo;
@@ -272,6 +272,9 @@ namespace EpicGames.Horde.Storage
 			return nodeRef.Locator;
 		}
 
+		/// <summary>
+		/// Compresses the current packet and schedule it to be written to storage
+		/// </summary>
 		void FlushPacket()
 		{
 			if (_packetWriter.Length > 0)
