Horde: Fix slow performance of bundle store tests.

diff --git a/Engine/Source/Programs/Shared/EpicGames.Horde.Tests/BundleStoreTests.cs b/Engine/Source/Programs/Shared/EpicGames.Horde.Tests/BundleStoreTests.cs
--- a/Engine/Source/Programs/Shared/EpicGames.Horde.Tests/BundleStoreTests.cs
+++ b/Engine/Source/Programs/Shared/EpicGames.Horde.Tests/BundleStoreTests.cs
@@ -319,8 +319,8 @@ namespace EpicGames.Horde.Tests
 				root = new DirectoryNode(DirectoryFlags.None);
 
 				ChunkingOptions options = new ChunkingOptions();
-				options.LeafOptions.MinSize = 1;
-				options.LeafOptions.TargetSize = 128;
+				options.LeafOptions.MinSize = 128;
+				options.LeafOptions.TargetSize = 256;
 				options.LeafOptions.MaxSize = 64 * 1024;
 
 				FileNode file = await FileNode.CreateAsync(data, options, writer, CancellationToken.None);
@@ -399,9 +399,34 @@ namespace EpicGames.Horde.Tests
 
 			FileNode world = await root.Files.First().ExpandAsync(reader, CancellationToken.None);
 
-			byte[] worldData = await GetFileDataAsync(reader, world);
-			Assert.AreEqual(data.Length, worldData.Length);
-			Assert.IsTrue(worldData.SequenceEqual(data));
+			int length = await CheckFileDataAsync(reader, world, data);
+			Assert.AreEqual(data.Length, length);
+		}
+
+		static async Task<int> CheckFileDataAsync(TreeReader reader, FileNode fileNode, ReadOnlyMemory<byte> data)
+		{
+			int offset = 0;
+			if (fileNode is LeafFileNode leafNode)
+			{
+				foreach (ReadOnlyMemory<byte> segment in leafNode.Data)
+				{
+					Assert.IsTrue(segment.Span.SequenceEqual(data.Span.Slice(offset, segment.Length)));
+					offset += segment.Length;
+				}
+			}
+			else if (fileNode is InteriorFileNode interiorFileNode)
+			{
+				foreach (TreeNodeRef<FileNode> childRef in interiorFileNode.Children)
+				{
+					FileNode child = await childRef.ExpandAsync(reader);
+					offset += await CheckFileDataAsync(reader, child, data.Slice(offset));
+				}
+			}
+			else
+			{
+				throw new NotImplementedException();
+			}
+			return offset;
 		}
 
 		static async Task<byte[]> GetFileDataAsync(TreeReader reader, FileNode fileNode)
