Horde: Let syncing without have table be a class-wide mode of operation for ManagedWorkspace

Previously it could be sent as a flag to SyncAsync but this was non-ideal when it comes to conforming. Which method to use is now passed through workspace config sent from Horde server.

diff --git a/Engine/Source/Programs/Horde/Horde.Agent.Tests/PerforceExecutorTest.cs b/Engine/Source/Programs/Horde/Horde.Agent.Tests/PerforceExecutorTest.cs
new file mode 100644
--- /dev/null
+++ b/Engine/Source/Programs/Horde/Horde.Agent.Tests/PerforceExecutorTest.cs
@@ -0,0 +1,24 @@
+// Copyright Epic Games, Inc. All Rights Reserved.
+
+using Horde.Agent.Execution;
+using Microsoft.VisualStudio.TestTools.UnitTesting;
+
+namespace Horde.Agent.Tests;
+
+[TestClass]
+public class PerforceExecutorTest
+{
+	[TestMethod]
+	public void ShouldUseHaveTable()
+	{
+		Assert.IsTrue(PerforceExecutor.ShouldUseHaveTable(null));
+		Assert.IsTrue(PerforceExecutor.ShouldUseHaveTable(""));
+		Assert.IsTrue(PerforceExecutor.ShouldUseHaveTable("name=managedWorkspace"));
+		Assert.IsTrue(PerforceExecutor.ShouldUseHaveTable("name=managedWorkspace$#@!@#"));
+		Assert.IsTrue(PerforceExecutor.ShouldUseHaveTable("name=managedWorkspace$#@!@#"));
+		Assert.IsTrue(PerforceExecutor.ShouldUseHaveTable("name=managedWorkspace&useHaveTable=true"));
+		
+		Assert.IsFalse(PerforceExecutor.ShouldUseHaveTable("name=managedWorkspace&useHaveTable=false"));
+		Assert.IsFalse(PerforceExecutor.ShouldUseHaveTable("name=ManagedWorkspace&useHaveTable=FalsE"));
+	}
+}
diff --git a/Engine/Source/Programs/Horde/Horde.Agent/Commands/Workspace/PopulateCacheCommand.cs b/Engine/Source/Programs/Horde/Horde.Agent/Commands/Workspace/PopulateCacheCommand.cs
--- a/Engine/Source/Programs/Horde/Horde.Agent/Commands/Workspace/PopulateCacheCommand.cs
+++ b/Engine/Source/Programs/Horde/Horde.Agent/Commands/Workspace/PopulateCacheCommand.cs
@@ -26,15 +26,9 @@ namespace Horde.Agent.Commands.Workspace
 		[CommandLine("-FakeSync")]
 		[Description("Simulates the sync without actually fetching any files")]
 		bool FakeSync { get; set; } = false;
-		
-		[CommandLine("-UseHaveTable")]
-		[Description("Update the have-table")]
-		string UseHaveTable { get; set; } = "true";
 
 		protected override async Task ExecuteAsync(IPerforceConnection perforce, ManagedWorkspace repo, ILogger logger)
 		{
-			bool useHaveTable = UseHaveTable.Equals("true", StringComparison.Ordinal);
-
 			List<string> expandedFilters = ExpandFilters(Filters);
 
 			List<PopulateRequest> populateRequests = new List<PopulateRequest>();
@@ -50,7 +44,7 @@ namespace Horde.Agent.Commands.Workspace
 				populateRequests.Add(new PopulateRequest(perforceClient, clientAndStreamParam.Substring(idx + 1), expandedFilters));
 			}
 
-			await repo.PopulateAsync(populateRequests, FakeSync, useHaveTable, CancellationToken.None);
+			await repo.PopulateAsync(populateRequests, FakeSync, CancellationToken.None);
 		}
 	}
 }
diff --git a/Engine/Source/Programs/Horde/Horde.Agent/Commands/Workspace/SyncCommand.cs b/Engine/Source/Programs/Horde/Horde.Agent/Commands/Workspace/SyncCommand.cs
--- a/Engine/Source/Programs/Horde/Horde.Agent/Commands/Workspace/SyncCommand.cs
+++ b/Engine/Source/Programs/Horde/Horde.Agent/Commands/Workspace/SyncCommand.cs
@@ -52,7 +52,7 @@ namespace Horde.Agent.Commands.Workspace
 			List<string> expandedFilters = ExpandFilters(Filters);
 
 			using IPerforceConnection perforceClient = await perforce.WithClientAsync(ClientName);
-			await repo.SyncAsync(perforceClient, StreamName, changeNumber, expandedFilters, !IncrementalSync, FakeSync, true, CacheFile, CancellationToken.None);
+			await repo.SyncAsync(perforceClient, StreamName, changeNumber, expandedFilters, !IncrementalSync, FakeSync, CacheFile, CancellationToken.None);
 
 			if (PreflightChange != -1)
 			{
diff --git a/Engine/Source/Programs/Horde/Horde.Agent/Commands/Workspace/WorkspaceCommand.cs b/Engine/Source/Programs/Horde/Horde.Agent/Commands/Workspace/WorkspaceCommand.cs
--- a/Engine/Source/Programs/Horde/Horde.Agent/Commands/Workspace/WorkspaceCommand.cs
+++ b/Engine/Source/Programs/Horde/Horde.Agent/Commands/Workspace/WorkspaceCommand.cs
@@ -34,6 +34,10 @@ namespace Horde.Agent.Commands.Workspace
 		[CommandLine("-PreferNativeClient")]
 		[Description("Prefer to use native Perforce client (instead of launching separate p4 process)")]
 		protected bool PreferNativeClient { get; set; } = false;
+		
+		[CommandLine("-UseHaveTable")]
+		[Description("Use have-table for syncing")]
+		string UseHaveTable { get; set; } = "true";
 
 		public override void Configure(CommandLineArguments arguments, ILogger logger)
 		{
@@ -62,10 +66,11 @@ namespace Horde.Agent.Commands.Workspace
 			PerforceSettings settings = new(ServerAndPort, UserName) { PreferNativeClient = PreferNativeClient };
 			using IPerforceConnection perforce = await PerforceConnection.CreateAsync(settings, logger);
 			InfoRecord info = await perforce.GetInfoAsync(InfoOptions.ShortOutput, CancellationToken.None);
+			bool useHaveTable = UseHaveTable.Equals("true", StringComparison.Ordinal);
 
 			using Logging.HordeLoggerProvider loggerProvider = new ();
 			ILogger repoLogger = loggerProvider.CreateLogger("Repository");
-			ManagedWorkspace repo = await ManagedWorkspace.LoadOrCreateAsync(info.ClientHost!, BaseDir, Overwrite, repoLogger, CancellationToken.None);
+			ManagedWorkspace repo = await ManagedWorkspace.LoadOrCreateAsync(info.ClientHost!, BaseDir, Overwrite, useHaveTable, repoLogger, CancellationToken.None);
 			await ExecuteAsync(perforce, repo, logger);
 			return 0;
 		}
diff --git a/Engine/Source/Programs/Horde/Horde.Agent/Execution/PerforceExecutor.cs b/Engine/Source/Programs/Horde/Horde.Agent/Execution/PerforceExecutor.cs
--- a/Engine/Source/Programs/Horde/Horde.Agent/Execution/PerforceExecutor.cs
+++ b/Engine/Source/Programs/Horde/Horde.Agent/Execution/PerforceExecutor.cs
@@ -2,6 +2,7 @@
 
 using System;
 using System.Collections.Generic;
+using System.Collections.Specialized;
 using System.IO;
 using System.Linq;
 using System.Net.Http;
@@ -9,6 +10,7 @@ using System.Runtime.InteropServices;
 using System.Text.RegularExpressions;
 using System.Threading;
 using System.Threading.Tasks;
+using System.Web;
 using EpicGames.Core;
 using EpicGames.Perforce;
 using EpicGames.Perforce.Managed;
@@ -34,7 +36,6 @@ namespace Horde.Agent.Execution
 
 		protected WorkspaceInfo? _autoSdkWorkspace;
 		protected WorkspaceInfo _workspace;
-		private bool _useHaveTable = true;
 
 		public PerforceExecutor(ISession session, string jobId, string batchId, string agentTypeName, AgentWorkspace? autoSdkWorkspaceInfo, AgentWorkspace workspaceInfo, DirectoryReference rootDir, IHttpClientFactory httpClientFactory)
 			: base(session, jobId, batchId, agentTypeName, httpClientFactory)
@@ -49,57 +50,50 @@ namespace Horde.Agent.Execution
 		public override async Task InitializeAsync(ILogger logger, CancellationToken cancellationToken)
 		{
 			await base.InitializeAsync(logger, cancellationToken);
-			
-			int index = _additionalArguments.FindIndex(x => x.Contains("-SyncWithoutHaveTable", StringComparison.OrdinalIgnoreCase));
-			if (index >= 0)
-			{
-				logger.LogInformation("Syncing without have table");
-				_additionalArguments.RemoveAt(index);
-				_useHaveTable = false;
-			}
 
-			// Setup and sync the autosdk workspace
+			// Setup and sync the AutoSDK workspace
 			if (_autoSdkWorkspaceInfo != null)
 			{
-				using (IScope scope = GlobalTracer.Instance.BuildSpan("Workspace").WithResourceName("AutoSDK").StartActive())
-				{
-					_autoSdkWorkspace = await Utility.WorkspaceInfo.SetupWorkspaceAsync(_autoSdkWorkspaceInfo, _rootDir, logger, cancellationToken);
+				using IScope _ = GlobalTracer.Instance.BuildSpan("Workspace").WithResourceName("AutoSDK").StartActive();
+				
+				bool useHaveTable = ShouldUseHaveTable(_autoSdkWorkspaceInfo.Method);
+				_autoSdkWorkspace = await Utility.WorkspaceInfo.SetupWorkspaceAsync(_autoSdkWorkspaceInfo, _rootDir, useHaveTable, logger, cancellationToken);
 
-					DirectoryReference legacyDir = DirectoryReference.Combine(_autoSdkWorkspace.MetadataDir, "HostWin64");
-					if (DirectoryReference.Exists(legacyDir))
+				DirectoryReference legacyDir = DirectoryReference.Combine(_autoSdkWorkspace.MetadataDir, "HostWin64");
+				if (DirectoryReference.Exists(legacyDir))
+				{
+					try
 					{
-						try
-						{
-							FileUtils.ForceDeleteDirectory(legacyDir);
-						}
-						catch(Exception ex)
-						{
-							logger.LogInformation(ex, "Unable to delete {Dir}", legacyDir);
-						}
+						FileUtils.ForceDeleteDirectory(legacyDir);
 					}
+					catch(Exception ex)
+					{
+						logger.LogInformation(ex, "Unable to delete {Dir}", legacyDir);
+					}
+				}
 
-					int autoSdkChangeNumber = await _autoSdkWorkspace.GetLatestChangeAsync(cancellationToken);
+				int autoSdkChangeNumber = await _autoSdkWorkspace.GetLatestChangeAsync(cancellationToken);
 
-					string syncText = $"Synced to CL {autoSdkChangeNumber}";
+				string syncText = $"Synced to CL {autoSdkChangeNumber}";
 
-					FileReference syncFile = FileReference.Combine(_autoSdkWorkspace.MetadataDir, "Synced.txt");
-					if (!FileReference.Exists(syncFile) || FileReference.ReadAllText(syncFile) != syncText)
-					{
-						FileReference.Delete(syncFile);
+				FileReference syncFile = FileReference.Combine(_autoSdkWorkspace.MetadataDir, "Synced.txt");
+				if (!FileReference.Exists(syncFile) || FileReference.ReadAllText(syncFile) != syncText)
+				{
+					FileReference.Delete(syncFile);
 
-						FileReference autoSdkCacheFile = FileReference.Combine(_autoSdkWorkspace.MetadataDir, "Contents.dat");
-						await WorkspaceInfo.UpdateLocalCacheMarker(autoSdkCacheFile, autoSdkChangeNumber, -1);
-						await _autoSdkWorkspace.SyncAsync(autoSdkChangeNumber, -1, autoSdkCacheFile, true, cancellationToken);
+					FileReference autoSdkCacheFile = FileReference.Combine(_autoSdkWorkspace.MetadataDir, "Contents.dat");
+					await WorkspaceInfo.UpdateLocalCacheMarker(autoSdkCacheFile, autoSdkChangeNumber, -1);
+					await _autoSdkWorkspace.SyncAsync(autoSdkChangeNumber, -1, autoSdkCacheFile, cancellationToken);
 
-						await FileReference.WriteAllTextAsync(syncFile, syncText);
-					}
+					await FileReference.WriteAllTextAsync(syncFile, syncText);
 				}
 			}
 
 			using (IScope scope = GlobalTracer.Instance.BuildSpan("Workspace").WithResourceName(_workspaceInfo.Identifier).StartActive())
 			{
 				// Sync the regular workspace
-				_workspace = await Utility.WorkspaceInfo.SetupWorkspaceAsync(_workspaceInfo, _rootDir, logger, cancellationToken);
+				bool useHaveTable = ShouldUseHaveTable(_workspaceInfo.Method);
+				_workspace = await Utility.WorkspaceInfo.SetupWorkspaceAsync(_workspaceInfo, _rootDir, useHaveTable, logger, cancellationToken);
 
 				// Figure out the change to build
 				if (_job.Change == 0)
@@ -115,7 +109,7 @@ namespace Horde.Agent.Execution
 
 				// Sync the workspace
 				int syncPreflightChange = (_job.ClonedPreflightChange != 0) ? _job.ClonedPreflightChange : _job.PreflightChange;
-				await _workspace.SyncAsync(_job.Change, syncPreflightChange, null, _useHaveTable, cancellationToken);
+				await _workspace.SyncAsync(_job.Change, syncPreflightChange, null, cancellationToken);
 
 				// Remove any cached BuildGraph manifests
 				DirectoryReference manifestDir = DirectoryReference.Combine(_workspace.WorkspaceDir, "Engine", "Saved", "BuildGraph");
@@ -233,7 +227,7 @@ namespace Horde.Agent.Execution
 			await _workspace.CleanAsync(cancellationToken);
 		}
 
-		public static async Task ConformAsync(DirectoryReference rootDir, IList<AgentWorkspace> pendingWorkspaces, bool removeUntrackedFiles, bool useHaveTable, ILogger logger, CancellationToken cancellationToken)
+		public static async Task ConformAsync(DirectoryReference rootDir, IList<AgentWorkspace> pendingWorkspaces, bool removeUntrackedFiles, ILogger logger, CancellationToken cancellationToken)
 		{
 			using IScope scope = GlobalTracer.Instance.BuildSpan("Conform").StartActive();
 			scope.Span.SetTag("workspaces", String.Join(',', pendingWorkspaces.Select(x => x.Identifier)));
@@ -249,7 +243,8 @@ namespace Horde.Agent.Execution
 			List<WorkspaceInfo> workspaces = new List<WorkspaceInfo>();
 			foreach (AgentWorkspace pendingWorkspace in pendingWorkspaces)
 			{
-				WorkspaceInfo workspace = await Utility.WorkspaceInfo.SetupWorkspaceAsync(pendingWorkspace, rootDir, logger, cancellationToken);
+				bool useHaveTable = ShouldUseHaveTable(pendingWorkspace.Method);
+				WorkspaceInfo workspace = await Utility.WorkspaceInfo.SetupWorkspaceAsync(pendingWorkspace, rootDir, useHaveTable, logger, cancellationToken);
 				workspaces.Add(workspace);
 			}
 
@@ -383,13 +378,13 @@ namespace Horde.Agent.Execution
 					if (populateRequests.Count == 1 && !firstWorkspace.RemoveUntrackedFiles && !removeUntrackedFiles)
 					{
 						await firstWorkspace.CleanAsync(cancellationToken);
-						syncFuncs.Add(() => firstWorkspace.SyncAsync(-1, -1, null, useHaveTable, cancellationToken));
+						syncFuncs.Add(() => firstWorkspace.SyncAsync(-1, -1, null, cancellationToken));
 					}
 					else
 					{
 						ManagedWorkspace repository = firstWorkspace.Repository;
-						Tuple<int, StreamSnapshot>[] streamStates = await repository.PopulateCleanAsync(populateRequests, useHaveTable, cancellationToken);
-						syncFuncs.Add(() => repository.PopulateSyncAsync(populateRequests, streamStates, false, useHaveTable, cancellationToken));
+						Tuple<int, StreamSnapshot>[] streamStates = await repository.PopulateCleanAsync(populateRequests, cancellationToken);
+						syncFuncs.Add(() => repository.PopulateSyncAsync(populateRequests, streamStates, false, cancellationToken));
 					}
 				}
 				foreach (Func<Task> syncFunc in syncFuncs)
@@ -405,6 +400,37 @@ namespace Horde.Agent.Execution
 				}
 			}
 		}
+
+		/// <summary>
+		/// Parse a text string as a query string to determine use of have table
+		/// </summary>
+		/// <param name="method">Method text string from workspace config</param>
+		/// <returns>True if have table should be enabled</returns>
+		public static bool ShouldUseHaveTable(string? method)
+		{
+			const string NameKey = "name";
+			const string ManagedWorkspaceValue = "managedWorkspace";
+			const string UseHaveTableKey = "useHaveTable";
+			
+			if (method == null)
+			{
+				return true;
+			}
+
+			NameValueCollection nameValues = HttpUtility.ParseQueryString(method);
+			string? name = nameValues[NameKey];
+			string? useHaveTable = nameValues[UseHaveTableKey];
+			
+			if (name != null && name.Equals(ManagedWorkspaceValue, StringComparison.OrdinalIgnoreCase))
+			{
+				if (useHaveTable != null && useHaveTable.Equals("false", StringComparison.OrdinalIgnoreCase))
+				{
+					return false;
+				}
+			}
+
+			return true;
+		}
 	}
 
 	class PerforceExecutorFactory : JobExecutorFactory
diff --git a/Engine/Source/Programs/Horde/Horde.Agent/Leases/Handlers/ConformHandler.cs b/Engine/Source/Programs/Horde/Horde.Agent/Leases/Handlers/ConformHandler.cs
--- a/Engine/Source/Programs/Horde/Horde.Agent/Leases/Handlers/ConformHandler.cs
+++ b/Engine/Source/Programs/Horde/Horde.Agent/Leases/Handlers/ConformHandler.cs
@@ -46,7 +46,7 @@ namespace Horde.Agent.Leases.Handlers
 				// Run the conform task
 				if (_settings.Executor.Equals(PerforceExecutor.Name, StringComparison.OrdinalIgnoreCase) && _settings.PerforceExecutor.RunConform)
 				{
-					await PerforceExecutor.ConformAsync(session.WorkingDir, pendingWorkspaces, removeUntrackedFiles, true, conformLogger, cancellationToken);
+					await PerforceExecutor.ConformAsync(session.WorkingDir, pendingWorkspaces, removeUntrackedFiles, conformLogger, cancellationToken);
 				}
 				else
 				{
diff --git a/Engine/Source/Programs/Horde/Horde.Agent/Utility/WorkspaceInfo.cs b/Engine/Source/Programs/Horde/Horde.Agent/Utility/WorkspaceInfo.cs
--- a/Engine/Source/Programs/Horde/Horde.Agent/Utility/WorkspaceInfo.cs
+++ b/Engine/Source/Programs/Horde/Horde.Agent/Utility/WorkspaceInfo.cs
@@ -118,10 +118,11 @@ namespace Horde.Agent.Utility
 		/// </summary>
 		/// <param name="workspace">The workspace definition</param>
 		/// <param name="rootDir">Root directory for storing the workspace</param>
+		/// <param name="useHaveTable">Use the client's have table when syncing</param>
 		/// <param name="logger">Logger output</param>
 		/// <param name="cancellationToken">Cancellation token</param>
 		/// <returns>New workspace info</returns>
-		public static async Task<WorkspaceInfo> SetupWorkspaceAsync(AgentWorkspace workspace, DirectoryReference rootDir, ILogger logger, CancellationToken cancellationToken)
+		public static async Task<WorkspaceInfo> SetupWorkspaceAsync(AgentWorkspace workspace, DirectoryReference rootDir, bool useHaveTable, ILogger logger, CancellationToken cancellationToken)
 		{
 			// Fill in the default credentials iff they are not set
 			string? serverAndPort = String.IsNullOrEmpty(workspace.ServerAndPort)? null : workspace.ServerAndPort;
@@ -153,7 +154,7 @@ namespace Horde.Agent.Utility
 					logger.LogInformation("Using locally logged in session for {UserName}", userName);
 				}
 			}
-			return await SetupWorkspaceAsync(perforce, workspace.Stream, workspace.Identifier, workspace.View, !workspace.Incremental, rootDir, logger, cancellationToken);
+			return await SetupWorkspaceAsync(perforce, workspace.Stream, workspace.Identifier, workspace.View, !workspace.Incremental, rootDir, useHaveTable, logger, cancellationToken);
 		}
 
 		/// <summary>
@@ -165,10 +166,11 @@ namespace Horde.Agent.Utility
 		/// <param name="view">View for this workspace</param>
 		/// <param name="removeUntrackedFiles">Whether untracked files should be removed when cleaning this workspace</param>
 		/// <param name="rootDir">Root directory for storing the workspace</param>
+		/// <param name="useHaveTable">Use the client's have table when syncing</param>
 		/// <param name="logger">Logger output</param>
 		/// <param name="cancellationToken">Cancellation token</param>
 		/// <returns>New workspace info</returns>
-		public static async Task<WorkspaceInfo> SetupWorkspaceAsync(IPerforceConnection perforce, string streamName, string identifier, IList<string> view, bool removeUntrackedFiles, DirectoryReference rootDir, ILogger logger, CancellationToken cancellationToken)
+		public static async Task<WorkspaceInfo> SetupWorkspaceAsync(IPerforceConnection perforce, string streamName, string identifier, IList<string> view, bool removeUntrackedFiles, DirectoryReference rootDir, bool useHaveTable, ILogger logger, CancellationToken cancellationToken)
 		{
 			// Get the host name, and fill in any missing metadata about the connection
 			InfoRecord info = await perforce.GetInfoAsync(InfoOptions.ShortOutput, cancellationToken);
@@ -178,6 +180,11 @@ namespace Horde.Agent.Utility
 			{
 				throw new Exception("Unable to determine Perforce host name");
 			}
+			
+			if (!useHaveTable)
+			{
+				logger.LogInformation("Skipping use of have table");
+			}
 
 			// replace invalid characters in the workspace identifier with a '+' character
 
@@ -211,7 +218,7 @@ namespace Horde.Agent.Utility
 			DirectoryReference workspaceDir = DirectoryReference.Combine(metadataDir, "Sync");
 
 			// Create the repository
-			ManagedWorkspace newRepository = await ManagedWorkspace.LoadOrCreateAsync(hostName, metadataDir, true, logger, cancellationToken);
+			ManagedWorkspace newRepository = await ManagedWorkspace.LoadOrCreateAsync(hostName, metadataDir, true, useHaveTable, logger, cancellationToken);
 			if (removeUntrackedFiles)
 			{
 				await newRepository.DeleteClientAsync(perforceClient, cancellationToken);
@@ -295,12 +302,11 @@ namespace Horde.Agent.Utility
 		/// <param name="change">The changelist to sync to</param>
 		/// <param name="preflightChange">Change to preflight, or 0</param>
 		/// <param name="cacheFile">Path to the cache file to use</param>
-		/// <param name="useHaveTable">Whether to update have table during sync</param>
 		/// <param name="cancellationToken">Cancellation token</param>
 		/// <returns>Async task</returns>
-		public async Task SyncAsync(int change, int preflightChange, FileReference? cacheFile, bool useHaveTable, CancellationToken cancellationToken)
+		public async Task SyncAsync(int change, int preflightChange, FileReference? cacheFile, CancellationToken cancellationToken)
 		{
-			await Repository.SyncAsync(PerforceClient, StreamName, change, View, RemoveUntrackedFiles, false, useHaveTable, cacheFile, cancellationToken);
+			await Repository.SyncAsync(PerforceClient, StreamName, change, View, RemoveUntrackedFiles, false, cacheFile, cancellationToken);
 
 			// Purge the cache for incremental workspaces
 			if (!RemoveUntrackedFiles)
diff --git a/Engine/Source/Programs/Shared/EpicGames.Perforce.Managed.Tests/ManagedWorkspaceTest.cs b/Engine/Source/Programs/Shared/EpicGames.Perforce.Managed.Tests/ManagedWorkspaceTest.cs
--- a/Engine/Source/Programs/Shared/EpicGames.Perforce.Managed.Tests/ManagedWorkspaceTest.cs
+++ b/Engine/Source/Programs/Shared/EpicGames.Perforce.Managed.Tests/ManagedWorkspaceTest.cs
@@ -32,10 +32,10 @@ public class ManagedWorkspaceTest : BasePerforceFixtureTest
 	[DataRow(false, DisplayName = "Without have-table")]
 	public async Task SyncSingleChangelist(bool useHaveTable)
 	{
-		ManagedWorkspace ws = await GetManagedWorkspace();
+		ManagedWorkspace ws = await GetManagedWorkspace(useHaveTable);
 		await AssertHaveTableFileCount(0);
 
-		await SyncAsync(ws, 6, useHaveTable);
+		await SyncAsync(ws, 6);
 		Stream.GetChangelist(6).AssertDepotFiles(SyncDir);
 		await Stream.GetChangelist(6).AssertHaveTableAsync(PerforceConnection, useHaveTable);
 	}
@@ -43,7 +43,7 @@ public class ManagedWorkspaceTest : BasePerforceFixtureTest
 	[TestMethod]
 	public async Task SyncBackwardsToOlderChangelistRemoveUntracked()
 	{
-		ManagedWorkspace ws = await GetManagedWorkspace();
+		ManagedWorkspace ws = await GetManagedWorkspace(true);
 
 		await SyncAsync(ws, 6, removeUntracked: false);
 		Stream.GetChangelist(6).AssertDepotFiles(SyncDir);
@@ -63,46 +63,28 @@ public class ManagedWorkspaceTest : BasePerforceFixtureTest
 	[DataRow(false, DisplayName = "Without have-table")]
 	public async Task SyncBackwardsToOlderChangelist(bool useHaveTable)
 	{
-		ManagedWorkspace ws = await GetManagedWorkspace();
+		ManagedWorkspace ws = await GetManagedWorkspace(useHaveTable);
 
-		await SyncAsync(ws, 6, useHaveTable);
+		await SyncAsync(ws, 6);
 		Stream.GetChangelist(6).AssertDepotFiles(SyncDir);
 		await Stream.GetChangelist(6).AssertHaveTableAsync(PerforceConnection, useHaveTable);
 		
-		await SyncAsync(ws, 7, useHaveTable);
+		await SyncAsync(ws, 7);
 		Stream.GetChangelist(7).AssertDepotFiles(SyncDir);
 		await Stream.GetChangelist(7).AssertHaveTableAsync(PerforceConnection, useHaveTable);
 		
 		// Go back one changelist
-		await SyncAsync(ws, 6, useHaveTable);
+		await SyncAsync(ws, 6);
 		Stream.GetChangelist(6).AssertDepotFiles(SyncDir);
 		await Stream.GetChangelist(6).AssertHaveTableAsync(PerforceConnection, useHaveTable);
 	}
 
-	[TestMethod]
-	public async Task SyncMixingUseOfHaveTable()
-	{
-		ManagedWorkspace ws = await GetManagedWorkspace();
-
-		await SyncAsync(ws, 1, useHaveTable: false);
-		Stream.GetChangelist(1).AssertDepotFiles(SyncDir);
-		await AssertHaveTableFileCount(0);
-		
-		await SyncAsync(ws, 2, useHaveTable: true);
-		Stream.GetChangelist(2).AssertDepotFiles(SyncDir);
-		await Stream.GetChangelist(2).AssertHaveTableAsync(PerforceConnection);
-		
-		await SyncAsync(ws, 3, useHaveTable: false);
-		Stream.GetChangelist(3).AssertDepotFiles(SyncDir);
-		await Stream.GetChangelist(2).AssertHaveTableAsync(PerforceConnection); // Have table should only match what was synced previously
-	}
-
 	[TestMethod]
 	[DataRow(true, DisplayName = "With have-table")]
 	[DataRow(false, DisplayName = "Without have-table")]
 	public async Task SyncUsingCacheFiles(bool useHaveTable)
 	{
-		ManagedWorkspace ws = await GetManagedWorkspace();
+		ManagedWorkspace ws = await GetManagedWorkspace(useHaveTable);
 
 		FileReference GetCacheFilePath(int changeNumber)
 		{
@@ -112,7 +94,7 @@ public class ManagedWorkspaceTest : BasePerforceFixtureTest
 		// Sync and create a new cache file per change number
 		foreach (ChangelistFixture cl in Stream.Changelists)
 		{
-			await SyncAsync(ws, cl.Number, useHaveTable: useHaveTable, cacheFile: GetCacheFilePath(cl.Number));
+			await SyncAsync(ws, cl.Number, cacheFile: GetCacheFilePath(cl.Number));
 			cl.AssertDepotFiles(SyncDir);
 			await cl.AssertHaveTableAsync(PerforceConnection, useHaveTable);
 		}
@@ -120,7 +102,7 @@ public class ManagedWorkspaceTest : BasePerforceFixtureTest
 		// Sync again but using the cache files created above
 		foreach (ChangelistFixture cl in Stream.Changelists.Reverse())
 		{
-			await SyncAsync(ws, cl.Number, useHaveTable: useHaveTable, cacheFile: GetCacheFilePath(cl.Number));
+			await SyncAsync(ws, cl.Number, cacheFile: GetCacheFilePath(cl.Number));
 			cl.AssertDepotFiles(SyncDir);
 			await cl.AssertHaveTableAsync(PerforceConnection, useHaveTable);
 		}
@@ -131,13 +113,13 @@ public class ManagedWorkspaceTest : BasePerforceFixtureTest
 	[DataRow(false, DisplayName = "Without have-table")]
 	public async Task Populate(bool useHaveTable)
 	{
-		ManagedWorkspace ws = await GetManagedWorkspace();
+		ManagedWorkspace ws = await GetManagedWorkspace(useHaveTable);
 
 		List<PopulateRequest> populateRequests = new ()
 		{
 			new PopulateRequest(PerforceConnection, StreamName, new List<string>())
 		};
-		await ws.PopulateAsync(populateRequests, false, useHaveTable, CancellationToken.None);
+		await ws.PopulateAsync(populateRequests, false, CancellationToken.None);
 		Stream.LatestChangelist.AssertDepotFiles(SyncDir);
 		await Stream.LatestChangelist.AssertHaveTableAsync(PerforceConnection, useHaveTable);
 	}
@@ -147,8 +129,8 @@ public class ManagedWorkspaceTest : BasePerforceFixtureTest
 	[DataRow(false, DisplayName = "Without have-table")]
 	public async Task Unshelve(bool useHaveTable)
 	{
-		ManagedWorkspace ws = await GetManagedWorkspace();
-		await SyncAsync(ws, 7, useHaveTable);
+		ManagedWorkspace ws = await GetManagedWorkspace(useHaveTable);
+		await SyncAsync(ws, 7);
 		await ws.UnshelveAsync(PerforceConnection, 8, CancellationToken.None);
 		Stream.GetChangelist(8).AssertDepotFiles(SyncDir);
 		
@@ -156,16 +138,16 @@ public class ManagedWorkspaceTest : BasePerforceFixtureTest
 		await Stream.GetChangelist(7).AssertHaveTableAsync(PerforceConnection, useHaveTable);
 	}
 
-	private async Task<ManagedWorkspace> GetManagedWorkspace()
+	private async Task<ManagedWorkspace> GetManagedWorkspace(bool useHaveTable)
 	{
-		ManagedWorkspace ws = await ManagedWorkspace.CreateAsync(Environment.MachineName, TempDir, _mwLogger, CancellationToken.None);
+		ManagedWorkspace ws = await ManagedWorkspace.CreateAsync(Environment.MachineName, TempDir, useHaveTable, _mwLogger, CancellationToken.None);
 		await ws.SetupAsync(PerforceConnection, StreamName, CancellationToken.None);
 		return ws;
 	}
 	
-	private async Task SyncAsync(ManagedWorkspace managedWorkspace, int changeNumber, bool useHaveTable = true, FileReference? cacheFile = null, bool removeUntracked = true)
+	private async Task SyncAsync(ManagedWorkspace managedWorkspace, int changeNumber, FileReference? cacheFile = null, bool removeUntracked = true)
 	{
-		await managedWorkspace.SyncAsync(PerforceConnection, StreamName, changeNumber, Array.Empty<string>(), removeUntracked, false, useHaveTable, cacheFile, CancellationToken.None);
+		await managedWorkspace.SyncAsync(PerforceConnection, StreamName, changeNumber, Array.Empty<string>(), removeUntracked, false, cacheFile, CancellationToken.None);
 	}
 	
 	private async Task AssertHaveTableFileCount(int expected)
diff --git a/Engine/Source/Programs/Shared/EpicGames.Perforce.Managed/ManagedWorkspace.cs b/Engine/Source/Programs/Shared/EpicGames.Perforce.Managed/ManagedWorkspace.cs
--- a/Engine/Source/Programs/Shared/EpicGames.Perforce.Managed/ManagedWorkspace.cs
+++ b/Engine/Source/Programs/Shared/EpicGames.Perforce.Managed/ManagedWorkspace.cs
@@ -135,11 +135,6 @@ namespace EpicGames.Perforce.Managed
 		/// </summary>
 		const string DataFileName = "Repository.dat";
 		
-		/// <summary>
-		/// Suffix added to client names that doesn't update the have table (syncing with -p option)
-		/// </summary>
-		const string NoHaveTableSuffix = "-NoHaveTable";
-
 		/// <summary>
 		/// Name of the host
 		/// </summary>
@@ -160,6 +155,15 @@ namespace EpicGames.Perforce.Managed
 		/// </summary>
 		readonly ILogger _logger;
 
+		/// <summary>
+		/// Use the client's have table when syncing.
+		/// 
+		/// When set to false, updates to the have table will be prevented through use of "sync -p".
+		/// Actual files to sync will gathered through "fstat".
+		/// This puts less strain on the Perforce server and can improve sync performance.
+		/// </summary>
+		readonly bool _useHaveTable;
+
 		/// <summary>
 		/// The root directory for the stash
 		/// </summary>
@@ -201,12 +205,14 @@ namespace EpicGames.Perforce.Managed
 		/// <param name="hostName">Name of the current host</param>
 		/// <param name="nextSequenceNumber">The next sequence number for operations</param>
 		/// <param name="baseDir">The root directory for the stash</param>
+		/// <param name="useHaveTable">Use the client's have table when syncing</param>
 		/// <param name="logger">The log output device</param>
-		private ManagedWorkspace(string hostName, uint nextSequenceNumber, DirectoryReference baseDir, ILogger logger)
+		private ManagedWorkspace(string hostName, uint nextSequenceNumber, DirectoryReference baseDir, bool useHaveTable, ILogger logger)
 		{
 			// Save the Perforce settings
 			_hostName = hostName;
 			_nextSequenceNumber = nextSequenceNumber;
+			_useHaveTable = useHaveTable;
 			_logger = logger;
 
 			// Get all the directories
@@ -229,16 +235,17 @@ namespace EpicGames.Perforce.Managed
 		/// <param name="hostName">Name of the current machine. Will be automatically detected from the host settings if not present.</param>
 		/// <param name="baseDir">The base directory for the repository</param>
 		/// <param name="overwrite">Whether to allow overwriting a repository that's not up to date</param>
+		/// <param name="useHaveTable">Use the client's have table when syncing</param>
 		/// <param name="logger">The logging interface</param>
 		/// <param name="cancellationToken">Cancellation token for this operation</param>
 		/// <returns></returns>
-		public static async Task<ManagedWorkspace> LoadOrCreateAsync(string hostName, DirectoryReference baseDir, bool overwrite, ILogger logger, CancellationToken cancellationToken)
+		public static async Task<ManagedWorkspace> LoadOrCreateAsync(string hostName, DirectoryReference baseDir, bool overwrite, bool useHaveTable, ILogger logger, CancellationToken cancellationToken)
 		{
 			if (Exists(baseDir))
 			{
 				try
 				{
-					return await LoadAsync(hostName, baseDir, logger, cancellationToken);
+					return await LoadAsync(hostName, baseDir, useHaveTable, logger, cancellationToken);
 				}
 				catch (Exception ex)
 				{
@@ -253,7 +260,7 @@ namespace EpicGames.Perforce.Managed
 				}
 			}
 
-			return await CreateAsync(hostName, baseDir, logger, cancellationToken);
+			return await CreateAsync(hostName, baseDir, useHaveTable, logger, cancellationToken);
 		}
 		/*
 				public static PerforceConnection GetPerforceConnection(PerforceConnection Perforce)
@@ -286,10 +293,11 @@ namespace EpicGames.Perforce.Managed
 		/// </summary>
 		/// <param name="hostName">Name of the current machine.</param>
 		/// <param name="baseDir">The base directory for the repository</param>
+		/// <param name="useHaveTable">Use the client's have table when syncing</param>
 		/// <param name="logger">The log output device</param>
 		/// <param name="cancellationToken">Cancellation token for this operation</param>
 		/// <returns>New repository instance</returns>
-		public static async Task<ManagedWorkspace> CreateAsync(string hostName, DirectoryReference baseDir, ILogger logger, CancellationToken cancellationToken)
+		public static async Task<ManagedWorkspace> CreateAsync(string hostName, DirectoryReference baseDir, bool useHaveTable, ILogger logger, CancellationToken cancellationToken)
 		{
 			logger.LogInformation("Creating repository at {Location}...", baseDir);
 
@@ -297,7 +305,7 @@ namespace EpicGames.Perforce.Managed
 			DirectoryReference.CreateDirectory(baseDir);
 			FileUtils.ForceDeleteDirectoryContents(baseDir);
 
-			ManagedWorkspace repo = new ManagedWorkspace(hostName, 1, baseDir, logger);
+			ManagedWorkspace repo = new ManagedWorkspace(hostName, 1, baseDir, useHaveTable, logger);
 			await repo.SaveAsync(TransactionState.Clean, cancellationToken);
 			repo.CreateCacheHierarchy();
 
@@ -337,9 +345,10 @@ namespace EpicGames.Perforce.Managed
 		/// </summary>
 		/// <param name="hostName">Name of the current host. Will be obtained from a 'p4 info' call if not specified</param>
 		/// <param name="baseDir">The base directory for the repository</param>
+		/// <param name="useHaveTable">Use the client's have table when syncing</param> 
 		/// <param name="logger">The log output device</param>
 		/// <param name="cancellationToken">Cancellation token for this command</param>
-		public static async Task<ManagedWorkspace> LoadAsync(string hostName, DirectoryReference baseDir, ILogger logger, CancellationToken cancellationToken)
+		public static async Task<ManagedWorkspace> LoadAsync(string hostName, DirectoryReference baseDir, bool useHaveTable, ILogger logger, CancellationToken cancellationToken)
 		{
 			if (!Exists(baseDir))
 			{
@@ -361,7 +370,7 @@ namespace EpicGames.Perforce.Managed
 			bool requiresRepair = reader.ReadBoolean();
 			uint nextSequenceNumber = reader.ReadUInt32();
 
-			ManagedWorkspace repo = new ManagedWorkspace(hostName, nextSequenceNumber, baseDir, logger);
+			ManagedWorkspace repo = new (hostName, nextSequenceNumber, baseDir, useHaveTable, logger);
 			repo._requiresRepair = requiresRepair;
 
 			int numTrackedFiles = reader.ReadInt32();
@@ -783,10 +792,9 @@ namespace EpicGames.Perforce.Managed
 		/// <param name="view">View of the workspace</param>
 		/// <param name="removeUntracked">Whether to remove untracked files from the workspace</param>
 		/// <param name="fakeSync">Whether to simulate the syncing operation rather than actually getting files from the server</param>
-		/// <param name="useHaveTable">Whether to update have table on the server during sync</param>
 		/// <param name="cacheFile">If set, uses the given file to cache the contents of the workspace. This can improve sync times when multiple machines sync the same workspace.</param>
 		/// <param name="cancellationToken">Cancellation token</param>
-		public async Task SyncAsync(IPerforceConnection perforce, string streamName, int changeNumber, IReadOnlyList<string> view, bool removeUntracked, bool fakeSync, bool useHaveTable, FileReference? cacheFile, CancellationToken cancellationToken)
+		public async Task SyncAsync(IPerforceConnection perforce, string streamName, int changeNumber, IReadOnlyList<string> view, bool removeUntracked, bool fakeSync, FileReference? cacheFile, CancellationToken cancellationToken)
 		{
 			Stopwatch timer = Stopwatch.StartNew();
 			if (changeNumber == -1)
@@ -798,7 +806,7 @@ namespace EpicGames.Perforce.Managed
 				_logger.LogInformation("Syncing to {StreamName} at CL {CL}", streamName, changeNumber);
 			}
 
-			using (_logger.WithProperty("useHaveTable", useHaveTable).BeginScope())
+			using (_logger.WithProperty("useHaveTable", _useHaveTable).BeginScope())
 			using (_logger.BeginIndentScope("  "))
 			{
 				// Update the client to the current stream
@@ -810,7 +818,7 @@ namespace EpicGames.Perforce.Managed
 					changeNumber = await GetLatestClientChangeAsync(perforce, cancellationToken);
 				}
 
-				if (useHaveTable)
+				if (_useHaveTable)
 				{
 					// Revert any open files
 					await RevertInternalAsync(perforce, cancellationToken);
@@ -829,7 +837,7 @@ namespace EpicGames.Perforce.Managed
 				StreamSnapshot? contents;
 				if (cacheFile == null)
 				{
-					if (useHaveTable)
+					if (_useHaveTable)
 					{
 						contents = await FindClientContentsAsync(perforce, changeNumber, cancellationToken);						
 					}
@@ -843,13 +851,13 @@ namespace EpicGames.Perforce.Managed
 					contents = await TryLoadClientContentsAsync(cacheFile, streamName, cancellationToken);
 					if (contents == null)
 					{
-						contents = await FindAndSaveClientContentsAsync(perforce, streamName, changeNumber, cacheFile, useHaveTable, cancellationToken);
+						contents = await FindAndSaveClientContentsAsync(perforce, streamName, changeNumber, cacheFile, cancellationToken);
 					}
 				}
 
 				// Sync all the appropriate files
 				await RemoveFilesFromWorkspaceAsync(contents, cancellationToken);
-				await AddFilesToWorkspaceAsync(perforce, contents, fakeSync, useHaveTable, cancellationToken);
+				await AddFilesToWorkspaceAsync(perforce, contents, fakeSync, cancellationToken);
 			}
 
 			_logger.LogInformation("Completed in {ElapsedTime}s", $"{timer.Elapsed.TotalSeconds:0.0}");
@@ -955,20 +963,20 @@ namespace EpicGames.Perforce.Managed
 		/// <summary>
 		/// Populates the cache with the head revision of the given streams.
 		/// </summary>
-		public async Task PopulateAsync(List<PopulateRequest> requests, bool fakeSync, bool useHaveTable, CancellationToken cancellationToken)
+		public async Task PopulateAsync(List<PopulateRequest> requests, bool fakeSync, CancellationToken cancellationToken)
 		{
 			_logger.LogInformation("Populating with {NumStreams} streams", requests.Count);
 			using (_logger.BeginIndentScope("  "))
 			{
-				Tuple<int, StreamSnapshot>[] streamState = await PopulateCleanAsync(requests, useHaveTable, cancellationToken);
-				await PopulateSyncAsync(requests, streamState, fakeSync, useHaveTable, cancellationToken);
+				Tuple<int, StreamSnapshot>[] streamState = await PopulateCleanAsync(requests, cancellationToken);
+				await PopulateSyncAsync(requests, streamState, fakeSync, cancellationToken);
 			}
 		}
 
 		/// <summary>
 		/// Perform the clean part of a populate command
 		/// </summary>
-		public async Task<Tuple<int, StreamSnapshot>[]> PopulateCleanAsync(List<PopulateRequest> requests, bool useHaveTable, CancellationToken cancellationToken)
+		public async Task<Tuple<int, StreamSnapshot>[]> PopulateCleanAsync(List<PopulateRequest> requests, CancellationToken cancellationToken)
 		{
 			// Revert all changes in each of the unique clients
 			foreach (PopulateRequest request in requests)
@@ -1001,7 +1009,7 @@ namespace EpicGames.Perforce.Managed
 					int changeNumber = await GetLatestClientChangeAsync(request.PerforceClient, cancellationToken);
 					_logger.LogInformation("Latest change is CL {CL}", changeNumber);
 
-					if (useHaveTable)
+					if (_useHaveTable)
 					{
 						await UpdateClientHaveTableAsync(request.PerforceClient, changeNumber, request.View, cancellationToken);
 
@@ -1071,7 +1079,7 @@ namespace EpicGames.Perforce.Managed
 		/// <summary>
 		/// Perform the sync part of a populate command
 		/// </summary>
-		public async Task PopulateSyncAsync(List<PopulateRequest> requests, Tuple<int, StreamSnapshot>[] streamState, bool fakeSync, bool useHaveTable, CancellationToken cancellationToken)
+		public async Task PopulateSyncAsync(List<PopulateRequest> requests, Tuple<int, StreamSnapshot>[] streamState, bool fakeSync, CancellationToken cancellationToken)
 		{
 			// Sync all the new files
 			for (int idx = 0; idx < requests.Count; idx++)
@@ -1082,7 +1090,7 @@ namespace EpicGames.Perforce.Managed
 
 				using (_logger.BeginIndentScope("  "))
 				{
-					if (useHaveTable)
+					if (_useHaveTable)
 					{
 						await DeleteClientAsync(request.PerforceClient, cancellationToken);
 						await UpdateClientAsync(request.PerforceClient, streamName, cancellationToken);
@@ -1093,7 +1101,7 @@ namespace EpicGames.Perforce.Managed
 
 					StreamSnapshot contents = streamState[idx].Item2;
 					await RemoveFilesFromWorkspaceAsync(contents, cancellationToken);
-					await AddFilesToWorkspaceAsync(request.PerforceClient, contents, fakeSync, useHaveTable, cancellationToken);
+					await AddFilesToWorkspaceAsync(request.PerforceClient, contents, fakeSync, cancellationToken);
 				}
 			}
 
@@ -1128,22 +1136,12 @@ namespace EpicGames.Perforce.Managed
 				_createdClients.Remove(clientName);
 			}
 
-			await DeleteAsync(GetClientName(perforceClient, true));
-			await DeleteAsync(GetClientName(perforceClient, false));
+			await DeleteAsync(perforceClient.Settings.ClientName!);
 		}
 
-		private static string GetClientName(IPerforceConnection perforceClient, bool useHaveTable)
-		{
-			if (!useHaveTable)
-			{
-				return perforceClient.Settings.ClientName + NoHaveTableSuffix;
-			}
-			return perforceClient.Settings.ClientName!;
-		}
-		
-		private async Task<ClientRecord> GetOrCreateClientAsync(IPerforceConnection perforceClient, string streamName, bool useHaveTable, CancellationToken cancellationToken)
+		private async Task<ClientRecord> GetOrCreateClientAsync(IPerforceConnection perforceClient, string streamName, CancellationToken cancellationToken)
 		{
-			string clientName = GetClientName(perforceClient, useHaveTable);
+			string clientName = perforceClient.Settings.ClientName!;//GetClientName(perforceClient, useHaveTable);
 			if (_createdClients.TryGetValue(clientName, out ClientRecord? client) && client.Stream == streamName)
 			{
 				return client;
@@ -1160,7 +1158,7 @@ namespace EpicGames.Perforce.Managed
 				
 				// Partitioned and read-only types store their have table separately on the server, compared to normal (writeable) clients
 				// Clients that sync without updating the have table cannot submit so they're marked as read-only. 
-				client.Type = useHaveTable ? "partitioned" : "readonly";
+				client.Type = _useHaveTable ? "partitioned" : "readonly";
 
 				using IPerforceConnection perforce = await perforceClient.WithoutClientAsync();
 
@@ -1186,9 +1184,7 @@ namespace EpicGames.Perforce.Managed
 		/// <param name="cancellationToken">Cancellation token</param>
 		private async Task UpdateClientAsync(IPerforceConnection perforceClient, string streamName, CancellationToken cancellationToken)
 		{
-			// Ensure a client exists both with and without have table use, using the correct stream
-			await GetOrCreateClientAsync(perforceClient, streamName, true, cancellationToken);
-			await GetOrCreateClientAsync(perforceClient, streamName, false, cancellationToken);
+			await GetOrCreateClientAsync(perforceClient, streamName, cancellationToken);
 
 			// Update the config file with the name of the client
 			FileReference configFile = FileReference.Combine(_baseDir, "p4.ini");
@@ -1637,12 +1633,11 @@ namespace EpicGames.Perforce.Managed
 		/// <param name="basePath">Base path for the stream</param>
 		/// <param name="changeNumber">The change number being synced. This must be specified in order to get the digest at the correct revision.</param>
 		/// <param name="cacheFile">Location of the file to save the cached contents</param>
-		/// <param name="useHaveTable">Whether to update the have table during sync</param>
 		/// <param name="cancellationToken">Cancellation token</param>
 		/// <returns>Contents of the workspace</returns>
-		private async Task<StreamSnapshotFromMemory> FindAndSaveClientContentsAsync(IPerforceConnection perforceClient, Utf8String basePath, int changeNumber, FileReference cacheFile, bool useHaveTable, CancellationToken cancellationToken)
+		private async Task<StreamSnapshotFromMemory> FindAndSaveClientContentsAsync(IPerforceConnection perforceClient, Utf8String basePath, int changeNumber, FileReference cacheFile, CancellationToken cancellationToken)
 		{
-			StreamSnapshotFromMemory contents = useHaveTable
+			StreamSnapshotFromMemory contents = _useHaveTable
 				? await FindClientContentsAsync(perforceClient, changeNumber, cancellationToken)
 				: await FindClientContentsWithoutHaveTableAsync(perforceClient, basePath.ToString(), changeNumber, cancellationToken);
 
@@ -1804,9 +1799,8 @@ namespace EpicGames.Perforce.Managed
 		/// <param name="client">The client connection</param>
 		/// <param name="stream">Contents of the stream</param>
 		/// <param name="fakeSync">Whether to simulate the sync operation, rather than actually syncing files</param>
-		/// <param name="useHaveTable">Whether to update the have table during sync</param>
 		/// <param name="cancellationToken">Cancellation token</param>
-		private async Task AddFilesToWorkspaceAsync(IPerforceConnection client, StreamSnapshot stream, bool fakeSync, bool useHaveTable, CancellationToken cancellationToken)
+		private async Task AddFilesToWorkspaceAsync(IPerforceConnection client, StreamSnapshot stream, bool fakeSync, CancellationToken cancellationToken)
 		{
 			// Make sure the repair flag is reset
 			await RunOptionalRepairAsync(cancellationToken);
@@ -1914,7 +1908,7 @@ namespace EpicGames.Perforce.Managed
 							{
 								(int batchBeginIdx, int batchEndIdx) = batches[nextBatchIdx];
 
-								Task task = Task.Run(() => SyncBatch(client, filesToSync, batchBeginIdx, batchEndIdx, fakeSync, useHaveTable, cancellationToken), cancellationToken);
+								Task task = Task.Run(() => SyncBatch(client, filesToSync, batchBeginIdx, batchEndIdx, fakeSync, cancellationToken), cancellationToken);
 								tasks[task] = nextBatchIdx++;
 							}
 
@@ -1966,10 +1960,9 @@ namespace EpicGames.Perforce.Managed
 		/// <param name="beginIdx">First file to sync</param>
 		/// <param name="endIdx">Index of the last file to sync (exclusive)</param>
 		/// <param name="fakeSync">Whether to fake a sync</param>
-		/// <param name="useHaveTable">Whether to update the have table during sync</param>
 		/// <param name="cancellationToken">Cancellation token for the request</param>
 		/// <returns>Async task</returns>
-		async Task SyncBatch(IPerforceConnection client, WorkspaceFileToSync[] filesToSync, int beginIdx, int endIdx, bool fakeSync, bool useHaveTable, CancellationToken cancellationToken)
+		async Task SyncBatch(IPerforceConnection client, WorkspaceFileToSync[] filesToSync, int beginIdx, int endIdx, bool fakeSync, CancellationToken cancellationToken)
 		{
 			if (fakeSync)
 			{
@@ -1991,7 +1984,7 @@ namespace EpicGames.Perforce.Managed
 					}
 				}
 
-				if (useHaveTable)
+				if (_useHaveTable)
 				{
 					using PerforceConnection clientWithFileList = new (client.Settings, client.Logger);
 					clientWithFileList.GlobalOptions.Add($"-x\"{syncFileName}\"");
@@ -2001,7 +1994,7 @@ namespace EpicGames.Perforce.Managed
 				{
 					// Ensure a client with an empty have table is used to not interfere with the DoNotUpdateHaveList option.
 					using PerforceConnection clientWithFileList = new (client.Settings, client.Logger);
-					clientWithFileList.ClientName = GetClientName(client, false);
+					clientWithFileList.ClientName = client.Settings.ClientName!;
 					clientWithFileList.GlobalOptions.Add($"-x\"{syncFileName}\"");
 					await clientWithFileList.SyncAsync(SyncOptions.DoNotUpdateHaveList | SyncOptions.FullDepotSyntax, -1, Array.Empty<string>(), cancellationToken).ToListAsync(cancellationToken);
 				}
