Horde: Fix missing hashes in serialized file references.

diff --git a/Engine/Source/Programs/Shared/EpicGames.Horde/Storage/Nodes/DirectoryNode.cs b/Engine/Source/Programs/Shared/EpicGames.Horde/Storage/Nodes/DirectoryNode.cs
--- a/Engine/Source/Programs/Shared/EpicGames.Horde/Storage/Nodes/DirectoryNode.cs
+++ b/Engine/Source/Programs/Shared/EpicGames.Horde/Storage/Nodes/DirectoryNode.cs
@@ -67,7 +67,12 @@ namespace EpicGames.Horde.Storage.Nodes
 		/// <summary>
 		/// Hash of the target node
 		/// </summary>
-		public IoHash Hash { get; private set; }
+		public IoHash Hash => (Target == null) ? _cachedHash : Target.Hash;
+
+		/// <summary>
+		/// Cached hash of the target node
+		/// </summary>
+		IoHash _cachedHash;
 
 		/// <summary>
 		/// Cached length of this node
@@ -94,6 +99,7 @@ namespace EpicGames.Horde.Storage.Nodes
 			Name = reader.ReadUtf8String();
 			Flags = (FileEntryFlags)reader.ReadUnsignedVarInt();
 
+			_cachedHash = reader.ReadIoHash();
 			_cachedLength = (long)reader.ReadUnsignedVarInt();
 		}
 
@@ -107,6 +113,7 @@ namespace EpicGames.Horde.Storage.Nodes
 
 			writer.WriteUtf8String(Name);
 			writer.WriteUnsignedVarInt((ulong)Flags);
+			writer.WriteIoHash(Hash);
 			writer.WriteUnsignedVarInt((ulong)Length);
 		}
 
@@ -140,8 +147,8 @@ namespace EpicGames.Horde.Storage.Nodes
 		{
 			base.OnCollapse();
 
+			_cachedHash = Target!.Hash;
 			_cachedLength = Target!.Length;
-			Hash = Target!.Hash;
 		}
 
 		/// <inheritdoc/>
