Fixed issue where UHT didn't properly detect circular structure definitions via TArray.

diff --git a/Engine/Source/Programs/Shared/EpicGames.UHT/Types/Properties/UhtArrayProperty.cs b/Engine/Source/Programs/Shared/EpicGames.UHT/Types/Properties/UhtArrayProperty.cs
--- a/Engine/Source/Programs/Shared/EpicGames.UHT/Types/Properties/UhtArrayProperty.cs
+++ b/Engine/Source/Programs/Shared/EpicGames.UHT/Types/Properties/UhtArrayProperty.cs
@@ -206,6 +206,14 @@ namespace EpicGames.UHT.Types
 					this.LogError("Replicated arrays with MemoryImageAllocators are not yet supported");
 				}
 			}
+
+			if (ValueProperty is UhtStructProperty structProperty)
+			{
+				if (structProperty.ScriptStruct == outerStruct)
+				{
+					this.LogError($"'Struct' recursion via arrays is unsupported for properties.");
+				}
+			}
 		}
 
 		/// <inheritdoc/>
diff --git a/Engine/Source/Programs/Shared/EpicGames.UHT/Types/UhtClass.cs b/Engine/Source/Programs/Shared/EpicGames.UHT/Types/UhtClass.cs
--- a/Engine/Source/Programs/Shared/EpicGames.UHT/Types/UhtClass.cs
+++ b/Engine/Source/Programs/Shared/EpicGames.UHT/Types/UhtClass.cs
@@ -970,7 +970,7 @@ namespace EpicGames.UHT.Types
 		}
 
 		/// <inheritdoc/>
-		public override bool ScanForInstancedReferenced(bool deepScan)
+		protected override bool ScanForInstancedReferencedInternal(bool deepScan)
 		{
 			if (ClassFlags.HasAnyFlags(EClassFlags.HasInstancedReference))
 			{
@@ -982,7 +982,7 @@ namespace EpicGames.UHT.Types
 				return true;
 			}
 
-			return base.ScanForInstancedReferenced(deepScan);
+			return base.ScanForInstancedReferencedInternal(deepScan);
 		}
 
 		private void MergeCategories()
diff --git a/Engine/Source/Programs/Shared/EpicGames.UHT/Types/UhtScriptStruct.cs b/Engine/Source/Programs/Shared/EpicGames.UHT/Types/UhtScriptStruct.cs
--- a/Engine/Source/Programs/Shared/EpicGames.UHT/Types/UhtScriptStruct.cs
+++ b/Engine/Source/Programs/Shared/EpicGames.UHT/Types/UhtScriptStruct.cs
@@ -804,7 +804,7 @@ namespace EpicGames.UHT.Types
 		}
 
 		/// <inheritdoc/>
-		public override bool ScanForInstancedReferenced(bool deepScan)
+		protected override bool ScanForInstancedReferencedInternal(bool deepScan)
 		{
 			if (ScriptStructFlags.HasAnyFlags(EStructFlags.HasInstancedReference))
 			{
@@ -816,7 +816,7 @@ namespace EpicGames.UHT.Types
 				return true;
 			}
 
-			return base.ScanForInstancedReferenced(deepScan);
+			return base.ScanForInstancedReferencedInternal(deepScan);
 		}
 		#endregion
 
diff --git a/Engine/Source/Programs/Shared/EpicGames.UHT/Types/UhtStruct.cs b/Engine/Source/Programs/Shared/EpicGames.UHT/Types/UhtStruct.cs
--- a/Engine/Source/Programs/Shared/EpicGames.UHT/Types/UhtStruct.cs
+++ b/Engine/Source/Programs/Shared/EpicGames.UHT/Types/UhtStruct.cs
@@ -4,6 +4,7 @@ using System;
 using System.Collections.Generic;
 using System.Diagnostics.CodeAnalysis;
 using System.Text.Json.Serialization;
+using System.Threading;
 using EpicGames.Core;
 using EpicGames.UHT.Parsers;
 using EpicGames.UHT.Tables;
@@ -110,6 +111,14 @@ namespace EpicGames.UHT.Types
 		[JsonIgnore]
 		public UhtStruct? SuperStruct => Super;
 
+		/// <summary>
+		/// Stack used to test to see if we are recursing in a ScanForInstanceReferenced call
+		/// </summary>
+		private readonly ThreadLocal<List<UhtType>> _scanForInstanceReferencedStack = new ThreadLocal<List<UhtType>>(() =>
+		{
+			return new List<UhtType>();
+		});
+
 		/// <summary>
 		/// Construct a new instance
 		/// </summary>
@@ -166,7 +175,32 @@ namespace EpicGames.UHT.Types
 		/// </summary>
 		/// <param name="deepScan">If true, the ScanForInstancedReferenced method on the properties will also be called.</param>
 		/// <returns></returns>
-		public virtual bool ScanForInstancedReferenced(bool deepScan)
+		public bool ScanForInstancedReferenced(bool deepScan)
+		{
+			List<UhtType> scanForInstanceReferencedStack = _scanForInstanceReferencedStack.Value!;
+			if (scanForInstanceReferencedStack.Contains(this))
+			{
+				return false;
+			}
+
+			scanForInstanceReferencedStack.Add(this);
+			try
+			{
+				return ScanForInstancedReferencedInternal(deepScan);
+			}
+			finally 
+			{
+				scanForInstanceReferencedStack.RemoveAt(scanForInstanceReferencedStack.Count - 1);
+			}
+		}
+
+		/// <summary>
+		/// Check properties to see if any instances are referenced.
+		/// This method does NOT cache the result.
+		/// </summary>
+		/// <param name="deepScan">If true, the ScanForInstancedReferenced method on the properties will also be called.</param>
+		/// <returns></returns>
+		protected virtual bool ScanForInstancedReferencedInternal(bool deepScan)
 		{
 			foreach (UhtType type in Children)
 			{
diff --git a/Engine/Source/Programs/UnrealHeaderTool/Private/UnrealTypeDefinitionInfo.cpp b/Engine/Source/Programs/UnrealHeaderTool/Private/UnrealTypeDefinitionInfo.cpp
--- a/Engine/Source/Programs/UnrealHeaderTool/Private/UnrealTypeDefinitionInfo.cpp
+++ b/Engine/Source/Programs/UnrealHeaderTool/Private/UnrealTypeDefinitionInfo.cpp
@@ -1970,6 +1970,16 @@ void FUnrealScriptStructDefinitionInfo::PostParseFinalizeInternal(EPostParseFina
 				}
 			}
 		}
+
+		// check for internal struct recursion via arrays
+		for (TSharedRef<FUnrealPropertyDefinitionInfo> PropertyDef : GetProperties())
+		{
+			const FPropertyBase& PropertyBase = PropertyDef->GetPropertyBase();
+			if (PropertyBase.ArrayType == EArrayType::Dynamic && PropertyBase.Type == EPropertyType::CPT_Struct && PropertyBase.TypeDef == this)
+			{
+				Throwf(TEXT("'Struct' recursion via arrays is unsupported for properties."));
+			}
+		}
 		break;
 
 	case EPostParseFinalizePhase::Phase2:
