Removing prior change to clean up old build records.  While the change fixed the issue with removed projects, it caused large number of installed build issues.

diff --git a/Engine/Source/Programs/Shared/EpicGames.Build/System/CompileScriptModules.cs b/Engine/Source/Programs/Shared/EpicGames.Build/System/CompileScriptModules.cs
--- a/Engine/Source/Programs/Shared/EpicGames.Build/System/CompileScriptModules.cs
+++ b/Engine/Source/Programs/Shared/EpicGames.Build/System/CompileScriptModules.cs
@@ -339,9 +339,8 @@ namespace UnrealBuildBase
 			}
 			else
 			{
-				// When the engine is not installed, delete any build record .json file that is not valid.
-				// It is important that we do !Valid to delete both invalid (out of date) and unknown (where csproj no long exists).
-				foreach (CsProjBuildRecordEntry Entry in BuildRecords.Values.Where(x => x.Status != CsProjBuildRecordStatus.Valid))
+				// when the engine is not installed, delete any build record .json file that is not valid
+				foreach (CsProjBuildRecordEntry Entry in BuildRecords.Values.Where(x => x.Status == CsProjBuildRecordStatus.Invalid))
 				{
 					if (Entry.BuildRecordFile != null)
 					{
