EpicGames.Perforce: Support view filtering for syncs not using have-table

With added tests for ManagedWorkspace to exercise views.

diff --git a/Engine/Source/Programs/Shared/EpicGames.Perforce.Managed.Tests/ManagedWorkspaceTest.cs b/Engine/Source/Programs/Shared/EpicGames.Perforce.Managed.Tests/ManagedWorkspaceTest.cs
--- a/Engine/Source/Programs/Shared/EpicGames.Perforce.Managed.Tests/ManagedWorkspaceTest.cs
+++ b/Engine/Source/Programs/Shared/EpicGames.Perforce.Managed.Tests/ManagedWorkspaceTest.cs
@@ -111,19 +111,74 @@ public class ManagedWorkspaceTest : BasePerforceFixtureTest
 	[TestMethod]
 	[DataRow(true, DisplayName = "With have-table")]
 	[DataRow(false, DisplayName = "Without have-table")]
-	public async Task Populate(bool useHaveTable)
+	public async Task SyncWithViewExclusivePathFirst(bool useHaveTable)
 	{
 		ManagedWorkspace ws = await GetManagedWorkspace(useHaveTable);
+		ChangelistFixture cl = Stream.GetChangelist(6);
 
-		List<PopulateRequest> populateRequests = new ()
-		{
-			new PopulateRequest(PerforceConnection, StreamName, new List<string>())
-		};
+		List<string> view = new() { "-/Data/...", "-/shared.h", };
+		await ws.SyncAsync(PerforceConnection, StreamName, cl.Number, view, true, false, null, CancellationToken.None);
+
+		List<DepotFileFixture> filtered = cl.StreamFiles
+			.Where(x => !x.DepotFile.Contains("shared.h", StringComparison.Ordinal))
+			.Where(x => !x.DepotFile.Contains("Data/", StringComparison.Ordinal)).ToList();
+
+		ChangelistFixture clViewApplied = new (cl.Number, cl.Description, filtered, cl.IsShelved);
+		clViewApplied.AssertDepotFiles(SyncDir);
+		await clViewApplied.AssertHaveTableAsync(PerforceConnection, useHaveTable);
+	}
+	
+	[TestMethod]
+	[DataRow(true, DisplayName = "With have-table")]
+	[DataRow(false, DisplayName = "Without have-table")]
+	public async Task SyncWithViewInclusivePathFirst(bool useHaveTable)
+	{
+		ManagedWorkspace ws = await GetManagedWorkspace(useHaveTable);
+		ChangelistFixture cl = Stream.GetChangelist(6);
+
+		List<string> view = new() { "/...", "-/Data/..." };
+		await ws.SyncAsync(PerforceConnection, StreamName, cl.Number, view, true, false, null, CancellationToken.None);
+
+		List<DepotFileFixture> filtered = cl.StreamFiles
+			.Where(x => !x.DepotFile.Contains("Data/", StringComparison.Ordinal)).ToList();
+
+		ChangelistFixture clViewApplied = new (cl.Number, cl.Description, filtered, cl.IsShelved);
+		clViewApplied.AssertDepotFiles(SyncDir);
+		await clViewApplied.AssertHaveTableAsync(PerforceConnection, useHaveTable);
+	}
+	
+	[TestMethod]
+	[DataRow(true, DisplayName = "With have-table")]
+	[DataRow(false, DisplayName = "Without have-table")]
+	public async Task Populate(bool useHaveTable)
+	{
+		ManagedWorkspace ws = await GetManagedWorkspace(useHaveTable);
+		List<PopulateRequest> populateRequests = new () { new PopulateRequest(PerforceConnection, StreamName, new List<string>()) };
 		await ws.PopulateAsync(populateRequests, false, CancellationToken.None);
 		Stream.LatestChangelist.AssertDepotFiles(SyncDir);
 		await Stream.LatestChangelist.AssertHaveTableAsync(PerforceConnection, useHaveTable);
 	}
 	
+	[TestMethod]
+	[DataRow(true, DisplayName = "With have-table")]
+	[DataRow(false, DisplayName = "Without have-table")]
+	public async Task PopulateWithView(bool useHaveTable)
+	{
+		ManagedWorkspace ws = await GetManagedWorkspace(useHaveTable);
+		List<string> view = new () { "-/Data/...", "-/shared.h" };
+		ChangelistFixture cl = Stream.LatestChangelist;
+		List<DepotFileFixture> filtered = cl.StreamFiles
+			.Where(x => !x.DepotFile.Contains("shared.h", StringComparison.Ordinal))
+			.Where(x => !x.DepotFile.Contains("Data/", StringComparison.Ordinal)).ToList();
+
+		List<PopulateRequest> populateRequests = new () { new PopulateRequest(PerforceConnection, StreamName, view) };
+		await ws.PopulateAsync(populateRequests, false, CancellationToken.None);
+		
+		ChangelistFixture clViewApplied = new (cl.Number, cl.Description, filtered, cl.IsShelved);
+		clViewApplied.AssertDepotFiles(SyncDir);
+		await clViewApplied.AssertHaveTableAsync(PerforceConnection, useHaveTable);
+	}
+	
 	[TestMethod]
 	[DataRow(true, DisplayName = "With have-table")]
 	[DataRow(false, DisplayName = "Without have-table")]
diff --git a/Engine/Source/Programs/Shared/EpicGames.Perforce.Managed/ManagedWorkspace.cs b/Engine/Source/Programs/Shared/EpicGames.Perforce.Managed/ManagedWorkspace.cs
--- a/Engine/Source/Programs/Shared/EpicGames.Perforce.Managed/ManagedWorkspace.cs
+++ b/Engine/Source/Programs/Shared/EpicGames.Perforce.Managed/ManagedWorkspace.cs
@@ -843,7 +843,7 @@ namespace EpicGames.Perforce.Managed
 					}
 					else
 					{
-						contents = await FindClientContentsWithoutHaveTableAsync(perforce, streamName, changeNumber, cancellationToken);
+						contents = await FindClientContentsWithoutHaveTableAsync(perforce, streamName, view, changeNumber, cancellationToken);
 					}
 				}
 				else
@@ -851,7 +851,7 @@ namespace EpicGames.Perforce.Managed
 					contents = await TryLoadClientContentsAsync(cacheFile, streamName, cancellationToken);
 					if (contents == null)
 					{
-						contents = await FindAndSaveClientContentsAsync(perforce, streamName, changeNumber, cacheFile, cancellationToken);
+						contents = await FindAndSaveClientContentsAsync(perforce, streamName, view, changeNumber, cacheFile, cancellationToken);
 					}
 				}
 
@@ -1018,7 +1018,7 @@ namespace EpicGames.Perforce.Managed
 					}
 					else
 					{
-						StreamSnapshot contents = await FindClientContentsWithoutHaveTableAsync(request.PerforceClient, streamName, changeNumber, cancellationToken);
+						StreamSnapshot contents = await FindClientContentsWithoutHaveTableAsync(request.PerforceClient, streamName, request.View, changeNumber, cancellationToken);
 						streamState[idx] = Tuple.Create(changeNumber, contents);
 					}
 
@@ -1553,53 +1553,71 @@ namespace EpicGames.Perforce.Managed
 
 			public Utf8String Digest => Values[(int)Field.digest].GetString();
 		}
-		
+
 		/// <summary>
 		/// Get the contents of the client without using the have table
 		/// </summary>
 		/// <param name="perforceClient">The client connection</param>
 		/// <param name="streamName">Name of stream</param>
+		/// <param name="view">View of the workspace</param>
 		/// <param name="changeNumber">The change number being synced. This must be specified in order to get the digest at the correct revision.</param>
 		/// <param name="cancellationToken">Cancellation token</param>
-		public async Task<StreamSnapshotFromMemory> FindClientContentsWithoutHaveTableAsync(IPerforceConnection perforceClient, string streamName, int changeNumber, CancellationToken cancellationToken)
+		public async Task<StreamSnapshotFromMemory> FindClientContentsWithoutHaveTableAsync(IPerforceConnection perforceClient, string streamName, IReadOnlyList<string> view, int changeNumber, CancellationToken cancellationToken)
 		{
-			DepotStreamTreeBuilder builder = new ();
+			DepotStreamTreeBuilder builder = new();
 			
-			StreamRecord streamRecord = await perforceClient.GetStreamAsync(streamName, true, cancellationToken);
-			PerforceViewMap viewMap = PerforceViewMap.Parse(streamRecord.View);
-			
-			// Re-use a single class instance as there can be millions of records
-			FStatRecordWithoutHaveTable record = new ();
-			void HandleRecord(PerforceRecord rawRecord)
+			using (Trace("FetchMetadata"))
+			using (ILoggerProgress scope = _logger.BeginProgressScope("Fetching metadata (without have table)..."))
 			{
-				// Copy into the values array
-				rawRecord.CopyInto(FStatRecordWithoutHaveTable.Utf8FieldNames, record.Values);
-				if (record.Digest.IsEmpty) { return; }
-				
-				if (viewMap.TryMapFile(record.DepotFile.ToString(), StringComparison.Ordinal, out string clientFile))
-				{
-					Md5Hash md5Hash = Md5Hash.Parse(record.Digest);
-					FileContentId fileContentId = new (md5Hash, record.HeadType);
-					builder.AddFile(clientFile, new StreamFile(record.DepotFile, record.FileSize, fileContentId, record.HeadRev));					
-				}
-				else
+				Stopwatch timer = Stopwatch.StartNew();
+
+				StreamRecord streamRecord = await perforceClient.GetStreamAsync(streamName, true, cancellationToken);
+				PerforceViewMap viewMap = PerforceViewMap.Parse(streamRecord.View);
+
+				// Use Horde's additional filtering taking place after stream view mapping
+				PerforceViewFilter viewFilter = PerforceViewFilter.Parse(view);
+
+				// Re-use a single class instance as there can be millions of records
+				FStatRecordWithoutHaveTable record = new();
+
+				void HandleRecord(PerforceRecord rawRecord)
 				{
-					_logger.LogError("Failed to view map depot file {DepotFile}", record.DepotFile.ToString());
+					// Copy into the values array
+					rawRecord.CopyInto(FStatRecordWithoutHaveTable.Utf8FieldNames, record.Values);
+					if (record.Digest.IsEmpty) { return; }
+
+					if (viewMap.TryMapFile(record.DepotFile.ToString(), StringComparison.OrdinalIgnoreCase, out string clientFile))
+					{
+						if (!viewFilter.IncludeFile(clientFile, StringComparison.OrdinalIgnoreCase))
+						{
+							return;
+						}
+
+						Md5Hash md5Hash = Md5Hash.Parse(record.Digest);
+						FileContentId fileContentId = new(md5Hash, record.HeadType);
+						builder.AddFile(clientFile, new StreamFile(record.DepotFile, record.FileSize, fileContentId, record.HeadRev));
+					}
+					else
+					{
+						_logger.LogError("Failed to view map depot file {DepotFile}", record.DepotFile.ToString());
+					}
 				}
-			}
 
-			string fileSpec = $"//{perforceClient.Settings.ClientName}/...@{changeNumber}";
-			List<string> arguments = new ();
-			arguments.Add("-Ol"); // Output fileSize and digest field
-			arguments.Add("-Os"); // Shorten output by excluding client workspace data (for instance, the clientFile field).
-			arguments.Add("-F"); // Filter any files not existing at current revision (filter below)
-			arguments.Add("^headAction=delete&^headAction=move/delete&^headAction=purge");
-			arguments.Add("-T"); // Include only the fields listed below
-			arguments.Add(String.Join(",", FStatRecordWithoutHaveTable.FieldNames));
-			arguments.Add(fileSpec);
+				string fileSpec = $"//{perforceClient.Settings.ClientName}/...@{changeNumber}";
+				List<string> arguments = new();
+				arguments.Add("-Ol"); // Output fileSize and digest field
+				arguments.Add("-Os"); // Shorten output by excluding client workspace data (for instance, the clientFile field).
+				arguments.Add("-F"); // Filter any files not existing at current revision (filter below)
+				arguments.Add("^headAction=delete&^headAction=move/delete&^headAction=purge");
+				arguments.Add("-T"); // Include only the fields listed below
+				arguments.Add(String.Join(",", FStatRecordWithoutHaveTable.FieldNames));
+				arguments.Add(fileSpec);
+
+				await perforceClient.RecordCommandAsync("fstat", arguments, null, HandleRecord, cancellationToken);
+
+				scope.Progress = $"({timer.Elapsed.TotalSeconds:0.0}s)";
+			}
 
-			await perforceClient.RecordCommandAsync("fstat", arguments, null, HandleRecord, cancellationToken);
-			
 			return new StreamSnapshotFromMemory(builder);
 		}
 
@@ -1631,15 +1649,16 @@ namespace EpicGames.Perforce.Managed
 		/// </summary>
 		/// <param name="perforceClient">The client connection</param>
 		/// <param name="basePath">Base path for the stream</param>
+		/// <param name="view">View of the workspace</param>
 		/// <param name="changeNumber">The change number being synced. This must be specified in order to get the digest at the correct revision.</param>
 		/// <param name="cacheFile">Location of the file to save the cached contents</param>
 		/// <param name="cancellationToken">Cancellation token</param>
 		/// <returns>Contents of the workspace</returns>
-		private async Task<StreamSnapshotFromMemory> FindAndSaveClientContentsAsync(IPerforceConnection perforceClient, Utf8String basePath, int changeNumber, FileReference cacheFile, CancellationToken cancellationToken)
+		private async Task<StreamSnapshotFromMemory> FindAndSaveClientContentsAsync(IPerforceConnection perforceClient, Utf8String basePath, IReadOnlyList<string> view, int changeNumber, FileReference cacheFile, CancellationToken cancellationToken)
 		{
 			StreamSnapshotFromMemory contents = _useHaveTable
 				? await FindClientContentsAsync(perforceClient, changeNumber, cancellationToken)
-				: await FindClientContentsWithoutHaveTableAsync(perforceClient, basePath.ToString(), changeNumber, cancellationToken);
+				: await FindClientContentsWithoutHaveTableAsync(perforceClient, basePath.ToString(), view, changeNumber, cancellationToken);
 
 			using (Trace("WriteMetadata"))
 			using (ILoggerProgress scope = _logger.BeginProgressScope($"Saving metadata to {cacheFile}..."))
diff --git a/Engine/Source/Programs/Shared/EpicGames.Perforce.Tests/EpicGames.Perforce.Tests.csproj b/Engine/Source/Programs/Shared/EpicGames.Perforce.Tests/EpicGames.Perforce.Tests.csproj
new file mode 100644
--- /dev/null
+++ b/Engine/Source/Programs/Shared/EpicGames.Perforce.Tests/EpicGames.Perforce.Tests.csproj
@@ -0,0 +1,37 @@
+<Project Sdk="Microsoft.NET.Sdk">
+
+  <PropertyGroup>
+    <TargetFramework>net6.0</TargetFramework>
+    <Configuration Condition=" '$(Configuration)' == '' ">Development</Configuration>
+    <OutputType>Library</OutputType>
+    <Deterministic>true</Deterministic>
+    <RootNamespace>EpicGames.Perforce.Tests</RootNamespace>
+    <GenerateAssemblyCompanyAttribute>false</GenerateAssemblyCompanyAttribute>
+    <GenerateAssemblyInformationalVersionAttribute>false</GenerateAssemblyInformationalVersionAttribute>
+    <GenerateAssemblyProductAttribute>false</GenerateAssemblyProductAttribute>
+    <GenerateAssemblyVersionAttribute>false</GenerateAssemblyVersionAttribute>
+    <Configurations>Debug;Release;Development</Configurations>
+    <IsPackable>false</IsPackable>
+    <Nullable>enable</Nullable>
+  </PropertyGroup>
+  <ItemGroup>
+    <Compile Include="..\..\Shared\MetaData.cs">
+      <Link>Properties\MetaData.cs</Link>
+    </Compile>
+  </ItemGroup>
+
+  <ItemGroup>
+    <PackageReference Include="Microsoft.NET.Test.Sdk" Version="17.3.0" />
+    <PackageReference Include="MSTest.TestAdapter" Version="2.2.10" />
+    <PackageReference Include="MSTest.TestFramework" Version="2.2.10" />
+    <PackageReference Include="coverlet.collector" Version="3.1.2">
+      <PrivateAssets>all</PrivateAssets>
+      <IncludeAssets>runtime; build; native; contentfiles; analyzers; buildtransitive</IncludeAssets>
+    </PackageReference>
+  </ItemGroup>
+
+  <ItemGroup>
+    <ProjectReference Include="..\EpicGames.Perforce.Managed\EpicGames.Perforce.Managed.csproj" />
+  </ItemGroup>
+
+</Project>
diff --git a/Engine/Source/Programs/Shared/EpicGames.Perforce.Tests/PerforceViewFilterTest.cs b/Engine/Source/Programs/Shared/EpicGames.Perforce.Tests/PerforceViewFilterTest.cs
new file mode 100644
--- /dev/null
+++ b/Engine/Source/Programs/Shared/EpicGames.Perforce.Tests/PerforceViewFilterTest.cs
@@ -0,0 +1,105 @@
+// Copyright Epic Games, Inc. All Rights Reserved.
+
+using System;
+using Microsoft.VisualStudio.TestTools.UnitTesting;
+
+namespace EpicGames.Perforce.Tests;
+
+[TestClass]
+public class PerforceViewFilterTest
+{
+	[TestMethod]
+	public void FileWildcard()
+	{
+		AssertInclude(true, "main.h", "*");
+		AssertInclude(true, "/main.h", "*");
+		AssertInclude(false, "/Engine/main.h", "*");
+		AssertInclude(true, "/Engine/main.h", "-*");
+		AssertInclude(false, "/Engine/Foo/main.h", "*");
+		AssertInclude(true, "/Engine/main.h", "-*");
+		AssertInclude(true, "/setup.bat", "-/...", "/*");
+		AssertInclude(true, "/setup.bat", "-/...", "*");
+		AssertInclude(false, "/Engine/main.cpp", "-/...", "*");
+	}
+
+	[TestMethod]
+	public void DirectoryWildcard()
+	{
+		AssertInclude(true, "/Engine/main.h", "/Engine/...");
+		AssertInclude(true, "/Engine/Foo/foo.h", "/Engine/...");
+		AssertInclude(true, "/foo.h", "/...");
+		AssertInclude(true, "/foo.h", "...");
+		AssertInclude(true, "foo.h", "...h");
+		AssertInclude(true, "/Engine/foo.h", "/...");
+		AssertInclude(true, "/Engine/main.cpp", "-/...", "/Engine/...");
+		AssertInclude(false, "/Engine/readme.txt", "-...", "/Docs/...");
+		AssertInclude(false, "/Engine/readme.txt", "/Docs/...", "-...");
+		AssertInclude(false, "/Docs/readme.txt", "-/...", "/Engine/...");
+		AssertInclude(false, "/Docs/readme.txt", "-...", "/Engine/...");
+		AssertInclude(false, "/Engine/main.cpp", "/...", "-/Engine/...");
+		
+		AssertInclude(false, "/Engine/main.cpp", "-...");
+		AssertInclude(false, "Engine/main.cpp", "-...");
+		AssertInclude(true, "/main.cpp", "-...", "*");
+		AssertInclude(true, "main.cpp", "-...", "*");
+		AssertInclude(true, "Engine/main.cpp", "-...", "Engine/...");
+		AssertInclude(true, "/Engine/main.cpp", "-...", "/Engine/...");
+		AssertInclude(true, "/Engine/main.cpp", "-...", "Engine/...");
+		AssertInclude(true, "Engine/main.cpp", "-...", "/Engine/...");
+	}
+	
+	[TestMethod]
+	public void DirectoryWildcardFileTypes()
+	{
+		AssertInclude(true, "/Engine/main.h", "/Engine/...h");
+		AssertInclude(true, "/Engine/Foo/foo.h", "/Engine/...h");
+		AssertInclude(true, "/foo.h", "/...h");
+		AssertInclude(true, "/foo.h", "...h");
+		AssertInclude(true, "foo.h", "...h");
+		AssertInclude(true, "/Engine/foo.h", "/...h");
+		AssertInclude(false, "/Engine/foo.h", "/...c");
+		AssertInclude(true, "/Engine/readme.txt", "-...h", "-...cpp");
+		AssertInclude(false, "/Engine/readme.txt", "/Docs/...", "-...h", "-...cpp");
+	}
+	
+	[TestMethod]
+	public void Exact()
+	{
+		// Leading slash is handled no matter if path or filter has one
+		AssertInclude(true, "/main.cpp", "/main.cpp");
+		AssertInclude(true, "/main.cpp", "main.cpp");
+		AssertInclude(true, "main.cpp", "main.cpp");
+		AssertInclude(true, "main.cpp", "/main.cpp");
+		
+		AssertInclude(true, "/Engine/main.cpp", "/Engine/main.cpp");
+		AssertInclude(false, "/Engine/main.cpp", "/Engine/main.c");
+		AssertInclude(false, "/Engine/main.cpp", "-/Engine/main.cpp");
+		AssertInclude(true, "/Engine/main.cpp", "-/Engine/main.c", "/Engine/main.cpp");
+	}
+
+	[TestMethod]
+	public void EntryParse()
+	{
+		AssertEntry("/Engine/...", true, "...", "/Engine/", "");
+		AssertEntry("/Engine/....cpp", true, "...", "/Engine/", ".cpp");
+		AssertEntry("/Engine/*", true, "*", "/Engine/", "");
+		AssertEntry("/Engine/*.h", true, "*", "/Engine/", ".h");
+		AssertEntry("/Engine/main.cpp", true, "", "/Engine/main.cpp", "");
+		AssertEntry("-/Engine/...", false, "...", "/Engine/", "");
+		AssertEntry("-/*", false, "*", "/", "");
+	}
+
+	private static void AssertInclude(bool expectIncluded, string path, params string[] filters)
+	{
+		Assert.AreEqual(expectIncluded, PerforceViewFilter.Parse(filters).IncludeFile(path, StringComparison.OrdinalIgnoreCase));
+	}
+	
+	private static void AssertEntry(string text, bool include, string wildcard, string prefix, string suffix)
+	{
+		PerforceViewFilterEntry e = PerforceViewFilterEntry.Parse(text);
+		Assert.AreEqual(include, e.Include);
+		Assert.AreEqual(wildcard, e.Wildcard);
+		Assert.AreEqual(prefix, e.Prefix);
+		Assert.AreEqual(suffix, e.Suffix);
+	}
+}
\ No newline at end of file
diff --git a/Engine/Source/Programs/Shared/EpicGames.Perforce/PerforceViewFilter.cs b/Engine/Source/Programs/Shared/EpicGames.Perforce/PerforceViewFilter.cs
new file mode 100644
--- /dev/null
+++ b/Engine/Source/Programs/Shared/EpicGames.Perforce/PerforceViewFilter.cs
@@ -0,0 +1,204 @@
+// Copyright Epic Games, Inc. All Rights Reserved.
+
+using System;
+using System.Collections.Generic;
+using System.Linq;
+using System.Text;
+using System.Text.RegularExpressions;
+
+namespace EpicGames.Perforce
+{
+	/// <summary>
+	/// Filters a view of files
+	/// </summary>
+	public class PerforceViewFilter
+	{
+		/// <summary>
+		/// List of entries making up the view
+		/// </summary>
+		public List<PerforceViewFilterEntry> Entries { get; }
+
+		/// <summary>
+		/// Construct from an existing set of entries
+		/// </summary>
+		/// <param name="entries"></param>
+		public PerforceViewFilter(IEnumerable<PerforceViewFilterEntry> entries)
+		{
+			Entries = entries.ToList();
+			
+			// When first entry is excluding, start by including everything. Otherwise, default is to exclude.
+			// This is to emulate the behavior of ManagedWorkspace.UpdateClientHaveTableAsync()
+			if (Entries.Count == 0 || !Entries[0].Include)
+			{
+				Entries.Insert(0, new PerforceViewFilterEntry(true, "...", "", ""));
+			}
+		}
+
+		/// <summary>
+		/// Construct a view filter from a set of entries
+		/// </summary>
+		/// <param name="entries"></param>
+		public static PerforceViewFilter Parse(IEnumerable<string> entries)
+		{
+			return new PerforceViewFilter(entries.Select(x => PerforceViewFilterEntry.Parse(x)));
+		}
+
+		/// <summary>
+		/// Determines if a file is included in the view
+		/// </summary>
+		/// <param name="file">The file to test</param>
+		/// <param name="comparison">The comparison type</param>
+		/// <returns>True if the file is included in the view</returns>
+		public bool IncludeFile(string file, StringComparison comparison)
+		{
+			bool included = false;
+			foreach (PerforceViewFilterEntry entry in Entries)
+			{
+				if (entry.MatchFile(file, comparison))
+				{
+					included = entry.Include;
+				}
+			}
+			return included;
+		}
+	}
+
+	/// <summary>
+	/// Entry within a view filter
+	/// </summary>
+	public class PerforceViewFilterEntry
+	{
+		/// <summary>
+		/// Whether to include files matching this pattern
+		/// </summary>
+		public bool Include { get; }
+
+		/// <summary>
+		/// The wildcard string - either '*' or '...'
+		/// </summary>
+		public string Wildcard { get; }
+
+		/// <summary>
+		/// The part of pattern before the wildcard
+		/// </summary>
+		public string Prefix { get; }
+
+		/// <summary>
+		/// The part of pattern after the wildcard. Perforce does not permit a slash to be in this part of the file spec filter.
+		/// </summary>
+		public string Suffix { get; }
+
+		/// <summary>
+		/// Tests if the entry has a file wildcard ('*')
+		/// </summary>
+		/// <returns>True if the entry has a file wildcard</returns>
+		public bool IsFileWildcard() => Wildcard.Length == 1;
+
+		/// <summary>
+		/// Tests if the entry has a path wildcard ('...')
+		/// </summary>
+		/// <returns>True if the entry has a path wildcard</returns>
+		public bool IsPathWildcard() => Wildcard.Length == 3;
+
+		/// <summary>
+		/// Constructor
+		/// </summary>
+		/// <param name="include"></param>
+		/// <param name="pathExpression"></param>
+		public PerforceViewFilterEntry(bool include, string pathExpression)
+		{
+			Include = include;
+		
+			Match match = Regex.Match(pathExpression, @"^(.*?)(\*|\.\.\.|%%1)(.*)$");
+			if (match.Success)
+			{
+				Prefix = match.Groups[1].Value;
+				Suffix = match.Groups[3].Value;
+				Wildcard = match.Groups[2].Value;
+			}
+			else
+			{
+				Prefix = pathExpression;
+				Suffix = String.Empty;
+				Wildcard = String.Empty;
+			}
+			
+			if (!Prefix.StartsWith('/'))
+			{
+				Prefix = '/' + Prefix;
+			}
+		}
+
+		/// <summary>
+		/// Constructor
+		/// </summary>
+		/// <param name="include"></param>
+		/// <param name="wildcard"></param>
+		/// <param name="prefix"></param>
+		/// <param name="suffix"></param>
+		public PerforceViewFilterEntry(bool include, string wildcard, string prefix, string suffix)
+		{
+			Include = include;
+			Wildcard = wildcard;
+			Prefix = prefix;
+			Suffix = suffix;
+		}
+
+		/// <summary>
+		/// Parse a view filter entry from a string, as returned by spec documents
+		/// </summary>
+		/// <param name="entry"></param>
+		/// <returns></returns>
+		public static PerforceViewFilterEntry Parse(string entry)
+		{
+			Match match = Regex.Match(entry, @"^\s*(-?)\s*([^ ]+)\s*$");
+			if (!match.Success)
+			{
+				throw new PerforceException($"Unable to parse view filter entry: {entry}");
+			}
+			return new PerforceViewFilterEntry(match.Groups[1].Length == 0, match.Groups[2].Value);
+		}
+		
+		/// <summary>
+		/// Determine if a file matches the current entry
+		/// </summary>
+		/// <param name="path">Path to the file</param>
+		/// <param name="comparison">The comparison type</param>
+		/// <returns>True if the path matches the entry</returns>
+		public bool MatchFile(string path, StringComparison comparison)
+		{
+			if (!path.StartsWith('/'))
+			{
+				path = '/' + path;
+			}
+			
+			if (Wildcard.Length == 0)
+			{
+				// Expect an exact match if no wildcard is specified
+				return String.Equals(path, Prefix, comparison);
+			}
+			
+			if (!path.StartsWith(Prefix, comparison) || !path.EndsWith(Suffix, comparison))
+			{
+				return false;
+			}
+			if (IsFileWildcard() && path.AsSpan(Prefix.Length, path.Length - Suffix.Length - Prefix.Length).IndexOf('/') != -1)
+			{
+				return false;
+			}
+			return true;
+		}
+
+		/// <inheritdoc/>
+		public override string ToString()
+		{
+			StringBuilder builder = new ();
+			if (!Include)
+			{
+				builder.Append('-');
+			}
+			builder.Append($"{Prefix}{Wildcard}{Suffix}");
+			return builder.ToString();
+		}
+	}
+}
