EpicGames.Perforce: Do not use target-typed object creation so project can compile for older NET versions (eg. UGS).

diff --git a/Engine/Source/Programs/Shared/EpicGames.Perforce/PerforceViewFilter.cs b/Engine/Source/Programs/Shared/EpicGames.Perforce/PerforceViewFilter.cs
--- a/Engine/Source/Programs/Shared/EpicGames.Perforce/PerforceViewFilter.cs
+++ b/Engine/Source/Programs/Shared/EpicGames.Perforce/PerforceViewFilter.cs
@@ -192,7 +192,7 @@ namespace EpicGames.Perforce
 		/// <inheritdoc/>
 		public override string ToString()
 		{
-			StringBuilder builder = new ();
+			StringBuilder builder = new StringBuilder();
 			if (!Include)
 			{
 				builder.Append('-');
