Horde: Support a custom data field in file entry nodes; can be used for storing P4 MD5s.

diff --git a/Engine/Source/Programs/Shared/EpicGames.Horde/Storage/Nodes/DirectoryNode.cs b/Engine/Source/Programs/Shared/EpicGames.Horde/Storage/Nodes/DirectoryNode.cs
--- a/Engine/Source/Programs/Shared/EpicGames.Horde/Storage/Nodes/DirectoryNode.cs
+++ b/Engine/Source/Programs/Shared/EpicGames.Horde/Storage/Nodes/DirectoryNode.cs
@@ -39,9 +39,9 @@ namespace EpicGames.Horde.Storage.Nodes
 		Text = 16,
 
 		/// <summary>
-		/// The data for this entry is a Perforce depot path and revision rather than the actual file contents.
+		/// Used to indicate that custom data is included in the output. Used internally for serialization; not exposed to users.
 		/// </summary>
-		PerforceDepotPathAndRevision = 32,
+		HasCustomData = 32,
 	}
 
 	/// <summary>
@@ -69,6 +69,11 @@ namespace EpicGames.Horde.Storage.Nodes
 		/// </summary>
 		public IoHash Hash => (Target == null) ? _cachedHash : Target.Hash;
 
+		/// <summary>
+		/// Custom user data for this file entry
+		/// </summary>
+		public ReadOnlyMemory<byte> CustomData { get; } = ReadOnlyMemory<byte>.Empty;
+
 		/// <summary>
 		/// Cached hash of the target node
 		/// </summary>
@@ -101,6 +106,12 @@ namespace EpicGames.Horde.Storage.Nodes
 
 			_cachedHash = reader.ReadIoHash();
 			_cachedLength = (long)reader.ReadUnsignedVarInt();
+
+			if ((Flags & FileEntryFlags.HasCustomData) != 0)
+			{
+				CustomData = reader.ReadVariableLengthBytes();
+				Flags &= ~FileEntryFlags.HasCustomData;
+			}
 		}
 
 		/// <summary>
@@ -111,10 +122,17 @@ namespace EpicGames.Horde.Storage.Nodes
 		{
 			base.Serialize(writer);
 
+			FileEntryFlags flags = (CustomData.Length > 0) ? (Flags | FileEntryFlags.HasCustomData) : (Flags & ~FileEntryFlags.HasCustomData);
+
 			writer.WriteUtf8String(Name);
-			writer.WriteUnsignedVarInt((ulong)Flags);
+			writer.WriteUnsignedVarInt((ulong)flags);
 			writer.WriteIoHash(Hash);
 			writer.WriteUnsignedVarInt((ulong)Length);
+
+			if((flags & FileEntryFlags.HasCustomData) != 0)
+			{
+				writer.WriteVariableLengthBytes(CustomData.Span);
+			}
 		}
 
 		/// <summary>
