[UnrealHeaderTool]
* Fixed so return type of uhtfunctions are taken into account when looking at needed includes
* Added so functions that produce rpc code and has params that are of class type adds includes for class types. This is needed to be able to cast to UObject* in the generated code

diff --git a/Engine/Source/Programs/Shared/EpicGames.UHT/Exporters/CodeGen/UhtHeaderCodeGeneratorCppFile.cs b/Engine/Source/Programs/Shared/EpicGames.UHT/Exporters/CodeGen/UhtHeaderCodeGeneratorCppFile.cs
--- a/Engine/Source/Programs/Shared/EpicGames.UHT/Exporters/CodeGen/UhtHeaderCodeGeneratorCppFile.cs
+++ b/Engine/Source/Programs/Shared/EpicGames.UHT/Exporters/CodeGen/UhtHeaderCodeGeneratorCppFile.cs
@@ -66,16 +66,21 @@ namespace EpicGames.UHT.Exporters.CodeGen
 									addedCoreNetHeader = true;
 								}
 							}
+
+							bool RequireIncludeForClasses = IsRpcFunction(function) && ShouldExportFunction(function);
+
 							foreach (UhtProperty property in function.Properties)
 							{
-								AddIncludeForProperty(property, addedIncludes, includesToAdd);
+								AddIncludeForProperty(property, addedIncludes, includesToAdd, RequireIncludeForClasses);
 							}
+							if (function.ReturnProperty != null)
+								AddIncludeForProperty(function.ReturnProperty, addedIncludes, includesToAdd, RequireIncludeForClasses);
 						}
 
 						// Properties
 						foreach (UhtProperty property in structObj.Properties)
 						{
-							AddIncludeForProperty(property, addedIncludes, includesToAdd);
+							AddIncludeForProperty(property, addedIncludes, includesToAdd, false);
 						}
 					}
 
@@ -350,7 +355,7 @@ namespace EpicGames.UHT.Exporters.CodeGen
 			}
 		}
 
-		private void AddIncludeForType(UhtProperty uhtProperty, HashSet<UhtHeaderFile> addedIncludes, IList<string> includesToAdd)
+		private void AddIncludeForType(UhtProperty uhtProperty, HashSet<UhtHeaderFile> addedIncludes, IList<string> includesToAdd, bool requireIncludeForClasses)
 		{
 			if (uhtProperty is UhtStructProperty structProperty)
 			{
@@ -360,20 +365,28 @@ namespace EpicGames.UHT.Exporters.CodeGen
 					includesToAdd.Add(HeaderInfos[scriptStruct.HeaderFile.HeaderFileTypeIndex].IncludePath);
 				}
 			}
+			else if (requireIncludeForClasses && uhtProperty is UhtClassProperty classProperty)
+			{
+				UhtClass uhtClass = classProperty.Class;
+				if (!uhtClass.HeaderFile.IsNoExportTypes && addedIncludes.Add(uhtClass.HeaderFile))
+				{
+					includesToAdd.Add(HeaderInfos[uhtClass.HeaderFile.HeaderFileTypeIndex].IncludePath);
+				}
+			}
 		}
 
-		private void AddIncludeForProperty(UhtProperty property, HashSet<UhtHeaderFile> addedIncludes, IList<string> includesToAdd)
+		private void AddIncludeForProperty(UhtProperty property, HashSet<UhtHeaderFile> addedIncludes, IList<string> includesToAdd, bool requireIncludeForClasses)
 		{
-			AddIncludeForType(property, addedIncludes, includesToAdd);
+			AddIncludeForType(property, addedIncludes, includesToAdd, requireIncludeForClasses);
 
 			if (property is UhtContainerBaseProperty containerProperty)
 			{
-				AddIncludeForType(containerProperty.ValueProperty, addedIncludes, includesToAdd);
+				AddIncludeForType(containerProperty.ValueProperty, addedIncludes, includesToAdd, false);
 			}
 
 			if (property is UhtMapProperty mapProperty)
 			{
-				AddIncludeForType(mapProperty.KeyProperty, addedIncludes, includesToAdd);
+				AddIncludeForType(mapProperty.KeyProperty, addedIncludes, includesToAdd, false);
 			}
 		}
 
