Allow EpicGames.Perforce.PerforceConnection.TryUndoChangeAsync to return reponses that contain warning/error messages with the PerforceGenericCode.Empty code.

- P4 undo cannot really be run on an empty changelist as it can only run on submitted changelists. However if 'changeNumber' points to a workspace that does not contain the files in 'changeNumberToUndo' then perforce will return error messages about this with the Empty code.
- This will allow us to improve the error reporting in some of our tools which call p4 undo.

diff --git a/Engine/Source/Programs/Shared/EpicGames.Perforce/PerforceConnection.cs b/Engine/Source/Programs/Shared/EpicGames.Perforce/PerforceConnection.cs
--- a/Engine/Source/Programs/Shared/EpicGames.Perforce/PerforceConnection.cs
+++ b/Engine/Source/Programs/Shared/EpicGames.Perforce/PerforceConnection.cs
@@ -4045,9 +4045,7 @@ namespace EpicGames.Perforce
 			arguments.Add($"-c{changeNumber}");
 			arguments.Add($"//...@{changeNumberToUndo}");
 
-			PerforceResponseList<UndoRecord> records = await CommandAsync<UndoRecord>(connection, "undo", arguments, null, cancellationToken);
-			records.RemoveAll(x => x.Error != null && x.Error.Generic == PerforceGenericCode.Empty);
-			return records;
+			return await CommandAsync<UndoRecord>(connection, "undo", arguments, null, cancellationToken);
 		}
 
 		#endregion
