[UHT]
* Fixed so more generated ctors are put in the .gen.cpp file in order to be able to forward declare more

Details:
* Promoted GetConstructorType from UhtHeaderCodeGeneratorHFile.cs to UhtHeaderCodeGenerator.cs so it could be used from UhtHeaderCodeGeneratorCppFile.cs
* Added flag to UhtClass to know if it is using legacy macro for body generation or not (implemented by Tim, bug fix for interfaces by me)
* Changed so declaration of ctors stay in .generated.h while implementation is moved to .gen.cpp

Parts implemented by and discussed with [REDACTED_PERSON]

diff --git a/Engine/Source/Programs/Shared/EpicGames.UHT/Exporters/CodeGen/UhtHeaderCodeGenerator.cs b/Engine/Source/Programs/Shared/EpicGames.UHT/Exporters/CodeGen/UhtHeaderCodeGenerator.cs
--- a/Engine/Source/Programs/Shared/EpicGames.UHT/Exporters/CodeGen/UhtHeaderCodeGenerator.cs
+++ b/Engine/Source/Programs/Shared/EpicGames.UHT/Exporters/CodeGen/UhtHeaderCodeGenerator.cs
@@ -459,6 +459,66 @@ namespace EpicGames.UHT.Exporters.CodeGen
 			return true;
 		}
 
+		/// <summary>
+		/// Type of constructor on the class regardless of explicit or generated.
+		/// </summary>
+		protected enum ConstructorType
+		{
+			ObjectInitializer,
+			Default,
+			ForbiddenDefault,
+		}
+
+		/// <summary>
+		/// Return the type of constructor the class will have regardless of if one has been explicitly 
+		/// declared or will be generated.
+		/// </summary>
+		/// <param name="classObj">Class in question</param>
+		/// <returns>Constructor type</returns>
+		protected static ConstructorType GetConstructorType(UhtClass classObj)
+		{
+			if (!classObj.ClassExportFlags.HasAnyFlags(UhtClassExportFlags.HasConstructor))
+			{
+
+				// Assume super class has OI constructor, this may not always be true but we should always be able to check this.
+				// In any case, it will default to old behavior before we even checked this.
+				bool superClassObjectInitializerConstructorDeclared = true;
+				UhtClass? superClass = classObj.SuperClass;
+				if (superClass != null)
+				{
+					if (!superClass.HeaderFile.IsNoExportTypes)
+					{
+						superClassObjectInitializerConstructorDeclared = GetConstructorType(superClass) == ConstructorType.ObjectInitializer;
+					}
+				}
+
+				if (superClassObjectInitializerConstructorDeclared)
+				{
+					return ConstructorType.ObjectInitializer;
+				}
+				else
+				{
+					return ConstructorType.Default;
+				}
+			}
+			else
+			{
+				if (classObj.ClassExportFlags.HasAnyFlags(UhtClassExportFlags.HasObjectInitializerConstructor))
+				{
+					return ConstructorType.ObjectInitializer;
+				}
+				else if (classObj.ClassExportFlags.HasAnyFlags(UhtClassExportFlags.HasDefaultConstructor))
+				{
+					return ConstructorType.Default;
+				}
+				else
+				{
+					return ConstructorType.ForbiddenDefault;
+				}
+			}
+		}
+
+
 		protected static string GetDelegateFunctionExportName(UhtFunction function)
 		{
 			const string DelegatePrefix = "delegate";
diff --git a/Engine/Source/Programs/Shared/EpicGames.UHT/Exporters/CodeGen/UhtHeaderCodeGeneratorCppFile.cs b/Engine/Source/Programs/Shared/EpicGames.UHT/Exporters/CodeGen/UhtHeaderCodeGeneratorCppFile.cs
--- a/Engine/Source/Programs/Shared/EpicGames.UHT/Exporters/CodeGen/UhtHeaderCodeGeneratorCppFile.cs
+++ b/Engine/Source/Programs/Shared/EpicGames.UHT/Exporters/CodeGen/UhtHeaderCodeGeneratorCppFile.cs
@@ -1093,6 +1093,23 @@ namespace EpicGames.UHT.Exporters.CodeGen
 				builder.Append("\tDEFINE_VTABLE_PTR_HELPER_CTOR(").Append(classObj.SourceName).Append(");\r\n");
 			}
 
+			if (!classObj.ClassExportFlags.HasAnyFlags(UhtClassExportFlags.HasConstructor))
+			{
+				if (!classObj.ClassExportFlags.HasAnyFlags(UhtClassExportFlags.UsesGeneratedBodyLegacy))
+				{
+					switch (GetConstructorType(classObj))
+					{
+						case ConstructorType.ObjectInitializer:
+							builder.Append('\t').Append(classObj.SourceName).Append("::").Append(classObj.SourceName).Append("(const FObjectInitializer& ObjectInitializer) : Super(ObjectInitializer) {}\r\n");
+							break;
+
+						case ConstructorType.Default:
+							builder.Append('\t').Append(classObj.SourceName).Append("::").Append(classObj.SourceName).Append("() {}\r\n");
+							break;
+					}
+				}
+			}
+
 			if (!classObj.ClassExportFlags.HasAnyFlags(UhtClassExportFlags.HasDestructor))
 			{
 				builder.Append('\t').Append(classObj.SourceName).Append("::~").Append(classObj.SourceName).Append("() {}\r\n");
diff --git a/Engine/Source/Programs/Shared/EpicGames.UHT/Exporters/CodeGen/UhtHeaderCodeGeneratorHFile.cs b/Engine/Source/Programs/Shared/EpicGames.UHT/Exporters/CodeGen/UhtHeaderCodeGeneratorHFile.cs
--- a/Engine/Source/Programs/Shared/EpicGames.UHT/Exporters/CodeGen/UhtHeaderCodeGeneratorHFile.cs
+++ b/Engine/Source/Programs/Shared/EpicGames.UHT/Exporters/CodeGen/UhtHeaderCodeGeneratorHFile.cs
@@ -977,65 +977,6 @@ namespace EpicGames.UHT.Exporters.CodeGen
 			return builder;
 		}
 
-		/// <summary>
-		/// Type of constructor on the class regardless of explicit or generated.
-		/// </summary>
-		enum ConstructorType
-		{
-			ObjectInitializer,
-			Default,
-			ForbiddenDefault,
-		}
-
-		/// <summary>
-		/// Return the type of constructor the class will have regardless of if one has been explicitly 
-		/// declared or will be generated.
-		/// </summary>
-		/// <param name="classObj">Class in question</param>
-		/// <returns>Constructor type</returns>
-		private static ConstructorType GetConstructorType(UhtClass classObj)
-		{
-			if (!classObj.ClassExportFlags.HasAnyFlags(UhtClassExportFlags.HasConstructor))
-			{
-
-				// Assume super class has OI constructor, this may not always be true but we should always be able to check this.
-				// In any case, it will default to old behavior before we even checked this.
-				bool superClassObjectInitializerConstructorDeclared = true;
-				UhtClass? superClass = classObj.SuperClass;
-				if (superClass != null)
-				{
-					if (!superClass.HeaderFile.IsNoExportTypes)
-					{
-						superClassObjectInitializerConstructorDeclared = GetConstructorType(superClass) == ConstructorType.ObjectInitializer;
-					}
-				}
-
-				if (superClassObjectInitializerConstructorDeclared)
-				{
-					return ConstructorType.ObjectInitializer;
-				}
-				else
-				{
-					return ConstructorType.Default;
-				}
-			}
-			else
-			{
-				if (classObj.ClassExportFlags.HasAnyFlags(UhtClassExportFlags.HasObjectInitializerConstructor))
-				{
-					return ConstructorType.ObjectInitializer;
-				}
-				else if (classObj.ClassExportFlags.HasAnyFlags(UhtClassExportFlags.HasDefaultConstructor))
-				{
-					return ConstructorType.Default;
-				}
-				else
-				{
-					return ConstructorType.ForbiddenDefault;
-				}
-			}
-		}
-
 		/// <summary>
 		/// Generates private copy-constructor declaration.
 		/// </summary>
@@ -1051,11 +992,11 @@ namespace EpicGames.UHT.Exporters.CodeGen
 				switch (GetConstructorType(classObj))
 				{
 					case ConstructorType.ObjectInitializer:
-						builder.Append('\t').Append(api).Append(classObj.SourceName).Append("(const FObjectInitializer& ObjectInitializer = FObjectInitializer::Get()) : Super(ObjectInitializer) { }; \\\r\n");
+						builder.Append('\t').Append(api).Append(classObj.SourceName).Append("(const FObjectInitializer& ObjectInitializer = FObjectInitializer::Get()); \\\r\n");
 						break;
 
 					case ConstructorType.Default:
-						builder.Append('\t').Append(api).Append(classObj.SourceName).Append("() { }; \\\r\n");
+						builder.Append('\t').Append(api).Append(classObj.SourceName).Append("(); \\\r\n");
 						break;
 				}
 			}
diff --git a/Engine/Source/Programs/Shared/EpicGames.UHT/Parsers/UhtClassParser.cs b/Engine/Source/Programs/Shared/EpicGames.UHT/Parsers/UhtClassParser.cs
--- a/Engine/Source/Programs/Shared/EpicGames.UHT/Parsers/UhtClassParser.cs
+++ b/Engine/Source/Programs/Shared/EpicGames.UHT/Parsers/UhtClassParser.cs
@@ -41,6 +41,7 @@ namespace EpicGames.UHT.Parsers
 			}
 			else
 			{
+				classObj.ClassExportFlags |= UhtClassExportFlags.UsesGeneratedBodyLegacy;
 				topScope.AccessSpecifier = UhtAccessSpecifier.Public;
 			}
 
diff --git a/Engine/Source/Programs/Shared/EpicGames.UHT/Parsers/UhtInterfaceClassParser.cs b/Engine/Source/Programs/Shared/EpicGames.UHT/Parsers/UhtInterfaceClassParser.cs
--- a/Engine/Source/Programs/Shared/EpicGames.UHT/Parsers/UhtInterfaceClassParser.cs
+++ b/Engine/Source/Programs/Shared/EpicGames.UHT/Parsers/UhtInterfaceClassParser.cs
@@ -42,6 +42,7 @@ namespace EpicGames.UHT.Parsers
 			if (token.IsValue("GENERATED_UINTERFACE_BODY"))
 			{
 				topScope.AccessSpecifier = UhtAccessSpecifier.Public;
+				classObj.ClassExportFlags |= UhtClassExportFlags.UsesGeneratedBodyLegacy;
 			}
 			return UhtParseResult.Handled;
 		}
diff --git a/Engine/Source/Programs/Shared/EpicGames.UHT/Parsers/UhtNativeInterfaceClassParser.cs b/Engine/Source/Programs/Shared/EpicGames.UHT/Parsers/UhtNativeInterfaceClassParser.cs
--- a/Engine/Source/Programs/Shared/EpicGames.UHT/Parsers/UhtNativeInterfaceClassParser.cs
+++ b/Engine/Source/Programs/Shared/EpicGames.UHT/Parsers/UhtNativeInterfaceClassParser.cs
@@ -53,6 +53,7 @@ namespace EpicGames.UHT.Parsers
 
 			if (token.IsValue("GENERATED_IINTERFACE_BODY"))
 			{
+				classObj.ClassExportFlags |= UhtClassExportFlags.UsesGeneratedBodyLegacy;
 				topScope.AccessSpecifier = UhtAccessSpecifier.Public;
 			}
 			return UhtParseResult.Handled;
diff --git a/Engine/Source/Programs/Shared/EpicGames.UHT/Types/UhtClass.cs b/Engine/Source/Programs/Shared/EpicGames.UHT/Types/UhtClass.cs
--- a/Engine/Source/Programs/Shared/EpicGames.UHT/Types/UhtClass.cs
+++ b/Engine/Source/Programs/Shared/EpicGames.UHT/Types/UhtClass.cs
@@ -88,6 +88,12 @@ namespace EpicGames.UHT.Types
 		/// Class has destructor
 		/// </summary>
 		HasDestructor = 1 << 11,
+
+		/// <summary>
+		/// The GENERATED_UCLASS_BODY, GENERATED_UINTERFACE_BODY, and GENERATED_IINTERFACE_BODY macros use the 
+		/// legacy generated body.  If this flag is set, generate the legacy instead of the GENERATED_BODY macros.
+		/// </summary>
+		UsesGeneratedBodyLegacy = 1 << 12,
 	}
 
 	/// <summary>
