Horde: Move logic for appending data to leaf file nodes into a static method, allowing for a more efficient method of chunking data directly into an output buffer.

diff --git a/Engine/Source/Programs/Shared/EpicGames.Horde.Tests/BundleStoreTests.cs b/Engine/Source/Programs/Shared/EpicGames.Horde.Tests/BundleStoreTests.cs
--- a/Engine/Source/Programs/Shared/EpicGames.Horde.Tests/BundleStoreTests.cs
+++ b/Engine/Source/Programs/Shared/EpicGames.Horde.Tests/BundleStoreTests.cs
@@ -248,7 +248,7 @@ namespace EpicGames.Horde.Tests
 				remaining = await node.AppendDataAsync(remaining, options, writer, CancellationToken.None);
 				nodes.Add(node);
 
-				IoHash hash = IoHash.Compute(node.Data);
+				IoHash hash = IoHash.Compute(node.Data.Span);
 				uniqueChunks[hash] = node.Length;
 			}
 
@@ -384,7 +384,7 @@ namespace EpicGames.Horde.Tests
 			{
 				LeafFileNode newLeafNode = (LeafFileNode)newNode;
 				Assert.AreEqual(oldLeafNode.Length, newLeafNode.Length);
-				Assert.IsTrue(oldLeafNode.Data.AsSingleSegment().Span.SequenceEqual(newLeafNode.Data.AsSingleSegment().Span));
+				Assert.IsTrue(oldLeafNode.Data.Span.SequenceEqual(newLeafNode.Data.Span));
 			}
 			else
 			{
@@ -408,11 +408,8 @@ namespace EpicGames.Horde.Tests
 			int offset = 0;
 			if (fileNode is LeafFileNode leafNode)
 			{
-				foreach (ReadOnlyMemory<byte> segment in leafNode.Data)
-				{
-					Assert.IsTrue(segment.Span.SequenceEqual(data.Span.Slice(offset, segment.Length)));
-					offset += segment.Length;
-				}
+				Assert.IsTrue(leafNode.Data.Span.SequenceEqual(data.Span.Slice(offset, leafNode.Data.Length)));
+				offset += leafNode.Data.Length;
 			}
 			else if (fileNode is InteriorFileNode interiorFileNode)
 			{
diff --git a/Engine/Source/Programs/Shared/EpicGames.Horde/Storage/Nodes/FileNode.cs b/Engine/Source/Programs/Shared/EpicGames.Horde/Storage/Nodes/FileNode.cs
--- a/Engine/Source/Programs/Shared/EpicGames.Horde/Storage/Nodes/FileNode.cs
+++ b/Engine/Source/Programs/Shared/EpicGames.Horde/Storage/Nodes/FileNode.cs
@@ -322,19 +322,14 @@ namespace EpicGames.Horde.Storage.Nodes
 
 		bool _isReadOnly;
 		uint _rollingHash;
-		DataSegment _firstSegment;
-		DataSegment _lastSegment;
-		ReadOnlySequence<byte> _writtenSequence; // Payload described by _firstSegment -> _lastSegment
+		int _length;
+		byte[]? _buffer;
 
 		/// <summary>
 		/// Create an empty leaf node
 		/// </summary>
 		public LeafFileNode()
 		{
-			_firstSegment = new DataSegment(0, ReadOnlyMemory<byte>.Empty);
-			_lastSegment = _firstSegment;
-
-			_writtenSequence = ReadOnlySequence<byte>.Empty;
 		}
 
 		/// <summary>
@@ -344,18 +339,15 @@ namespace EpicGames.Horde.Storage.Nodes
 		{
 			_isReadOnly = true;
 			_rollingHash = reader.ReadUInt32();
-
-			_firstSegment = new DataSegment(0, reader.ReadVariableLengthBytes());
-			_lastSegment = _firstSegment;
-
-			_writtenSequence = new ReadOnlySequence<byte>(_firstSegment.Memory);
+			_buffer = reader.ReadVariableLengthBytes().ToArray();
+			_length = _buffer.Length;
 		}
 
 		/// <inheritdoc/>
 		public override void Serialize(ITreeNodeWriter writer)
 		{
 			writer.WriteUInt32(_rollingHash);
-			writer.WriteVariableLengthBytes(_writtenSequence);
+			writer.WriteVariableLengthBytes(Data.Span);
 		}
 
 		/// <inheritdoc/>
@@ -370,138 +362,117 @@ namespace EpicGames.Horde.Storage.Nodes
 		/// <summary>
 		/// Gets the data for this node
 		/// </summary>
-		public ReadOnlySequence<byte> Data => _writtenSequence;
+		public ReadOnlyMemory<byte> Data => _buffer.AsMemory(0, _length);
 
 		/// <inheritdoc/>
-		public override long Length => _writtenSequence.Length;
+		public override long Length => _length;
 
 		/// <inheritdoc/>
 		public override ValueTask<ReadOnlyMemory<byte>> AppendDataAsync(ReadOnlyMemory<byte> newData, ChunkingOptions options, TreeWriter writer, CancellationToken cancellationToken)
 		{
-			return new ValueTask<ReadOnlyMemory<byte>>(AppendData(newData, options));
-		}
-
-		ReadOnlyMemory<byte> AppendData(ReadOnlyMemory<byte> newData, ChunkingOptions options)
-		{
+			ReadOnlyMemory<byte> result;
 			if (_isReadOnly)
 			{
-				return newData;
-			}
-
-			// If the target option sizes are fixed, just chunk the data along fixed boundaries
-			if (options.LeafOptions.MinSize == options.LeafOptions.TargetSize && options.LeafOptions.MaxSize == options.LeafOptions.TargetSize)
-			{
-				int length = Math.Min(newData.Length, options.LeafOptions.MaxSize - (int)Length);
-				AppendLeafData(newData.Span.Slice(0, length), 0);
-				newData = newData.Slice(length);
+				result = newData;
 			}
 			else
 			{
-				// Fast path for appending data to the buffer up to the chunk window size
-				int windowSize = options.LeafOptions.MinSize;
-				if (Length < windowSize)
+				if (_buffer == null || _buffer.Length < options.LeafOptions.MaxSize)
 				{
-					int appendLength = Math.Min(windowSize - (int)Length, newData.Length);
-					AppendLeafData(newData.Span.Slice(0, appendLength));
-					newData = newData.Slice(appendLength);
+					Array.Resize(ref _buffer, options.LeafOptions.MaxSize);
 				}
 
-				// Cap the maximum amount of data to append to this node
-				int maxLength = Math.Min(newData.Length, options.LeafOptions.MaxSize - (int)Length);
-				if (maxLength > 0)
-				{
-					ReadOnlySpan<byte> inputSpan = newData.Span.Slice(0, maxLength);
-					int length = AppendLeafDataToChunkBoundary(inputSpan, options);
-					newData = newData.Slice(length);
-				}
-			}
+				int appendLength = AppendData(_buffer.AsSpan(0, _length), newData.Span, ref _rollingHash, options.LeafOptions);
+				newData.Slice(0, appendLength).CopyTo(_buffer.AsMemory(_length));
+				_length += appendLength;
 
-			// Mark this node as complete if we've reached the max size
-			if (Length == options.LeafOptions.MaxSize)
-			{
-				_isReadOnly = true;
+				result = newData.Slice(appendLength);
+				_isReadOnly = (result.Length > 0);
 			}
-			return newData;
+			return new ValueTask<ReadOnlyMemory<byte>>(result);
 		}
 
-		private int AppendLeafDataToChunkBoundary(ReadOnlySpan<byte> headSpan, ChunkingOptions options)
+		/// <summary>
+		/// Determines how much data to append to an existing leaf node
+		/// </summary>
+		/// <param name="currentData">Current data in the leaf node</param>
+		/// <param name="appendData">Data to be appended</param>
+		/// <param name="rollingHash">Current BuzHash of the data</param>
+		/// <param name="options">Options for chunking the data</param>
+		/// <returns>The number of bytes to append</returns>
+		internal static int AppendData(ReadOnlySpan<byte> currentData, ReadOnlySpan<byte> appendData, ref uint rollingHash, ChunkingOptionsForNodeType options)
 		{
-			int windowSize = options.LeafOptions.MinSize;
-			Debug.Assert(Length >= windowSize);
+			// If the target option sizes are fixed, just chunk the data along fixed boundaries
+			if (options.MinSize == options.TargetSize && options.MaxSize == options.TargetSize)
+			{
+				return Math.Min(appendData.Length, options.MaxSize - (int)currentData.Length);
+			}
 
-			// Get the threshold for the rolling hash
-			uint newRollingHash = _rollingHash;
-			uint rollingHashThreshold = (uint)((1L << 32) / options.LeafOptions.TargetSize);
+			// Cap the append data span to the maximum amount we can add
+			int maxAppendLength = options.MaxSize - currentData.Length;
+			if (maxAppendLength < appendData.Length)
+			{
+				appendData = appendData.Slice(0, maxAppendLength);
+			}
+
+			// Length of the data to be appended
+			int appendLength = 0;
 
-			// Offset within the head span, updated as we step through it.
-			int offset = 0;
+			// Fast path for appending data to the buffer up to the chunk window size
+			int windowSize = options.MinSize;
+			if (currentData.Length < windowSize)
+			{
+				appendLength = Math.Min(windowSize - (int)currentData.Length, appendData.Length);
+				rollingHash = BuzHash.Add(rollingHash, appendData.Slice(0, appendLength));
+			}
 
-			// Step the window through the tail end of the existing payload window. In this state, update the hash to remove data from the current payload, and add data from the new payload.
-			int tailLength = Math.Min(headSpan.Length, windowSize);
-			ReadOnlySequence<byte> tailSequence = _writtenSequence.Slice(_writtenSequence.Length - windowSize, tailLength);
+			// Get the threshold for the rolling hash
+			uint rollingHashThreshold = (uint)((1L << 32) / options.TargetSize);
 
-			foreach (ReadOnlyMemory<byte> tailSegment in tailSequence)
+			// Step through the part of the data where the tail of the window is in currentData, and the head of the window is in appendData.
+			if(appendLength < appendData.Length && windowSize > appendLength)
 			{
-				int count = BuzHash.Update(tailSegment.Span, headSpan.Slice(offset, tailSegment.Length), rollingHashThreshold, ref newRollingHash);
+				int overlap = windowSize - appendLength;
+				int overlapLength = Math.Min(appendData.Length - appendLength, overlap);
+
+				ReadOnlySpan<byte> tailSpan = currentData.Slice(currentData.Length - overlap, overlapLength);
+				ReadOnlySpan<byte> headSpan = appendData.Slice(appendLength, overlapLength);
+
+				int count = BuzHash.Update(tailSpan, headSpan, rollingHashThreshold, ref rollingHash);
 				if (count != -1)
 				{
-					offset += count;
-					AppendLeafData(headSpan.Slice(0, offset), newRollingHash);
-					_isReadOnly = true;
-					return offset;
+					appendLength += count;
+					return appendLength;
 				}
-				offset += tailSegment.Length;
+
+				appendLength += headSpan.Length;
 			}
 
-			// Step through the new window until we get to a chunk boundary.
-			if (offset < headSpan.Length)
+			// Step through the rest of the data which is completely contained in appendData.
+			if (appendLength < appendData.Length)
 			{
-				int count = BuzHash.Update(headSpan.Slice(offset - windowSize, headSpan.Length - offset), headSpan.Slice(offset), rollingHashThreshold, ref newRollingHash);
+				Debug.Assert(appendLength >= windowSize);
+					
+				ReadOnlySpan<byte> tailSpan = appendData.Slice(appendLength - windowSize, appendData.Length - windowSize);
+				ReadOnlySpan<byte> headSpan = appendData.Slice(appendLength);
+
+				int count = BuzHash.Update(tailSpan, headSpan, rollingHashThreshold, ref rollingHash);
 				if (count != -1)
 				{
-					offset += count;
-					AppendLeafData(headSpan.Slice(0, offset), newRollingHash);
-					_isReadOnly = true;
-					return offset;
+					appendLength += count;
+					return appendLength;
 				}
-			}
-
-			// Otherwise just append all the data.
-			AppendLeafData(headSpan, newRollingHash);
-			return headSpan.Length;
-		}
 
-		private void AppendLeafData(ReadOnlySpan<byte> leafData)
-		{
-			uint newRollingHash = BuzHash.Add(_rollingHash, leafData);
-			AppendLeafData(leafData, newRollingHash);
-		}
-
-		private void AppendLeafData(ReadOnlySpan<byte> leafData, uint newRollingHash)
-		{
-			_rollingHash = newRollingHash;
-
-			DataSegment segment = new DataSegment(_lastSegment.RunningIndex + _lastSegment.Memory.Length, leafData.ToArray());
-			if (_writtenSequence.Length == 0)
-			{
-				_firstSegment = segment;
+				appendLength += headSpan.Length;
 			}
-			else
-			{
-				_lastSegment.SetNext(segment);
-			}
-			_lastSegment = segment;
 
-			_writtenSequence = new ReadOnlySequence<byte>(_firstSegment, 0, _lastSegment, _lastSegment.Memory.Length);
+			return appendLength;
 		}
 
 		/// <inheritdoc/>
 		public override async Task CopyToStreamAsync(TreeReader reader, Stream outputStream, CancellationToken cancellationToken)
 		{
-			foreach (ReadOnlyMemory<byte> segment in _writtenSequence)
-			{
-				await outputStream.WriteAsync(segment, cancellationToken);
-			}
+			await outputStream.WriteAsync(Data, cancellationToken);
 		}
 	}
 
