EpicGames.Perforce: Ensure reused Perforce clients are reset when syncing without have-table

diff --git a/Engine/Source/Programs/Shared/EpicGames.Perforce.Managed.Tests/ManagedWorkspaceTest.cs b/Engine/Source/Programs/Shared/EpicGames.Perforce.Managed.Tests/ManagedWorkspaceTest.cs
--- a/Engine/Source/Programs/Shared/EpicGames.Perforce.Managed.Tests/ManagedWorkspaceTest.cs
+++ b/Engine/Source/Programs/Shared/EpicGames.Perforce.Managed.Tests/ManagedWorkspaceTest.cs
@@ -192,6 +192,27 @@ public class ManagedWorkspaceTest : BasePerforceFixtureTest
 		// Have-table still correspond to CL 7 as CL 8 is shelved, only p4 printed to workspace
 		await Stream.GetChangelist(7).AssertHaveTableAsync(PerforceConnection, useHaveTable);
 	}
+	
+	[TestMethod]
+	public async Task ReusePerforceClientWithoutHaveTable()
+	{
+		// Sync with have-table as normal
+		ManagedWorkspace wsWithHave = await GetManagedWorkspace(true);
+		await wsWithHave.SetupAsync(PerforceConnection, StreamName, CancellationToken.None);
+		ChangelistFixture cl = Stream.GetChangelist(6);
+		await SyncAsync(wsWithHave, cl.Number);
+		cl.AssertDepotFiles(SyncDir);
+		await cl.AssertHaveTableAsync(PerforceConnection);
+
+		// Create a new ManagedWorkspace without have-table but re-use same Perforce connection
+		// The have-table remnants from above should not interfere with this workspace
+		ManagedWorkspace wsWithoutHave = await GetManagedWorkspace(false);
+		await wsWithoutHave.SetupAsync(PerforceConnection, StreamName, CancellationToken.None);
+		await SyncAsync(wsWithoutHave, cl.Number);
+		await AssertHaveTableFileCount(0);
+		await cl.AssertHaveTableAsync(PerforceConnection, useHaveTable: false);
+		cl.AssertDepotFiles(SyncDir);
+	}
 
 	private async Task<ManagedWorkspace> GetManagedWorkspace(bool useHaveTable)
 	{
diff --git a/Engine/Source/Programs/Shared/EpicGames.Perforce.Managed/ManagedWorkspace.cs b/Engine/Source/Programs/Shared/EpicGames.Perforce.Managed/ManagedWorkspace.cs
--- a/Engine/Source/Programs/Shared/EpicGames.Perforce.Managed/ManagedWorkspace.cs
+++ b/Engine/Source/Programs/Shared/EpicGames.Perforce.Managed/ManagedWorkspace.cs
@@ -1169,6 +1169,12 @@ namespace EpicGames.Perforce.Managed
 					await perforceClient.CreateClientAsync(client, cancellationToken);
 				}
 
+				if (!_useHaveTable)
+				{
+					// If have table is not used, make client is fully reset as it may have been re-used
+					await UpdateHaveTablePathAsync(perforceClient, $"//{clientName}/...#0", cancellationToken);
+				}
+
 				status.Progress = $"({timer.Elapsed.TotalSeconds:0.0}s)";
 			}
 			_createdClients[clientName] = client;
