EpicGames.Core: Add utility method to remove a span of bytes from the hash, as well as some unit tests.

diff --git a/Engine/Source/Programs/Shared/EpicGames.Core.Tests/BuzHashTests.cs b/Engine/Source/Programs/Shared/EpicGames.Core.Tests/BuzHashTests.cs
new file mode 100644
--- /dev/null
+++ b/Engine/Source/Programs/Shared/EpicGames.Core.Tests/BuzHashTests.cs
@@ -0,0 +1,55 @@
+// Copyright Epic Games, Inc. All Rights Reserved.
+
+using System;
+using System.Collections.Generic;
+using System.Linq;
+using System.Security.Cryptography;
+using System.Text;
+using System.Threading.Tasks;
+using Microsoft.VisualStudio.TestTools.UnitTesting;
+
+namespace EpicGames.Core.Tests
+{
+	[TestClass]
+	public class BuzHashTests
+	{
+		[TestMethod]
+		public void AddSub()
+		{
+			byte[] bytes = RandomNumberGenerator.GetBytes(4096);
+
+			const int WindowLength = 32;
+
+			uint rollingHash = BuzHash.Add(0, bytes.AsSpan(0, WindowLength - 1));
+			for (int idx = WindowLength; idx < bytes.Length; idx++)
+			{
+				rollingHash = BuzHash.Add(rollingHash, bytes[idx - 1]);
+
+				uint newRollingHash = BuzHash.Add(0, bytes.AsSpan(idx - WindowLength, WindowLength));
+				Assert.AreEqual(rollingHash, newRollingHash);
+
+				rollingHash = BuzHash.Sub(rollingHash, bytes[idx - WindowLength], WindowLength);
+			}
+		}
+
+		[TestMethod]
+		public void AddSubStepped()
+		{
+			byte[] bytes = RandomNumberGenerator.GetBytes(4096);
+
+			const int WindowLength = 32;
+			const int WindowStep = 8;
+
+			uint rollingHash = BuzHash.Add(0, bytes.AsSpan(0, WindowLength - WindowStep));
+			for (int idx = WindowLength; idx < bytes.Length; idx += WindowStep)
+			{
+				rollingHash = BuzHash.Add(rollingHash, bytes.AsSpan(idx - WindowStep, WindowStep));
+
+				uint newRollingHash = BuzHash.Add(0, bytes.AsSpan(idx - WindowLength, WindowLength));
+				Assert.AreEqual(rollingHash, newRollingHash);
+
+				rollingHash = BuzHash.Sub(rollingHash, bytes.AsSpan(idx - WindowLength, WindowStep), WindowLength);
+			}
+		}
+	}
+}
diff --git a/Engine/Source/Programs/Shared/EpicGames.Core/BuzHash.cs b/Engine/Source/Programs/Shared/EpicGames.Core/BuzHash.cs
--- a/Engine/Source/Programs/Shared/EpicGames.Core/BuzHash.cs
+++ b/Engine/Source/Programs/Shared/EpicGames.Core/BuzHash.cs
@@ -136,6 +136,26 @@ namespace EpicGames.Core
 			return state ^ Rol32(s_table[value], (int)(count - 1));
 		}
 
+		/// <summary>
+		/// Static method for appending a range of data to a hash value
+		/// </summary>
+		/// <param name="state">The current hash value</param>
+		/// <param name="data">Data to append to the hash</param>
+		/// <param name="count">Size of the window</param>
+		/// <returns>New hash value</returns>
+		public static unsafe uint Sub(uint state, ReadOnlySpan<byte> data, int count)
+		{
+			fixed (uint* table = s_table)
+			fixed (byte* dataPtr = data)
+			{
+				for (int idx = 0; idx < data.Length; idx++)
+				{
+					state = state ^ Rol32(s_table[dataPtr[idx]], --count);
+				}
+			}
+			return state;
+		}
+
 		/// <summary>
 		/// Update the hash with the given data
 		/// </summary>
@@ -197,5 +217,15 @@ namespace EpicGames.Core
 			n &= 31;
 			return ((v) << (n)) | ((v) >> (32 - n));
 		}
+
+		/// <summary>
+		/// Gets the threshold below which hash values should statistically occur on average after the given number of entries
+		/// </summary>
+		/// <param name="targetSize"></param>
+		/// <returns></returns>
+		public static uint GetThreshold(int targetSize)
+		{
+			return (uint)((1L << 32) / targetSize);
+		}
 	}
 }
