EpicGames.Core: Add ChunkedMemoryWriter for writing byte sequences to rented buffers.

diff --git a/Engine/Source/Programs/Shared/EpicGames.Core/MemoryWriter.cs b/Engine/Source/Programs/Shared/EpicGames.Core/MemoryWriter.cs
--- a/Engine/Source/Programs/Shared/EpicGames.Core/MemoryWriter.cs
+++ b/Engine/Source/Programs/Shared/EpicGames.Core/MemoryWriter.cs
@@ -152,6 +152,173 @@ namespace EpicGames.Core
 		}
 	}
 
+	/// <summary>
+	/// Class for building byte sequences, similar to StringBuilder. Allocates memory in chunks to avoid copying data.
+	/// </summary>
+	public sealed class ChunkedMemoryWriter : IMemoryWriter, IDisposable
+	{
+		class Chunk
+		{
+			public readonly int RunningIndex;
+			public readonly IMemoryOwner<byte> Owner;
+			public readonly Memory<byte> Data;
+			public int Length;
+
+			public Chunk(int runningIndex, int size)
+			{
+				RunningIndex = runningIndex;
+				Owner = MemoryPool<byte>.Shared.Rent(size);
+				Data = Owner.Memory;
+			}
+
+			public ReadOnlySpan<byte> WrittenSpan => WrittenMemory.Span;
+			public ReadOnlyMemory<byte> WrittenMemory => Data.Slice(0, Length);
+		}
+
+		readonly List<Chunk> _chunks = new List<Chunk>();
+		readonly int _chunkSize;
+		Chunk _currentChunk;
+
+		/// <summary>
+		/// Length of the current sequence
+		/// </summary>
+		public int Length { get; private set; }
+
+		/// <summary>
+		/// Constructor
+		/// </summary>
+		/// <param name="chunkSize"></param>
+		public ChunkedMemoryWriter(int chunkSize)
+			: this(chunkSize, chunkSize)
+		{
+		}
+
+		/// <summary>
+		/// Constructor
+		/// </summary>
+		/// <param name="initialSize">Size of the initial chunk</param>
+		/// <param name="chunkSize">Default size for subsequent chunks</param>
+		public ChunkedMemoryWriter(int initialSize = 4096, int chunkSize = 4096)
+		{
+			_currentChunk = new Chunk(0, initialSize);
+			_chunks.Add(_currentChunk);
+			_chunkSize = chunkSize;
+		}
+
+		/// <inheritdoc/>
+		public void Dispose()
+		{
+			Clear();
+		}
+
+		/// <summary>
+		/// Clear the current builder
+		/// </summary>
+		public void Clear()
+		{
+			foreach (Chunk chunk in _chunks)
+			{
+				chunk.Owner.Dispose();
+			}
+
+			_chunks.Clear();
+			Length = 0;
+		}
+
+		/// <inheritdoc/>
+		public Memory<byte> GetMemory(int minSize)
+		{
+			if (_currentChunk.Length + minSize > _currentChunk.Data.Length)
+			{
+				_currentChunk = new Chunk(_currentChunk.RunningIndex + _currentChunk.Length, Math.Max(minSize, _chunkSize));
+				_chunks.Add(_currentChunk);
+			}
+			return _currentChunk.Data.Slice(_currentChunk.Length);
+		}
+
+		/// <inheritdoc/>
+		public void Advance(int length)
+		{
+			if (length < 0 || _currentChunk.Length + length > _currentChunk.Data.Length)
+			{
+				throw new ArgumentException("Length exceeds possible size of data written to output buffer", nameof(length));
+			}
+
+			_currentChunk.Length += length;
+			Length += length;
+		}
+
+		/// <summary>
+		/// Appends data in this builder to the given sequence
+		/// </summary>
+		/// <param name="builder">Sequence builder</param>
+		public void AppendTo(ReadOnlySequenceBuilder<byte> builder)
+		{
+			foreach (Chunk chunk in _chunks)
+			{
+				builder.Append(chunk.WrittenMemory);
+			}
+		}
+
+		/// <summary>
+		/// Gets a sequence representing the bytes that have been written so far
+		/// </summary>
+		/// <returns>Sequence of bytes</returns>
+		public ReadOnlySequence<byte> AsSequence()
+		{
+			ReadOnlySequenceBuilder<byte> builder = new ReadOnlySequenceBuilder<byte>();
+			AppendTo(builder);
+			return builder.Construct();
+		}
+
+		/// <summary>
+		/// Gets a sequence representing the bytes that have been written so far, starting at the given offset
+		/// </summary>
+		/// <param name="offset">Offset to start from</param>
+		/// <returns>Sequence of bytes</returns>
+		public ReadOnlySequence<byte> AsSequence(int offset)
+		{
+			// TODO: could do a binary search for offset and work forwards from there
+			return AsSequence().Slice(offset);
+		}
+
+		/// <summary>
+		/// Gets a sequence representing the bytes that have been written so far, starting at the given offset and 
+		/// </summary>
+		/// <param name="offset">Offset to start from</param>
+		/// <param name="length">Length of the sequence to return</param>
+		/// <returns>Sequence of bytes</returns>
+		public ReadOnlySequence<byte> AsSequence(int offset, int length)
+		{
+			// TODO: could do a binary search for offset and work fowards from there
+			return AsSequence().Slice(offset, length);
+		}
+
+		/// <summary>
+		/// Copies the data to the given output span
+		/// </summary>
+		/// <param name="span">Span to write to</param>
+		public void CopyTo(Span<byte> span)
+		{
+			foreach (Chunk chunk in _chunks)
+			{
+				Span<byte> output = span.Slice(chunk.RunningIndex);
+				chunk.WrittenSpan.CopyTo(output);
+			}
+		}
+
+		/// <summary>
+		/// Create a byte array from the sequence
+		/// </summary>
+		/// <returns>Byte array containing the current buffer data</returns>
+		public byte[] ToByteArray()
+		{
+			byte[] data = new byte[Length];
+			CopyTo(data);
+			return data;
+		}
+	}
+
 	/// <summary>
 	/// Extension methods for <see cref="IMemoryWriter"/>
 	/// </summary>
