Horde: Optimized path for creating bundles from directory structures.

* NodeHandle objects now encapsulate handles to immutable nodes which have been written but not yet flushed to disk. TreeNodeRef objects are now implemented using NodeHandles.
* TreeWriter supports a direct path of writing to its internal buffer, and returning NodeHandle objects to written data.
* FileNodeWriter writes chunked data directly into TreeWriter objects without having to buffer inside TreeNode objects first.
* TreeWriter compresses data packets and writes final packets to storage asynchronously.

diff --git a/Engine/Source/Programs/Horde/Horde.Agent.Tests/JsonRpcLoggerTests.cs b/Engine/Source/Programs/Horde/Horde.Agent.Tests/JsonRpcLoggerTests.cs
--- a/Engine/Source/Programs/Horde/Horde.Agent.Tests/JsonRpcLoggerTests.cs
+++ b/Engine/Source/Programs/Horde/Horde.Agent.Tests/JsonRpcLoggerTests.cs
@@ -104,14 +104,14 @@ namespace Horde.Agent.Tests
 
 		class FakeJsonRpcLoggerBackend : JsonRpcAndStorageLogSink
 		{
-			public NodeLocator Target { get; private set; }
+			public NodeHandle? Target { get; private set; }
 
 			public FakeJsonRpcLoggerBackend(IRpcConnection connection, string logId, IJsonRpcLogSink inner, IStorageClient store, ILogger logger)
 				: base(connection, logId, inner, store, logger)
 			{
 			}
 
-			protected override Task UpdateLogAsync(NodeLocator target, int lineCount, CancellationToken cancellationToken)
+			protected override Task UpdateLogAsync(NodeHandle target, int lineCount, CancellationToken cancellationToken)
 			{
 				Target = target;
 				return Task.CompletedTask;
@@ -157,7 +157,7 @@ namespace Horde.Agent.Tests
 						logger.LogInformation("Testing {Number}", idx);
 					}
 				}
-				file = await reader.ReadNodeAsync<LogNode>(sink.Target!);
+				file = await reader.ReadNodeAsync<LogNode>(sink.Target!.Locator!);
 			}
 
 			// Check the index text
diff --git a/Engine/Source/Programs/Horde/Horde.Agent/Commands/Bundles/BundleCreate.cs b/Engine/Source/Programs/Horde/Horde.Agent/Commands/Bundles/BundleCreate.cs
--- a/Engine/Source/Programs/Horde/Horde.Agent/Commands/Bundles/BundleCreate.cs
+++ b/Engine/Source/Programs/Horde/Horde.Agent/Commands/Bundles/BundleCreate.cs
@@ -33,7 +33,7 @@ namespace Horde.Agent.Commands.Bundles
 		{
 			IStorageClient store = await _storageClientFactory.GetClientAsync(NamespaceId);
 
-			TreeWriter writer = new TreeWriter(store, prefix: RefName.Text);
+			using TreeWriter writer = new TreeWriter(store, prefix: RefName.Text);
 
 			DirectoryNode node = new DirectoryNode(DirectoryFlags.None);
 
diff --git a/Engine/Source/Programs/Horde/Horde.Agent/Commands/Bundles/BundleDump.cs b/Engine/Source/Programs/Horde/Horde.Agent/Commands/Bundles/BundleDump.cs
--- a/Engine/Source/Programs/Horde/Horde.Agent/Commands/Bundles/BundleDump.cs
+++ b/Engine/Source/Programs/Horde/Horde.Agent/Commands/Bundles/BundleDump.cs
@@ -55,11 +55,11 @@ namespace Horde.Agent.Commands.Bundles
 			{
 				IStorageClient store = await _storageClientFactory.GetClientAsync(NamespaceId);
 
-				NodeLocator locator = await store.ReadRefTargetAsync(RefName.Value);
-				logger.LogInformation("Ref {RefName} -> {Locator}", RefName.Value, locator);
+				NodeHandle handle = await store.ReadRefTargetAsync(RefName.Value);
+				logger.LogInformation("Ref {RefName} -> {Locator}", RefName.Value, handle);
 
-				bundle = await store.ReadBundleAsync(locator.Blob);
-				logger.LogInformation("Summary for blob {BlobId}", locator.Blob);
+				bundle = await store.ReadBundleAsync(handle.Locator.Blob);
+				logger.LogInformation("Summary for blob {BlobId}", handle.Locator.Blob);
 			}
 			else
 			{
diff --git a/Engine/Source/Programs/Horde/Horde.Agent/Commands/Vcs/VcsCommit.cs b/Engine/Source/Programs/Horde/Horde.Agent/Commands/Vcs/VcsCommit.cs
--- a/Engine/Source/Programs/Horde/Horde.Agent/Commands/Vcs/VcsCommit.cs
+++ b/Engine/Source/Programs/Horde/Horde.Agent/Commands/Vcs/VcsCommit.cs
@@ -74,7 +74,7 @@ namespace Horde.Agent.Commands.Vcs
 			List<(DirectoryNodeRef, DirectoryState)> directories = new List<(DirectoryNodeRef, DirectoryState)>();
 			await UpdateTreeAsync(reader, rootRef, rootDir, oldState, newState, files, directories);
 
-			TreeWriter writer = new TreeWriter(store);
+			using TreeWriter writer = new TreeWriter(store);
 			await DirectoryNode.CopyFromDirectoryAsync(files.ConvertAll(x => (x.Item1, x.Item2)), new ChunkingOptions(), writer, CancellationToken.None);
 
 			CommitNode newTip;
diff --git a/Engine/Source/Programs/Horde/Horde.Agent/Execution/JobExecutor.cs b/Engine/Source/Programs/Horde/Horde.Agent/Execution/JobExecutor.cs
--- a/Engine/Source/Programs/Horde/Horde.Agent/Execution/JobExecutor.cs
+++ b/Engine/Source/Programs/Horde/Horde.Agent/Execution/JobExecutor.cs
@@ -866,7 +866,7 @@ namespace Horde.Agent.Execution
 				RefName refName = TempStorage.GetRefNameForNode(RefPrefix, step.Name);
 
 				TreeOptions treeOptions = new TreeOptions();
-				TreeWriter treeWriter = new TreeWriter(storage, treeOptions, refName.Text);
+				using TreeWriter treeWriter = new TreeWriter(storage, treeOptions, refName.Text);
 
 				TempStorageNode outputNode = new TempStorageNode();
 
@@ -1286,7 +1286,7 @@ namespace Horde.Agent.Execution
 				RefName refName = new RefName(RefName.Sanitize($"{RefPrefix}/artifacts"));
 
 				TreeOptions treeOptions = new TreeOptions();
-				TreeWriter treeWriter = new TreeWriter(store, treeOptions, refName.Text);
+				using TreeWriter treeWriter = new TreeWriter(store, treeOptions, refName.Text);
 
 				DirectoryNode directoryNode = new DirectoryNode(DirectoryFlags.None);
 
diff --git a/Engine/Source/Programs/Horde/Horde.Agent/Utility/JsonRpcLogger.cs b/Engine/Source/Programs/Horde/Horde.Agent/Utility/JsonRpcLogger.cs
--- a/Engine/Source/Programs/Horde/Horde.Agent/Utility/JsonRpcLogger.cs
+++ b/Engine/Source/Programs/Horde/Horde.Agent/Utility/JsonRpcLogger.cs
@@ -144,6 +144,11 @@ namespace Horde.Agent.Parser
 			{
 				await _inner.DisposeAsync();
 			}
+
+			if (_writer != null)
+			{
+				_writer.Dispose();
+			}
 		}
 
 		async Task TickTailAsync()
@@ -213,7 +218,7 @@ namespace Horde.Agent.Parser
 
 			if (request.Flush || _bufferLength > FlushLength)
 			{
-				NodeLocator target = await _builder.FlushAsync(_writer, request.Flush, cancellationToken);
+				NodeHandle target = await _builder.FlushAsync(_writer, request.Flush, cancellationToken);
 				await UpdateLogAsync(target, _builder.LineCount, cancellationToken);
 				_bufferLength = 0;
 			}
@@ -223,7 +228,7 @@ namespace Horde.Agent.Parser
 
 		#region RPC calls
 
-		protected virtual async Task UpdateLogAsync(NodeLocator target, int lineCount, CancellationToken cancellationToken)
+		protected virtual async Task UpdateLogAsync(NodeHandle target, int lineCount, CancellationToken cancellationToken)
 		{
 			_logger.LogDebug("Updating log {LogId} to line {LineCount}, target {Locator}", _logId, lineCount, target);
 
diff --git a/Engine/Source/Programs/Horde/Horde.Build.Tests/BlobStoreTests.cs b/Engine/Source/Programs/Horde/Horde.Build.Tests/BlobStoreTests.cs
--- a/Engine/Source/Programs/Horde/Horde.Build.Tests/BlobStoreTests.cs
+++ b/Engine/Source/Programs/Horde/Horde.Build.Tests/BlobStoreTests.cs
@@ -129,15 +129,15 @@ namespace Horde.Build.Tests
 			for(int idx = 0; idx < 2; idx++)
 			{
 				RefName refName = new RefName("hello");
-				await store.WriteRefTargetAsync(refName, new NodeLocator(locator3, 0));
-				NodeLocator refTarget = await store.ReadRefTargetAsync(refName);
-				Assert.AreEqual(locator3, refTarget.Blob);
+				await store.WriteRefTargetAsync(refName, new NodeHandle(bundle3.Header.Exports[0].Hash, locator3, 0));
+				NodeHandle refTarget = await store.ReadRefTargetAsync(refName);
+				Assert.AreEqual(locator3, refTarget.Locator.Blob);
 			}
 
 			RefName refName2 = new RefName("hello2");
 
-			NodeLocator refTargetId2 = await store.WriteRefAsync(refName2, CreateTestBundle(input3, new BlobLocator[] { locator1, locator2 }), 0);
-			Blob refTarget2 = await ReadBlobAsync(store, refTargetId2.Blob);
+			NodeHandle refTargetId2 = await store.WriteRefAsync(refName2, CreateTestBundle(input3, new BlobLocator[] { locator1, locator2 }), 0);
+			Blob refTarget2 = await ReadBlobAsync(store, refTargetId2.Locator.Blob);
 			Assert.IsTrue(refTarget2.Data.Span.SequenceEqual(input3));
 			Assert.IsTrue(refTarget2.References.SequenceEqual(new BlobLocator[] { locator1, locator2 }));
 		}
@@ -149,33 +149,39 @@ namespace Horde.Build.Tests
 
 			Bundle bundle1 = CreateTestBundle(new byte[] { 1, 2, 3 }, Array.Empty<BlobLocator>());
 			BlobLocator locator1 = await store.WriteBundleAsync(bundle1);
-			NodeLocator target = new NodeLocator(locator1, 0);
+			NodeHandle target = new NodeHandle(bundle1.Header.Exports[0].Hash, locator1, 0);
 
 			await store.WriteRefTargetAsync("test-ref-1", target);
 			await store.WriteRefTargetAsync("test-ref-2", target, new RefOptions { Lifetime = TimeSpan.FromMinutes(30.0), Extend = true });
 			await store.WriteRefTargetAsync("test-ref-3", target, new RefOptions { Lifetime = TimeSpan.FromMinutes(30.0), Extend = false });
 
-			Assert.AreEqual(target, await store.TryReadRefTargetAsync("test-ref-1"));
-			Assert.AreEqual(target, await store.TryReadRefTargetAsync("test-ref-2"));
-			Assert.AreEqual(target, await store.TryReadRefTargetAsync("test-ref-3"));
+			Assert.AreEqual(target.Locator, await TryReadRefTargetAsync(store, "test-ref-1"));
+			Assert.AreEqual(target.Locator, await TryReadRefTargetAsync(store, "test-ref-2"));
+			Assert.AreEqual(target.Locator, await TryReadRefTargetAsync(store, "test-ref-3"));
 
 			await Clock.AdvanceAsync(TimeSpan.FromMinutes(25.0));
 
-			Assert.AreEqual(target, await store.TryReadRefTargetAsync("test-ref-1"));
-			Assert.AreEqual(target, await store.TryReadRefTargetAsync("test-ref-2"));
-			Assert.AreEqual(target, await store.TryReadRefTargetAsync("test-ref-3"));
+			Assert.AreEqual(target.Locator, await TryReadRefTargetAsync(store, "test-ref-1"));
+			Assert.AreEqual(target.Locator, await TryReadRefTargetAsync(store, "test-ref-2"));
+			Assert.AreEqual(target.Locator, await TryReadRefTargetAsync(store, "test-ref-3"));
 
 			await Clock.AdvanceAsync(TimeSpan.FromMinutes(25.0));
 
-			Assert.AreEqual(target, await store.TryReadRefTargetAsync("test-ref-1"));
-			Assert.AreEqual(target, await store.TryReadRefTargetAsync("test-ref-2"));
-			Assert.AreEqual(default, await store.TryReadRefTargetAsync("test-ref-3"));
+			Assert.AreEqual(target.Locator, await TryReadRefTargetAsync(store, "test-ref-1"));
+			Assert.AreEqual(target.Locator, await TryReadRefTargetAsync(store, "test-ref-2"));
+			Assert.AreEqual(default, await TryReadRefTargetAsync(store, "test-ref-3"));
 
 			await Clock.AdvanceAsync(TimeSpan.FromMinutes(35.0));
 
-			Assert.AreEqual(target, await store.TryReadRefTargetAsync("test-ref-1"));
-			Assert.AreEqual(default, await store.TryReadRefTargetAsync("test-ref-2"));
-			Assert.AreEqual(default, await store.TryReadRefTargetAsync("test-ref-3"));
+			Assert.AreEqual(target.Locator, await TryReadRefTargetAsync(store, "test-ref-1"));
+			Assert.AreEqual(default, await TryReadRefTargetAsync(store, "test-ref-2"));
+			Assert.AreEqual(default, await TryReadRefTargetAsync(store, "test-ref-3"));
+		}
+
+		static async Task<NodeLocator> TryReadRefTargetAsync(IStorageClient store, RefName name)
+		{
+			NodeHandle? handle = await store.TryReadRefTargetAsync(name);
+			return handle?.Locator ?? default;
 		}
 	}
 }
diff --git a/Engine/Source/Programs/Horde/Horde.Build.Tests/GcServiceTests.cs b/Engine/Source/Programs/Horde/Horde.Build.Tests/GcServiceTests.cs
--- a/Engine/Source/Programs/Horde/Horde.Build.Tests/GcServiceTests.cs
+++ b/Engine/Source/Programs/Horde/Horde.Build.Tests/GcServiceTests.cs
@@ -46,7 +46,7 @@ namespace Horde.Build.Tests
 				int blobIdx = (int)(random.NextDouble() * blobs.Length);
 				if (roots.Add(blobs[blobIdx]))
 				{
-					await store.WriteRefTargetAsync(new RefName($"ref-{idx}"), new NodeLocator(blobs[blobIdx], 0));
+					await store.WriteRefTargetAsync(new RefName($"ref-{idx}"), new NodeHandle(IoHash.Zero, blobs[blobIdx], 0));
 				}
 			}
 
diff --git a/Engine/Source/Programs/Horde/Horde.Build/Commands/Bundles/BundleWriteTest.cs b/Engine/Source/Programs/Horde/Horde.Build/Commands/Bundles/BundleWriteTest.cs
--- a/Engine/Source/Programs/Horde/Horde.Build/Commands/Bundles/BundleWriteTest.cs
+++ b/Engine/Source/Programs/Horde/Horde.Build/Commands/Bundles/BundleWriteTest.cs
@@ -21,10 +21,10 @@ namespace Horde.Build.Commands.Bundles
 			public Task DeleteRefAsync(RefName name, CancellationToken cancellationToken = default) => Task.CompletedTask;
 			public Task<Stream> ReadBlobAsync(BlobLocator locator, CancellationToken cancellationToken = default) => throw new NotImplementedException();
 			public Task<Stream> ReadBlobRangeAsync(BlobLocator locator, int offset, int length, CancellationToken cancellationToken = default) => throw new NotImplementedException();
-			public Task<NodeLocator> TryReadRefTargetAsync(RefName name, DateTime cacheTime = default, CancellationToken cancellationToken = default) => throw new NotImplementedException();
+			public Task<NodeHandle?> TryReadRefTargetAsync(RefName name, DateTime cacheTime = default, CancellationToken cancellationToken = default) => throw new NotImplementedException();
 			public Task<BlobLocator> WriteBlobAsync(Stream stream, Utf8String prefix = default, CancellationToken cancellationToken = default) => Task.FromResult(BlobLocator.Create(HostId.Empty));
-			public Task<NodeLocator> WriteRefAsync(RefName name, Bundle bundle, int exportIdx, Utf8String prefix = default, RefOptions? options = null, CancellationToken cancellationToken = default) => throw new NotImplementedException();
-			public Task WriteRefTargetAsync(RefName name, NodeLocator target, RefOptions? options = null, CancellationToken cancellationToken = default) => Task.CompletedTask;
+			public Task<NodeHandle> WriteRefAsync(RefName name, Bundle bundle, int exportIdx, Utf8String prefix = default, RefOptions? options = null, CancellationToken cancellationToken = default) => throw new NotImplementedException();
+			public Task WriteRefTargetAsync(RefName name, NodeHandle target, RefOptions? options = null, CancellationToken cancellationToken = default) => Task.CompletedTask;
 		}
 
 		public override async Task<int> ExecuteAsync(ILogger logger)
@@ -33,9 +33,10 @@ namespace Horde.Build.Commands.Bundles
 
 			TreeOptions options = new TreeOptions();
 			options.CompressionFormat = BundleCompressionFormat.None;
-			TreeWriter writer = new TreeWriter(store, options);
+			using TreeWriter writer = new TreeWriter(store, options);
 
 			ChunkingOptions chunkingOptions = new ChunkingOptions();
+//			chunkingOptions.LeafOptions = new ChunkingOptionsForNodeType(64 * 1024);
 
 			FileNode node = new LeafFileNode();
 
@@ -46,9 +47,10 @@ namespace Horde.Build.Commands.Bundles
 			long length = 0;
 			double nextTime = 2.0;
 
+			FileNodeWriter fileNodeWriter = new FileNodeWriter(writer, chunkingOptions);
 			for (; ; )
 			{
-				node = await node.AppendAsync(buffer, chunkingOptions, writer, default);
+				await fileNodeWriter.AppendAsync(buffer, default);
 				length += buffer.Length;
 
 				double time = timer.Elapsed.TotalSeconds;
diff --git a/Engine/Source/Programs/Horde/Horde.Build/Commands/Bundles/CreateCommand.cs b/Engine/Source/Programs/Horde/Horde.Build/Commands/Bundles/CreateCommand.cs
--- a/Engine/Source/Programs/Horde/Horde.Build/Commands/Bundles/CreateCommand.cs
+++ b/Engine/Source/Programs/Horde/Horde.Build/Commands/Bundles/CreateCommand.cs
@@ -35,7 +35,7 @@ namespace Horde.Build.Commands.Bundles
 			using ServiceProvider serviceProvider = Startup.CreateServiceProvider(_configuration, _loggerProvider);
 
 			IStorageClient store = serviceProvider.GetRequiredService<IStorageClient<ReplicationService>>();
-			TreeWriter writer = new TreeWriter(store, prefix: RefName.Text);
+			using TreeWriter writer = new TreeWriter(store, prefix: RefName.Text);
 
 			DirectoryNode node = new DirectoryNode(DirectoryFlags.None);
 			await node.CopyFromDirectoryAsync(InputDir.ToDirectoryInfo(), new ChunkingOptions(), writer, CancellationToken.None);
diff --git a/Engine/Source/Programs/Horde/Horde.Build/Logs/LogRpcService.cs b/Engine/Source/Programs/Horde/Horde.Build/Logs/LogRpcService.cs
--- a/Engine/Source/Programs/Horde/Horde.Build/Logs/LogRpcService.cs
+++ b/Engine/Source/Programs/Horde/Horde.Build/Logs/LogRpcService.cs
@@ -53,7 +53,7 @@ namespace Horde.Build.Logs
 
 			IStorageClient store = await _storageService.GetClientAsync(Namespace.Logs, context.CancellationToken);
 			_logger.LogInformation("Updating {LogId} to node {RefTarget}", request.LogId, request.Target);
-			await store.WriteRefTargetAsync(new RefName(request.LogId), NodeLocator.Parse(request.Target));
+			await store.WriteRefTargetAsync(new RefName(request.LogId), NodeHandle.Parse(request.Target));
 
 			await _logFileCollection.UpdateLineCountAsync(logFile, request.LineCount, CancellationToken.None);
 
diff --git a/Engine/Source/Programs/Horde/Horde.Build/Perforce/PerforceReplicator.cs b/Engine/Source/Programs/Horde/Horde.Build/Perforce/PerforceReplicator.cs
--- a/Engine/Source/Programs/Horde/Horde.Build/Perforce/PerforceReplicator.cs
+++ b/Engine/Source/Programs/Horde/Horde.Build/Perforce/PerforceReplicator.cs
@@ -96,14 +96,23 @@ namespace Horde.Build.Perforce
 			class Handle
 			{
 				public Utf8String _path;
-				public FileNode _node = null!;
+				public FileNodeWriter _fileWriter = null!;
 				public long _size;
+				public long _sizeWritten;
 				public readonly IncrementalHash _hash = IncrementalHash.CreateHash(HashAlgorithmName.MD5);
 			}
 
+			readonly TreeWriter _writer;
+			readonly ChunkingOptions _options;
 			readonly Stack<Handle> _freeHandles = new Stack<Handle>();
 			readonly Dictionary<int, Handle> _openHandles = new Dictionary<int, Handle>();
 
+			public FileWriter(TreeWriter writer, ChunkingOptions options)
+			{
+				_writer = writer;
+				_options = options;
+			}
+
 			public void Dispose()
 			{
 				foreach (Handle handle in _freeHandles)
@@ -120,34 +129,39 @@ namespace Horde.Build.Perforce
 			{
 				Handle handle = (_freeHandles.Count > 0)? _freeHandles.Pop() : new Handle();
 
+				handle._fileWriter.Reset();
 				handle._path = path;
-				handle._node = new LeafFileNode();
 				handle._size = size;
 
 				_openHandles.Add(fd, handle);
 			}
 
-			public async Task AppendAsync(int fd, ReadOnlyMemory<byte> data, ChunkingOptions options, TreeWriter writer, CancellationToken cancellationToken)
+			public async Task AppendAsync(int fd, ReadOnlyMemory<byte> data, CancellationToken cancellationToken)
 			{
 				Handle handle = _openHandles[fd];
 				handle._hash.AppendData(data.Span);
-				handle._node = await handle._node.AppendAsync(data, options, writer, cancellationToken);
+
+				await handle._fileWriter.AppendAsync(data, cancellationToken);
+				handle._sizeWritten += data.Length;
 			}
 
-			public void Close(int fd, out Utf8String path, out FileNode node, out byte[] hash)
+			public async Task<(Utf8String, NodeHandle, byte[], long)> CloseAsync(int fd, CancellationToken cancellationToken)
 			{
 				Handle handle = _openHandles[fd];
-				if (handle._node.Length != handle._size)
+				if (handle._sizeWritten != handle._size)
 				{
-					throw new ReplicationException($"Invalid size for replicated file '{handle._path}'. Expected {handle._size}, got {handle._node.Length}.");
+					throw new ReplicationException($"Invalid size for replicated file '{handle._path}'. Expected {handle._size}, got {handle._sizeWritten}.");
 				}
 
-				path = handle._path;
-				node = handle._node;
-				hash = handle._hash.GetHashAndReset();
+				Utf8String path = handle._path;
+				NodeHandle node = await handle._fileWriter.FlushAsync(cancellationToken);
+				byte[] hash = handle._hash.GetHashAndReset();
+				long length = handle._fileWriter.Length;
 
 				_openHandles.Remove(fd);
 				_freeHandles.Push(handle);
+
+				return (path, node, hash, length);
 			}
 		}
 
@@ -294,7 +308,7 @@ namespace Horde.Build.Perforce
 			_logger.LogInformation("Total sync size: {Size:n1}mb", totalSize / (1024.0 * 1024.0));
 
 			// Create the tree writer
-			TreeWriter writer = new TreeWriter(store, refName, options.TreeOptions);
+			using TreeWriter writer = new TreeWriter(store, refName, options.TreeOptions);
 
 			// Sync incrementally
 			long syncedSize = 0;
@@ -349,7 +363,7 @@ namespace Horde.Build.Perforce
 				Stopwatch processTimer = new Stopwatch();
 				Stopwatch gcTimer = new Stopwatch();
 
-				using FileWriter fileWriter = new FileWriter();
+				using FileWriter fileWriter = new FileWriter(writer, options.ChunkingOptions);
 				await foreach (PerforceResponse response in perforce.StreamCommandAsync("sync", Array.Empty<string>(), syncPaths, null, typeof(SyncRecord), true, default))
 				{
 					PerforceError? error = response.Error;
@@ -389,13 +403,13 @@ namespace Horde.Build.Perforce
 						}
 						else if (io.Command == PerforceIoCommand.Write)
 						{
-							await fileWriter.AppendAsync(io.File, io.Payload, options.ChunkingOptions, writer, cancellationToken);
+							await fileWriter.AppendAsync(io.File, io.Payload, cancellationToken);
 						}
 						else if (io.Command == PerforceIoCommand.Close)
 						{
-							fileWriter.Close(io.File, out Utf8String path, out FileNode node, out byte[] hash);
+							(Utf8String path, NodeHandle handle, byte[] hash, long length) = await fileWriter.CloseAsync(io.File, cancellationToken);
 
-							FileEntry entry = await root.AddFileByPathAsync(reader, path, FileEntryFlags.None, node, cancellationToken);
+							FileEntry entry = await root.AddFileByPathAsync(reader, path, FileEntryFlags.None, length, handle, cancellationToken);
 							entry.CustomData = hash;
 							await writer.WriteAsync(entry, cancellationToken);
 						}
diff --git a/Engine/Source/Programs/Horde/Horde.Build/Storage/StorageController.cs b/Engine/Source/Programs/Horde/Horde.Build/Storage/StorageController.cs
--- a/Engine/Source/Programs/Horde/Horde.Build/Storage/StorageController.cs
+++ b/Engine/Source/Programs/Horde/Horde.Build/Storage/StorageController.cs
@@ -44,6 +44,11 @@ namespace Horde.Build.Storage
 	/// </summary>
 	public class WriteRefRequest
 	{
+		/// <summary>
+		/// Hash of the target node
+		/// </summary>
+		public IoHash Hash { get; set; }
+
 		/// <summary>
 		/// Locator for the target blob
 		/// </summary>
@@ -65,6 +70,11 @@ namespace Horde.Build.Storage
 	/// </summary>
 	public class ReadRefResponse
 	{
+		/// <summary>
+		/// Hash of the target node
+		/// </summary>
+		public IoHash Hash { get; set; }
+
 		/// <summary>
 		/// Locator for the target blob
 		/// </summary>
@@ -78,10 +88,11 @@ namespace Horde.Build.Storage
 		/// <summary>
 		/// Constructor
 		/// </summary>
-		public ReadRefResponse(NodeLocator target)
+		public ReadRefResponse(NodeHandle target)
 		{
-			Blob = target.Blob;
-			ExportIdx = target.ExportIdx;
+			Hash = target.Hash;
+			Blob = target.Locator.Blob;
+			ExportIdx = target.Locator.ExportIdx;
 		}
 	}
 
@@ -238,7 +249,7 @@ namespace Horde.Build.Storage
 			}
 
 			IStorageClient client = await _storageService.GetClientAsync(namespaceId, cancellationToken);
-			NodeLocator target = new NodeLocator(request.Blob, request.ExportIdx);
+			NodeHandle target = new NodeHandle(request.Hash, request.Blob, request.ExportIdx);
 			await client.WriteRefTargetAsync(refName, target, request.Options, cancellationToken);
 
 			return Ok();
@@ -261,8 +272,8 @@ namespace Horde.Build.Storage
 
 			IStorageClient client = await _storageService.GetClientAsync(namespaceId, cancellationToken);
 
-			NodeLocator target = await client.TryReadRefTargetAsync(refName, cancellationToken: cancellationToken);
-			if (!target.IsValid())
+			NodeHandle? target = await client.TryReadRefTargetAsync(refName, cancellationToken: cancellationToken);
+			if (target == null)
 			{
 				return NotFound();
 			}
diff --git a/Engine/Source/Programs/Horde/Horde.Build/Storage/StorageService.cs b/Engine/Source/Programs/Horde/Horde.Build/Storage/StorageService.cs
--- a/Engine/Source/Programs/Horde/Horde.Build/Storage/StorageService.cs
+++ b/Engine/Source/Programs/Horde/Horde.Build/Storage/StorageService.cs
@@ -212,10 +212,10 @@ namespace Horde.Build.Storage
 			#region Refs
 
 			/// <inheritdoc/>
-			public override Task<NodeLocator> TryReadRefTargetAsync(RefName name, DateTime cacheTime = default, CancellationToken cancellationToken = default) => _outer.TryReadRefTargetAsync(NamespaceId, name, cacheTime, cancellationToken);
+			public override Task<NodeHandle?> TryReadRefTargetAsync(RefName name, DateTime cacheTime = default, CancellationToken cancellationToken = default) => _outer.TryReadRefTargetAsync(NamespaceId, name, cacheTime, cancellationToken);
 
 			/// <inheritdoc/>
-			public override Task WriteRefTargetAsync(RefName name, NodeLocator target, RefOptions? options = null, CancellationToken cancellationToken = default) => _outer.WriteRefTargetAsync(NamespaceId, name, target, options, cancellationToken);
+			public override Task WriteRefTargetAsync(RefName name, NodeHandle target, RefOptions? options = null, CancellationToken cancellationToken = default) => _outer.WriteRefTargetAsync(NamespaceId, name, target, options, cancellationToken);
 
 			/// <inheritdoc/>
 			public override Task DeleteRefAsync(RefName name, CancellationToken cancellationToken = default) => _outer.DeleteRefAsync(NamespaceId, name, cancellationToken);
@@ -301,6 +301,9 @@ namespace Horde.Build.Storage
 			[BsonElement("name")]
 			public RefName Name { get; set; }
 
+			[BsonElement("hash")]
+			public IoHash Hash { get; set; }
+
 			[BsonElement("blob")]
 			public BlobLocator BlobLocator { get; set; }
 
@@ -317,7 +320,7 @@ namespace Horde.Build.Storage
 			public TimeSpan? Lifetime { get; set; }
 
 			[BsonIgnore]
-			public NodeLocator Target => new NodeLocator(BlobLocator, ExportIdx);
+			public NodeHandle Target => new NodeHandle(Hash, BlobLocator, ExportIdx);
 
 			public RefInfo()
 			{
@@ -325,13 +328,14 @@ namespace Horde.Build.Storage
 				BlobLocator = BlobLocator.Empty;
 			}
 
-			public RefInfo(NamespaceId namespaceId, RefName name, NodeLocator target, ObjectId blobInfoId)
+			public RefInfo(NamespaceId namespaceId, RefName name, NodeHandle target, ObjectId blobInfoId)
 			{
 				NamespaceId = namespaceId;
 				Name = name;
-				BlobLocator = target.Blob;
+				Hash = target.Hash;
+				BlobLocator = target.Locator.Blob;
 				BlobInfoId = blobInfoId;
-				ExportIdx = target.ExportIdx;
+				ExportIdx = target.Locator.ExportIdx;
 			}
 
 			public bool HasExpired(DateTime utcNow) => ExpiresAtUtc.HasValue && utcNow >= ExpiresAtUtc.Value;
@@ -855,7 +859,7 @@ namespace Horde.Build.Storage
 		}
 
 		/// <inheritdoc/>
-		async Task<NodeLocator> TryReadRefTargetAsync(NamespaceId namespaceId, RefName name, DateTime cacheTime = default, CancellationToken cancellationToken = default)
+		async Task<NodeHandle?> TryReadRefTargetAsync(NamespaceId namespaceId, RefName name, DateTime cacheTime = default, CancellationToken cancellationToken = default)
 		{
 			RefCacheValue entry;
 			if (!_cache.TryGetValue(name, out entry) || entry.Time < cacheTime)
@@ -866,7 +870,7 @@ namespace Horde.Build.Storage
 
 			if (entry.Value == null)
 			{
-				return default;
+				return null;
 			}
 
 			if (entry.Value.ExpiresAtUtc != null)
@@ -887,12 +891,12 @@ namespace Horde.Build.Storage
 		}
 
 		/// <inheritdoc/>
-		async Task WriteRefTargetAsync(NamespaceId namespaceId, RefName name, NodeLocator target, RefOptions? options = null, CancellationToken cancellationToken = default)
+		async Task WriteRefTargetAsync(NamespaceId namespaceId, RefName name, NodeHandle target, RefOptions? options = null, CancellationToken cancellationToken = default)
 		{
-			BlobInfo? newBlobInfo = await _blobCollection.Find(x => x.NamespaceId == namespaceId && x.BlobId == target.Blob.BlobId).FirstOrDefaultAsync(cancellationToken);
+			BlobInfo? newBlobInfo = await _blobCollection.Find(x => x.NamespaceId == namespaceId && x.BlobId == target.Locator.Blob.BlobId).FirstOrDefaultAsync(cancellationToken);
 			if (newBlobInfo == null)
 			{
-				throw new Exception($"Invalid/unknown blob identifier '{target.Blob.BlobId}' in namespace {namespaceId}");
+				throw new Exception($"Invalid/unknown blob identifier '{target.Locator.Blob.BlobId}' in namespace {namespaceId}");
 			}
 
 			RefInfo newRefInfo = new RefInfo(namespaceId, name, target, newBlobInfo.Id);
diff --git a/Engine/Source/Programs/Shared/EpicGames.Horde.Tests/BundleStoreTests.cs b/Engine/Source/Programs/Shared/EpicGames.Horde.Tests/BundleStoreTests.cs
--- a/Engine/Source/Programs/Shared/EpicGames.Horde.Tests/BundleStoreTests.cs
+++ b/Engine/Source/Programs/Shared/EpicGames.Horde.Tests/BundleStoreTests.cs
@@ -79,7 +79,7 @@ namespace EpicGames.Horde.Tests
 		{
 			// Generate a tree
 			{
-				TreeWriter writer = new TreeWriter(store, options, "test");
+				using TreeWriter writer = new TreeWriter(store, options, "test");
 
 				SimpleNode node1 = new SimpleNode(new ReadOnlySequence<byte>(new byte[] { 1 }), Array.Empty<TreeNodeRef<SimpleNode>>());
 				SimpleNode node2 = new SimpleNode(new ReadOnlySequence<byte>(new byte[] { 2 }), new[] { new TreeNodeRef<SimpleNode>(node1) });
@@ -135,6 +135,20 @@ namespace EpicGames.Horde.Tests
 			Assert.AreEqual(0, refs1.Count);
 		}
 
+		[TestMethod]
+		public async Task SimpleNodeAsync()
+		{
+			MemoryStorageClient store = new MemoryStorageClient();
+
+			RefName refName = new RefName("test");
+			await store.WriteNodeAsync(refName, new SimpleNode(new ReadOnlySequence<byte>(new byte[] { (byte)123 }), Array.Empty<TreeNodeRef<SimpleNode>>()));
+
+			TreeReader reader = new TreeReader(store, null, s_readOptions, NullLogger.Instance);
+			SimpleNode node = await reader.ReadNodeAsync<SimpleNode>(refName);
+
+			Assert.AreEqual(123, node.Data.FirstSpan[0]);
+		}
+
 		[TestMethod]
 		public async Task DirectoryNodesAsync()
 		{
@@ -180,13 +194,14 @@ namespace EpicGames.Horde.Tests
 
 			// Generate a tree
 			{
-				TreeWriter writer = new TreeWriter(store, new TreeOptions());
+				using TreeWriter writer = new TreeWriter(store, new TreeOptions());
 
 				DirectoryNode root = new DirectoryNode(DirectoryFlags.None);
 				DirectoryNode hello = root.AddDirectory("hello");
 
-				FileNode world = await FileNode.CreateAsync(Encoding.UTF8.GetBytes("world"), new ChunkingOptions(), writer, CancellationToken.None);
-				hello.AddFile("world", FileEntryFlags.None, world);
+				FileNodeWriter fileWriter = new FileNodeWriter(writer, new ChunkingOptions());
+				NodeHandle fileHandle = await fileWriter.CreateAsync(Encoding.UTF8.GetBytes("world"), CancellationToken.None);
+				hello.AddFile("world", FileEntryFlags.None, fileWriter.Length, fileHandle);
 
 				await writer.WriteAsync(new RefName("test"), root);
 
@@ -216,83 +231,6 @@ namespace EpicGames.Horde.Tests
 			Assert.IsTrue(worldData.SequenceEqual(Encoding.UTF8.GetBytes("world")));
 		}
 
-		[TestMethod]
-		public async Task ChunkingTests()
-		{
-			using IMemoryCache cache = new MemoryCache(new MemoryCacheOptions());
-
-			MemoryStorageClient store = new MemoryStorageClient();
-			TreeWriter writer = new TreeWriter(store, new TreeOptions());
-
-			byte[] chunk = RandomNumberGenerator.GetBytes(256);
-
-			byte[] data = new byte[chunk.Length * 1024];
-			for (int idx = 0; idx * chunk.Length < data.Length; idx++)
-			{
-				chunk.CopyTo(data.AsSpan(idx * chunk.Length));
-			}
-
-			ChunkingOptions options = new ChunkingOptions();
-			options.LeafOptions.MinSize = 1;
-			options.LeafOptions.TargetSize = 64;
-			options.LeafOptions.MaxSize = 1024;
-
-			ReadOnlyMemory<byte> remaining = data;
-
-			Dictionary<IoHash, long> uniqueChunks = new Dictionary<IoHash, long>();
-
-			List<LeafFileNode> nodes = new List<LeafFileNode>();
-			while (remaining.Length > 0)
-			{
-				LeafFileNode node = new LeafFileNode();
-				remaining = await node.AppendDataAsync(remaining, options, writer, CancellationToken.None);
-				nodes.Add(node);
-
-				IoHash hash = IoHash.Compute(node.Data.Span);
-				uniqueChunks[hash] = node.Length;
-			}
-
-			long uniqueSize = uniqueChunks.Sum(x => x.Value);
-			Assert.IsTrue(uniqueSize < data.Length / 3);
-		}
-		/*
-		[TestMethod]
-		public async Task ChildWriterAsync()
-		{
-			using IMemoryCache cache = new MemoryCache(new MemoryCacheOptions());
-
-			InMemoryBlobStore store = new InMemoryBlobStore(cache);
-
-			{
-				ITreeWriter writer = store.CreateTreeWriter(new TreeOptions());
-
-				ITreeWriter childWriter1 = writer.CreateChildWriter();
-				ITreeBlobRef childNode1 = await childWriter1.WriteNodeAsync(new ReadOnlySequence<byte>(new byte[] { 4, 5, 6 }), Array.Empty<ITreeBlobRef>());
-
-				ITreeWriter childWriter2 = writer.CreateChildWriter();
-				ITreeBlobRef childNode2 = await childWriter2.WriteNodeAsync(new ReadOnlySequence<byte>(new byte[] { 7, 8, 9 }), Array.Empty<ITreeBlobRef>());
-
-				RefName refName = new RefName("test");
-				await writer.WriteRefAsync(refName, new ReadOnlySequence<byte>(new byte[] { 1, 2, 3 }), new[] { childNode1, childNode2 });
-
-				Assert.AreEqual(3, store.Blobs.Count);
-				Assert.AreEqual(1, store.Refs.Count);
-
-				ITreeBlob? rootBlob = await store.TryReadTreeAsync(refName);
-				Assert.IsNotNull(rootBlob);
-				Assert.IsTrue(rootBlob.Data.AsSingleSegment().Span.SequenceEqual(new byte[] { 1, 2, 3 }));
-				Assert.AreEqual(2, rootBlob.Refs.Count);
-
-				ITreeBlob? childBlob1 = await rootBlob.Refs[0].GetTargetAsync();
-				Assert.IsTrue(childBlob1.Data.AsSingleSegment().Span.SequenceEqual(new byte[] { 4, 5, 6 }));
-				Assert.AreEqual(0, childBlob1.Refs.Count);
-
-				ITreeBlob? childBlob2 = await rootBlob.Refs[1].GetTargetAsync();
-				Assert.IsTrue(childBlob2.Data.AsSingleSegment().Span.SequenceEqual(new byte[] { 7, 8, 9 }));
-				Assert.AreEqual(0, childBlob2.Refs.Count);
-			}
-		}
-		*/
 		[TestMethod]
 		public async Task LargeFileTestAsync()
 		{
@@ -314,7 +252,7 @@ namespace EpicGames.Horde.Tests
 			// Generate a tree
 			DirectoryNode root;
 			{
-				TreeWriter writer = new TreeWriter(store, new TreeOptions { MaxBlobSize = 1024 });
+				using TreeWriter writer = new TreeWriter(store, new TreeOptions { MaxBlobSize = 1024 });
 
 				root = new DirectoryNode(DirectoryFlags.None);
 
@@ -323,8 +261,9 @@ namespace EpicGames.Horde.Tests
 				options.LeafOptions.TargetSize = 256;
 				options.LeafOptions.MaxSize = 64 * 1024;
 
-				FileNode file = await FileNode.CreateAsync(data, options, writer, CancellationToken.None);
-				root.AddFile("test", FileEntryFlags.None, file);
+				FileNodeWriter fileWriter = new FileNodeWriter(writer, options);
+				NodeHandle handle = await fileWriter.CreateAsync(data, CancellationToken.None);
+				root.AddFile("test", FileEntryFlags.None, fileWriter.Length, handle);
 
 				await writer.WriteAsync(new RefName("test"), root);
 
@@ -366,10 +305,7 @@ namespace EpicGames.Horde.Tests
 			if (oldNode is InteriorFileNode oldInteriorNode)
 			{
 				InteriorFileNode newInteriorNode = (InteriorFileNode)newNode;
-
-				Assert.AreEqual(oldInteriorNode.Length, newInteriorNode.Length);
 				Assert.AreEqual(oldInteriorNode.Children.Count, newInteriorNode.Children.Count);
-				Assert.AreEqual(oldInteriorNode.RollingHash, newInteriorNode.RollingHash);
 
 				int index = 0;
 				foreach ((TreeNodeRef<FileNode> oldFileRef, TreeNodeRef<FileNode> newFileRef) in oldInteriorNode.Children.Zip(newInteriorNode.Children))
@@ -383,7 +319,6 @@ namespace EpicGames.Horde.Tests
 			else if (oldNode is LeafFileNode oldLeafNode)
 			{
 				LeafFileNode newLeafNode = (LeafFileNode)newNode;
-				Assert.AreEqual(oldLeafNode.Length, newLeafNode.Length);
 				Assert.IsTrue(oldLeafNode.Data.Span.SequenceEqual(newLeafNode.Data.Span));
 			}
 			else
diff --git a/Engine/Source/Programs/Shared/EpicGames.Horde.Tests/BundleTests.cs b/Engine/Source/Programs/Shared/EpicGames.Horde.Tests/BundleTests.cs
--- a/Engine/Source/Programs/Shared/EpicGames.Horde.Tests/BundleTests.cs
+++ b/Engine/Source/Programs/Shared/EpicGames.Horde.Tests/BundleTests.cs
@@ -67,23 +67,28 @@ namespace EpicGames.Horde.Tests
 		{
 			RefName refName = new RefName("test");
 			TreeReader reader = new TreeReader(_storage, null, NullLogger.Instance);
-			TreeWriter writer = new TreeWriter(_storage, new TreeOptions(), refName.Text);
+			using TreeWriter writer = new TreeWriter(_storage, new TreeOptions(), refName.Text);
 
 			ChunkingOptions options = new ChunkingOptions();
 			options.LeafOptions = new ChunkingOptionsForNodeType(8, 8, 8);
 
-			FileNode node = new LeafFileNode();
-			node = await node.AppendAsync(new byte[7], options, writer, CancellationToken.None);
+			FileNodeWriter fileNodeWriter = new FileNodeWriter(writer, options);
+
+			NodeHandle handle;
+			FileNode node;
+
+			handle = await fileNodeWriter.CreateAsync(new byte[7], CancellationToken.None);
+			node = await reader.ReadNodeAsync<FileNode>(handle.Locator);
 			Assert.IsTrue(node is LeafFileNode);
 			Assert.AreEqual(7, ((LeafFileNode)node).Data.Length);
 
-			node = new LeafFileNode();
-			node = await node.AppendAsync(new byte[8], options, writer, CancellationToken.None);
+			handle = await fileNodeWriter.CreateAsync(new byte[8], CancellationToken.None);
+			node = await reader.ReadNodeAsync<FileNode>(handle.Locator);
 			Assert.IsTrue(node is LeafFileNode);
 			Assert.AreEqual(8, ((LeafFileNode)node).Data.Length);
 
-			node = new LeafFileNode();
-			node = await node.AppendAsync(new byte[9], options, writer, CancellationToken.None);
+			handle = await fileNodeWriter.CreateAsync(new byte[9], CancellationToken.None);
+			node = await reader.ReadNodeAsync<FileNode>(handle.Locator);
 			Assert.IsTrue(node is InteriorFileNode);
 			Assert.AreEqual(2, ((InteriorFileNode)node).Children.Count);
 
@@ -184,8 +189,8 @@ namespace EpicGames.Horde.Tests
 			Assert.AreEqual(1, store.Blobs.Count);
 
 			// Check the ref
-			NodeLocator refTarget = await store.ReadRefTargetAsync(refName);
-			Bundle bundle = await store.ReadBundleAsync(refTarget.Blob);
+			NodeHandle refTarget = await store.ReadRefTargetAsync(refName);
+			Bundle bundle = await store.ReadBundleAsync(refTarget.Locator.Blob);
 			Assert.AreEqual(0, bundle.Header.Imports.Count);
 			Assert.AreEqual(3, bundle.Header.Exports.Count);
 
@@ -267,236 +272,5 @@ namespace EpicGames.Horde.Tests
 				Assert.IsNotNull(newNode4);
 			}
 		}
-		/*
-		[TestMethod]
-		public async Task CompactTest()
-		{
-			BundleOptions options = new BundleOptions();
-			options.MaxBlobSize = 1024 * 1024;
-			options.MaxInlineBlobSize = 1;
-
-			using Bundle<DirectoryNode> newBundle = Bundle.Create<DirectoryNode>(_storageClient, _namespaceId, new DirectoryNode(), options, _cache);
-
-			DirectoryNode node1 = newBundle.Root.AddDirectory("node1");
-			DirectoryNode node2 = node1.AddDirectory("node2");
-			DirectoryNode node3 = node2.AddDirectory("node3");
-			DirectoryNode node4 = newBundle.Root.AddDirectory("node4"); // same contents as node 3
-
-			RefId refId1 = new RefId("ref1");
-			await newBundle.WriteAsync(_bucketId, refId1, CbObject.Empty, false);
-
-			Assert.AreEqual(1, _storageClient.Refs.Count);
-			Assert.AreEqual(1, _storageClient.Blobs.Count);
-
-			IRef ref1 = _storageClient.Refs[(_namespaceId, _bucketId, refId1)];
-
-			BundleRoot root1 = CbSerializer.Deserialize<BundleRoot>(ref1.Value);
-			BundleObject rootObject1 = root1.Object;
-			Assert.AreEqual(1, rootObject1.Exports.Count);
-			Assert.AreEqual(1, rootObject1.ImportObjects.Count);
-
-			IoHash leafHash1 = rootObject1.ImportObjects[0].Object.Hash;
-			BundleObject leafObject1 = CbSerializer.Deserialize<BundleObject>(_storageClient.Blobs[(_namespaceId, leafHash1)]);
-			Assert.AreEqual(3, leafObject1.Exports.Count); // node1 + node2 + node3 (== node4)
-			Assert.AreEqual(0, leafObject1.ImportObjects.Count);
-
-			// Remove one of the nodes from the root without compacting. the existing blob should be reused.
-			newBundle.Root.DeleteDirectory("node1");
-
-			RefId refId2 = new RefId("ref2");
-			await newBundle.WriteAsync(_bucketId, refId2, CbObject.Empty, false);
-
-			IRef ref2 = _storageClient.Refs[(_namespaceId, _bucketId, refId2)];
-
-			BundleRoot root2 = CbSerializer.Deserialize<BundleRoot>(ref2.Value);
-			BundleObject rootObject2 = root2.Object;
-			Assert.AreEqual(1, rootObject2.Exports.Count);
-			Assert.AreEqual(1, rootObject2.ImportObjects.Count);
-
-			IoHash leafHash2 = rootObject2.ImportObjects[0].Object.Hash;
-			Assert.AreEqual(leafHash1, leafHash2);
-			Assert.AreEqual(3, leafObject1.Exports.Count); // unused: node1 + node2 + node3, used: node4 (== node3)
-			Assert.AreEqual(0, leafObject1.ImportObjects.Count);
-
-			// Repack it and check that we make a new object
-			RefId refId3 = new RefId("ref3");
-			await newBundle.WriteAsync(_bucketId, refId3, CbObject.Empty, true);
-
-			IRef ref3 = _storageClient.Refs[(_namespaceId, _bucketId, refId3)];
-
-			BundleRoot root3 = CbSerializer.Deserialize<BundleRoot>(ref3.Value);
-			BundleObject rootObject3 = root3.Object;
-			Assert.AreEqual(1, rootObject3.Exports.Count);
-			Assert.AreEqual(1, rootObject3.ImportObjects.Count);
-
-			IoHash leafHash3 = rootObject3.ImportObjects[0].Object.Hash;
-			Assert.AreNotEqual(leafHash1, leafHash3);
-
-			BundleObject leafObject3 = CbSerializer.Deserialize<BundleObject>(_storageClient.Blobs[(_namespaceId, leafHash3)]);
-			Assert.AreEqual(1, leafObject3.Exports.Count);
-			Assert.AreEqual(0, leafObject3.ImportObjects.Count);
-		}
-*/
-		[TestMethod]
-		public async Task CoreAppendTest()
-		{
-			TreeReader reader = new TreeReader(_storage, null, NullLogger.Instance);
-			TreeWriter writer = new TreeWriter(_storage);
-
-			byte[] data = new byte[4096];
-			new Random(0).NextBytes(data);
-
-			ChunkingOptions options = new ChunkingOptions();
-			options.LeafOptions = new ChunkingOptionsForNodeType(16, 64, 256);
-
-			FileNode node = new LeafFileNode();
-			for (int idx = 0; idx < data.Length; idx++)
-			{
-				node = await node.AppendAsync(data.AsMemory(idx, 1), options, writer, CancellationToken.None);
-
-				byte[] outputData = await node.ToByteArrayAsync(reader, CancellationToken.None);
-				Assert.IsTrue(data.AsMemory(0, idx + 1).Span.SequenceEqual(outputData.AsSpan(0, idx + 1)));
-			}
-		}
-
-		[TestMethod]
-		public async Task FixedSizeChunkingTests()
-		{
-			ChunkingOptions options = new ChunkingOptions();
-			options.LeafOptions = new ChunkingOptionsForNodeType(64, 64, 64);
-			options.InteriorOptions = new ChunkingOptionsForNodeType(IoHash.NumBytes * 4, IoHash.NumBytes * 4, IoHash.NumBytes * 4);
-
-			await ChunkingTests(options);
-		}
-
-		[TestMethod]
-		public async Task VariableSizeChunkingTests()
-		{
-			ChunkingOptions options = new ChunkingOptions();
-			options.LeafOptions = new ChunkingOptionsForNodeType(32, 64, 96);
-			options.InteriorOptions = new ChunkingOptionsForNodeType(IoHash.NumBytes * 1, IoHash.NumBytes * 4, IoHash.NumBytes * 12);
-
-			await ChunkingTests(options);
-		}
-
-		async Task ChunkingTests(ChunkingOptions options)
-		{
-			TreeWriter writer = new TreeWriter(_storage);
-
-			byte[] data = new byte[4096];
-			new Random(0).NextBytes(data);
-
-			for (int idx = 0; idx < data.Length; idx++)
-			{
-				data[idx] = (byte)idx;
-			}
-
-			FileNode root = new LeafFileNode();
-
-			const int NumIterations = 100;
-			for (int idx = 0; idx < NumIterations; idx++)
-			{
-				root = await root.AppendAsync(data, options, writer, CancellationToken.None);
-			}
-
-			TreeReader reader = new TreeReader(_storage, null, NullLogger.Instance);
-			byte[] result = await root.ToByteArrayAsync(reader, CancellationToken.None);
-			Assert.AreEqual(NumIterations * data.Length, result.Length);
-
-			for (int idx = 0; idx < NumIterations; idx++)
-			{
-				ReadOnlyMemory<byte> spanData = result.AsMemory(idx * data.Length, data.Length);
-				Assert.IsTrue(spanData.Span.SequenceEqual(data));
-			}
-
-			await CheckSizes(reader, root, options, true);
-		}
-
-		async Task CheckSizes(TreeReader reader, FileNode node, ChunkingOptions options, bool rightmost)
-		{
-			if (node is LeafFileNode leafNode)
-			{
-				Assert.IsTrue(rightmost || leafNode.Data.Length >= options.LeafOptions.MinSize);
-				Assert.IsTrue(leafNode.Data.Length <= options.LeafOptions.MaxSize);
-			}
-			else
-			{
-				InteriorFileNode interiorNode = (InteriorFileNode)node;
-
-				Assert.IsTrue(rightmost || interiorNode.Children.Count >= options.InteriorOptions.MinSize);
-				Assert.IsTrue(interiorNode.Children.Count <= options.InteriorOptions.MaxSize);
-
-				int childCount = interiorNode.Children.Count;
-				for (int idx = 0; idx < childCount; idx++)
-				{
-					FileNode childNode = await interiorNode.Children[idx].ExpandAsync(reader, CancellationToken.None);
-					await CheckSizes(reader, childNode, options, idx == childCount - 1);
-				}
-			}
-		}
-
-#if false
-		[TestMethod]
-		public async Task SpillTestAsync()
-		{
-			TreeOptions options = new TreeOptions();
-			options.MaxBlobSize = 1;
-
-			RefName refName = new RefName("ref");
-			ITreeWriter writer = _storage.CreateTreeWriter(refName.Text);
-
-			DirectoryNode root = new DirectoryNode(DirectoryFlags.None);
-
-			long totalLength = 0;
-			for (int idxA = 0; idxA < 10; idxA++)
-			{
-				DirectoryNode nodeA = root.AddDirectory($"{idxA}");
-				for (int idxB = 0; idxB < 10; idxB++)
-				{
-					DirectoryNode nodeB = nodeA.AddDirectory($"{idxB}");
-					for (int idxC = 0; idxC < 10; idxC++)
-					{
-						DirectoryNode nodeC = nodeB.AddDirectory($"{idxC}");
-						for (int idxD = 0; idxD < 10; idxD++)
-						{
-							FileEntry file = nodeC.AddFile($"{idxD}", FileEntryFlags.None);
-							byte[] data = Encoding.UTF8.GetBytes($"This is file {idxA}/{idxB}/{idxC}/{idxD}");
-							totalLength += data.Length;
-							await file.AppendAsync(data, new ChunkingOptions(), writer, CancellationToken.None);
-						}
-					}
-
-					int oldWorkingSetSize = GetWorkingSetSize(root);
-					await _storage.WriteTreeAsync(refName, root, new TreeOptions());
-					int newWorkingSetSize = GetWorkingSetSize(root);
-					Assert.IsTrue(newWorkingSetSize <= oldWorkingSetSize);
-					Assert.IsTrue(newWorkingSetSize <= 20);
-				}
-			}
-
-			Assert.IsTrue(_storage.Blobs.Count > 0);
-			Assert.IsTrue(_storage.Refs.Count == 1);
-
-			await _storage.WriteTreeAsync(refName, root, new TreeOptions(), CancellationToken.None);
-
-			Assert.AreEqual(totalLength, root.Length);
-
-			Assert.IsTrue(_storage.Blobs.Count > 0);
-			Assert.IsTrue(_storage.Refs.Count == 1);
-		}
-
-		int GetWorkingSetSize(TreeNode node)
-		{
-			int size = 0;
-			foreach (TreeNodeRef nodeRef in node.GetReferences())
-			{
-				if (nodeRef.Node != null)
-				{
-					size += 1 + GetWorkingSetSize(nodeRef.Node);
-				}
-			}
-			return size;
-		}
-#endif
 	}
 }
diff --git a/Engine/Source/Programs/Shared/EpicGames.Horde.Tests/ChunkingTests.cs b/Engine/Source/Programs/Shared/EpicGames.Horde.Tests/ChunkingTests.cs
--- a/Engine/Source/Programs/Shared/EpicGames.Horde.Tests/ChunkingTests.cs
+++ b/Engine/Source/Programs/Shared/EpicGames.Horde.Tests/ChunkingTests.cs
@@ -47,47 +47,6 @@ namespace EpicGames.Horde.Tests
 			}
 		}
 
-		[TestMethod]
-		public async Task BasicChunkingTests()
-		{
-			using MemoryCache cache = new MemoryCache(new MemoryCacheOptions());
-
-			MemoryStorageClient store = new MemoryStorageClient();
-			TreeReader reader = new TreeReader(store, cache, NullLogger.Instance);
-			TreeWriter writer = new TreeWriter(store, new TreeOptions(), "test");
-
-			ChunkingOptions options = new ChunkingOptions();
-			options.LeafOptions = new ChunkingOptionsForNodeType(8, 8, 8);
-
-			FileNode node = new LeafFileNode();
-			node = await node.AppendAsync(new byte[7], options, writer, CancellationToken.None);
-			Assert.AreEqual(7, node.Length);
-			Assert.IsTrue(node is LeafFileNode);
-			Assert.AreEqual(7, (await node.ToByteArrayAsync(reader, CancellationToken.None)).Length);
-
-			node = new LeafFileNode();
-			node = await node.AppendAsync(new byte[8], options, writer, CancellationToken.None);
-			Assert.AreEqual(8, node.Length);
-			Assert.IsTrue(node is LeafFileNode);
-			Assert.AreEqual(8, (await node.ToByteArrayAsync(reader, CancellationToken.None)).Length);
-
-			node = new LeafFileNode();
-			node = await node.AppendAsync(new byte[9], options, writer, CancellationToken.None);
-			Assert.AreEqual(9, node.Length);
-			Assert.IsTrue(node is InteriorFileNode);
-			Assert.AreEqual(2, ((InteriorFileNode)node).Children.Count);
-
-			FileNode? childNode1 = await ((TreeNodeRef<FileNode>)((InteriorFileNode)node).Children[0]).ExpandAsync(reader);
-			Assert.IsNotNull(childNode1);
-			Assert.IsTrue(childNode1 is LeafFileNode);
-			Assert.AreEqual(8, (await childNode1!.ToByteArrayAsync(reader, CancellationToken.None)).Length);
-
-			FileNode? childNode2 = await ((TreeNodeRef<FileNode>)((InteriorFileNode)node).Children[1]).ExpandAsync(reader);
-			Assert.IsNotNull(childNode2);
-			Assert.IsTrue(childNode2 is LeafFileNode);
-			Assert.AreEqual(1, (await childNode2!.ToByteArrayAsync(reader, CancellationToken.None)).Length);
-		}
-
 		[TestMethod]
 		public async Task FixedSizeChunkingTests()
 		{
@@ -108,13 +67,13 @@ namespace EpicGames.Horde.Tests
 			await TestChunkingAsync(options);
 		}
 
-		static async Task TestChunkingAsync(ChunkingOptions options)
+		async Task TestChunkingAsync(ChunkingOptions options)
 		{
 			using MemoryCache cache = new MemoryCache(new MemoryCacheOptions());
 
 			MemoryStorageClient store = new MemoryStorageClient();
 			TreeReader reader = new TreeReader(store, cache, NullLogger.Instance);
-			TreeWriter writer = new TreeWriter(store, new TreeOptions());
+			using TreeWriter writer = new TreeWriter(store, new TreeOptions());
 
 			byte[] data = new byte[4096];
 			new Random(0).NextBytes(data);
@@ -124,14 +83,22 @@ namespace EpicGames.Horde.Tests
 				data[idx] = (byte)idx;
 			}
 
-			FileNode root = new LeafFileNode();
+			NodeHandle handle;
 
 			const int NumIterations = 100;
-			for (int idx = 0; idx < NumIterations; idx++)
 			{
-				root = await root.AppendAsync(data, options, writer, CancellationToken.None);
+				FileNodeWriter fileWriter = new FileNodeWriter(writer, options);
+
+				for (int idx = 0; idx < NumIterations; idx++)
+				{
+					await fileWriter.AppendAsync(data, CancellationToken.None);
+				}
+
+				handle = await fileWriter.FlushAsync(CancellationToken.None);
 			}
 
+			FileNode root = await reader.ReadNodeAsync<FileNode>(handle.Locator);
+
 			byte[] result;
 			using (MemoryStream stream = new MemoryStream())
 			{
@@ -146,6 +113,31 @@ namespace EpicGames.Horde.Tests
 				ReadOnlyMemory<byte> spanData = result.AsMemory(idx * data.Length, data.Length);
 				Assert.IsTrue(spanData.Span.SequenceEqual(data));
 			}
+
+			await CheckSizes(reader, root, options, true);
+		}
+
+		async Task CheckSizes(TreeReader reader, FileNode node, ChunkingOptions options, bool rightmost)
+		{
+			if (node is LeafFileNode leafNode)
+			{
+				Assert.IsTrue(rightmost || leafNode.Data.Length >= options.LeafOptions.MinSize);
+				Assert.IsTrue(leafNode.Data.Length <= options.LeafOptions.MaxSize);
+			}
+			else
+			{
+				InteriorFileNode interiorNode = (InteriorFileNode)node;
+
+				Assert.IsTrue(rightmost || interiorNode.Children.Count * IoHash.NumBytes >= options.InteriorOptions.MinSize);
+				Assert.IsTrue(interiorNode.Children.Count <= options.InteriorOptions.MaxSize);
+
+				int childCount = interiorNode.Children.Count;
+				for (int idx = 0; idx < childCount; idx++)
+				{
+					FileNode childNode = await interiorNode.Children[idx].ExpandAsync(reader, CancellationToken.None);
+					await CheckSizes(reader, childNode, options, idx == childCount - 1);
+				}
+			}
 		}
 	}
 }
diff --git a/Engine/Source/Programs/Shared/EpicGames.Horde.Tests/LogTests.cs b/Engine/Source/Programs/Shared/EpicGames.Horde.Tests/LogTests.cs
--- a/Engine/Source/Programs/Shared/EpicGames.Horde.Tests/LogTests.cs
+++ b/Engine/Source/Programs/Shared/EpicGames.Horde.Tests/LogTests.cs
@@ -87,8 +87,12 @@ namespace EpicGames.Horde.Tests
 			}
 
 			// Flush it to storage, and read the finished log node
-			NodeLocator target = await builder.FlushAsync(new TreeWriter(_store), true, CancellationToken.None);
-			LogNode log = await _reader.ReadNodeAsync<LogNode>(target);
+			NodeHandle target;
+			using (TreeWriter writer = new TreeWriter(_store))
+			{
+				target = await builder.FlushAsync(writer, true, CancellationToken.None);
+			}
+			LogNode log = await _reader.ReadNodeAsync<LogNode>(target.Locator);
 
 			// Read the data back out and check it's the same
 			byte[] readData = new byte[_data.Length];
@@ -129,10 +133,15 @@ namespace EpicGames.Horde.Tests
 			{
 				builder.WriteData(Encoding.UTF8.GetBytes(lines[lineIdx]));
 			}
-			NodeLocator rootNodeLocator = await builder.FlushAsync(new TreeWriter(_store), true, CancellationToken.None);
+
+			NodeHandle rootNodeHandle;
+			using (TreeWriter writer = new TreeWriter(_store))
+			{
+				rootNodeHandle = await builder.FlushAsync(writer, true, CancellationToken.None);
+			}
 
 			// Read it back in and test the index
-			LogNode rootNode = await _reader.ReadNodeAsync<LogNode>(rootNodeLocator);
+			LogNode rootNode = await _reader.ReadNodeAsync<LogNode>(rootNodeHandle.Locator);
 			LogIndexNode index = await rootNode.IndexRef.ExpandAsync(_reader);
 		
 			for (int lineIdx = 0; lineIdx < lines.Length; lineIdx++)
diff --git a/Engine/Source/Programs/Shared/EpicGames.Horde/Logs/LogChunkNode.cs b/Engine/Source/Programs/Shared/EpicGames.Horde/Logs/LogChunkNode.cs
--- a/Engine/Source/Programs/Shared/EpicGames.Horde/Logs/LogChunkNode.cs
+++ b/Engine/Source/Programs/Shared/EpicGames.Horde/Logs/LogChunkNode.cs
@@ -244,7 +244,7 @@ namespace EpicGames.Horde.Logs
 		}
 
 		/// <inheritdoc/>
-		protected override void Serialize(IMemoryWriter writer)
+		public override void Serialize(ITreeNodeWriter writer)
 		{
 			base.Serialize(writer);
 
diff --git a/Engine/Source/Programs/Shared/EpicGames.Horde/Logs/LogNode.cs b/Engine/Source/Programs/Shared/EpicGames.Horde/Logs/LogNode.cs
--- a/Engine/Source/Programs/Shared/EpicGames.Horde/Logs/LogNode.cs
+++ b/Engine/Source/Programs/Shared/EpicGames.Horde/Logs/LogNode.cs
@@ -269,7 +269,7 @@ namespace EpicGames.Horde.Logs
 		/// <param name="writer">Writer for the output nodes</param>
 		/// <param name="complete">Whether the log is complete</param>
 		/// <param name="cancellationToken">Cancellation token for the operation</param>
-		public async Task<NodeLocator> FlushAsync(TreeWriter writer, bool complete, CancellationToken cancellationToken)
+		public async Task<NodeHandle> FlushAsync(TreeWriter writer, bool complete, CancellationToken cancellationToken)
 		{
 			// Capture the new data that needs to be written
 			IReadOnlyList<LogChunkNode> writeTextChunks;
@@ -299,7 +299,7 @@ namespace EpicGames.Horde.Logs
 
 			LogNode newRoot = new LogNode(_format, lineCount, length, newJsonChunkRefs, new TreeNodeRef<LogIndexNode>(newIndex), complete);
 
-			NodeLocator newRootLocator = await writer.FlushAsync(newRoot, cancellationToken);
+			NodeHandle newRootHandle = await writer.FlushAsync(newRoot, cancellationToken);
 
 			// Update the new state
 			lock (_lockObject)
@@ -310,7 +310,7 @@ namespace EpicGames.Horde.Logs
 				_indexTextBuilder.Remove(writeIndexTextChunks.Count);
 			}
 
-			return newRootLocator;
+			return newRootHandle;
 		}
 	}
 
diff --git a/Engine/Source/Programs/Shared/EpicGames.Horde/Storage/Backends/FileStorageClient.cs b/Engine/Source/Programs/Shared/EpicGames.Horde/Storage/Backends/FileStorageClient.cs
--- a/Engine/Source/Programs/Shared/EpicGames.Horde/Storage/Backends/FileStorageClient.cs
+++ b/Engine/Source/Programs/Shared/EpicGames.Horde/Storage/Backends/FileStorageClient.cs
@@ -1,6 +1,7 @@
 // Copyright Epic Games, Inc. All Rights Reserved.
 
 using System;
+using System.Diagnostics;
 using System.Globalization;
 using System.IO;
 using System.Threading;
@@ -82,21 +83,21 @@ namespace EpicGames.Horde.Storage.Backends
 		}
 
 		/// <inheritdoc/>
-		public override async Task<NodeLocator> TryReadRefTargetAsync(RefName name, DateTime cacheTime = default, CancellationToken cancellationToken = default)
+		public override async Task<NodeHandle?> TryReadRefTargetAsync(RefName name, DateTime cacheTime = default, CancellationToken cancellationToken = default)
 		{
 			FileReference file = GetRefFile(name);
 			if (!FileReference.Exists(file))
 			{
-				return default;
+				return null;
 			}
 
 			_logger.LogInformation("Reading {File}", file);
 			string text = await FileReference.ReadAllTextAsync(file, cancellationToken);
-			return NodeLocator.Parse(text);
+			return NodeHandle.Parse(text);
 		}
 
 		/// <inheritdoc/>
-		public override async Task WriteRefTargetAsync(RefName name, NodeLocator target, RefOptions? options = null, CancellationToken cancellationToken = default)
+		public override async Task WriteRefTargetAsync(RefName name, NodeHandle target, RefOptions? options = null, CancellationToken cancellationToken = default)
 		{
 			FileReference file = GetRefFile(name);
 			DirectoryReference.CreateDirectory(file.Directory);
diff --git a/Engine/Source/Programs/Shared/EpicGames.Horde/Storage/Backends/HttpStorageClient.cs b/Engine/Source/Programs/Shared/EpicGames.Horde/Storage/Backends/HttpStorageClient.cs
--- a/Engine/Source/Programs/Shared/EpicGames.Horde/Storage/Backends/HttpStorageClient.cs
+++ b/Engine/Source/Programs/Shared/EpicGames.Horde/Storage/Backends/HttpStorageClient.cs
@@ -35,6 +35,7 @@ namespace EpicGames.Horde.Storage.Backends
 
 		class ReadRefResponse
 		{
+			public IoHash Hash { get; set; }
 			public BlobLocator Blob { get; set; }
 			public int ExportIdx { get; set; }
 		}
@@ -195,7 +196,7 @@ namespace EpicGames.Horde.Storage.Backends
 		}
 
 		/// <inheritdoc/>
-		public override async Task<NodeLocator> TryReadRefTargetAsync(RefName name, DateTime cacheTime = default, CancellationToken cancellationToken = default)
+		public override async Task<NodeHandle?> TryReadRefTargetAsync(RefName name, DateTime cacheTime = default, CancellationToken cancellationToken = default)
 		{
 			using (HttpClient httpClient = _createClient())
 			{
@@ -206,7 +207,7 @@ namespace EpicGames.Horde.Storage.Backends
 						if (response.StatusCode == HttpStatusCode.NotFound)
 						{
 							_logger.LogDebug("Read ref {RefName} -> None", name);
-							return default;
+							return null;
 						}
 						else if (!response.IsSuccessStatusCode)
 						{
@@ -218,7 +219,7 @@ namespace EpicGames.Horde.Storage.Backends
 							response.EnsureSuccessStatusCode();
 							ReadRefResponse? data = await response.Content.ReadFromJsonAsync<ReadRefResponse>(cancellationToken: cancellationToken);
 							_logger.LogDebug("Read ref {RefName} -> {Blob}#{ExportIdx}", name, data!.Blob, data!.ExportIdx);
-							return new NodeLocator(data!.Blob, data!.ExportIdx);
+							return new NodeHandle(data.Hash, new NodeLocator(data!.Blob, data!.ExportIdx));
 						}
 					}
 				}
@@ -226,12 +227,12 @@ namespace EpicGames.Horde.Storage.Backends
 		}
 
 		/// <inheritdoc/>
-		public override async Task WriteRefTargetAsync(RefName name, NodeLocator target, RefOptions? options = null, CancellationToken cancellationToken = default)
+		public override async Task WriteRefTargetAsync(RefName name, NodeHandle target, RefOptions? options = null, CancellationToken cancellationToken = default)
 		{
 			_logger.LogDebug("Writing ref {RefName} -> {RefTarget}", name, target);
 			using (HttpClient httpClient = _createClient())
 			{
-				using (HttpResponseMessage response = await httpClient.PutAsync($"refs/{name}", new { blob = target.Blob, exportIdx = target.ExportIdx, options }, cancellationToken))
+				using (HttpResponseMessage response = await httpClient.PutAsync($"refs/{name}", new { hash = target.Hash, blob = target.Locator.Blob, exportIdx = target.Locator.ExportIdx, options }, cancellationToken))
 				{
 					response.EnsureSuccessStatusCode();
 				}
diff --git a/Engine/Source/Programs/Shared/EpicGames.Horde/Storage/Backends/MemoryStorageClient.cs b/Engine/Source/Programs/Shared/EpicGames.Horde/Storage/Backends/MemoryStorageClient.cs
--- a/Engine/Source/Programs/Shared/EpicGames.Horde/Storage/Backends/MemoryStorageClient.cs
+++ b/Engine/Source/Programs/Shared/EpicGames.Horde/Storage/Backends/MemoryStorageClient.cs
@@ -26,13 +26,13 @@ namespace EpicGames.Horde.Storage.Backends
 		/// <summary>
 		/// Map of ref name to ref data
 		/// </summary>
-		readonly ConcurrentDictionary<RefName, NodeLocator> _refs = new ConcurrentDictionary<RefName, NodeLocator>();
+		readonly ConcurrentDictionary<RefName, NodeHandle> _refs = new ConcurrentDictionary<RefName, NodeHandle>();
 
 		/// <inheritdoc cref="_blobs"/>
 		public IReadOnlyDictionary<BlobLocator, Bundle> Blobs => _blobs;
 
 		/// <inheritdoc cref="_refs"/>
-		public IReadOnlyDictionary<RefName, NodeLocator> Refs => _refs;
+		public IReadOnlyDictionary<RefName, NodeHandle> Refs => _refs;
 
 		/// <summary>
 		/// Constructor
@@ -73,15 +73,15 @@ namespace EpicGames.Horde.Storage.Backends
 		public override Task DeleteRefAsync(RefName name, CancellationToken cancellationToken) => Task.FromResult(_refs.TryRemove(name, out _));
 
 		/// <inheritdoc/>
-		public override Task<NodeLocator> TryReadRefTargetAsync(RefName name, DateTime cacheTime = default, CancellationToken cancellationToken = default)
+		public override Task<NodeHandle?> TryReadRefTargetAsync(RefName name, DateTime cacheTime = default, CancellationToken cancellationToken = default)
 		{
-			NodeLocator refTarget;
+			NodeHandle? refTarget;
 			_refs.TryGetValue(name, out refTarget);
 			return Task.FromResult(refTarget);
 		}
 
 		/// <inheritdoc/>
-		public override Task WriteRefTargetAsync(RefName name, NodeLocator target, RefOptions? options = null, CancellationToken cancellationToken = default)
+		public override Task WriteRefTargetAsync(RefName name, NodeHandle target, RefOptions? options = null, CancellationToken cancellationToken = default)
 		{
 			_refs[name] = target;
 			return Task.CompletedTask;
diff --git a/Engine/Source/Programs/Shared/EpicGames.Horde/Storage/Bundle.cs b/Engine/Source/Programs/Shared/EpicGames.Horde/Storage/Bundle.cs
--- a/Engine/Source/Programs/Shared/EpicGames.Horde/Storage/Bundle.cs
+++ b/Engine/Source/Programs/Shared/EpicGames.Horde/Storage/Bundle.cs
@@ -323,7 +323,7 @@ namespace EpicGames.Horde.Storage
 	/// <summary>
 	/// Information about a type within a bundle
 	/// </summary>
-	[DebuggerDisplay("{Tag}#{Version}")]
+	[DebuggerDisplay("{Guid}#{Version}")]
 	public class BundleType : IEquatable<BundleType>
 	{
 		/// <summary>
@@ -576,32 +576,37 @@ namespace EpicGames.Horde.Storage
 		/// <summary>
 		/// Compress a data packet
 		/// </summary>
-		/// <param name="format"></param>
-		/// <param name="input"></param>
+		/// <param name="format">Format for the compressed data</param>
+		/// <param name="input">The data to compress</param>
+		/// <param name="writer">Writer for output data</param>
 		/// <returns>The compressed data</returns>
-		public static ReadOnlyMemory<byte> Compress(BundleCompressionFormat format, ReadOnlyMemory<byte> input)
+		public static int Compress(BundleCompressionFormat format, ReadOnlyMemory<byte> input, IMemoryWriter writer)
 		{
 			switch (format)
 			{
 				case BundleCompressionFormat.None:
 					{
-						return input;
+						writer.WriteFixedLengthBytes(input.Span);
+						return input.Length;
 					}
 				case BundleCompressionFormat.LZ4:
 					{
 						int maxSize = LZ4Codec.MaximumOutputSize(input.Length);
-						using (IMemoryOwner<byte> buffer = MemoryPool<byte>.Shared.Rent(maxSize))
-						{
-							int encodedLength = LZ4Codec.Encode(input.Span, buffer.Memory.Span);
-							return buffer.Memory.Slice(0, encodedLength).ToArray();
-						}
+
+						Span<byte> buffer = writer.GetSpan(maxSize);
+						int encodedLength = LZ4Codec.Encode(input.Span, buffer);
+
+						writer.Advance(encodedLength);
+						return encodedLength;
 					}
 				case BundleCompressionFormat.Gzip:
 					{
 						using MemoryStream outputStream = new MemoryStream(input.Length);
 						using GZipStream gzipStream = new GZipStream(outputStream, CompressionLevel.Fastest);
 						gzipStream.Write(input.Span);
-						return outputStream.ToArray();
+
+						writer.WriteFixedLengthBytes(outputStream.ToArray());
+						return (int)outputStream.Length;
 					}
 				case BundleCompressionFormat.Oodle:
 					{
diff --git a/Engine/Source/Programs/Shared/EpicGames.Horde/Storage/IStorageClient.cs b/Engine/Source/Programs/Shared/EpicGames.Horde/Storage/IStorageClient.cs
--- a/Engine/Source/Programs/Shared/EpicGames.Horde/Storage/IStorageClient.cs
+++ b/Engine/Source/Programs/Shared/EpicGames.Horde/Storage/IStorageClient.cs
@@ -157,7 +157,7 @@ namespace EpicGames.Horde.Storage
 		/// <param name="cacheTime">Minimum coherency for any cached value to be returned</param>
 		/// <param name="cancellationToken">Cancellation token for the operation</param>
 		/// <returns>Node pointed to by the ref</returns>
-		Task<NodeLocator> TryReadRefTargetAsync(RefName name, DateTime cacheTime = default, CancellationToken cancellationToken = default);
+		Task<NodeHandle?> TryReadRefTargetAsync(RefName name, DateTime cacheTime = default, CancellationToken cancellationToken = default);
 
 		/// <summary>
 		/// Writes a new ref to the store which points to a new blob
@@ -169,17 +169,17 @@ namespace EpicGames.Horde.Storage
 		/// <param name="options">Options for the new ref</param>
 		/// <param name="cancellationToken">Cancellation token for the operation</param>
 		/// <returns>Unique identifier for the blob</returns>
-		Task<NodeLocator> WriteRefAsync(RefName name, Bundle bundle, int exportIdx, Utf8String prefix = default, RefOptions? options = null, CancellationToken cancellationToken = default);
+		Task<NodeHandle> WriteRefAsync(RefName name, Bundle bundle, int exportIdx, Utf8String prefix = default, RefOptions? options = null, CancellationToken cancellationToken = default);
 
 		/// <summary>
 		/// Writes a new ref to the store
 		/// </summary>
 		/// <param name="name">Ref to write</param>
-		/// <param name="target">The target for the ref</param>
+		/// <param name="handle">Handle to the target node</param>
 		/// <param name="options">Options for the new ref</param>
 		/// <param name="cancellationToken">Cancellation token for the operation</param>
 		/// <returns>Unique identifier for the blob</returns>
-		Task WriteRefTargetAsync(RefName name, NodeLocator target, RefOptions? options = null, CancellationToken cancellationToken = default);
+		Task WriteRefTargetAsync(RefName name, NodeHandle handle, RefOptions? options = null, CancellationToken cancellationToken = default);
 
 		/// <summary>
 		/// Reads data for a ref from the store
@@ -242,13 +242,13 @@ namespace EpicGames.Horde.Storage
 			public Task DeleteRefAsync(RefName name, CancellationToken cancellationToken = default) => _inner.DeleteRefAsync(name, cancellationToken);
 
 			/// <inheritdoc/>
-			public Task<NodeLocator> TryReadRefTargetAsync(RefName name, DateTime cacheTime = default, CancellationToken cancellationToken = default) => _inner.TryReadRefTargetAsync(name, cacheTime, cancellationToken);
+			public Task<NodeHandle?> TryReadRefTargetAsync(RefName name, DateTime cacheTime = default, CancellationToken cancellationToken = default) => _inner.TryReadRefTargetAsync(name, cacheTime, cancellationToken);
 
 			/// <inheritdoc/>
-			public Task<NodeLocator> WriteRefAsync(RefName name, Bundle bundle, int exportIdx = 0, Utf8String prefix = default, RefOptions? options = null, CancellationToken cancellationToken = default) => _inner.WriteRefAsync(name, bundle, exportIdx, prefix, options, cancellationToken);
+			public Task<NodeHandle> WriteRefAsync(RefName name, Bundle bundle, int exportIdx = 0, Utf8String prefix = default, RefOptions? options = null, CancellationToken cancellationToken = default) => _inner.WriteRefAsync(name, bundle, exportIdx, prefix, options, cancellationToken);
 
 			/// <inheritdoc/>
-			public Task WriteRefTargetAsync(RefName name, NodeLocator target, RefOptions? options = null, CancellationToken cancellationToken = default) => _inner.WriteRefTargetAsync(name, target, options, cancellationToken);
+			public Task WriteRefTargetAsync(RefName name, NodeHandle target, RefOptions? options = null, CancellationToken cancellationToken = default) => _inner.WriteRefTargetAsync(name, target, options, cancellationToken);
 
 			#endregion
 		}
@@ -336,8 +336,8 @@ namespace EpicGames.Horde.Storage
 		/// <returns>True if the ref exists, false if it did not exist</returns>
 		public static async Task<bool> HasRefAsync(this IStorageClient store, RefName name, DateTime cacheTime = default, CancellationToken cancellationToken = default)
 		{
-			NodeLocator target = await store.TryReadRefTargetAsync(name, cacheTime, cancellationToken);
-			return target.IsValid();
+			NodeHandle? target = await store.TryReadRefTargetAsync(name, cacheTime, cancellationToken);
+			return target != null;
 		}
 
 		/// <summary>
@@ -361,7 +361,7 @@ namespace EpicGames.Horde.Storage
 		/// <param name="maxAge">Maximum age of any cached ref</param>
 		/// <param name="cancellationToken">Cancellation token for the operation</param>
 		/// <returns>The ref target</returns>
-		public static Task<NodeLocator> TryReadRefTargetAsync(this IStorageClient store, RefName name, TimeSpan maxAge, CancellationToken cancellationToken = default)
+		public static Task<NodeHandle?> TryReadRefTargetAsync(this IStorageClient store, RefName name, TimeSpan maxAge, CancellationToken cancellationToken = default)
 		{
 			return store.TryReadRefTargetAsync(name, DateTime.UtcNow - maxAge, cancellationToken);
 		}
@@ -374,10 +374,10 @@ namespace EpicGames.Horde.Storage
 		/// <param name="cacheTime">Minimum coherency of any cached result</param>
 		/// <param name="cancellationToken">Cancellation token for the operation</param>
 		/// <returns>The ref target</returns>
-		public static async Task<NodeLocator> ReadRefTargetAsync(this IStorageClient store, RefName name, DateTime cacheTime = default, CancellationToken cancellationToken = default)
+		public static async Task<NodeHandle> ReadRefTargetAsync(this IStorageClient store, RefName name, DateTime cacheTime = default, CancellationToken cancellationToken = default)
 		{
-			NodeLocator refTarget = await store.TryReadRefTargetAsync(name, cacheTime, cancellationToken);
-			if (!refTarget.IsValid())
+			NodeHandle? refTarget = await store.TryReadRefTargetAsync(name, cacheTime, cancellationToken);
+			if (refTarget == null)
 			{
 				throw new RefNameNotFoundException(name);
 			}
@@ -392,7 +392,7 @@ namespace EpicGames.Horde.Storage
 		/// <param name="maxAge">Maximum age for any cached result</param>
 		/// <param name="cancellationToken">Cancellation token for the operation</param>
 		/// <returns>The blob instance</returns>
-		public static Task<NodeLocator> ReadRefTargetAsync(this IStorageClient store, RefName name, TimeSpan maxAge, CancellationToken cancellationToken = default)
+		public static Task<NodeHandle> ReadRefTargetAsync(this IStorageClient store, RefName name, TimeSpan maxAge, CancellationToken cancellationToken = default)
 		{
 			return ReadRefTargetAsync(store, name, DateTime.UtcNow - maxAge, cancellationToken);
 		}
diff --git a/Engine/Source/Programs/Shared/EpicGames.Horde/Storage/Nodes/CbNode.cs b/Engine/Source/Programs/Shared/EpicGames.Horde/Storage/Nodes/CbNode.cs
--- a/Engine/Source/Programs/Shared/EpicGames.Horde/Storage/Nodes/CbNode.cs
+++ b/Engine/Source/Programs/Shared/EpicGames.Horde/Storage/Nodes/CbNode.cs
@@ -2,8 +2,10 @@
 
 using EpicGames.Core;
 using EpicGames.Serialization;
+using System;
 using System.Collections.Generic;
 using System.Linq;
+using System.Reflection.Metadata;
 
 namespace EpicGames.Horde.Storage.Nodes
 {
@@ -13,6 +15,37 @@ namespace EpicGames.Horde.Storage.Nodes
 	[TreeNode("{34A0793F-8364-42F4-8632-98A71C843229}", 1)]
 	public class CbNode : TreeNode
 	{
+		class HandleMapper
+		{
+			public IReadOnlyList<NodeLocator> Locators { get; }
+			public List<NodeHandle> Handles { get; }
+
+			public HandleMapper(IReadOnlyList<NodeLocator> locators)
+			{
+				Locators = locators;
+				Handles = new List<NodeHandle>(locators.Count);
+			}
+
+			public void IterateField(CbField field)
+			{
+				if (field.IsAttachment())
+				{
+					NodeLocator locator = Locators[Handles.Count];
+					Handles.Add(new NodeHandle(field.AsAttachment(), locator));
+				}
+				else if (field.IsArray())
+				{
+					CbArray array = field.AsArray();
+					array.IterateAttachments(IterateField);
+				}
+				else if (field.IsObject())
+				{
+					CbObject obj = field.AsObject();
+					obj.IterateAttachments(IterateField);
+				}
+			}
+		}
+
 		/// <summary>
 		/// The compact binary object
 		/// </summary>
@@ -41,7 +74,11 @@ namespace EpicGames.Horde.Storage.Nodes
 		public CbNode(ITreeNodeReader reader)
 		{
 			Object = new CbObject(reader.ReadFixedLengthBytes(reader.Length));
-			References = reader.References.ConvertAll(x => new TreeNodeRef(x));
+
+			HandleMapper mapper = new HandleMapper(reader.References);
+			Object.IterateAttachments(mapper.IterateField);
+
+			References = mapper.Handles.ConvertAll(x => new TreeNodeRef(x));
 		}
 
 		/// <inheritdoc/>
diff --git a/Engine/Source/Programs/Shared/EpicGames.Horde/Storage/Nodes/DirectoryNode.cs b/Engine/Source/Programs/Shared/EpicGames.Horde/Storage/Nodes/DirectoryNode.cs
--- a/Engine/Source/Programs/Shared/EpicGames.Horde/Storage/Nodes/DirectoryNode.cs
+++ b/Engine/Source/Programs/Shared/EpicGames.Horde/Storage/Nodes/DirectoryNode.cs
@@ -62,36 +62,28 @@ namespace EpicGames.Horde.Storage.Nodes
 		/// <summary>
 		/// Length of this entry
 		/// </summary>
-		public long Length => (Target == null) ? _cachedLength : Target.Length;
+		public long Length { get; }
 
 		/// <summary>
 		/// Hash of the target node
 		/// </summary>
-		public IoHash Hash => (Target == null) ? _cachedHash : Target.Hash;
+		public IoHash Hash { get; }
 
 		/// <summary>
 		/// Custom user data for this file entry
 		/// </summary>
 		public ReadOnlyMemory<byte> CustomData { get; set; } = ReadOnlyMemory<byte>.Empty;
 
-		/// <summary>
-		/// Cached hash of the target node
-		/// </summary>
-		IoHash _cachedHash;
-
-		/// <summary>
-		/// Cached length of this node
-		/// </summary>
-		long _cachedLength;
-
 		/// <summary>
 		/// Constructor
 		/// </summary>
-		public FileEntry(Utf8String name, FileEntryFlags flags, FileNode node)
+		public FileEntry(Utf8String name, FileEntryFlags flags, long length, NodeHandle node)
 			: base(node)
 		{
 			Name = name;
 			Flags = flags;
+			Length = length;
+			Hash = node.Hash;
 		}
 
 		/// <summary>
@@ -103,9 +95,8 @@ namespace EpicGames.Horde.Storage.Nodes
 		{
 			Name = reader.ReadUtf8String();
 			Flags = (FileEntryFlags)reader.ReadUnsignedVarInt();
-
-			_cachedHash = reader.ReadIoHash();
-			_cachedLength = (long)reader.ReadUnsignedVarInt();
+			Length = (long)reader.ReadUnsignedVarInt();
+			Hash = reader.ReadIoHash();
 
 			if ((Flags & FileEntryFlags.HasCustomData) != 0)
 			{
@@ -118,7 +109,7 @@ namespace EpicGames.Horde.Storage.Nodes
 		/// Serialize this entry
 		/// </summary>
 		/// <param name="writer"></param>
-		protected override void Serialize(IMemoryWriter writer)
+		public override void Serialize(ITreeNodeWriter writer)
 		{
 			base.Serialize(writer);
 
@@ -126,10 +117,10 @@ namespace EpicGames.Horde.Storage.Nodes
 
 			writer.WriteUtf8String(Name);
 			writer.WriteUnsignedVarInt((ulong)flags);
-			writer.WriteIoHash(Hash);
 			writer.WriteUnsignedVarInt((ulong)Length);
+			writer.WriteIoHash(Hash);
 
-			if((flags & FileEntryFlags.HasCustomData) != 0)
+			if ((flags & FileEntryFlags.HasCustomData) != 0)
 			{
 				writer.WriteVariableLengthBytes(CustomData.Span);
 			}
@@ -160,15 +151,6 @@ namespace EpicGames.Horde.Storage.Nodes
 			await node.CopyToFileAsync(reader, file, cancellationToken);
 		}
 
-		/// <inheritdoc/>
-		protected override void OnCollapse()
-		{
-			base.OnCollapse();
-
-			_cachedHash = Target!.Hash;
-			_cachedLength = Target!.Length;
-		}
-
 		/// <inheritdoc/>
 		public override string ToString() => Name.ToString();
 	}
@@ -215,9 +197,10 @@ namespace EpicGames.Horde.Storage.Nodes
 		/// Serialize this directory entry to disk
 		/// </summary>
 		/// <param name="writer"></param>
-		protected override void Serialize(IMemoryWriter writer)
+		public override void Serialize(ITreeNodeWriter writer)
 		{
 			base.Serialize(writer);
+
 			writer.WriteUnsignedVarInt((ulong)Length);
 		}
 
@@ -272,9 +255,10 @@ namespace EpicGames.Horde.Storage.Nodes
 		/// Serialize this directory entry to disk
 		/// </summary>
 		/// <param name="writer"></param>
-		protected override void Serialize(IMemoryWriter writer)
+		public override void Serialize(ITreeNodeWriter writer)
 		{
 			base.Serialize(writer);
+
 			writer.WriteUtf8String(Name);
 		}
 
@@ -433,11 +417,12 @@ namespace EpicGames.Horde.Storage.Nodes
 		/// </summary>
 		/// <param name="name">Name of the new directory</param>
 		/// <param name="flags">Flags for the new file</param>
-		/// <param name="node">The file data</param>
+		/// <param name="length">Length of the file</param>
+		/// <param name="handle">Handle to the file data</param>
 		/// <returns>The new directory object</returns>
-		public FileEntry AddFile(Utf8String name, FileEntryFlags flags, FileNode node)
+		public FileEntry AddFile(Utf8String name, FileEntryFlags flags, long length, NodeHandle handle)
 		{
-			FileEntry entry = new FileEntry(name, flags, node);
+			FileEntry entry = new FileEntry(name, flags, length, handle);
 			AddFile(entry);
 			return entry;
 		}
@@ -448,10 +433,11 @@ namespace EpicGames.Horde.Storage.Nodes
 		/// <param name="reader">Reader for node data</param>
 		/// <param name="path">Path to the file</param>
 		/// <param name="flags">Flags for the new file</param>
-		/// <param name="node">THe file node</param>
+		/// <param name="handle">The file node</param>
+		/// <param name="length">Length of the node</param>
 		/// <param name="cancellationToken">Cancellation token for the operation</param>
 		/// <returns>The new directory object</returns>
-		public async ValueTask<FileEntry> AddFileByPathAsync(TreeReader reader, Utf8String path, FileEntryFlags flags, FileNode node, CancellationToken cancellationToken = default)
+		public async ValueTask<FileEntry> AddFileByPathAsync(TreeReader reader, Utf8String path, FileEntryFlags flags, long length, NodeHandle handle, CancellationToken cancellationToken = default)
 		{
 			DirectoryNode directory = this;
 
@@ -463,26 +449,26 @@ namespace EpicGames.Horde.Storage.Nodes
 
 			for (; ; )
 			{
-				int length = 0;
-				for (; ; length++)
+				int pathLength = 0;
+				for (; ; pathLength++)
 				{
-					if (length == remainingPath.Length)
+					if (pathLength == remainingPath.Length)
 					{
-						return directory.AddFile(remainingPath, flags, node);
+						return directory.AddFile(remainingPath, flags, length, handle);
 					}
 
-					byte character = remainingPath[length];
+					byte character = remainingPath[pathLength];
 					if (character == '\\' || character == '/')
 					{
 						break;
 					}
 				}
 
-				if (length > 0)
+				if (pathLength > 0)
 				{
-					directory = await directory.FindOrAddDirectoryAsync(reader, remainingPath.Slice(0, length), cancellationToken);
+					directory = await directory.FindOrAddDirectoryAsync(reader, remainingPath.Slice(0, pathLength), cancellationToken);
 				}
-				remainingPath = remainingPath.Slice(length + 1);
+				remainingPath = remainingPath.Slice(pathLength + 1);
 			}
 		}
 
@@ -733,12 +719,14 @@ namespace EpicGames.Horde.Storage.Nodes
 
 		static async Task CopyFilesAsync(List<(DirectoryNode DirectoryNode, FileInfo FileInfo)> files, int minIdx, int maxIdx, FileEntry[] entries, ChunkingOptions options, TreeWriter baseWriter, CancellationToken cancellationToken)
 		{
-			TreeWriter writer = new TreeWriter(baseWriter);
+			using TreeWriter writer = new TreeWriter(baseWriter);
+
+			FileNodeWriter fileNodeWriter = new FileNodeWriter(writer, options);
 			for(int idx = minIdx; idx < maxIdx; idx++)
 			{
 				FileInfo fileInfo = files[idx].FileInfo;
-				FileNode fileNode = await FileNode.CreateAsync(fileInfo, options, writer, cancellationToken);
-				entries[idx] = new FileEntry(fileInfo.Name, FileEntryFlags.None, fileNode);
+				NodeHandle handle = await fileNodeWriter.CreateAsync(fileInfo, cancellationToken);
+				entries[idx] = new FileEntry(fileInfo.Name, FileEntryFlags.None, fileNodeWriter.Length, handle);
 			}
 			await writer.FlushAsync(cancellationToken);
 		}
@@ -772,11 +760,4 @@ namespace EpicGames.Horde.Storage.Nodes
 			await Task.WhenAll(tasks);
 		}
 	}
-
-	/// <summary>
-	/// Extension methods for <see cref="DirectoryNode"/>
-	/// </summary>
-	public static class DirectoryNodeExtensions
-	{
-	}
 }
diff --git a/Engine/Source/Programs/Shared/EpicGames.Horde/Storage/Nodes/FileNode.cs b/Engine/Source/Programs/Shared/EpicGames.Horde/Storage/Nodes/FileNode.cs
--- a/Engine/Source/Programs/Shared/EpicGames.Horde/Storage/Nodes/FileNode.cs
+++ b/Engine/Source/Programs/Shared/EpicGames.Horde/Storage/Nodes/FileNode.cs
@@ -12,6 +12,7 @@ using System.Diagnostics;
 using System.IO;
 using System.Linq;
 using System.Reflection;
+using System.Reflection.Metadata;
 using System.Threading;
 using System.Threading.Tasks;
 using System.Xml.Linq;
@@ -63,6 +64,15 @@ namespace EpicGames.Horde.Storage.Nodes
 		/// </summary>
 		public int TargetSize { get; set; }
 
+		/// <summary>
+		/// Constructor
+		/// </summary>
+		/// <param name="size">Fixed size chunks to use</param>
+		public ChunkingOptionsForNodeType(int size)
+			: this(size, size, size)
+		{
+		}
+
 		/// <summary>
 		/// Default constructor
 		/// </summary>
@@ -81,186 +91,6 @@ namespace EpicGames.Horde.Storage.Nodes
 	/// </summary>
 	public abstract class FileNode : TreeNode
 	{
-		/// <summary>
-		/// Length of the node
-		/// </summary>
-		public abstract long Length { get; }
-
-		/// <summary>
-		/// Rolling hash for the current node
-		/// </summary>
-		public abstract uint RollingHash { get; }
-
-		/// <summary>
-		/// Whether the node can have data appended to it
-		/// </summary>
-		/// <returns>True if the node is read only</returns>
-		public abstract bool IsReadOnly();
-
-		/// <summary>
-		/// Creates a file node from a block of memory
-		/// </summary>
-		/// <param name="memory">The memory to read from</param>
-		/// <param name="options">Options for chunking the data</param>
-		/// <param name="writer">Writer for new tree nodes</param>
-		/// <param name="cancellationToken">Cancellation token for the operation</param>
-		/// <returns>New file node</returns>
-		public static Task<FileNode> CreateAsync(ReadOnlyMemory<byte> memory, ChunkingOptions options, TreeWriter writer, CancellationToken cancellationToken)
-		{
-			using ReadOnlyMemoryStream stream = new ReadOnlyMemoryStream(memory);
-			return CreateAsync(stream, options, writer, cancellationToken);
-		}
-
-		/// <summary>
-		/// Creates a file node from a file
-		/// </summary>
-		/// <param name="fileInfo">The stream to read from</param>
-		/// <param name="options">Options for chunking the data</param>
-		/// <param name="writer">Writer for new tree nodes</param>
-		/// <param name="cancellationToken">Cancellation token for the operation</param>
-		/// <returns>New file node</returns>
-		public static async Task<FileNode> CreateAsync(FileInfo fileInfo, ChunkingOptions options, TreeWriter writer, CancellationToken cancellationToken)
-		{
-			using (Stream stream = fileInfo.OpenRead())
-			{
-				return await FileNode.CreateAsync(stream, options, writer, cancellationToken);
-			}
-		}
-
-		/// <summary>
-		/// Creates a file node from a stream
-		/// </summary>
-		/// <param name="stream">The stream to read from</param>
-		/// <param name="options">Options for chunking the data</param>
-		/// <param name="writer">Writer for new tree nodes</param>
-		/// <param name="cancellationToken">Cancellation token for the operation</param>
-		/// <returns>New file node</returns>
-		public static async Task<FileNode> CreateAsync(Stream stream, ChunkingOptions options, TreeWriter writer, CancellationToken cancellationToken)
-		{
-			FileNode node = new LeafFileNode();
-
-			byte[] buffer = new byte[4 * 1024];
-			for (; ; )
-			{
-				int numBytes = await stream.ReadAsync(buffer, cancellationToken);
-				if (numBytes == 0)
-				{
-					break;
-				}
-				node = await node.AppendAsync(buffer.AsMemory(0, numBytes), options, writer, cancellationToken);
-			}
-
-			return node;
-		}
-
-		/// <summary>
-		/// Append data to this chunk. Must only be called on the root node in a chunk tree.
-		/// </summary>
-		/// <param name="input">The data to write</param>
-		/// <param name="options">Settings for chunking the data</param>
-		/// <param name="writer">Writer for new tree nodes</param>
-		/// <param name="cancellationToken">Cancellation token for the operation</param>
-		public ValueTask<FileNode> AppendAsync(ReadOnlyMemory<byte> input, ChunkingOptions options, TreeWriter writer, CancellationToken cancellationToken)
-		{
-			return AppendAsync(this, input, options, writer, cancellationToken);
-		}
-
-		/// <summary>
-		/// Appends data to this file
-		/// </summary>
-		/// <param name="stream">Data to append to the file</param>
-		/// <param name="options">Options for chunking the data</param>
-		/// <param name="writer">Writer for new node data</param>
-		/// <param name="cancellationToken">Cancellation token for the operation</param>
-		/// <returns>New node at the root of this file</returns>
-		public ValueTask<FileNode> AppendAsync(Stream stream, ChunkingOptions options, TreeWriter writer, CancellationToken cancellationToken)
-		{
-			return AppendAsync(this, stream, options, writer, cancellationToken);
-		}
-
-		/// <summary>
-		/// Appends data to this file
-		/// </summary>
-		/// <param name="node">Node to append to</param>
-		/// <param name="stream">Data to append to the file</param>
-		/// <param name="options">Options for chunking the data</param>
-		/// <param name="writer">Writer for new node data</param>
-		/// <param name="cancellationToken">Cancellation token for the operation</param>
-		/// <returns>New node at the root of this file</returns>
-		static async ValueTask<FileNode> AppendAsync(FileNode node, Stream stream, ChunkingOptions options, TreeWriter writer, CancellationToken cancellationToken)
-		{
-			const int BufferLength = 32 * 1024;
-
-			using IMemoryOwner<byte> owner = MemoryPool<byte>.Shared.Rent(BufferLength * 2);
-			Memory<byte> buffer = owner.Memory;
-
-			int readBufferOffset = 0;
-			Memory<byte> appendBuffer = Memory<byte>.Empty;
-			for (; ; )
-			{
-				// Start a read into memory
-				Memory<byte> readBuffer = buffer.Slice(readBufferOffset, BufferLength);
-				Task<int> readTask = Task.Run(async () => await stream.ReadAsync(readBuffer, cancellationToken), cancellationToken);
-
-				// In the meantime, append the last data that was read to the tree
-				if (appendBuffer.Length > 0)
-				{
-					node = await node.AppendAsync(appendBuffer, options, writer, cancellationToken);
-				}
-
-				// Wait for the read to finish
-				int numBytes = await readTask;
-				if (numBytes == 0)
-				{
-					break;
-				}
-
-				// Switch the buffers around
-				appendBuffer = readBuffer.Slice(0, numBytes);
-				readBufferOffset ^= BufferLength;
-			}
-			return node;
-		}
-
-		static async ValueTask<FileNode> AppendAsync(FileNode root, ReadOnlyMemory<byte> input, ChunkingOptions options, TreeWriter writer, CancellationToken cancellationToken)
-		{
-			for (; ; )
-			{
-				// Append as much data as possible to the existing tree
-				input = await root.AppendToNodeAsync(input, options, writer, cancellationToken);
-				if (input.IsEmpty)
-				{
-					break;
-				}
-
-				// Increase the height of the tree by pushing the contents of this node into a new child node
-				root = new InteriorFileNode(root.Length, root);
-			}
-			return root;
-		}
-
-		private async Task<ReadOnlyMemory<byte>> AppendToNodeAsync(ReadOnlyMemory<byte> appendData, ChunkingOptions options, TreeWriter writer, CancellationToken cancellationToken)
-		{
-			if (appendData.Length == 0 || IsReadOnly())
-			{
-				return appendData;
-			}
-			else
-			{
-				return await AppendDataAsync(appendData, options, writer, cancellationToken);
-			}
-		}
-
-		/// <summary>
-		/// Attempt to append data to the current node.
-		/// </summary>
-		/// <param name="newData">The data to append</param>
-		/// <param name="options">Options for chunking the data</param>
-		/// <param name="writer">Writer for new tree nodes</param>
-		/// <param name="cancellationToken">Cancellation token for the operation</param>
-		/// <returns>Remaining data in the buffer</returns>
-		public abstract ValueTask<ReadOnlyMemory<byte>> AppendDataAsync(ReadOnlyMemory<byte> newData, ChunkingOptions options, TreeWriter writer, CancellationToken cancellationToken);
-
 		/// <summary>
 		/// Copies the contents of this node and its children to the given output stream
 		/// </summary>
@@ -306,24 +136,10 @@ namespace EpicGames.Horde.Storage.Nodes
 	[TreeNode("{B27AFB68-9E20-4A4B-A4D8-788A4098D439}", 1)]
 	public sealed class LeafFileNode : FileNode
 	{
-		class DataSegment : ReadOnlySequenceSegment<byte>
-		{
-			public DataSegment(long runningIndex, ReadOnlyMemory<byte> data)
-			{
-				RunningIndex = runningIndex;
-				Memory = data;
-			}
-
-			public void SetNext(DataSegment next)
-			{
-				Next = next;
-			}
-		}
-
-		bool _isReadOnly;
-		uint _rollingHash;
-		int _length;
-		byte[]? _buffer;
+		/// <summary>
+		/// Data for this node
+		/// </summary>
+		public ReadOnlyMemory<byte> Data { get; }
 
 		/// <summary>
 		/// Create an empty leaf node
@@ -337,61 +153,18 @@ namespace EpicGames.Horde.Storage.Nodes
 		/// </summary>
 		public LeafFileNode(ITreeNodeReader reader)
 		{
-			_isReadOnly = true;
-			_rollingHash = reader.ReadUInt32();
-			_buffer = reader.ReadVariableLengthBytes().ToArray();
-			_length = _buffer.Length;
+			Data = reader.ReadFixedLengthBytes(reader.Length);
 		}
 
 		/// <inheritdoc/>
 		public override void Serialize(ITreeNodeWriter writer)
 		{
-			writer.WriteUInt32(_rollingHash);
-			writer.WriteVariableLengthBytes(Data.Span);
+			writer.WriteFixedLengthBytes(Data.Span);
 		}
 
 		/// <inheritdoc/>
 		public override IEnumerable<TreeNodeRef> EnumerateRefs() => Enumerable.Empty<TreeNodeRef>();
 
-		/// <inheritdoc/>
-		public override bool IsReadOnly() => _isReadOnly;
-
-		/// <inheritdoc/>
-		public override uint RollingHash => _rollingHash;
-
-		/// <summary>
-		/// Gets the data for this node
-		/// </summary>
-		public ReadOnlyMemory<byte> Data => _buffer.AsMemory(0, _length);
-
-		/// <inheritdoc/>
-		public override long Length => _length;
-
-		/// <inheritdoc/>
-		public override ValueTask<ReadOnlyMemory<byte>> AppendDataAsync(ReadOnlyMemory<byte> newData, ChunkingOptions options, TreeWriter writer, CancellationToken cancellationToken)
-		{
-			ReadOnlyMemory<byte> result;
-			if (_isReadOnly)
-			{
-				result = newData;
-			}
-			else
-			{
-				if (_buffer == null || _buffer.Length < options.LeafOptions.MaxSize)
-				{
-					Array.Resize(ref _buffer, options.LeafOptions.MaxSize);
-				}
-
-				int appendLength = AppendData(_buffer.AsSpan(0, _length), newData.Span, ref _rollingHash, options.LeafOptions);
-				newData.Slice(0, appendLength).CopyTo(_buffer.AsMemory(_length));
-				_length += appendLength;
-
-				result = newData.Slice(appendLength);
-				_isReadOnly = (result.Length > 0);
-			}
-			return new ValueTask<ReadOnlyMemory<byte>>(result);
-		}
-
 		/// <summary>
 		/// Determines how much data to append to an existing leaf node
 		/// </summary>
@@ -482,31 +255,18 @@ namespace EpicGames.Horde.Storage.Nodes
 	[TreeNode("{F4DEDDBC-70CB-4C7A-8347-F011AFCCCDB9}", 1)]
 	public class InteriorFileNode : FileNode
 	{
-		bool _isReadOnly;
-		uint _rollingHash;
-		long _length;
-		readonly List<TreeNodeRef<FileNode>> _children = new List<TreeNodeRef<FileNode>>();
-
 		/// <summary>
 		/// Child nodes
 		/// </summary>
-		public IReadOnlyList<TreeNodeRef<FileNode>> Children => _children;
-
-		/// <inheritdoc/>
-		public override long Length => _length;
-
-		/// <inheritdoc/>
-		public override uint RollingHash => _rollingHash;
+		public IReadOnlyList<TreeNodeRef<FileNode>> Children { get; }
 
 		/// <summary>
 		/// Constructor
 		/// </summary>
-		/// <param name="length"></param>
-		/// <param name="child"></param>
-		public InteriorFileNode(long length, FileNode child)
+		/// <param name="children"></param>
+		public InteriorFileNode(IReadOnlyList<TreeNodeRef<FileNode>> children)
 		{
-			_length = length;
-			_children.Add(new TreeNodeRef<FileNode>(child));
+			Children = children;
 		}
 
 		/// <summary>
@@ -514,100 +274,383 @@ namespace EpicGames.Horde.Storage.Nodes
 		/// </summary>
 		public InteriorFileNode(ITreeNodeReader reader)
 		{
-			_rollingHash = reader.ReadUInt32();
-			_length = (long)reader.ReadUnsignedVarInt();
-			_isReadOnly = reader.ReadBoolean();
-			_children = reader.ReadList(() => reader.ReadRef<FileNode>());
+			TreeNodeRef<FileNode>[] children = new TreeNodeRef<FileNode>[reader.Length / IoHash.NumBytes];
+			for (int idx = 0; idx < children.Length; idx++)
+			{
+				children[idx] = reader.ReadRef<FileNode>();
+			}
+			Children = children;
 		}
 
-		/// <inheritdoc/>
-		public override bool IsReadOnly() => _isReadOnly;
-
 		/// <inheritdoc/>
 		public override void Serialize(ITreeNodeWriter writer)
 		{
-			writer.WriteUInt32(_rollingHash);
-			writer.WriteUnsignedVarInt((ulong)_length);
-			writer.WriteBoolean(_isReadOnly);
-			writer.WriteList(_children, x => writer.WriteRef(x));
+			foreach (TreeNodeRef<FileNode> child in Children)
+			{
+				writer.WriteRef(child);
+			}
 		}
 
 		/// <inheritdoc/>
-		public override IEnumerable<TreeNodeRef> EnumerateRefs() => _children;
+		public override IEnumerable<TreeNodeRef> EnumerateRefs() => Children;
 
-		/// <inheritdoc/>
-		public override async ValueTask<ReadOnlyMemory<byte>> AppendDataAsync(ReadOnlyMemory<byte> newData, ChunkingOptions options, TreeWriter writer, CancellationToken cancellationToken)
+		/// <summary>
+		/// Test whether the current node is complete
+		/// </summary>
+		/// <param name="currentData"></param>
+		/// <param name="rollingHash"></param>
+		/// <param name="options"></param>
+		/// <returns></returns>
+		internal static bool IsComplete(ReadOnlySpan<byte> currentData, uint rollingHash, ChunkingOptionsForNodeType options)
 		{
-			for (; ; )
+			if (currentData.Length + IoHash.NumBytes > options.MaxSize)
 			{
-				Debug.Assert(_children != null);
+				return true;
+			}
 
-				// Try to write to the last node
-				if (_children.Count > 0)
+			if (currentData.Length >= options.MinSize)
+			{
+				uint rollingHashThreshold = BuzHash.GetThreshold(options.TargetSize);
+				if (rollingHash < rollingHashThreshold)
 				{
-					TreeNodeRef<FileNode> lastNodeRef = _children[^1];
-
-					FileNode? lastNode = lastNodeRef.Target;
-					if (lastNode != null)
-					{
-						// Update the length to match the new node
-						_length -= lastNode.Length;
-						newData = await lastNode.AppendDataAsync(newData, options, writer, cancellationToken);
-						_length += lastNode.Length;
-
-						// If the last node is complete, write it to the buffer
-						if (lastNode.IsReadOnly())
-						{
-							// Write the last node to allow it to be flushed
-							await writer.WriteAsync(lastNodeRef, cancellationToken);
-
-							// Update the hash
-							AppendChildHash(lastNode.RollingHash);
-
-							// Check if it's time to finish this chunk
-							uint hashThreshold = (uint)(((1L << 32) * IoHash.NumBytes) / options.LeafOptions.TargetSize);
-							if ((_children.Count >= options.InteriorOptions.MinSize && _rollingHash < hashThreshold) || (_children.Count >= options.InteriorOptions.MaxSize))
-							{
-								_isReadOnly = true;
-								return newData;
-							}
-						}
-
-						// Bail out if there's nothing left to write
-						if (newData.Length == 0)
-						{
-							return newData;
-						}
-
-						// Collapse the final node
-						await writer.WriteAsync(Children[^1], cancellationToken);
-					}
+					return true;
 				}
-
-				// Add a new child node
-				_children.Add(new TreeNodeRef<FileNode>(new LeafFileNode()));
 			}
+
+			return false;
 		}
 
 		/// <summary>
-		/// Updates the rolling hash to append a child hash
+		/// Append a new hash to this interior node
 		/// </summary>
-		/// <param name="childHash">The child hash to append</param>
-		void AppendChildHash(uint childHash)
+		/// <param name="currentData">Current data for the node</param>
+		/// <param name="hash">Hash of the child node</param>
+		/// <param name="rollingHash">Current rolling hash for the node</param>
+		/// <param name="options">Options for chunking the node</param>
+		/// <returns>True if the hash could be appended, false otherwise</returns>
+		internal static void AppendData(ReadOnlySpan<byte> currentData, IoHash hash, ref uint rollingHash, ChunkingOptionsForNodeType options)
 		{
-			Span<byte> hashData = stackalloc byte[4];
-			BinaryPrimitives.WriteUInt32LittleEndian(hashData, childHash);
-			_rollingHash = BuzHash.Add(_rollingHash, hashData);
+			Span<byte> hashData = stackalloc byte[IoHash.NumBytes];
+			hash.CopyTo(hashData);
+
+			rollingHash = BuzHash.Add(rollingHash, hashData);
+
+			int windowSize = options.MinSize - (options.MinSize % IoHash.NumBytes);
+			if (currentData.Length > windowSize)
+			{
+				ReadOnlySpan<byte> removeData = currentData.Slice(currentData.Length - windowSize, IoHash.NumBytes);
+				rollingHash = BuzHash.Sub(rollingHash, removeData, windowSize + IoHash.NumBytes);
+			}
 		}
 
 		/// <inheritdoc/>
 		public override async Task CopyToStreamAsync(TreeReader reader, Stream outputStream, CancellationToken cancellationToken)
 		{
-			foreach (TreeNodeRef<FileNode> childNodeRef in _children)
+			foreach (TreeNodeRef<FileNode> childNodeRef in Children)
 			{
 				FileNode childNode = await childNodeRef.ExpandAsync(reader, cancellationToken);
 				await childNode.CopyToStreamAsync(reader, outputStream, cancellationToken);
 			}
 		}
 	}
+
+	/// <summary>
+	/// Utility class for generating FileNode data directly into <see cref="TreeWriter"/> instances, without constructing node representations first.
+	/// </summary>
+	public class FileNodeWriter
+	{
+		class InteriorNodeState
+		{
+			public InteriorNodeState? _parent;
+			public readonly ArrayMemoryWriter _data;
+			public readonly List<NodeHandle> _children = new List<NodeHandle>();
+
+			public uint _rollingHash;
+
+			public InteriorNodeState(int maxSize)
+			{
+				_data = new ArrayMemoryWriter(maxSize);
+			}
+
+			public void Reset()
+			{
+				_rollingHash = 0;
+				_data.Clear();
+				_children.Clear();
+			}
+
+			public void Write(NodeHandle handle)
+			{
+				_children.Add(handle);
+				_data.WriteIoHash(handle.Hash);
+			}
+		}
+
+		static readonly BundleType s_leafNodeType = TreeNodeExtensions.GetBundleType(typeof(LeafFileNode));
+		static readonly BundleType s_interiorNodeType = TreeNodeExtensions.GetBundleType(typeof(InteriorFileNode));
+
+		readonly TreeWriter _writer;
+		readonly ChunkingOptions _options;
+
+		// Tree state
+		InteriorNodeState? _topInteriorNode = null;
+		long _totalLength;
+		readonly Stack<InteriorNodeState> _freeInteriorNodes = new Stack<InteriorNodeState>();
+
+		// Leaf node state
+		uint _leafHash;
+		int _leafLength;
+
+		/// <summary>
+		/// Length of the file so far
+		/// </summary>
+		public long Length => _totalLength;
+
+		/// <summary>
+		/// Constructor
+		/// </summary>
+		/// <param name="writer">Writer for new nodes</param>
+		/// <param name="options">Chunking options</param>
+		public FileNodeWriter(TreeWriter writer, ChunkingOptions options)
+		{
+			_writer = writer;
+			_options = options;
+		}
+
+		/// <summary>
+		/// Reset the current state
+		/// </summary>
+		public void Reset()
+		{
+			FreeInteriorNodes();
+			ResetLeafState();
+			_totalLength = 0;
+		}
+
+		/// <summary>
+		/// Resets the state of the current leaf node
+		/// </summary>
+		void ResetLeafState()
+		{
+			_leafHash = 0;
+			_leafLength = 0;
+		}
+
+		/// <summary>
+		/// Creates a new interior node state
+		/// </summary>
+		/// <returns>State object</returns>
+		InteriorNodeState CreateInteriorNode()
+		{
+			InteriorNodeState? result;
+			if (!_freeInteriorNodes.TryPop(out result))
+			{
+				result = new InteriorNodeState(_options.InteriorOptions.MaxSize);
+			}
+			return result;
+		}
+
+		/// <summary>
+		/// Free all the current interior nodes
+		/// </summary>
+		void FreeInteriorNodes()
+		{
+			while (_topInteriorNode != null)
+			{
+				InteriorNodeState current = _topInteriorNode;
+				_topInteriorNode = _topInteriorNode._parent;
+				current.Reset();
+				_freeInteriorNodes.Push(current);
+			}
+		}
+
+		/// <summary>
+		/// Creates data for the given file
+		/// </summary>
+		/// <param name="fileInfo">File to append</param>
+		/// <param name="cancellationToken">Cancellation token for the operation</param>
+		public async Task<NodeHandle> CreateAsync(FileInfo fileInfo, CancellationToken cancellationToken)
+		{
+			using (FileStream stream = fileInfo.OpenRead())
+			{
+				return await CreateAsync(stream, cancellationToken);
+			}
+		}
+
+		/// <summary>
+		/// Creates data from the given stream
+		/// </summary>
+		/// <param name="stream">Stream to append</param>
+		/// <param name="cancellationToken">Cancellation token for the operation</param>
+		public async Task<NodeHandle> CreateAsync(Stream stream, CancellationToken cancellationToken)
+		{
+			Reset();
+			await AppendAsync(stream, cancellationToken);
+			return await FlushAsync(cancellationToken);
+		}
+
+		/// <summary>
+		/// Creates data from the given data
+		/// </summary>
+		/// <param name="data">Stream to append</param>
+		/// <param name="cancellationToken">Cancellation token for the operation</param>
+		public async Task<NodeHandle> CreateAsync(ReadOnlyMemory<byte> data, CancellationToken cancellationToken)
+		{
+			Reset();
+			await AppendAsync(data, cancellationToken);
+			return await FlushAsync(cancellationToken);
+		}
+
+		/// <summary>
+		/// Appends data to the current file
+		/// </summary>
+		/// <param name="stream">Stream containing data to append</param>
+		/// <param name="cancellationToken">Cancellation token for the operation</param>
+		public async Task AppendAsync(Stream stream, CancellationToken cancellationToken)
+		{
+			const int BufferLength = 32 * 1024;
+
+			using IMemoryOwner<byte> owner = MemoryPool<byte>.Shared.Rent(BufferLength * 2);
+			Memory<byte> buffer = owner.Memory;
+
+			int readBufferOffset = 0;
+			Memory<byte> appendBuffer = Memory<byte>.Empty;
+			for (; ; )
+			{
+				// Start a read into memory
+				Memory<byte> readBuffer = buffer.Slice(readBufferOffset, BufferLength);
+				Task<int> readTask = Task.Run(async () => await stream.ReadAsync(readBuffer, cancellationToken), cancellationToken);
+
+				// In the meantime, append the last data that was read to the tree
+				if (appendBuffer.Length > 0)
+				{
+					await AppendAsync(appendBuffer, cancellationToken);
+				}
+
+				// Wait for the read to finish
+				int numBytes = await readTask;
+				if (numBytes == 0)
+				{
+					break;
+				}
+
+				// Switch the buffers around
+				appendBuffer = readBuffer.Slice(0, numBytes);
+				readBufferOffset ^= BufferLength;
+			}
+		}
+
+		/// <summary>
+		/// Appends data to the current file
+		/// </summary>
+		/// <param name="data">Data to append</param>
+		/// <param name="cancellationToken">Cancellation token for the operation</param>
+		public async Task AppendAsync(ReadOnlyMemory<byte> data, CancellationToken cancellationToken)
+		{
+			Memory<byte> buffer = _writer.GetOutputBuffer(_leafLength, _leafLength);
+			for (; ; )
+			{
+				// Append data to the current leaf node
+				int appendLength = LeafFileNode.AppendData(buffer.Span.Slice(0, _leafLength), data.Span, ref _leafHash, _options.LeafOptions);
+
+				buffer = _writer.GetOutputBuffer(_leafLength, _leafLength + appendLength);
+				data.Slice(0, appendLength).CopyTo(buffer.Slice(_leafLength));
+
+				_leafLength += appendLength;
+				data = data.Slice(appendLength);
+
+				_totalLength += appendLength;
+
+				if (data.Length == 0)
+				{
+					break;
+				}
+
+				// Flush the leaf node and any interior nodes that are full
+				NodeHandle handle = await WriteLeafNodeAsync(cancellationToken);
+				ResetLeafState();
+				await AddToInteriorNodeAsync(handle, cancellationToken);
+			}
+		}
+
+		async Task AddToInteriorNodeAsync(NodeHandle handle, CancellationToken cancellationToken)
+		{
+			_topInteriorNode ??= CreateInteriorNode();
+			await AddToInteriorNodeAsync(_topInteriorNode, handle, cancellationToken);
+		}
+
+		async Task AddToInteriorNodeAsync(InteriorNodeState interiorNode, NodeHandle handle, CancellationToken cancellationToken)
+		{
+			// If the node is already full, flush it
+			if (InteriorFileNode.IsComplete(interiorNode._data.WrittenSpan, interiorNode._rollingHash, _options.InteriorOptions))
+			{
+				NodeHandle interiorNodeHandle = await WriteInteriorNodeAndResetAsync(interiorNode, cancellationToken);
+				interiorNode._parent ??= CreateInteriorNode();
+				await AddToInteriorNodeAsync(interiorNode._parent, interiorNodeHandle, cancellationToken);
+			}
+
+			// Add this handle
+			InteriorFileNode.AppendData(interiorNode._data.WrittenSpan, handle.Hash, ref interiorNode._rollingHash, _options.InteriorOptions);
+			interiorNode.Write(handle);
+		}
+
+		/// <summary>
+		/// Complete the current file, and write all open nodes to the underlying writer
+		/// </summary>
+		/// <param name="cancellationToken">Cancellation token for the operation</param>
+		/// <returns>Handle to the root node</returns>
+		public async Task<NodeHandle> CompleteAsync(CancellationToken cancellationToken)
+		{
+			NodeHandle handle = await WriteLeafNodeAsync(cancellationToken);
+			ResetLeafState();
+
+			for (InteriorNodeState? state = _topInteriorNode; state != null; state = state._parent)
+			{
+				await AddToInteriorNodeAsync(state, handle, cancellationToken);
+				handle = await WriteInteriorNodeAndResetAsync(state, cancellationToken);
+			}
+
+			FreeInteriorNodes();
+			return handle;
+		}
+
+		/// <summary>
+		/// Flush the state of the writer
+		/// </summary>
+		/// <param name="cancellationToken">Cancellation token for the operation</param>
+		/// <returns>Handle to the root FileNode</returns>
+		public async Task<NodeHandle> FlushAsync(CancellationToken cancellationToken)
+		{
+			NodeHandle handle = await CompleteAsync(cancellationToken);
+			await _writer.FlushAsync(cancellationToken);
+			return handle;
+		}
+
+		/// <summary>
+		/// Writes the state of the given interior node to storage
+		/// </summary>
+		/// <param name="state"></param>
+		/// <param name="cancellationToken"></param>
+		/// <returns></returns>
+		async ValueTask<NodeHandle> WriteInteriorNodeAndResetAsync(InteriorNodeState state, CancellationToken cancellationToken)
+		{
+			Memory<byte> buffer = _writer.GetOutputBuffer(0, state._data.Length);
+			state._data.WrittenMemory.CopyTo(buffer);
+
+			NodeHandle handle = await _writer.WriteNodeAsync(state._data.Length, state._children, s_interiorNodeType, cancellationToken);
+			state.Reset();
+
+			return handle;
+		}
+
+		/// <summary>
+		/// Writes the contents of the current leaf node to storage
+		/// </summary>
+		/// <param name="cancellationToken">Cancellation token for the operation</param>
+		/// <returns>Handle to the written leaf node</returns>
+		async ValueTask<NodeHandle> WriteLeafNodeAsync(CancellationToken cancellationToken)
+		{
+			return await _writer.WriteNodeAsync(_leafLength, Array.Empty<NodeHandle>(), s_leafNodeType, cancellationToken);
+		}
+	}
 }
diff --git a/Engine/Source/Programs/Shared/EpicGames.Horde/Storage/StorageClientBase.cs b/Engine/Source/Programs/Shared/EpicGames.Horde/Storage/StorageClientBase.cs
--- a/Engine/Source/Programs/Shared/EpicGames.Horde/Storage/StorageClientBase.cs
+++ b/Engine/Source/Programs/Shared/EpicGames.Horde/Storage/StorageClientBase.cs
@@ -39,19 +39,19 @@ namespace EpicGames.Horde.Storage
 		public abstract Task DeleteRefAsync(RefName name, CancellationToken cancellationToken = default);
 
 		/// <inheritdoc/>
-		public abstract Task<NodeLocator> TryReadRefTargetAsync(RefName name, DateTime cacheTime = default, CancellationToken cancellationToken = default);
+		public abstract Task<NodeHandle?> TryReadRefTargetAsync(RefName name, DateTime cacheTime = default, CancellationToken cancellationToken = default);
 
 		/// <inheritdoc/>
-		public virtual async Task<NodeLocator> WriteRefAsync(RefName name, Bundle bundle, int exportIdx = 0, Utf8String prefix = default, RefOptions? options = null, CancellationToken cancellationToken = default)
+		public virtual async Task<NodeHandle> WriteRefAsync(RefName name, Bundle bundle, int exportIdx = 0, Utf8String prefix = default, RefOptions? options = null, CancellationToken cancellationToken = default)
 		{
 			BlobLocator locator = await this.WriteBundleAsync(bundle, prefix, cancellationToken);
-			NodeLocator target = new NodeLocator(locator, exportIdx);
+			NodeHandle target = new NodeHandle(bundle.Header.Exports[exportIdx].Hash, locator, exportIdx);
 			await WriteRefTargetAsync(name, target, options, cancellationToken);
 			return target;
 		}
 
 		/// <inheritdoc/>
-		public abstract Task WriteRefTargetAsync(RefName name, NodeLocator target, RefOptions? options = null, CancellationToken cancellationToken = default);
+		public abstract Task WriteRefTargetAsync(RefName name, NodeHandle target, RefOptions? options = null, CancellationToken cancellationToken = default);
 
 		#endregion
 	}
diff --git a/Engine/Source/Programs/Shared/EpicGames.Horde/Storage/TreeNode.cs b/Engine/Source/Programs/Shared/EpicGames.Horde/Storage/TreeNode.cs
--- a/Engine/Source/Programs/Shared/EpicGames.Horde/Storage/TreeNode.cs
+++ b/Engine/Source/Programs/Shared/EpicGames.Horde/Storage/TreeNode.cs
@@ -70,10 +70,10 @@ namespace EpicGames.Horde.Storage
 	public interface ITreeNodeWriter : IMemoryWriter
 	{
 		/// <summary>
-		/// Write a reference to another node
+		/// Writes a reference to another node
 		/// </summary>
-		/// <param name="target">Target of the reference</param>
-		void WriteRef(TreeNodeRef target);
+		/// <param name="handle">Handle to the target node</param>
+		void WriteNodeHandle(NodeHandle handle);
 	}
 
 	/// <summary>
@@ -162,6 +162,16 @@ namespace EpicGames.Horde.Storage
 			}
 		}
 
+		/// <summary>
+		/// Writes a ref to storage
+		/// </summary>
+		/// <param name="writer">Writer to serialize to</param>
+		/// <param name="value">Value to write</param>
+		public static void WriteRef(this ITreeNodeWriter writer, TreeNodeRef value)
+		{
+			value.Serialize(writer);
+		}
+
 		/// <summary>
 		/// Writes an optional ref value to storage
 		/// </summary>
@@ -191,9 +201,9 @@ namespace EpicGames.Horde.Storage
 		/// <param name="refOptions">Options for the ref</param>
 		/// <param name="cancellationToken">Cancellation token for the operation</param>
 		/// <returns>Location of node targetted by the ref</returns>
-		public static async Task<NodeLocator> WriteNodeAsync(this IStorageClient store, RefName name, TreeNode node, TreeOptions? options = null, Utf8String prefix = default, RefOptions? refOptions = null, CancellationToken cancellationToken = default)
+		public static async Task<NodeHandle> WriteNodeAsync(this IStorageClient store, RefName name, TreeNode node, TreeOptions? options = null, Utf8String prefix = default, RefOptions? refOptions = null, CancellationToken cancellationToken = default)
 		{
-			TreeWriter writer = new TreeWriter(store, options, prefix.IsEmpty ? name.Text : prefix);
+			using TreeWriter writer = new TreeWriter(store, options, prefix.IsEmpty ? name.Text : prefix);
 			return await writer.WriteAsync(name, node, refOptions, cancellationToken);
 		}
 
diff --git a/Engine/Source/Programs/Shared/EpicGames.Horde/Storage/TreeNodeRef.cs b/Engine/Source/Programs/Shared/EpicGames.Horde/Storage/TreeNodeRef.cs
--- a/Engine/Source/Programs/Shared/EpicGames.Horde/Storage/TreeNodeRef.cs
+++ b/Engine/Source/Programs/Shared/EpicGames.Horde/Storage/TreeNodeRef.cs
@@ -1,10 +1,12 @@
 // Copyright Epic Games, Inc. All Rights Reserved.
 
 using System;
+using System.Collections.Generic;
+using System.Diagnostics;
+using System.Linq;
 using System.Threading;
 using System.Threading.Tasks;
 using EpicGames.Core;
-using EpicGames.Serialization;
 
 namespace EpicGames.Horde.Storage
 {
@@ -24,20 +26,26 @@ namespace EpicGames.Horde.Storage
 	///  * In memory but target has been modified
 	///      * Locator is set but may not reflect the current node state, Target is valid, IsDirty() returns true. 
 	///      * Writing a ref transitions it to the "in storage" state. Calling Collapse is a no-op.
-	///      
+	///
 	/// The <see cref="OnCollapse"/> and <see cref="OnExpand"/> methods allow overriden implementations to cache information about the target.
+	///
+	/// Each ref must have EXACTLY one owner; sharing of refs between objects is not permitted and will break change tracking.
+	/// Multiple refs MAY point to the same target object.
+	/// 
+	/// To read an untracked object that can be added to a new ref, call <see cref="TreeReader.ReadNodeAsync(NodeLocator, CancellationToken)"/> 
+	/// directly, or use <see cref="TreeNodeRef{T}.ExpandCopyAsync(TreeReader, CancellationToken)"/>.
 	/// </summary>
 	public class TreeNodeRef
 	{
 		/// <summary>
-		/// Locator for the blob containing this node. Invalid for nodes in memory.
+		/// The target node, or null if the node is not resident in memory.
 		/// </summary>
-		public NodeLocator Locator { get; private set; }
+		public TreeNode? Target { get; private set; }
 
 		/// <summary>
-		/// The target node, or null if the node is not resident in memory.
+		/// Handle to the node if in storage (or pending write to storage)
 		/// </summary>
-		public TreeNode? Target { get; private set; }
+		public NodeHandle? Handle { get; private set; }
 
 		/// <summary>
 		/// Creates a reference to a node in memory.
@@ -51,25 +59,25 @@ namespace EpicGames.Horde.Storage
 		/// <summary>
 		/// Creates a reference to a node in storage.
 		/// </summary>
-		/// <param name="locator">Node to reference</param>
-		public TreeNodeRef(NodeLocator locator)
+		/// <param name="handle">Handle to the referenced node</param>
+		public TreeNodeRef(NodeHandle handle)
 		{
-			Locator = locator;
+			Handle = handle;
 		}
 
 		/// <summary>
 		/// Deserialization constructor
 		/// </summary>
 		/// <param name="reader"></param>
-		public TreeNodeRef(ITreeNodeReader reader) : this(reader.ReadNodeLocator())
+		public TreeNodeRef(ITreeNodeReader reader) : this(reader.ReadNodeHandle())
 		{
 		}
 
 		/// <summary>
-		/// Determines whether the the referenced node has modified from the last version written to storage
+		/// Determines whether the the referenced node has been modified from the last version written to storage
 		/// </summary>
 		/// <returns></returns>
-		public bool IsDirty() => Target != null && (Target.Revision != 0 || !Locator.IsValid());
+		public bool IsDirty() => Target != null && (Handle == null || Target.Hash != Handle.Hash);
 
 		/// <summary>
 		/// Update the reference to refer to a node in memory.
@@ -81,44 +89,38 @@ namespace EpicGames.Horde.Storage
 				throw new InvalidOperationException("Cannot mark a ref as dirty without having expanded it.");
 			}
 
-			Locator = default;
+			Handle = null;
 		}
 
 		/// <summary>
 		/// Updates the hash and revision number for the ref.
 		/// </summary>
 		/// <returns></returns>
-		internal void MarkAsPendingWrite()
+		internal void MarkAsPendingWrite(NodeHandle handle)
 		{
 			OnCollapse();
+			Handle = handle;
 		}
 
 		/// <summary>
 		/// Update the reference to refer to a location in storage.
 		/// </summary>
-		/// <param name="locator">Location of the node</param>
-		/// <param name="revision">Revision number that was written</param>
-		internal void MarkAsWritten(NodeLocator locator, uint revision)
+		internal void MarkAsWritten()
 		{
-			Locator = locator;
-			if (Target != null && Target.Revision == revision)
+			if (Target != null && Handle != null && Target.Hash == Handle.Hash)
 			{
 				Target = null;
 			}
 		}
 
-		/// <summary>
-		/// Internal helper method to allow implementations of <see cref="ITreeNodeWriter"/> to call <see cref="Serialize(IMemoryWriter)"/>.
-		/// </summary>
-		/// <param name="writer"></param>
-		internal void SerializeInternal(IMemoryWriter writer) => Serialize(writer);
-
 		/// <summary>
 		/// Serialize the node to the given writer
 		/// </summary>
 		/// <param name="writer"></param>
-		protected virtual void Serialize(IMemoryWriter writer)
+		public virtual void Serialize(ITreeNodeWriter writer)
 		{
+			Debug.Assert(Handle != null);
+			writer.WriteNodeHandle(Handle);
 		}
 
 		/// <summary>
@@ -131,7 +133,7 @@ namespace EpicGames.Horde.Storage
 		{
 			if (Target == null)
 			{
-				Target = await reader.ReadNodeAsync(Locator, cancellationToken);
+				Target = await reader.ReadNodeAsync(Handle!.Locator, cancellationToken);
 				OnExpand();
 			}
 			return Target;
@@ -186,8 +188,8 @@ namespace EpicGames.Horde.Storage
 		/// <summary>
 		/// Constructor
 		/// </summary>
-		/// <param name="target">The referenced node</param>
-		public TreeNodeRef(NodeLocator target) : base(target)
+		/// <param name="handle">Handle to the referenced node</param>
+		public TreeNodeRef(NodeHandle handle) : base(handle)
 		{
 		}
 
@@ -218,7 +220,175 @@ namespace EpicGames.Horde.Storage
 		/// <returns></returns>
 		public async ValueTask<T> ExpandCopyAsync(TreeReader reader, CancellationToken cancellationToken = default)
 		{
-			return await reader.ReadNodeAsync<T>(Locator, cancellationToken);
+			if (Handle == null)
+			{
+				throw new InvalidOperationException("TreeNodeRef has not been serialized to storage");
+			}
+			return await reader.ReadNodeAsync<T>(Handle.Locator, cancellationToken);
+		}
+	}
+
+	/// <summary>
+	/// Extension methods for writing node
+	/// </summary>
+	public static class TreeNodeRefExtensions
+	{
+		// Implementation of INodeWriter that tracks refs
+		class NodeWriter : ITreeNodeWriter
+		{
+			readonly TreeWriter _treeWriter;
+
+			Memory<byte> _memory;
+			int _length;
+			readonly IReadOnlyList<NodeHandle> _refs;
+			int _refIdx;
+
+			public int Length => _length;
+
+			public NodeWriter(TreeWriter treeWriter, IReadOnlyList<NodeHandle> refs)
+			{
+				_treeWriter = treeWriter;
+				_memory = treeWriter.GetOutputBuffer(0, 256 * 1024);
+				_refs = refs;
+			}
+
+			public void WriteNodeHandle(NodeHandle target)
+			{
+				if (_refs[_refIdx] != target)
+				{
+					throw new InvalidOperationException("Referenced node does not match the handle returned by owner's EnumerateRefs method.");
+				}
+
+				this.WriteIoHash(target.Hash);
+				_refIdx++;
+			}
+
+			public Memory<byte> GetMemory(int minSize = 1)
+			{
+				int newLength = _length + minSize;
+				if (newLength > _memory.Length)
+				{
+					_memory = _treeWriter.GetOutputBuffer(_length, Math.Max(_memory.Length * 2, newLength));
+				}
+				return _memory.Slice(_length);
+			}
+
+			public void Advance(int length) => _length += length;
+		}
+
+		/// <summary>
+		/// Class used to track nodes which are pending write (and the state of the object when the write was started)
+		/// </summary>
+		internal class NodeWriteCallback : TreeWriter.WriteCallback
+		{
+			readonly TreeNodeRef _nodeRef;
+			readonly NodeHandle _handle;
+
+			public NodeWriteCallback(TreeNodeRef nodeRef, NodeHandle handle)
+			{
+				_nodeRef = nodeRef;
+				_handle = handle;
+			}
+
+			public override void OnWrite()
+			{
+				if (_nodeRef.Handle == _handle)
+				{
+					_nodeRef.MarkAsWritten();
+				}
+			}
+		}
+
+		/// <summary>
+		/// Writes an individual node to storage
+		/// </summary>
+		/// <param name="writer">Writer to serialize nodes to</param>
+		/// <param name="nodeRef">Reference to the node</param>
+		/// <param name="cancellationToken">Cancellation token for the operation</param>
+		/// <returns>A flag indicating whether the node is dirty, and if it is, an optional bundle that contains it</returns>
+		public static async ValueTask<NodeHandle> WriteAsync(this TreeWriter writer, TreeNodeRef nodeRef, CancellationToken cancellationToken)
+		{
+			// Check we actually have a target node. If we don't, we don't need to write anything.
+			TreeNode? target = nodeRef.Target;
+			if (target == null)
+			{
+				Debug.Assert(nodeRef.Handle != null);
+				return nodeRef.Handle;
+			}
+
+			// Write all the nodes it references, and mark the ref as dirty if any of them change.
+			List<TreeNodeRef> nextRefs = target.EnumerateRefs().ToList();
+			List<NodeHandle> nextRefHandles = new List<NodeHandle>(nextRefs.Count);
+			foreach (TreeNodeRef nextRef in nextRefs)
+			{
+				NodeHandle nextRefHandle = await WriteAsync(writer, nextRef, cancellationToken);
+				if (!nextRefHandle.Locator.IsValid())
+				{
+					nodeRef.MarkAsDirty();
+				}
+				nextRefHandles.Add(nextRefHandle);
+			}
+
+			// If the target node hasn't been modified, use the existing serialized state.
+			if (!nodeRef.IsDirty())
+			{
+				// Make sure the locator is valid. The node may be queued for writing but not flushed to disk yet.
+				Debug.Assert(nodeRef.Handle != null);
+				if (nodeRef.Handle.Locator.IsValid())
+				{
+					nodeRef.Collapse();
+				}
+				return nodeRef.Handle;
+			}
+
+			// Serialize the node
+			NodeWriter nodeWriter = new NodeWriter(writer, nextRefHandles);
+			target.Serialize(nodeWriter);
+
+			// Write the final data
+			NodeHandle handle = await writer.WriteNodeAsync(nodeWriter.Length, nextRefHandles, target.GetBundleType(), cancellationToken);
+			target.Hash = handle.Hash;
+
+			NodeWriteCallback writeState = new NodeWriteCallback(nodeRef, handle);
+			handle.AddWriteCallback(writeState);
+
+			nodeRef.MarkAsPendingWrite(handle);
+			return handle;
+		}
+
+		/// <summary>
+		/// Writes a node to the given ref
+		/// </summary>
+		/// <param name="writer"></param>
+		/// <param name="name">Name of the ref to write</param>
+		/// <param name="node"></param>
+		/// <param name="options"></param>
+		/// <param name="cancellationToken"></param>
+		/// <returns></returns>
+		public static async Task<NodeHandle> WriteAsync(this TreeWriter writer, RefName name, TreeNode node, RefOptions? options = null, CancellationToken cancellationToken = default)
+		{
+			TreeNodeRef nodeRef = new TreeNodeRef(node);
+			await writer.WriteAsync(nodeRef, cancellationToken);
+			await writer.FlushAsync(cancellationToken);
+
+			Debug.Assert(nodeRef.Handle != null);
+			await writer.Store.WriteRefTargetAsync(name, nodeRef.Handle, options, cancellationToken);
+			return nodeRef.Handle;
+		}
+
+		/// <summary>
+		/// Flushes all the current nodes to storage
+		/// </summary>
+		/// <param name="writer"></param>
+		/// <param name="root">Root for the tree</param>
+		/// <param name="cancellationToken">Cancellation token for the operation</param>
+		/// <returns></returns>
+		public static async Task<NodeHandle> FlushAsync(this TreeWriter writer, TreeNode root, CancellationToken cancellationToken)
+		{
+			TreeNodeRef rootRef = new TreeNodeRef(root);
+			await writer.WriteAsync(rootRef, cancellationToken);
+			await writer.FlushAsync(cancellationToken);
+			return rootRef.Handle!;
 		}
 	}
 }
diff --git a/Engine/Source/Programs/Shared/EpicGames.Horde/Storage/TreeReader.cs b/Engine/Source/Programs/Shared/EpicGames.Horde/Storage/TreeReader.cs
--- a/Engine/Source/Programs/Shared/EpicGames.Horde/Storage/TreeReader.cs
+++ b/Engine/Source/Programs/Shared/EpicGames.Horde/Storage/TreeReader.cs
@@ -37,14 +37,14 @@ namespace EpicGames.Horde.Storage
 		IoHash Hash { get; }
 
 		/// <summary>
-		/// References to other nodes
+		/// Locations of all referenced nodes.
 		/// </summary>
 		IReadOnlyList<NodeLocator> References { get; }
 
 		/// <summary>
 		/// Reads a reference to another node
 		/// </summary>
-		NodeLocator ReadNodeLocator();
+		NodeHandle ReadNodeHandle();
 	}
 
 	/// <summary>
@@ -237,6 +237,8 @@ namespace EpicGames.Horde.Storage
 			readonly BundleType _type;
 			readonly int _length;
 
+			int _refIdx;
+
 			public NodeReader(ReadOnlyMemory<byte> data, IoHash hash, IReadOnlyList<NodeLocator> refs, BundleType type)
 				: base(data)
 			{
@@ -254,10 +256,10 @@ namespace EpicGames.Horde.Storage
 
 			public IReadOnlyList<NodeLocator> References => _refs;
 
-			public NodeLocator ReadNodeLocator()
+			public NodeHandle ReadNodeHandle()
 			{
-				int idx = (int)this.ReadUnsignedVarInt();
-				return _refs[idx];
+				IoHash hash = this.ReadIoHash();
+				return new NodeHandle(hash, _refs[_refIdx++]);
 			}
 		}
 
@@ -361,7 +363,7 @@ namespace EpicGames.Horde.Storage
 						QueuedHeader? queuedHeader = _queuedHeaders.FirstOrDefault(x => !currentBlobIds.Contains(x.BlobId));
 						if (queuedHeader != null)
 						{
-							Task task = Task.Run(() => PerformHeaderReadAsync(queuedHeader, cancellationToken), cancellationToken);
+							Task task = Task.Run(() => PerformHeaderReadGuardedAsync(queuedHeader, cancellationToken), cancellationToken);
 							currentTasks.Add((queuedHeader.BlobId, task));
 							continue;
 						}
@@ -370,7 +372,7 @@ namespace EpicGames.Horde.Storage
 						QueuedPacket? queuedPacket = _queuedPackets.FirstOrDefault(x => !currentBlobIds.Contains(x.BlobId));
 						if (queuedPacket != null)
 						{
-							Task task = Task.Run(() => PerformPacketReadAsync(queuedPacket, cancellationToken), cancellationToken);
+							Task task = Task.Run(() => PerformPacketReadGuardedAsync(queuedPacket, cancellationToken), cancellationToken);
 							currentTasks.Add((queuedPacket.BlobId, task));
 							continue;
 						}
@@ -406,6 +408,23 @@ namespace EpicGames.Horde.Storage
 			}
 		}
 
+		/// <summary>
+		/// Reads a bundle header from the queue
+		/// </summary>
+		/// <param name="queuedHeader">The header to read</param>
+		/// <param name="cancellationToken">Cancellation token for the operation</param>
+		async Task PerformHeaderReadGuardedAsync(QueuedHeader queuedHeader, CancellationToken cancellationToken)
+		{
+			try
+			{
+				await PerformHeaderReadAsync(queuedHeader, cancellationToken);
+			}
+			catch (Exception ex)
+			{
+				queuedHeader.CompletionSource.SetException(ex);
+			}
+		}
+
 		/// <summary>
 		/// Reads a bundle header from the queue
 		/// </summary>
@@ -478,6 +497,23 @@ namespace EpicGames.Horde.Storage
 			}
 		}
 
+		/// <summary>
+		/// Reads a packet from storage
+		/// </summary>
+		/// <param name="queuedPacket">The packet to read</param>
+		/// <param name="cancellationToken">Cancellation token for the operation</param>
+		async Task PerformPacketReadGuardedAsync(QueuedPacket queuedPacket, CancellationToken cancellationToken)
+		{
+			try
+			{
+				await PerformPacketReadAsync(queuedPacket, cancellationToken);
+			}
+			catch (Exception ex)
+			{
+				queuedPacket.CompletionSource.SetException(ex);
+			}
+		}
+
 		/// <summary>
 		/// Reads a packet from storage
 		/// </summary>
@@ -735,12 +771,12 @@ namespace EpicGames.Horde.Storage
 		/// <returns>Node for the given ref, or null if it does not exist</returns>
 		public async Task<TNode?> TryReadNodeAsync<TNode>(RefName name, DateTime cacheTime = default, CancellationToken cancellationToken = default) where TNode : TreeNode
 		{
-			NodeLocator refTarget = await _store.TryReadRefTargetAsync(name, cacheTime, cancellationToken);
-			if (!refTarget.IsValid())
+			NodeHandle? refTarget = await _store.TryReadRefTargetAsync(name, cacheTime, cancellationToken);
+			if (refTarget == null)
 			{
 				return null;
 			}
-			return await ReadNodeAsync<TNode>(refTarget, cancellationToken);
+			return await ReadNodeAsync<TNode>(refTarget.Locator, cancellationToken);
 		}
 
 		/// <summary>
@@ -752,8 +788,8 @@ namespace EpicGames.Horde.Storage
 		/// <returns>Node for the given ref, or null if it does not exist</returns>
 		public async Task<TreeNodeRef<TNode>?> TryReadNodeRefAsync<TNode>(RefName name, DateTime cacheTime = default, CancellationToken cancellationToken = default) where TNode : TreeNode
 		{
-			NodeLocator refTarget = await _store.TryReadRefTargetAsync(name, cacheTime, cancellationToken);
-			if (!refTarget.IsValid())
+			NodeHandle? refTarget = await _store.TryReadRefTargetAsync(name, cacheTime, cancellationToken);
+			if (refTarget == null)
 			{
 				return null;
 			}
diff --git a/Engine/Source/Programs/Shared/EpicGames.Horde/Storage/TreeWriter.cs b/Engine/Source/Programs/Shared/EpicGames.Horde/Storage/TreeWriter.cs
--- a/Engine/Source/Programs/Shared/EpicGames.Horde/Storage/TreeWriter.cs
+++ b/Engine/Source/Programs/Shared/EpicGames.Horde/Storage/TreeWriter.cs
@@ -4,10 +4,14 @@ using System;
 using System.Buffers;
 using System.Collections.Generic;
 using System.Diagnostics;
+using System.Diagnostics.CodeAnalysis;
 using System.Linq;
 using System.Threading;
 using System.Threading.Tasks;
 using EpicGames.Core;
+using Microsoft.CodeAnalysis;
+using Microsoft.Extensions.Caching.Memory;
+using static EpicGames.Horde.Storage.TreeWriter;
 
 namespace EpicGames.Horde.Storage
 {
@@ -30,338 +34,518 @@ namespace EpicGames.Horde.Storage
 		/// Minimum size of a block to be compressed
 		/// </summary>
 		public int MinCompressionPacketSize { get; set; } = 16 * 1024;
+
+		/// <summary>
+		/// Maximum amount of data to store in memory. This includes any background writes as well as bundles being built.
+		/// </summary>
+		public int MaxInMemoryDataLength { get; set; } = 128 * 1024 * 1024;
+
+		/// <summary>
+		/// Number of nodes to cache
+		/// </summary>
+		public int NodeCacheSize { get; set; } = 1024;
 	}
 
 	/// <summary>
-	/// Writes nodes of a tree to an <see cref="IStorageClient"/>, packed into bundles.
+	/// Unique identifier for a node
 	/// </summary>
-	public class TreeWriter
+	/// <param name="Hash">Hash of the node</param>
+	/// <param name="Type">Type of the node</param>
+	public record NodeKey(IoHash Hash, BundleType Type);
+
+	/// <summary>
+	/// Handle to a node. Can be used to reference nodes that have not been flushed yet.
+	/// </summary>
+	public class NodeHandle
 	{
-		// Information about a TreeNodeRef that was written, and needs to be fixed up after flushing
-		class NodeRefInfo
+		/// <summary>
+		/// Hash of the target node
+		/// </summary>
+		public IoHash Hash { get; }
+
+		/// <summary>
+		/// Location of the node in storage
+		/// </summary>
+		public NodeLocator Locator { get; protected set; }
+
+		/// <summary>
+		/// Used to track the bundle that is being written to
+		/// </summary>
+		internal PendingBundle? PendingBundle { get; set; }
+
+		/// <summary>
+		/// Constructor
+		/// </summary>
+		/// <param name="hash">Hash of the target node</param>
+		/// <param name="locator">Location of the node in storage</param>
+		public NodeHandle(IoHash hash, NodeLocator locator)
+		{
+			Hash = hash;
+			Locator = locator;
+		}
+
+		/// <summary>
+		/// Constructor
+		/// </summary>
+		/// <param name="hash">Hash of the target node</param>
+		/// <param name="locator">Location of the node in storage</param>
+		/// <param name="exportIdx">Index of the bundle export</param>
+		public NodeHandle(IoHash hash, BlobLocator locator, int exportIdx)
+			: this(hash, new NodeLocator(locator, exportIdx))
 		{
-			public readonly TreeNodeRef NodeRef;
-			public readonly uint Revision;
-			public NodeRefInfo? Next;
+		}
 
-			public NodeRefInfo(TreeNodeRef nodeRef, NodeRefInfo? next)
+		/// <summary>
+		/// Adds a callback to be executed once the node has been written. Triggers immediately if the node has already been written.
+		/// </summary>
+		/// <param name="callback">Action to be executed after the write</param>
+		internal void AddWriteCallback(WriteCallback callback)
+		{
+			PendingBundle? pendingBundle = PendingBundle;
+			if (pendingBundle == null || !pendingBundle.TryAddWriteCallback(callback))
 			{
-				NodeRef = nodeRef;
-				Revision = nodeRef.Target!.Revision;
-				Next = next;
+				Debug.Assert(Locator.IsValid());
+				callback.OnWrite();
 			}
 		}
 
-		// Unique identifier for an output node
-		record OutputNodeKey(IoHash Hash, BundleType Type);
-
-		// Information about a unique output node. Note that multiple node refs may de-duplicate to the same output node.
-		class OutputNodeInfo : NodeRefInfo
+		/// <summary>
+		/// Parse a node handle value from a string
+		/// </summary>
+		/// <param name="text">Text to parse</param>
+		/// <returns>Parsed node handle</returns>
+		public static NodeHandle Parse(string text)
 		{
-			public readonly OutputNodeKey Key;
-			public readonly int Length;
-			public readonly IReadOnlyList<TreeNodeRef> Refs;
-
-			public OutputNodeInfo(OutputNodeKey key, int length, IReadOnlyList<TreeNodeRef> refs, TreeNodeRef nodeRef)
-				: base(nodeRef, null)
+			int hashLength = IoHash.NumBytes * 2;
+			if (text.Length == hashLength)
 			{
-				Key = key;
-				Length = length;
-				Refs = refs;
+				return new NodeHandle(IoHash.Parse(text), default);
+			}
+			else if (text[hashLength] == '@')
+			{
+				return new NodeHandle(IoHash.Parse(text.Substring(0, hashLength)), NodeLocator.Parse(text.Substring(hashLength + 1)));
+			}
+			else
+			{
+				throw new FormatException("Invalid NodeHandle value");
 			}
-
-			public void Add(TreeNodeRef nodeRef) => Next = new NodeRefInfo(nodeRef, Next);
 		}
 
-		class NodeWriter : ITreeNodeWriter
+		/// <inheritdoc/>
+		public override string ToString()
 		{
-			public ByteArrayBuilder Data = new ByteArrayBuilder();
-
-			readonly IReadOnlyDictionary<TreeNodeRef, int> _refToIndex;
-
-			public int Length => Data.Length;
-
-			public NodeWriter(IReadOnlyList<TreeNodeRef> refs)
+			if (Locator.IsValid())
 			{
-				_refToIndex = new Dictionary<TreeNodeRef, int>(refs.Distinct().Select((x, i) => KeyValuePair.Create(x, i)));
+				return $"{Hash}@{Locator}";
 			}
-
-			public void WriteRef(TreeNodeRef target)
+			else
 			{
-				if (!_refToIndex.TryGetValue(target, out int index))
-				{
-					throw new InvalidOperationException($"Cannot serialize value; was not returned by owner's {nameof(target.Target.EnumerateRefs)} method.");
-				}
-				Data.WriteUnsignedVarInt(index);
-				target.SerializeInternal(this);
+				return $"{Hash}";
 			}
-
-			public Memory<byte> GetMemory(int minSize = 1) => Data.GetMemory(minSize);
-
-			public void Advance(int length) => Data.Advance(length);
 		}
+	}
 
-		static readonly TreeOptions s_defaultOptions = new TreeOptions();
-
-		readonly IStorageClient _store;
-		readonly TreeOptions _options;
-		readonly Utf8String _prefix;
-
-		// Writer used to create packets
-		readonly ArrayMemoryWriter _packetWriter;
-
-		// Map of keys to nodes in the queue
-		readonly Dictionary<OutputNodeKey, OutputNodeInfo> _nodeKeyToInfo = new Dictionary<OutputNodeKey, OutputNodeInfo>();
-
-		// Map of hashes to existing nodes. Empty/invalid locators are used for items that are queued for writing, but not available yet.
-		readonly Dictionary<OutputNodeKey, NodeLocator> _nodeKeyToLocator = new Dictionary<OutputNodeKey, NodeLocator>(); // TODO: this needs to include some additional state from external sources.
-
-		// Queue of nodes for the current bundle
-		readonly List<OutputNodeInfo> _queue = new List<OutputNodeInfo>();
-
-		// Map of node refs to their serialized state
-		readonly HashSet<(TreeNodeRef, uint)> _queuedRefs = new HashSet<(TreeNodeRef, uint)>();
-
-		// List of packets in the current bundle
-		readonly List<BundlePacket> _packets = new List<BundlePacket>();
-
-		// List of compressed packets
-		readonly List<ReadOnlyMemory<byte>> _encodedPackets = new List<ReadOnlyMemory<byte>>();
+	/// <summary>
+	/// Index of known nodes that can be used for deduplication.
+	/// </summary>
+	public class NodeCache
+	{
+		readonly object _lockObject = new object();
+		readonly int _maxKeys;
+		readonly Queue<NodeKey> _nodeKeys = new Queue<NodeKey>();
+		readonly Dictionary<NodeKey, NodeHandle> _nodeKeyToHandle = new Dictionary<NodeKey, NodeHandle>();
 
-		// Total size of compressed packets in the current bundle
-		int _currentBundleLength;
+		/// <summary>
+		/// Constructor
+		/// </summary>
+		/// <param name="maxKeys">Maximum number of node keys to keep in the cache</param>
+		public NodeCache(int maxKeys)
+		{
+			_maxKeys = maxKeys;
+			_nodeKeys = new Queue<NodeKey>(maxKeys);
+			_nodeKeyToHandle = new Dictionary<NodeKey, NodeHandle>(maxKeys);
+		}
 
 		/// <summary>
-		/// Accessor for the store backing this writer
+		/// Adds a new node handle to the cache
 		/// </summary>
-		public IStorageClient Store => _store;
+		/// <param name="key">Unique node key</param>
+		/// <param name="handle">Handle to the node</param>
+		public void Add(NodeKey key, NodeHandle handle)
+		{
+			lock (_lockObject)
+			{
+				AddInternal(key, handle);
+			}
+		}
 
 		/// <summary>
-		/// Constructor
+		/// Adds nodes exported from a bundle to the cache
 		/// </summary>
-		/// <param name="store">Store to write data to</param>
-		/// <param name="options">Options for the writer</param>
-		/// <param name="prefix">Prefix for blobs written to the store</param>
-		public TreeWriter(IStorageClient store, TreeOptions? options = null, Utf8String prefix = default)
+		/// <param name="locator">Locator for the bundle</param>
+		/// <param name="header">The bundle header</param>
+		public void Add(BlobLocator locator, BundleHeader header)
 		{
-			_store = store;
-			_options = options ?? s_defaultOptions;
-			_prefix = prefix;
-			_packetWriter = new ArrayMemoryWriter(_options.MinCompressionPacketSize);
+			lock (_lockObject)
+			{
+				for (int idx = 0; idx < header.Exports.Count; idx++)
+				{
+					BundleExport export = header.Exports[idx];
+
+					NodeKey key = new NodeKey(export.Hash, header.Types[export.TypeIdx]);
+					NodeLocator node = new NodeLocator(locator, idx);
+					NodeHandle handle = new NodeHandle(export.Hash, node);
+
+					AddInternal(key, handle);
+				}
+			}
+		}
+
+		void AddInternal(NodeKey key, NodeHandle handle)
+		{
+			NodeKey? prevKey;
+			if (_nodeKeys.Count == _maxKeys && _nodeKeys.TryDequeue(out prevKey))
+			{
+				_nodeKeyToHandle.Remove(prevKey);
+			}
+			_nodeKeyToHandle.TryAdd(key, handle);
 		}
 
 		/// <summary>
-		/// Constructor
+		/// Find a node within the cache
 		/// </summary>
-		/// <param name="store">Store to write data to</param>
-		/// <param name="refName">Ref being written. Will be used as a prefix for storing blobs.</param>
-		/// <param name="options">Options for the writer</param>
-		public TreeWriter(IStorageClient store, RefName refName, TreeOptions? options = null)
-			: this(store, options, refName.Text)
+		/// <param name="key">Key to look up in the cache</param>
+		/// <param name="handle">Handle for the node</param>
+		/// <returns>True if the node was found</returns>
+		public bool TryGetNode(NodeKey key, [NotNullWhen(true)] out NodeHandle? handle)
 		{
+			lock (_lockObject)
+			{
+				return _nodeKeyToHandle.TryGetValue(key, out handle);
+			}
 		}
+	}
 
+	/// <summary>
+	/// Writes nodes of a tree to an <see cref="IStorageClient"/>, packed into bundles. Each <see cref="TreeWriter"/> instance is single threaded,
+	/// but multiple instances may be written to in parallel.
+	/// </summary>
+	public sealed class TreeWriter : IDisposable
+	{
 		/// <summary>
-		/// Copies settings from another tree writer
+		/// Object to receive notifications on a node being written
 		/// </summary>
-		/// <param name="other">Other instance to copy from</param>
-		public TreeWriter(TreeWriter other)
-			: this(other._store, other._options, other._prefix)
+		internal abstract class WriteCallback
 		{
+			internal WriteCallback? _next;
+
+			public abstract void OnWrite();
 		}
 
-		/// <inheritdoc/>
-		public async Task<bool> WriteAsync(TreeNodeRef nodeRef, CancellationToken cancellationToken = default)
+		// Information about a unique output node. Note that multiple node refs may de-duplicate to the same output node.
+		internal class PendingNode : NodeHandle
 		{
-			// Check we actually have a target node. If we don't, we don't need to write anything.
-			TreeNode? target = nodeRef.Target;
-			if (target == null)
-			{
-				Debug.Assert(nodeRef.Locator.IsValid());
-				return false;
-			}
+			public readonly NodeKey Key;
+			public readonly int Length;
+			public readonly NodeHandle[] Refs;
 
-			// Early-out if this ref is already queued for writing.
-			uint revision = target.Revision;
-			if (!_queuedRefs.Add((nodeRef, revision)))
+			public PendingNode(NodeKey key, int length, IReadOnlyList<NodeHandle> refs, PendingBundle pendingBundle)
+				: base(key.Hash, default)
 			{
-				return true;
+				Key = key;
+				Length = length;
+				Refs = refs.ToArray();
+				PendingBundle = pendingBundle;
 			}
 
-			// Write all the nodes it references, and mark the ref as dirty if any of them change.
-			List<TreeNodeRef> nextRefs = target.EnumerateRefs().ToList();
-			foreach (TreeNodeRef nextRef in nextRefs)
+			public void MarkAsWritten(NodeLocator locator)
 			{
-				if (await WriteAsync(nextRef, cancellationToken))
-				{
-					nodeRef.MarkAsDirty();
-				}
+				Debug.Assert(!Locator.IsValid());
+				Locator = locator;
+				PendingBundle = null;
 			}
+		}
 
-			// If the target node hasn't been modified, use the existing serialized state.
-			if (!nodeRef.IsDirty())
+		// Information about a bundle being built. Metadata operations are synchronous, compression/writes are asynchronous.
+		internal class PendingBundle : IDisposable
+		{
+			class WriteCallbackSentinel : WriteCallback
 			{
-				// Make sure the locator is valid. The node may be queued for writing but not flushed to disk yet.
-				Debug.Assert(nodeRef.Locator.IsValid());
-				nodeRef.Collapse();
-				return false;
+				public override void OnWrite() => throw new NotImplementedException();
 			}
 
-			// Serialize the node
-			NodeWriter nodeWriter = new NodeWriter(nextRefs);
-			target.Serialize(nodeWriter);
+			static readonly WriteCallbackSentinel s_callbackSentinel = new WriteCallbackSentinel();
 
-			// Get the hash for the new blob
-			ReadOnlySequence<byte> data = nodeWriter.Data.AsSequence();
-			target.Hash = IoHash.Compute(data);
-			nodeRef.MarkAsPendingWrite();
+			readonly int _maxPacketSize;
+			readonly BundleCompressionFormat _compressionFormat;
+
+			// The current packet being built
+			IMemoryOwner<byte>? _currentPacket;
+
+			// Length of the packet that has been written
+			int _currentPacketLength;
+
+			// Map of keys to nodes in the queue
+			public readonly Dictionary<NodeKey, PendingNode> _nodeKeyToInfo = new Dictionary<NodeKey, PendingNode>();
+
+			// Queue of nodes for the current bundle
+			readonly List<PendingNode> _queue = new List<PendingNode>();
 
-			OutputNodeKey nodeKey = new OutputNodeKey(target.Hash, target.GetBundleType());
+			// List of packets in the current bundle
+			readonly List<BundlePacket> _packets = new List<BundlePacket>();
 
-			// Check if we're already tracking a node with the same hash
-			OutputNodeInfo? nodeInfo;
-			if (_nodeKeyToInfo.TryGetValue(nodeKey, out nodeInfo))
+			// List of compressed packets
+			readonly ChunkedMemoryWriter _encodedPacketWriter;
+
+			// Set of all direct dependencies from this bundle
+			readonly HashSet<Task> _dependencies = new HashSet<Task>();
+
+			// Total size of uncompressed data in the current bundle
+			public int UncompressedLength { get; private set; }
+
+			// Whether all nodes have been added to a bundle
+			public bool IsReadOnly { get; private set; }
+
+			// List of post-write callbacks
+			WriteCallback? _callbacks = null;
+
+			// Task used to compress data in the background
+			Task _writeTask = Task.CompletedTask;
+
+			// Event which is signalled after the bundle is written to storage
+			readonly TaskCompletionSource<bool> _completeEvent = new TaskCompletionSource<bool>();
+
+			// Task signalled after the write is complete
+			public Task CompleteTask => _completeEvent.Task;
+
+			public PendingBundle(int maxPacketSize, int maxBlobSize, BundleCompressionFormat compressionFormat)
 			{
-				// This node is already queued for writing to storage
-				nodeInfo.Add(nodeRef);
+				_maxPacketSize = maxPacketSize;
+				_compressionFormat = compressionFormat;
+				_encodedPacketWriter = new ChunkedMemoryWriter(maxBlobSize);
 			}
-			else if (_nodeKeyToLocator.TryGetValue(nodeKey, out NodeLocator locator))
+
+			/// <inheritdoc/>
+			public void Dispose()
 			{
-				// This node has already been written to storage
-				nodeRef.MarkAsWritten(locator, nodeRef.Target!.Revision);
+				if (_currentPacket != null)
+				{
+					_currentPacket.Dispose();
+					_currentPacket = null;
+				}
+				_encodedPacketWriter.Dispose();
 			}
-			else
+
+			// Whether this bundle has finished writing
+			public bool IsComplete() => CompleteTask.IsCompleted;
+
+			// Whether this bundle can be written
+			public bool CanComplete() => !IsReadOnly && _dependencies.Count == 0;
+
+			// Add a dependency onto another bundle
+			public void AddDependencyOn(PendingBundle bundle)
 			{
-				// Write the node if we don't already have it
-				_packetWriter.WriteFixedLengthBytes(data);
+				if (bundle != this)
+				{
+					_dependencies.Add(bundle.CompleteTask);
+				}
+			}
 
-				nodeInfo = new OutputNodeInfo(nodeKey, (int)data.Length, nextRefs, nodeRef);
-				_queue.Add(nodeInfo);
-				_nodeKeyToInfo.Add(nodeKey, nodeInfo);
+			// Adds a callback after writing
+			public bool TryAddWriteCallback(WriteCallback callback)
+			{
+				for (; ; )
+				{
+					WriteCallback? tail = _callbacks;
+					if (tail == s_callbackSentinel)
+					{
+						return false;
+					}
 
-				if (_packetWriter.Length > Math.Min(_options.MinCompressionPacketSize, _options.MaxBlobSize))
+					callback._next = tail;
+
+					if (Interlocked.CompareExchange(ref _callbacks, callback, tail) == tail)
+					{
+						return true;
+					}
+				}
+			}
+
+			// Try to add a ref to an existing output node
+			public bool TryAddRefToExistingNode(NodeKey nodeKey, [NotNullWhen(true)] out PendingNode? handle)
+			{
+				PendingNode? pendingNode;
+				if (_nodeKeyToInfo.TryGetValue(nodeKey, out pendingNode))
 				{
-					FlushPacket();
+					handle = pendingNode;
+					return true;
 				}
+				else
+				{
+					handle = null;
+					return false;
+				}
+			}
 
-				if (_currentBundleLength > _options.MaxBlobSize)
+			// Starts writing a node to the buffer
+			public Memory<byte> GetBuffer(int usedSize, int desiredSize)
+			{
+				if (_currentPacket == null)
+				{
+					_currentPacket = MemoryPool<byte>.Shared.Rent(Math.Max(_maxPacketSize, desiredSize));
+				}
+				else if (_currentPacketLength + desiredSize > _maxPacketSize)
 				{
-					await FlushAsync(cancellationToken);
+					IMemoryOwner<byte> nextPacket = MemoryPool<byte>.Shared.Rent(Math.Max(_maxPacketSize, desiredSize));
+					if (usedSize > 0)
+					{
+						_currentPacket.Memory.Slice(_currentPacketLength, usedSize).CopyTo(nextPacket.Memory);
+					}
+					FlushPacket();
+					_currentPacket = nextPacket;
+					return _currentPacket.Memory.Slice(_currentPacketLength, desiredSize);
 				}
+				return _currentPacket.Memory.Slice(_currentPacketLength, desiredSize);
 			}
 
-			return true;
-		}
+			// Finish a node write
+			public PendingNode WriteNode(NodeKey nodeKey, int size, IReadOnlyList<NodeHandle> refs)
+			{
+				_currentPacketLength += size;
+				UncompressedLength += size;
 
-		/// <summary>
-		/// Writes a node to the given ref
-		/// </summary>
-		/// <param name="name">Name of the ref to write</param>
-		/// <param name="node"></param>
-		/// <param name="options"></param>
-		/// <param name="cancellationToken"></param>
-		/// <returns></returns>
-		public async Task<NodeLocator> WriteAsync(RefName name, TreeNode node, RefOptions? options = null, CancellationToken cancellationToken = default)
-		{
-			TreeNodeRef nodeRef = new TreeNodeRef(node);
-			await WriteAsync(nodeRef, cancellationToken);
-			await FlushAsync(cancellationToken);
+				PendingNode pendingNode = new PendingNode(nodeKey, (int)size, refs, this);
+				_queue.Add(pendingNode);
+				_nodeKeyToInfo.Add(nodeKey, pendingNode);
 
-			Debug.Assert(nodeRef.Locator.IsValid());
-			await _store.WriteRefTargetAsync(name, nodeRef.Locator, options, cancellationToken);
-			return nodeRef.Locator;
-		}
+				return pendingNode;
+			}
 
-		/// <summary>
-		/// Compresses the current packet and schedule it to be written to storage
-		/// </summary>
-		void FlushPacket()
-		{
-			if (_packetWriter.Length > 0)
+			// Compresses the current packet and schedule it to be written to storage
+			void FlushPacket()
 			{
-				ReadOnlyMemory<byte> encodedData = BundleData.Compress(_options.CompressionFormat, _packetWriter.WrittenMemory);
-				_encodedPackets.Add(encodedData);
+				if (_currentPacket != null)
+				{
+					IMemoryOwner<byte> currentPacket = _currentPacket;
+					int currentPacketLength = _currentPacketLength;
 
-				BundlePacket packet = new BundlePacket(encodedData.Length, _packetWriter.Length);
-				_packets.Add(packet);
+					if (currentPacketLength > 0)
+					{
+						_writeTask = _writeTask.ContinueWith(x => CompressPacket(currentPacket, currentPacketLength), TaskScheduler.Default);
+					}
+					else
+					{
+						currentPacket.Dispose();
+					}
 
-				_packetWriter.Clear();
-				_currentBundleLength += encodedData.Length;
+					_currentPacket = null;
+					_currentPacketLength = 0;
+				}
 			}
-		}
 
-		/// <summary>
-		/// Flushes all the current nodes to storage
-		/// </summary>
-		/// <param name="root">Root for the tree</param>
-		/// <param name="cancellationToken">Cancellation token for the operation</param>
-		/// <returns></returns>
-		public async Task<NodeLocator> FlushAsync(TreeNode root, CancellationToken cancellationToken)
-		{
-			TreeNodeRef rootRef = new TreeNodeRef(root);
-			await WriteAsync(rootRef, cancellationToken);
-			await FlushAsync(cancellationToken);
-			return rootRef.Locator;
-		}
+			void CompressPacket(IMemoryOwner<byte> packetData, int length)
+			{
+				int encodedLength = BundleData.Compress(_compressionFormat, packetData.Memory.Slice(0, length), _encodedPacketWriter);
 
-		/// <summary>
-		/// Flushes all the current nodes to storage
-		/// </summary>
-		/// <param name="cancellationToken">Cancellation token for the operation</param>
-		/// <returns></returns>
-		public async Task FlushAsync(CancellationToken cancellationToken)
-		{
-			FlushPacket();
+				BundlePacket packet = new BundlePacket(encodedLength, length);
+				_packets.Add(packet);
+
+				packetData.Dispose();
+			}
 
-			if (_queue.Count > 0)
+			// Mark the bundle as complete
+			public void MarkAsComplete(IStorageClient store, Utf8String prefix)
 			{
-				// Create a set from the nodes to be written. We use this to determine references that are imported.
-				HashSet<TreeNodeRef> nodeSet = new HashSet<TreeNodeRef>();
-				foreach (OutputNodeInfo nodeInfo in _queue)
+				if (!IsReadOnly)
+				{
+					FlushPacket();
+					Task prevWriteTask = _writeTask;
+					_writeTask = Task.Run(() => CompleteAsync(prevWriteTask, store, prefix));
+					IsReadOnly = true;
+				}
+			}
+
+			async Task CompleteAsync(Task prevWriteTask, IStorageClient store, Utf8String prefix)
+			{
+				try
 				{
-					for (NodeRefInfo? nodeRefInfo = nodeInfo; nodeRefInfo != null; nodeRefInfo = nodeRefInfo.Next)
+					await prevWriteTask;
+					await Task.WhenAll(_dependencies);
+
+					Bundle bundle = CreateBundle();
+					BlobLocator locator = await store.WriteBundleAsync(bundle, prefix);
+
+					for (int idx = 0; idx < _queue.Count; idx++)
+					{
+						NodeLocator nodeLocator = new NodeLocator(locator, idx);
+						_queue[idx].MarkAsWritten(nodeLocator);
+					}
+
+					WriteCallback? callback = Interlocked.Exchange(ref _callbacks, s_callbackSentinel);
+					while (callback != null)
 					{
-						nodeSet.Add(nodeRefInfo.NodeRef);
+						callback.OnWrite();
+						callback = callback._next;
 					}
+
+					_completeEvent.SetResult(true);
 				}
-				
+				catch (Exception ex)
+				{
+					_completeEvent.SetException(ex);
+				}
+			}
+
+			// Flush the current write state
+			public Task FlushAsync()
+			{
+				FlushPacket();
+				return _writeTask;
+			}
+
+			// Mark the bundle as complete and create a bundle with the current state
+			public Bundle CreateBundle()
+			{
+				// Create a set from the nodes to be written. We use this to determine references that are imported.
+				HashSet<NodeHandle> nodeSet = new HashSet<NodeHandle>(_queue);
+
 				// Find all the imported nodes by bundle
-				Dictionary<BlobLocator, List<(int, TreeNodeRef)>> bundleToImports = new Dictionary<BlobLocator, List<(int, TreeNodeRef)>>();
-				foreach (OutputNodeInfo nodeInfo in _queue)
+				Dictionary<BlobLocator, List<(int, NodeHandle)>> bundleToImports = new Dictionary<BlobLocator, List<(int, NodeHandle)>>();
+				foreach (PendingNode nodeInfo in _queue)
 				{
-					foreach (TreeNodeRef refKey in nodeInfo.Refs)
+					foreach (NodeHandle handle in nodeInfo.Refs)
 					{
-						if (nodeSet.Add(refKey))
+						if (nodeSet.Add(handle))
 						{
-							Debug.Assert(refKey.Locator.IsValid());
-							NodeLocator refLocator = refKey.Locator;
+							Debug.Assert(handle.Locator.IsValid());
+							NodeLocator refLocator = handle.Locator;
 
-							List<(int, TreeNodeRef)>? importedNodes;
+							List<(int, NodeHandle)>? importedNodes;
 							if (!bundleToImports.TryGetValue(refLocator.Blob, out importedNodes))
 							{
-								importedNodes = new List<(int, TreeNodeRef)>();
+								importedNodes = new List<(int, NodeHandle)>();
 								bundleToImports.Add(refLocator.Blob, importedNodes);
 							}
 
-							importedNodes.Add((refLocator.ExportIdx, refKey));
+							importedNodes.Add((refLocator.ExportIdx, handle));
 						}
 					}
 				}
 
 				// Map from node hash to index, with imported nodes first, ordered by blob, and exported nodes second.
 				int nodeIdx = 0;
-				Dictionary<TreeNodeRef, int> nodeToIndex = new Dictionary<TreeNodeRef, int>();
+				Dictionary<NodeHandle, int> nodeToIndex = new Dictionary<NodeHandle, int>();
 
 				// Add all the imports and assign them identifiers
 				List<BundleImport> imports = new List<BundleImport>();
-				foreach ((BlobLocator blobLocator, List<(int, TreeNodeRef)> importEntries) in bundleToImports)
+				foreach ((BlobLocator blobLocator, List<(int, NodeHandle)> importEntries) in bundleToImports)
 				{
 					importEntries.SortBy(x => x.Item1);
 
 					int[] entries = new int[importEntries.Count];
 					for (int idx = 0; idx < importEntries.Count; idx++)
 					{
-						TreeNodeRef key = importEntries[idx].Item2;
+						NodeHandle key = importEntries[idx].Item2;
 						nodeToIndex.Add(key, nodeIdx++);
 						entries[idx] = importEntries[idx].Item1;
 					}
@@ -375,58 +559,231 @@ namespace EpicGames.Horde.Storage
 
 				// Create the export list
 				List<BundleExport> exports = new List<BundleExport>(_queue.Count);
-				foreach (OutputNodeInfo nodeInfo in _queue)
+				foreach (PendingNode nodeInfo in _queue)
 				{
-					int typeIdx = FindOrAddType(types, typeToIndex, nodeInfo.Key.Type);
-					int[] references = nodeInfo.Refs.Select(x => nodeToIndex[x]).ToArray();
-					BundleExport export = new BundleExport(typeIdx, nodeInfo.Key.Hash, nodeInfo.Length, references);
-					for (NodeRefInfo? nodeRefInfo = nodeInfo; nodeRefInfo != null; nodeRefInfo = nodeRefInfo.Next)
+					int typeIdx;
+					if (!typeToIndex.TryGetValue(nodeInfo.Key.Type, out typeIdx))
 					{
-						nodeToIndex[nodeRefInfo.NodeRef] = nodeIdx;
+						typeIdx = types.Count;
+						typeToIndex.Add(nodeInfo.Key.Type, typeIdx);
+						types.Add(nodeInfo.Key.Type);
 					}
+
+					int[] references = nodeInfo.Refs.Select(x => nodeToIndex[x]).ToArray();
+					BundleExport export = new BundleExport(typeIdx, nodeInfo.Key.Hash, nodeInfo.Length, references);
+					nodeToIndex[nodeInfo] = nodeIdx;
 					exports.Add(export);
 					nodeIdx++;
 				}
 
-				// Create the bundle
-				BundleHeader header = new BundleHeader(_options.CompressionFormat, types, imports, exports, _packets.ToArray());
-				Bundle bundle = new Bundle(header, _encodedPackets.ToArray());
-
-				// Write the bundle
-				BlobLocator locator = await _store.WriteBundleAsync(bundle, _prefix, cancellationToken);
-				for (int idx = 0; idx < _queue.Count; idx++)
+				// Get the memory for each packet
+				List<ReadOnlyMemory<byte>> packetData = new List<ReadOnlyMemory<byte>>(_packets.Count);
+				foreach (ReadOnlyMemory<byte> segment in _encodedPacketWriter.AsSequence())
 				{
-					NodeLocator nodeLocator = new NodeLocator(locator, idx);
-
-					OutputNodeInfo nodeInfo = _queue[idx];
-					for (NodeRefInfo? nodeRefInfo = nodeInfo; nodeRefInfo != null; nodeRefInfo = nodeRefInfo.Next)
+					int offset = 0;
+					while(packetData.Count < _packets.Count)
 					{
-						nodeRefInfo.NodeRef.MarkAsWritten(nodeLocator, nodeRefInfo.Revision);
+						BundlePacket packet = _packets[packetData.Count];
+						int next = offset + packet.EncodedLength;
+
+						if (next > segment.Length)
+						{
+							break;
+						}
+						packetData.Add(segment.Slice(offset, packet.EncodedLength));
+						offset = next;
 					}
+				}
+
+				// Create the bundle
+				BundleHeader header = new BundleHeader(_compressionFormat, types, imports, exports, _packets.ToArray());
+				return new Bundle(header, packetData);
+			}
+		}
+
+		static readonly TreeOptions s_defaultOptions = new TreeOptions();
+
+		readonly IStorageClient _store;
+		readonly TreeOptions _options;
+		readonly Utf8String _prefix;
+
+		readonly NodeCache _nodeCache;
+		readonly Queue<PendingBundle> _writeQueue = new Queue<PendingBundle>();
+
+		PendingBundle? _currentBundle;
+		int _memoryFootprint;
+
+		/// <summary>
+		/// Accessor for the store backing this writer
+		/// </summary>
+		public IStorageClient Store => _store;
+
+		/// <summary>
+		/// Cache of nodes to deduplicate against
+		/// </summary>
+		public NodeCache NodeCache => _nodeCache;
+
+		/// <summary>
+		/// Constructor
+		/// </summary>
+		/// <param name="store">Store to write data to</param>
+		/// <param name="options">Options for the writer</param>
+		/// <param name="prefix">Prefix for blobs written to the store</param>
+		/// <param name="nodeCache">Cache of nodes for deduplication</param>
+		public TreeWriter(IStorageClient store, TreeOptions? options = null, Utf8String prefix = default, NodeCache? nodeCache = null)
+		{
+			_store = store;
+			_options = options ?? s_defaultOptions;
+			_prefix = prefix;
+			_nodeCache = nodeCache ?? new NodeCache(_options.NodeCacheSize);
+		}
+
+		/// <summary>
+		/// Constructor
+		/// </summary>
+		/// <param name="store">Store to write data to</param>
+		/// <param name="refName">Ref being written. Will be used as a prefix for storing blobs.</param>
+		/// <param name="options">Options for the writer</param>
+		public TreeWriter(IStorageClient store, RefName refName, TreeOptions? options = null)
+			: this(store, options, refName.Text)
+		{
+		}
+
+		/// <summary>
+		/// Copy constructor
+		/// </summary>
+		/// <param name="other"></param>
+		public TreeWriter(TreeWriter other)
+			: this(other._store, other._options, other._prefix, other._nodeCache)
+		{
+		}
+
+		/// <inheritdoc/>
+		public void Dispose()
+		{
+			if (_currentBundle != null)
+			{
+				_currentBundle.Dispose();
+				_currentBundle = null;
+			}
+		}
+
+		/// <summary>
+		/// Mark this writer as complete, allowing its data to be serialized.
+		/// </summary>
+		public void Complete()
+		{
+			if (_currentBundle != null)
+			{
+				_currentBundle.MarkAsComplete(_store, _prefix);
+				_writeQueue.Enqueue(_currentBundle);
+				_currentBundle = null;
+			}
+		}
+
+		/// <summary>
+		/// Gets an output buffer for writing.
+		/// </summary>
+		/// <param name="usedSize">Current size in the existing buffer that has been written to</param>
+		/// <param name="desiredSize">Desired size of the returned buffer</param>
+		/// <returns>Buffer to be written into.</returns>
+		public Memory<byte> GetOutputBuffer(int usedSize, int desiredSize)
+		{
+			// If the bundle is full, start the process of writing it to disk
+			if (_currentBundle != null && _currentBundle.UncompressedLength > 0 && _currentBundle.UncompressedLength + desiredSize > _options.MaxBlobSize)
+			{
+				Complete();
+			}
+
+			// If we don't yet have a bundle in this writer, create one
+			PendingBundle currentBundle = GetCurrentBundle();
+			return currentBundle.GetBuffer(usedSize, desiredSize);
+		}
+
+		/// <summary>
+		/// Finish writing a node.
+		/// </summary>
+		/// <param name="size">Used size of the buffer</param>
+		/// <param name="references">References to other nodes</param>
+		/// <param name="type">Type of the node that was written</param>
+		/// <param name="cancellationToken">Cancellation token for the operation</param>
+		/// <returns>Handle to the written node</returns>
+		public async ValueTask<NodeHandle> WriteNodeAsync(int size, IReadOnlyList<NodeHandle> references, BundleType type, CancellationToken cancellationToken = default)
+		{
+			PendingBundle currentBundle = GetCurrentBundle();
+
+			// Get the hash for the new blob
+			ReadOnlyMemory<byte> memory = currentBundle.GetBuffer(size, size);
+			IoHash hash = IoHash.Compute(memory.Span);
+
+			// Create a unique key for the new node
+			NodeKey nodeKey = new NodeKey(hash, type);
 
-					_nodeKeyToLocator[nodeInfo.Key] = nodeLocator;
+			// Check if we have a matching node already in storage
+			if (_nodeCache.TryGetNode(nodeKey, out NodeHandle? handle))
+			{
+				return handle;
+			}
+
+			// Append this node data
+			PendingNode pendingNode = currentBundle.WriteNode(nodeKey, size, references);
+			_memoryFootprint += size;
+			_nodeCache.Add(nodeKey, pendingNode);
+
+			// Add dependencies on all bundles containing a dependent node
+			foreach (NodeHandle reference in references)
+			{
+				if (reference.PendingBundle != null)
+				{
+					currentBundle.AddDependencyOn(reference.PendingBundle);
 				}
+			}
+
+			// If the bundle is full, start the process of writing it to disk
+			if (currentBundle.UncompressedLength > _options.MaxBlobSize)
+			{
+				Complete();
+			}
 
-				// Reset the output state
-				_packets.Clear();
-				_encodedPackets.Clear();
-				_queue.Clear();
-				_queuedRefs.Clear();
-				_nodeKeyToInfo.Clear();
-				_currentBundleLength = 0;
+			// Remove any complete bundle writes
+			while (_writeQueue.Count > 0 && (_writeQueue.Peek().IsComplete() || _memoryFootprint > _options.MaxInMemoryDataLength))
+			{
+				await WaitForWriteAsync(cancellationToken);
 			}
+
+			return pendingNode;
 		}
 
-		static int FindOrAddType(List<BundleType> types, Dictionary<BundleType, int> typeToIndex, BundleType type)
+		PendingBundle GetCurrentBundle()
+		{
+			if (_currentBundle == null || _currentBundle.IsReadOnly)
+			{
+				int bufferSize = (int)(_options.MinCompressionPacketSize * 1.2);
+				_currentBundle = new PendingBundle(bufferSize, _options.MaxBlobSize, _options.CompressionFormat);
+			}
+			return _currentBundle;
+		}
+
+		async Task WaitForWriteAsync(CancellationToken cancellationToken)
+		{
+			PendingBundle writtenBundle = _writeQueue.Dequeue();
+			await await Task.WhenAny(writtenBundle.CompleteTask, Task.Delay(-1, cancellationToken));
+			_memoryFootprint -= writtenBundle.UncompressedLength;
+			writtenBundle.Dispose();
+		}
+
+		/// <summary>
+		/// Flushes all the current nodes to storage
+		/// </summary>
+		/// <param name="cancellationToken">Cancellation token for the operation</param>
+		/// <returns></returns>
+		public async Task FlushAsync(CancellationToken cancellationToken)
 		{
-			int typeIdx;
-			if (!typeToIndex.TryGetValue(type, out typeIdx))
+			Complete();
+			while (_writeQueue.Count > 0)
 			{
-				typeIdx = types.Count;
-				typeToIndex.Add(type, typeIdx);
-				types.Add(type);
+				await WaitForWriteAsync(cancellationToken);
 			}
-			return typeIdx;
 		}
 	}
 }
