EpicGames.Perforce.Native: Add permissions to I/O operation messages.

diff --git a/Engine/Source/Programs/Shared/EpicGames.Perforce.Native/Client.cpp b/Engine/Source/Programs/Shared/EpicGames.Perforce.Native/Client.cpp
--- a/Engine/Source/Programs/Shared/EpicGames.Perforce.Native/Client.cpp
+++ b/Engine/Source/Programs/Shared/EpicGames.Perforce.Native/Client.cpp
@@ -494,16 +494,25 @@ public:
 
 		int PathLen = path.Length() + 1;
 
-		int BufferLen = PathLen + sizeof(int) + sizeof(int);
+		int BufferLen = PathLen + (sizeof(int) * 3);
 
 		char* Buffer = new char[BufferLen];
-		memcpy(Buffer, path.Text(), (size_t)PathLen);
+
+		char* BufferEnd = Buffer;
+		memcpy(BufferEnd, path.Text(), (size_t)PathLen);
+		BufferEnd += PathLen;
 
 		int TypeInt = (int)Type;
-		memcpy(Buffer + PathLen, &TypeInt, sizeof(int));
+		memcpy(BufferEnd, &TypeInt, sizeof(int));
+		BufferEnd += sizeof(int);
 
 		int ModeInt = (int)mode;
-		memcpy(Buffer + PathLen + sizeof(int), &ModeInt, sizeof(int));
+		memcpy(BufferEnd, &ModeInt, sizeof(int));
+		BufferEnd += sizeof(int);
+
+		int PermsInt = (int)perms;
+		memcpy(BufferEnd, &PermsInt, sizeof(int));
+		BufferEnd += sizeof(int);
 
 		FileId = ++NextFileId;
 		User.OutputIo(FileId, "open", Buffer, BufferLen);
