Horde: Missing file.

diff --git a/Engine/Source/Programs/Shared/EpicGames.Horde/Storage/Nodes/DirectoryNode.cs b/Engine/Source/Programs/Shared/EpicGames.Horde/Storage/Nodes/DirectoryNode.cs
--- a/Engine/Source/Programs/Shared/EpicGames.Horde/Storage/Nodes/DirectoryNode.cs
+++ b/Engine/Source/Programs/Shared/EpicGames.Horde/Storage/Nodes/DirectoryNode.cs
@@ -42,6 +42,11 @@ namespace EpicGames.Horde.Storage.Nodes
 		/// Used to indicate that custom data is included in the output. Used internally for serialization; not exposed to users.
 		/// </summary>
 		HasCustomData = 32,
+
+		/// <summary>
+		/// File should be materialized as UTF-16 (but is stored as a UTF-8 source)
+		/// </summary>
+		Utf16 = 64,
 	}
 
 	/// <summary>
