Fixed C# UHT so that enumerations can not contain value names of TRUE or FALSE (case insensitive).

diff --git a/Engine/Source/Programs/Shared/EpicGames.UHT/Parsers/UhtEnumParser.cs b/Engine/Source/Programs/Shared/EpicGames.UHT/Parsers/UhtEnumParser.cs
--- a/Engine/Source/Programs/Shared/EpicGames.UHT/Parsers/UhtEnumParser.cs
+++ b/Engine/Source/Programs/Shared/EpicGames.UHT/Parsers/UhtEnumParser.cs
@@ -1,7 +1,6 @@
 // Copyright Epic Games, Inc. All Rights Reserved.
 
 using System.Diagnostics.CodeAnalysis;
-using System.Runtime.CompilerServices;
 using EpicGames.Core;
 using EpicGames.UHT.Tables;
 using EpicGames.UHT.Tokenizer;
@@ -136,7 +135,7 @@ namespace EpicGames.UHT.Parsers
 						if (tagToken.IsValue("true", true) || tagToken.IsValue("false", true))
 						{
 							// C++ UHT compatibility - TODO
-							topScope.TokenReader.LogError("Enumerations can't have a elements named 'true' or 'false' regardless of case");
+							topScope.TokenReader.LogError("Enumerations can't have any elements named 'true' or 'false' regardless of case");
 						}
 
 						StringView fullEnumName;
