EpicGames.Core: Add overloads for IMemoryReader and IMemoryWriter that take IMemoryReader/IMemoryWriter as parameters, to avoid having to capture them in a lambda (and/or allocate memory for the delegate).

Also add a more optimal path for allocating zero-length arrays and strings.

diff --git a/Engine/Source/Programs/Shared/EpicGames.Core/MemoryReader.cs b/Engine/Source/Programs/Shared/EpicGames.Core/MemoryReader.cs
--- a/Engine/Source/Programs/Shared/EpicGames.Core/MemoryReader.cs
+++ b/Engine/Source/Programs/Shared/EpicGames.Core/MemoryReader.cs
@@ -263,6 +263,24 @@ namespace EpicGames.Core
 			return list;
 		}
 
+		/// <summary>
+		/// Reads a variable length list
+		/// </summary>
+		/// <param name="reader">Reader to deserialize from</param>
+		/// <param name="readItem">Delegate to write an individual item</param>
+		public static List<T> ReadList<T>(this IMemoryReader reader, Func<IMemoryReader, T> readItem)
+		{
+			int length = (int)reader.ReadUnsignedVarInt();
+
+			List<T> list = new List<T>(length);
+			for (int idx = 0; idx < length; idx++)
+			{
+				list.Add(readItem(reader));
+			}
+
+			return list;
+		}
+
 		/// <summary>
 		/// Reads a variable length array
 		/// </summary>
@@ -274,6 +292,17 @@ namespace EpicGames.Core
 			return ReadFixedLengthArray(reader, length, readItem);
 		}
 
+		/// <summary>
+		/// Reads a variable length array
+		/// </summary>
+		/// <param name="reader">Reader to deserialize from</param>
+		/// <param name="readItem">Delegate to write an individual item</param>
+		public static T[] ReadVariableLengthArray<T>(this IMemoryReader reader, Func<IMemoryReader, T> readItem)
+		{
+			int length = (int)reader.ReadUnsignedVarInt();
+			return ReadFixedLengthArray(reader, length, readItem);
+		}
+
 		/// <summary>
 		/// Reads a variable length array
 		/// </summary>
@@ -294,6 +323,11 @@ namespace EpicGames.Core
 		[System.Diagnostics.CodeAnalysis.SuppressMessage("Style", "IDE0060:Remove unused parameter")]
 		public static T[] ReadFixedLengthArray<T>(this IMemoryReader reader, int length, Func<T> readItem)
 		{
+			if (length == 0)
+			{
+				return Array.Empty<T>();
+			}
+
 			T[] array = new T[length];
 			for (int idx = 0; idx < length; idx++)
 			{
@@ -302,6 +336,27 @@ namespace EpicGames.Core
 			return array;
 		}
 
+		/// <summary>
+		/// Writes a fixed length array
+		/// </summary>
+		/// <param name="reader">Reader to deserialize from</param>
+		/// <param name="length">Length of the array to read</param>
+		/// <param name="readItem">Delegate to read an individual item</param>
+		public static T[] ReadFixedLengthArray<T>(this IMemoryReader reader, int length, Func<IMemoryReader, T> readItem)
+		{
+			if (length == 0)
+			{
+				return Array.Empty<T>();
+			}
+
+			T[] array = new T[length];
+			for (int idx = 0; idx < length; idx++)
+			{
+				array[idx] = readItem(reader);
+			}
+			return array;
+		}
+
 		/// <summary>
 		/// Reads a dictionary from the writer
 		/// </summary>
@@ -322,6 +377,26 @@ namespace EpicGames.Core
 			}
 		}
 
+		/// <summary>
+		/// Reads a dictionary from the writer
+		/// </summary>
+		/// <param name="reader">Reader to serialize from</param>
+		/// <param name="dictionary">The dictionary to read</param>
+		/// <param name="readKey">Delegate to write an individual key</param>
+		/// <param name="readValue">Delegate to write an individual value</param>
+		public static void ReadDictionary<TKey, TValue>(this IMemoryReader reader, Dictionary<TKey, TValue> dictionary, Func<IMemoryReader, TKey> readKey, Func<IMemoryReader, TValue> readValue) where TKey : notnull
+		{
+			int count = (int)reader.ReadUnsignedVarInt();
+			dictionary.EnsureCapacity(count);
+
+			for (int idx = 0; idx < count; idx++)
+			{
+				TKey key = readKey(reader);
+				TValue value = readValue(reader);
+				dictionary.Add(key, value);
+			}
+		}
+
 		/// <summary>
 		/// Reads a dictionary from the writer
 		/// </summary>
@@ -336,6 +411,20 @@ namespace EpicGames.Core
 			return dictionary;
 		}
 
+		/// <summary>
+		/// Reads a dictionary from the writer
+		/// </summary>
+		/// <param name="reader">Reader to serialize from</param>
+		/// <param name="readKey">Delegate to write an individual key</param>
+		/// <param name="readValue">Delegate to write an individual value</param>
+		/// <param name="comparer">Comparer for the new dictionary</param>
+		public static Dictionary<TKey, TValue> ReadDictionary<TKey, TValue>(this IMemoryReader reader, Func<IMemoryReader, TKey> readKey, Func<IMemoryReader, TValue> readValue, IEqualityComparer<TKey>? comparer = null) where TKey : notnull
+		{
+			Dictionary<TKey, TValue> dictionary = new Dictionary<TKey, TValue>(comparer);
+			ReadDictionary(reader, dictionary, readKey, readValue);
+			return dictionary;
+		}
+
 		/// <summary>
 		/// Read a string from memory
 		/// </summary>
@@ -352,6 +441,10 @@ namespace EpicGames.Core
 		public static string ReadString(this IMemoryReader reader, Encoding encoding)
 		{
 			int length = (int)reader.ReadUnsignedVarInt();
+			if (length == 0)
+			{
+				return String.Empty;
+			}
 
 			ReadOnlySpan<byte> span = reader.GetSpan(length).Slice(0, length);
 			string str = encoding.GetString(span);
diff --git a/Engine/Source/Programs/Shared/EpicGames.Core/MemoryWriter.cs b/Engine/Source/Programs/Shared/EpicGames.Core/MemoryWriter.cs
--- a/Engine/Source/Programs/Shared/EpicGames.Core/MemoryWriter.cs
+++ b/Engine/Source/Programs/Shared/EpicGames.Core/MemoryWriter.cs
@@ -512,6 +512,17 @@ namespace EpicGames.Core
 			WriteVariableLengthArray(writer, list, writeItem);
 		}
 
+		/// <summary>
+		/// Writes a variable length array
+		/// </summary>
+		/// <param name="writer">Writer to serialize to</param>
+		/// <param name="list">The array to write</param>
+		/// <param name="writeItem">Delegate to write an individual item</param>
+		public static void WriteList<T>(this IMemoryWriter writer, IReadOnlyList<T> list, Action<IMemoryWriter, T> writeItem)
+		{
+			WriteVariableLengthArray(writer, list, writeItem);
+		}
+
 		/// <summary>
 		/// Writes a variable length span of bytes
 		/// </summary>
@@ -561,6 +572,18 @@ namespace EpicGames.Core
 			WriteFixedLengthArray(writer, list, writeItem);
 		}
 
+		/// <summary>
+		/// Writes a variable length array
+		/// </summary>
+		/// <param name="writer">Writer to serialize to</param>
+		/// <param name="list">The array to write</param>
+		/// <param name="writeItem">Delegate to write an individual item</param>
+		public static void WriteVariableLengthArray<T>(this IMemoryWriter writer, IReadOnlyList<T> list, Action<IMemoryWriter, T> writeItem)
+		{
+			writer.WriteUnsignedVarInt(list.Count);
+			WriteFixedLengthArray(writer, list, writeItem);
+		}
+
 		/// <summary>
 		/// Writes a variable length array
 		/// </summary>
@@ -587,6 +610,19 @@ namespace EpicGames.Core
 				writeItem(list[idx]);
 			}
 		}
+		/// <summary>
+		/// Writes a fixed length array
+		/// </summary>
+		/// <param name="writer">Writer to serialize to</param>
+		/// <param name="list">The array to write</param>
+		/// <param name="writeItem">Delegate to write an individual item</param>
+		public static void WriteFixedLengthArray<T>(this IMemoryWriter writer, IReadOnlyList<T> list, Action<IMemoryWriter, T> writeItem)
+		{
+			for (int idx = 0; idx < list.Count; idx++)
+			{
+				writeItem(writer, list[idx]);
+			}
+		}
 
 		/// <summary>
 		/// Writes a dictionary to the writer
@@ -605,6 +641,23 @@ namespace EpicGames.Core
 			}
 		}
 
+		/// <summary>
+		/// Writes a dictionary to the writer
+		/// </summary>
+		/// <param name="writer">Writer to serialize to</param>
+		/// <param name="dictionary">The dictionary to write</param>
+		/// <param name="writeKey">Delegate to write an individual key</param>
+		/// <param name="writeValue">Delegate to write an individual value</param>
+		public static void WriteDictionary<TKey, TValue>(this IMemoryWriter writer, IReadOnlyDictionary<TKey, TValue> dictionary, Action<IMemoryWriter, TKey> writeKey, Action<IMemoryWriter, TValue> writeValue) where TKey : notnull
+		{
+			writer.WriteUnsignedVarInt(dictionary.Count);
+			foreach (KeyValuePair<TKey, TValue> kvp in dictionary)
+			{
+				writeKey(writer, kvp.Key);
+				writeValue(writer, kvp.Value);
+			}
+		}
+
 		/// <summary>
 		/// Write a string to memory
 		/// </summary>
diff --git a/Engine/Source/Programs/Shared/EpicGames.Core/Utf8String.cs b/Engine/Source/Programs/Shared/EpicGames.Core/Utf8String.cs
--- a/Engine/Source/Programs/Shared/EpicGames.Core/Utf8String.cs
+++ b/Engine/Source/Programs/Shared/EpicGames.Core/Utf8String.cs
@@ -90,6 +90,11 @@ namespace EpicGames.Core
 		/// <returns></returns>
 		public Utf8String Clone()
 		{
+			if (Memory.Length == 0)
+			{
+				return default;
+			}
+
 			byte[] newBuffer = new byte[Memory.Length];
 			Memory.CopyTo(newBuffer);
 			return new Utf8String(newBuffer);
