EpicGames.Perforce.Native: Prevent Close messages being sent twice for the same file handle.

diff --git a/Engine/Source/Programs/Shared/EpicGames.Perforce.Native/Client.cpp b/Engine/Source/Programs/Shared/EpicGames.Perforce.Native/Client.cpp
--- a/Engine/Source/Programs/Shared/EpicGames.Perforce.Native/Client.cpp
+++ b/Engine/Source/Programs/Shared/EpicGames.Perforce.Native/Client.cpp
@@ -535,7 +535,11 @@ public:
 
 	virtual void Close(Error* e) override
 	{
-		User.OutputIo(FileId, "close", nullptr, 0);
+		if (FileId != -1)
+		{
+			User.OutputIo(FileId, "close", nullptr, 0);
+			FileId = -1;
+		}
 	}
 
 	virtual int Stat() override
