EpicGames.Core: Change IMemoryWriter to derive from IBufferWriter<byte>.

diff --git a/Engine/Source/Programs/Shared/EpicGames.Core/ByteArrayBuilder.cs b/Engine/Source/Programs/Shared/EpicGames.Core/ByteArrayBuilder.cs
--- a/Engine/Source/Programs/Shared/EpicGames.Core/ByteArrayBuilder.cs
+++ b/Engine/Source/Programs/Shared/EpicGames.Core/ByteArrayBuilder.cs
@@ -69,11 +69,15 @@ namespace EpicGames.Core
 		}
 
 		/// <inheritdoc/>
-		public Memory<byte> GetMemory(int minSize)
+		public Span<byte> GetSpan(int sizeHint) => GetMemory(sizeHint).Span;
+
+		/// <inheritdoc/>
+		public Memory<byte> GetMemory(int sizeHint)
 		{
-			if (_currentChunk.Length + minSize > _currentChunk.Data.Length)
+			int requiredSize = _currentChunk.Length + Math.Max(sizeHint, 1);
+			if (requiredSize > _currentChunk.Data.Length)
 			{
-				_currentChunk = new Chunk(_currentChunk.RunningIndex + _currentChunk.Length, Math.Max(minSize, _chunkSize));
+				_currentChunk = new Chunk(_currentChunk.RunningIndex + _currentChunk.Length, Math.Max(sizeHint, _chunkSize));
 				_chunks.Add(_currentChunk);
 			}
 			return _currentChunk.Data.AsMemory(_currentChunk.Length);
diff --git a/Engine/Source/Programs/Shared/EpicGames.Core/MemoryWriter.cs b/Engine/Source/Programs/Shared/EpicGames.Core/MemoryWriter.cs
--- a/Engine/Source/Programs/Shared/EpicGames.Core/MemoryWriter.cs
+++ b/Engine/Source/Programs/Shared/EpicGames.Core/MemoryWriter.cs
@@ -4,7 +4,6 @@ using System;
 using System.Buffers;
 using System.Buffers.Binary;
 using System.Collections.Generic;
-using System.IO;
 using System.Text;
 
 namespace EpicGames.Core
@@ -12,25 +11,12 @@ namespace EpicGames.Core
 	/// <summary>
 	/// Interface for serializing to a memory buffer
 	/// </summary>
-	public interface IMemoryWriter
+	public interface IMemoryWriter : IBufferWriter<byte>
 	{
 		/// <summary>
 		/// Length of the written data
 		/// </summary>
 		int Length { get; }
-
-		/// <summary>
-		/// Gets a block of memory with at least the given size
-		/// </summary>
-		/// <param name="minSize">Minimum size of the returned data</param>
-		/// <returns>Memory of at least the given size. Note that the returned buffer may be larger than the requested size.</returns>
-		Memory<byte> GetMemory(int minSize = 1);
-
-		/// <summary>
-		/// Updates the current position within the input buffer
-		/// </summary>
-		/// <param name="length">Number of bytes to advance by</param>
-		void Advance(int length);
 	}
 
 	/// <summary>
@@ -77,7 +63,10 @@ namespace EpicGames.Core
 		public int Length => _length;
 
 		/// <inheritdoc/>
-		public Memory<byte> GetMemory(int length) => _memory.Slice(_length);
+		public Span<byte> GetSpan(int sizeHint) => _memory.Slice(_length).Span;
+
+		/// <inheritdoc/>
+		public Memory<byte> GetMemory(int sizeHint) => _memory.Slice(_length);
 
 		/// <inheritdoc/>
 		public void Advance(int length) => _length += length;
@@ -135,9 +124,12 @@ namespace EpicGames.Core
 		}
 
 		/// <inheritdoc/>
-		public Memory<byte> GetMemory(int minSize = 1)
+		public Span<byte> GetSpan(int sizeHint = 0) => GetMemory(sizeHint).Span;
+
+		/// <inheritdoc/>
+		public Memory<byte> GetMemory(int minSize = 0)
 		{
-			int requiredLength = _length + minSize;
+			int requiredLength = _length + Math.Max(minSize, 1);
 			if (_data.Length < requiredLength)
 			{
 				Array.Resize(ref _data, requiredLength + 4096);
@@ -226,11 +218,15 @@ namespace EpicGames.Core
 		}
 
 		/// <inheritdoc/>
-		public Memory<byte> GetMemory(int minSize)
+		public Span<byte> GetSpan(int minSize) => GetMemory(minSize).Span;
+
+		/// <inheritdoc/>
+		public Memory<byte> GetMemory(int sizeHint)
 		{
-			if (_currentChunk.Length + minSize > _currentChunk.Data.Length)
+			int requiredSize = _currentChunk.Length + Math.Max(sizeHint, 1);
+			if (requiredSize > _currentChunk.Data.Length)
 			{
-				_currentChunk = new Chunk(_currentChunk.RunningIndex + _currentChunk.Length, Math.Max(minSize, _chunkSize));
+				_currentChunk = new Chunk(_currentChunk.RunningIndex + _currentChunk.Length, Math.Max(sizeHint, _chunkSize));
 				_chunks.Add(_currentChunk);
 			}
 			return _currentChunk.Data.Slice(_currentChunk.Length);
