EpicGames.Redis: Support using protobuf (via protobuf-net) for serializing structures to Redis strings.

diff --git a/Engine/Source/Programs/Shared/EpicGames.Redis.Tests/SerializationTests.cs b/Engine/Source/Programs/Shared/EpicGames.Redis.Tests/SerializationTests.cs
--- a/Engine/Source/Programs/Shared/EpicGames.Redis.Tests/SerializationTests.cs
+++ b/Engine/Source/Programs/Shared/EpicGames.Redis.Tests/SerializationTests.cs
@@ -3,13 +3,20 @@
 using System;
 using System.ComponentModel;
 using System.Globalization;
+using System.Runtime.Serialization;
 using System.IO;
 using System.Text;
 using System.Threading.Tasks;
 using EpicGames.Core;
 using EpicGames.Redis.Converters;
 using Microsoft.VisualStudio.TestTools.UnitTesting;
+using ProtoBuf;
+using ProtoBuf.Meta;
 using StackExchange.Redis;
+using System.Buffers;
+using System.Collections.Generic;
+using System.Reflection.Emit;
+using System.Reflection;
 
 namespace EpicGames.Redis.Tests
 {
@@ -115,5 +122,27 @@ namespace EpicGames.Redis.Tests
 			Assert.AreEqual(123, output.Number);
 			Assert.AreEqual(123, output.Instance.Number);
 		}
+
+		[ProtoContract]
+		class ProtoClass
+		{
+			[ProtoMember(1)]
+			public string Text { get; set; } = null!;
+		}
+
+		[TestMethod]
+		public void ProtobufTest()
+		{
+			ProtoClass cls = new ProtoClass();
+			cls.Text = "hello world";
+
+			RedisValue value = RedisSerializer.Serialize(cls);
+
+			ProtoClass clsPb = Serializer.Deserialize<ProtoClass>(value);
+			Assert.AreEqual(clsPb.Text, "hello world");
+
+			ProtoClass clsRedis = RedisSerializer.Deserialize<ProtoClass>(value);
+			Assert.AreEqual(clsRedis.Text, "hello world");
+		}
 	}
 }
diff --git a/Engine/Source/Programs/Shared/EpicGames.Redis/Converters/RedisProtobufConverter.cs b/Engine/Source/Programs/Shared/EpicGames.Redis/Converters/RedisProtobufConverter.cs
new file mode 100644
--- /dev/null
+++ b/Engine/Source/Programs/Shared/EpicGames.Redis/Converters/RedisProtobufConverter.cs
@@ -0,0 +1,36 @@
+// Copyright Epic Games, Inc. All Rights Reserved.
+
+using EpicGames.Core;
+using ProtoBuf;
+using StackExchange.Redis;
+using System.IO;
+
+namespace EpicGames.Redis.Converters
+{
+	/// <summary>
+	/// Converter for records to Redis values using ProtoBuf-Net annotations.
+	/// </summary>
+	/// <typeparam name="T">The record type</typeparam>
+	public class RedisProtobufConverter<T> : IRedisConverter<T>
+	{
+		/// <inheritdoc/>
+		public RedisValue ToRedisValue(T value)
+		{
+			using (MemoryStream stream = new MemoryStream())
+			{
+				Serializer.Serialize(stream, value);
+				return stream.ToArray();
+			}
+		}
+
+		/// <inheritdoc/>
+		public T FromRedisValue(RedisValue value)
+		{
+			using (ReadOnlyMemoryStream stream = new ReadOnlyMemoryStream(value))
+			{
+				T result = Serializer.Deserialize<T>(stream);
+				return result;
+			}
+		}
+	}
+}
diff --git a/Engine/Source/Programs/Shared/EpicGames.Redis/EpicGames.Redis.csproj b/Engine/Source/Programs/Shared/EpicGames.Redis/EpicGames.Redis.csproj
--- a/Engine/Source/Programs/Shared/EpicGames.Redis/EpicGames.Redis.csproj
+++ b/Engine/Source/Programs/Shared/EpicGames.Redis/EpicGames.Redis.csproj
@@ -29,6 +29,7 @@
   </ItemGroup>
   <ItemGroup>
     <PackageReference Include="Microsoft.Extensions.Logging" Version="6.0.0" />
+    <PackageReference Include="protobuf-net" Version="3.1.26" />
     <PackageReference Include="StackExchange.Redis" Version="2.6.48" />
   </ItemGroup>
   <ItemGroup>
diff --git a/Engine/Source/Programs/Shared/EpicGames.Redis/RedisSerializer.cs b/Engine/Source/Programs/Shared/EpicGames.Redis/RedisSerializer.cs
--- a/Engine/Source/Programs/Shared/EpicGames.Redis/RedisSerializer.cs
+++ b/Engine/Source/Programs/Shared/EpicGames.Redis/RedisSerializer.cs
@@ -3,6 +3,7 @@
 using EpicGames.Core;
 using EpicGames.Redis.Converters;
 using EpicGames.Serialization;
+using ProtoBuf;
 using StackExchange.Redis;
 using System;
 using System.Collections.Generic;
@@ -201,6 +202,13 @@ namespace EpicGames.Redis
 				return (IRedisConverter<T>)nativeConverter;
 			}
 
+			// Check if the type supports protobuf serialization
+			ProtoContractAttribute? protoAttribute = type.GetCustomAttribute<ProtoContractAttribute>();
+			if (protoAttribute != null)
+			{
+				return new RedisProtobufConverter<T>();
+			}
+
 			// Check if there's a regular converter we can use to convert to/from a string
 			TypeConverter? converter = TypeDescriptor.GetConverter(type);
 			if (converter != null)
