Horde: Implement IBufferWriter methods in ITreeNodeWriter.

diff --git a/Engine/Source/Programs/Shared/EpicGames.Horde/Storage/TreeNodeRef.cs b/Engine/Source/Programs/Shared/EpicGames.Horde/Storage/TreeNodeRef.cs
--- a/Engine/Source/Programs/Shared/EpicGames.Horde/Storage/TreeNodeRef.cs
+++ b/Engine/Source/Programs/Shared/EpicGames.Horde/Storage/TreeNodeRef.cs
@@ -263,11 +263,14 @@ namespace EpicGames.Horde.Storage
 				_refIdx++;
 			}
 
-			public Memory<byte> GetMemory(int minSize = 1)
+			public Span<byte> GetSpan(int sizeHint = 0) => GetMemory(sizeHint).Span;
+
+			public Memory<byte> GetMemory(int sizeHint = 0)
 			{
-				int newLength = _length + minSize;
+				int newLength = _length + Math.Max(sizeHint, 1);
 				if (newLength > _memory.Length)
 				{
+					newLength = _length + Math.Max(sizeHint, 1024);
 					_memory = _treeWriter.GetOutputBuffer(_length, Math.Max(_memory.Length * 2, newLength));
 				}
 				return _memory.Slice(_length);
