EpicGames.Perforce: Fix static analysis warnings and remove custom editorconfig file.

diff --git a/Engine/Source/Programs/Shared/EpicGames.Perforce/.editorconfig b/Engine/Source/Programs/Shared/EpicGames.Perforce/.editorconfig
deleted file mode 100644
diff --git a/Engine/Source/Programs/Shared/EpicGames.Perforce/EpicGames.Perforce.csproj b/Engine/Source/Programs/Shared/EpicGames.Perforce/EpicGames.Perforce.csproj
--- a/Engine/Source/Programs/Shared/EpicGames.Perforce/EpicGames.Perforce.csproj
+++ b/Engine/Source/Programs/Shared/EpicGames.Perforce/EpicGames.Perforce.csproj
@@ -7,7 +7,7 @@
     <RootNamespace>EpicGames.Perforce</RootNamespace>
     <GenerateAssemblyInfo>false</GenerateAssemblyInfo>
     <GenerateTargetFrameworkAttribute>false</GenerateTargetFrameworkAttribute>
-    <Configurations>Debug;Release;Development</Configurations>
+    <Configurations>Analyze;Debug;Release;Development</Configurations>
     <nullable>enable</nullable>
     <DebugType>pdbonly</DebugType>
   </PropertyGroup>
@@ -22,6 +22,12 @@
     <DocumentationFile></DocumentationFile>
     <Optimize>true</Optimize>
   </PropertyGroup>
+  <PropertyGroup Condition=" '$(Configuration)|$(Platform)' == 'Analyze|AnyCPU' ">
+    <RunAnalyzersDuringBuild>True</RunAnalyzersDuringBuild>
+    <TreatWarningsAsErrors>false</TreatWarningsAsErrors>
+    <GenerateDocumentationFile>True</GenerateDocumentationFile>
+    <DocumentationFile></DocumentationFile>
+  </PropertyGroup>
   <ItemGroup>
     <None Remove="EpicGames.Perforce.uatbuildrecord" />
     <None Remove="EpicGames.Perforce.xml" />
diff --git a/Engine/Source/Programs/Shared/EpicGames.Perforce/FileSpecList.cs b/Engine/Source/Programs/Shared/EpicGames.Perforce/FileSpecList.cs
--- a/Engine/Source/Programs/Shared/EpicGames.Perforce/FileSpecList.cs
+++ b/Engine/Source/Programs/Shared/EpicGames.Perforce/FileSpecList.cs
@@ -9,7 +9,7 @@ namespace EpicGames.Perforce
 	/// <summary>
 	/// Wrapper for a list of filespecs. Allows implicit conversion from string (a single entry) or list.
 	/// </summary>
-	public struct FileSpecList : IEquatable<FileSpecList>
+	public readonly struct FileSpecList : IEquatable<FileSpecList>
 	{
 		/// <summary>
 		/// Empty filespec list
diff --git a/Engine/Source/Programs/Shared/EpicGames.Perforce/IPerforceConnection.cs b/Engine/Source/Programs/Shared/EpicGames.Perforce/IPerforceConnection.cs
--- a/Engine/Source/Programs/Shared/EpicGames.Perforce/IPerforceConnection.cs
+++ b/Engine/Source/Programs/Shared/EpicGames.Perforce/IPerforceConnection.cs
@@ -2,8 +2,6 @@
 
 using System;
 using System.Collections.Generic;
-using System.Threading;
-using System.Threading.Tasks;
 using Microsoft.Extensions.Logging;
 
 namespace EpicGames.Perforce
diff --git a/Engine/Source/Programs/Shared/EpicGames.Perforce/IPerforceOutput.cs b/Engine/Source/Programs/Shared/EpicGames.Perforce/IPerforceOutput.cs
--- a/Engine/Source/Programs/Shared/EpicGames.Perforce/IPerforceOutput.cs
+++ b/Engine/Source/Programs/Shared/EpicGames.Perforce/IPerforceOutput.cs
@@ -497,10 +497,10 @@ namespace EpicGames.Perforce
 			ulong requiredTagsBitMask = 0;
 
 			// Create the new record
-			object? newRecord = recordInfo._createInstance();
+			object? newRecord = recordInfo.CreateInstance();
 			if (newRecord == null)
 			{
-				throw new InvalidDataException($"Unable to construct record of type {recordInfo._type}");
+				throw new InvalidDataException($"Unable to construct record of type {recordInfo.Type}");
 			}
 
 			// Get the record info, and parse it into the object
@@ -550,9 +550,9 @@ namespace EpicGames.Perforce
 
 				// Try to find the matching field
 				CachedTagInfo? tagInfo;
-				if (recordInfo._nameToInfo.TryGetValue(tag, out tagInfo))
+				if (recordInfo.NameToInfo.TryGetValue(tag, out tagInfo))
 				{
-					requiredTagsBitMask |= tagInfo._requiredTagBitMask;
+					requiredTagsBitMask |= tagInfo.RequiredTagBitMask;
 				}
 
 				// Check whether it's a subobject or part of the current object.
@@ -570,10 +570,10 @@ namespace EpicGames.Perforce
 					if (tagInfo != null)
 					{
 						// Get the list field
-						System.Collections.IList? list = (System.Collections.IList?)tagInfo._property.GetValue(newRecord);
+						System.Collections.IList? list = (System.Collections.IList?)tagInfo.PropertyInfo.GetValue(newRecord);
 						if (list == null)
 						{
-							throw new PerforceException($"Empty list for {tagInfo._property.Name}");
+							throw new PerforceException($"Empty list for {tagInfo.PropertyInfo.Name}");
 						}
 
 						// Check the suffix matches the index of the next element
@@ -589,16 +589,16 @@ namespace EpicGames.Perforce
 							return false;
 						}
 					}
-					else if (recordInfo._subElementProperty != null)
+					else if (recordInfo.SubElementProperty != null)
 					{
 						// Move back to the start of this tag
 						bufferPos = startBufferPos;
 
 						// Get the list field
-						System.Collections.IList? list = (System.Collections.IList?)recordInfo._subElementProperty.GetValue(newRecord);
+						System.Collections.IList? list = (System.Collections.IList?)recordInfo.SubElementProperty.GetValue(newRecord);
 						if (list == null)
 						{
-							throw new PerforceException($"Invalid field for {recordInfo._subElementProperty.Name}");
+							throw new PerforceException($"Invalid field for {recordInfo.SubElementProperty.Name}");
 						}
 
 						// Check the suffix matches the index of the next element
@@ -609,7 +609,7 @@ namespace EpicGames.Perforce
 
 						// Parse the subobject and add it to the list
 						object? subRecord;
-						if (!TryReadTypedRecord(buffer, ref bufferPos, suffix, recordInfo._subElementRecordInfo!, out subRecord))
+						if (!TryReadTypedRecord(buffer, ref bufferPos, suffix, recordInfo.SubElementRecordInfo!, out subRecord))
 						{
 							record = null;
 							return false;
@@ -635,10 +635,10 @@ namespace EpicGames.Perforce
 			}
 
 			// Make sure we've got all the required tags we need
-			if (requiredTagsBitMask != recordInfo._requiredTagsBitMask)
+			if (requiredTagsBitMask != recordInfo.RequiredTagsBitMask)
 			{
-				string missingTagNames = String.Join(", ", recordInfo._nameToInfo.Where(x => (requiredTagsBitMask | x.Value._requiredTagBitMask) != requiredTagsBitMask).Select(x => x.Key));
-				throw new PerforceException("Missing '{0}' tag when parsing '{1}'", missingTagNames, recordInfo._type.Name);
+				string missingTagNames = String.Join(", ", recordInfo.NameToInfo.Where(x => (requiredTagsBitMask | x.Value.RequiredTagBitMask) != requiredTagsBitMask).Select(x => x.Key));
+				throw new PerforceException("Missing '{0}' tag when parsing '{1}'", missingTagNames, recordInfo.Type.Name);
 			}
 
 			// Construct the response object with the record
@@ -673,10 +673,7 @@ namespace EpicGames.Perforce
 				{
 					return false;
 				}
-				if (tagInfo != null)
-				{
-					tagInfo._setFromString(newRecord, @string);
-				}
+				tagInfo?.ReadFromString(newRecord, @string);
 			}
 			else if (valueType == 'i')
 			{
@@ -685,10 +682,7 @@ namespace EpicGames.Perforce
 				{
 					return false;
 				}
-				if (tagInfo != null)
-				{
-					tagInfo._setFromInteger(newRecord, integer);
-				}
+				tagInfo?.ReadFromInteger(newRecord, integer);
 			}
 			else
 			{
diff --git a/Engine/Source/Programs/Shared/EpicGames.Perforce/PerforceEnvironment.cs b/Engine/Source/Programs/Shared/EpicGames.Perforce/PerforceEnvironment.cs
--- a/Engine/Source/Programs/Shared/EpicGames.Perforce/PerforceEnvironment.cs
+++ b/Engine/Source/Programs/Shared/EpicGames.Perforce/PerforceEnvironment.cs
@@ -7,7 +7,6 @@ using System.Diagnostics;
 using System.IO;
 using System.Linq;
 using System.Runtime.InteropServices;
-using System.Runtime.Versioning;
 using EpicGames.Core;
 using Microsoft.Win32;
 
diff --git a/Engine/Source/Programs/Shared/EpicGames.Perforce/PerforceRecord.cs b/Engine/Source/Programs/Shared/EpicGames.Perforce/PerforceRecord.cs
--- a/Engine/Source/Programs/Shared/EpicGames.Perforce/PerforceRecord.cs
+++ b/Engine/Source/Programs/Shared/EpicGames.Perforce/PerforceRecord.cs
@@ -31,7 +31,7 @@ namespace EpicGames.Perforce
 	/// <summary>
 	/// Wrapper for values returned by Perforce
 	/// </summary>
-	public struct PerforceValue
+	public readonly struct PerforceValue
 	{
 		/// <summary>
 		/// The raw data for the value, including type, size, and payload
@@ -336,7 +336,6 @@ namespace EpicGames.Perforce
 			return stream.ToArray();
 		}
 
-
 		/// <summary>
 		/// Serializes a list of key/value pairs into binary format.
 		/// </summary>
diff --git a/Engine/Source/Programs/Shared/EpicGames.Perforce/PerforceReflection.cs b/Engine/Source/Programs/Shared/EpicGames.Perforce/PerforceReflection.cs
--- a/Engine/Source/Programs/Shared/EpicGames.Perforce/PerforceReflection.cs
+++ b/Engine/Source/Programs/Shared/EpicGames.Perforce/PerforceReflection.cs
@@ -24,55 +24,55 @@ namespace EpicGames.Perforce
 	}
 
 	/// <summary>
-	/// Stores cached information about a field with a P4Tag attribute
+	/// Stores cached information about a property with a <see cref="PerforceTagAttribute"/> attribute.
 	/// </summary>
 	class CachedTagInfo
 	{
 		/// <summary>
 		/// Name of the tag. Specified in the attribute or inferred from the field name.
 		/// </summary>
-		public Utf8String _name;
+		public Utf8String Name { get; }
 
 		/// <summary>
 		/// Whether this tag is optional or not.
 		/// </summary>
-		public bool _optional;
+		public bool Optional { get; }
 
 		/// <summary>
 		/// The property containing the value of this data.
 		/// </summary>
-		public PropertyInfo _property;
+		public PropertyInfo PropertyInfo { get; }
 
 		/// <summary>
 		/// Parser for this field type
 		/// </summary>
-		public Action<object, int> _setFromInteger;
+		public Action<object, int> ReadFromInteger { get; set; }
 
 		/// <summary>
 		/// Parser for this field type
 		/// </summary>
-		public Action<object, Utf8String> _setFromString;
+		public Action<object, Utf8String> ReadFromString { get; set; }
 
 		/// <summary>
 		/// Index into the bitmask of required types
 		/// </summary>
-		public ulong _requiredTagBitMask;
+		public ulong RequiredTagBitMask { get; }
 
 		/// <summary>
 		/// Constructor
 		/// </summary>
 		/// <param name="name"></param>
 		/// <param name="optional"></param>
-		/// <param name="property"></param>
+		/// <param name="propertyInfo"></param>
 		/// <param name="requiredTagBitMask"></param>
-		public CachedTagInfo(Utf8String name, bool optional, PropertyInfo property, ulong requiredTagBitMask)
+		public CachedTagInfo(Utf8String name, bool optional, PropertyInfo propertyInfo, ulong requiredTagBitMask)
 		{
-			_name = name;
-			_optional = optional;
-			_property = property;
-			_requiredTagBitMask = requiredTagBitMask;
-			_setFromInteger = (obj, value) => throw new PerforceException($"Field {name} was not expecting an integer value.");
-			_setFromString = (obj, @string) => throw new PerforceException($"Field {name} was not expecting a string value.");
+			Name = name;
+			Optional = optional;
+			PropertyInfo = propertyInfo;
+			RequiredTagBitMask = requiredTagBitMask;
+			ReadFromInteger = (obj, value) => throw new PerforceException($"Field {name} was not expecting an integer value.");
+			ReadFromString = (obj, str) => throw new PerforceException($"Field {name} was not expecting a string value.");
 		}
 	}
 
@@ -90,42 +90,42 @@ namespace EpicGames.Perforce
 		/// <summary>
 		/// Type of the record
 		/// </summary>
-		public Type _type;
+		public Type Type { get; }
 
 		/// <summary>
 		/// Method to construct this record
 		/// </summary>
-		public CreateRecordDelegate _createInstance;
+		public CreateRecordDelegate CreateInstance { get; }
 
 		/// <summary>
 		/// List of fields in the record. These should be ordered to match P4 output for maximum efficiency.
 		/// </summary>
-		public List<CachedTagInfo> _fields = new List<CachedTagInfo>();
+		public List<CachedTagInfo> Properties { get; } = new List<CachedTagInfo>();
 
 		/// <summary>
 		/// Map of name to tag info
 		/// </summary>
-		public Dictionary<Utf8String, CachedTagInfo> _nameToInfo = new Dictionary<Utf8String, CachedTagInfo>();
+		public Dictionary<Utf8String, CachedTagInfo> NameToInfo { get; set; } = new Dictionary<Utf8String, CachedTagInfo>();
 
 		/// <summary>
 		/// Bitmask of all the required tags. Formed by bitwise-or'ing the RequiredTagBitMask fields for each required CachedTagInfo.
 		/// </summary>
-		public ulong _requiredTagsBitMask;
+		public ulong RequiredTagsBitMask { get; set; }
 
 		/// <summary>
 		/// The type of records to create for subelements
 		/// </summary>
-		public Type? _subElementType;
+		public Type? SubElementType { get; set; }
 
 		/// <summary>
 		/// The cached record info for the subelement type
 		/// </summary>
-		public CachedRecordInfo? _subElementRecordInfo;
+		public CachedRecordInfo? SubElementRecordInfo { get; set; }
 
 		/// <summary>
 		/// Property containing subelements
 		/// </summary>
-		public PropertyInfo? _subElementProperty;
+		public PropertyInfo? SubElementProperty { get; set; }
 
 		/// <summary>
 		/// Constructor
@@ -133,7 +133,7 @@ namespace EpicGames.Perforce
 		/// <param name="type">The record type</param>
 		public CachedRecordInfo(Type type)
 		{
-			_type = type;
+			Type = type;
 
 			ConstructorInfo? constructor = type.GetConstructor(BindingFlags.Public | BindingFlags.NonPublic | BindingFlags.Instance, null, Type.EmptyTypes, null);
 			if (constructor == null)
@@ -145,7 +145,7 @@ namespace EpicGames.Perforce
 			ILGenerator generator = dynamicMethod.GetILGenerator();
 			generator.Emit(OpCodes.Newobj, constructor);
 			generator.Emit(OpCodes.Ret);
-			_createInstance = (CreateRecordDelegate)dynamicMethod.CreateDelegate(typeof(CreateRecordDelegate));
+			CreateInstance = (CreateRecordDelegate)dynamicMethod.CreateDelegate(typeof(CreateRecordDelegate));
 		}
 	}
 
@@ -413,12 +413,12 @@ namespace EpicGames.Perforce
 						ulong requiredTagBitMask = 0;
 						if (!tagAttribute.Optional)
 						{
-							requiredTagBitMask = record._requiredTagsBitMask + 1;
+							requiredTagBitMask = record.RequiredTagsBitMask + 1;
 							if (requiredTagBitMask == 0)
 							{
 								throw new PerforceException("Too many required tags in {0}; max is {1}", recordType.Name, sizeof(ulong) * 8);
 							}
-							record._requiredTagsBitMask |= requiredTagBitMask;
+							record.RequiredTagsBitMask |= requiredTagBitMask;
 						}
 
 						CachedTagInfo tagInfo = new CachedTagInfo(new Utf8String(tagName), tagAttribute.Optional, property, requiredTagBitMask);
@@ -428,67 +428,67 @@ namespace EpicGames.Perforce
 						PropertyInfo propertyCopy = property;
 						if (fieldType == typeof(DateTime))
 						{
-							tagInfo._setFromString = (obj, @string) => propertyCopy.SetValue(obj, ParseStringAsDateTime(@string));
+							tagInfo.ReadFromString = (obj, value) => propertyCopy.SetValue(obj, ParseStringAsDateTime(value));
 						}
 						else if (fieldType == typeof(bool))
 						{
-							tagInfo._setFromString = (obj, @string) => propertyCopy.SetValue(obj, ParseStringAsBool(@string));
+							tagInfo.ReadFromString = (obj, value) => propertyCopy.SetValue(obj, ParseStringAsBool(value));
 						}
 						else if (fieldType == typeof(Nullable<bool>))
 						{
-							tagInfo._setFromString = (obj, @string) => propertyCopy.SetValue(obj, ParseStringAsNullableBool(@string));
+							tagInfo.ReadFromString = (obj, value) => propertyCopy.SetValue(obj, ParseStringAsNullableBool(value));
 						}
 						else if (fieldType == typeof(int))
 						{
-							tagInfo._setFromInteger = (obj, @int) => propertyCopy.SetValue(obj, @int);
-							tagInfo._setFromString = (obj, @string) => propertyCopy.SetValue(obj, ParseStringAsInt(@string));
+							tagInfo.ReadFromInteger = (obj, value) => propertyCopy.SetValue(obj, value);
+							tagInfo.ReadFromString = (obj, value) => propertyCopy.SetValue(obj, ParseStringAsInt(value));
 						}
 						else if (fieldType == typeof(long))
 						{
-							tagInfo._setFromString = (obj, @string) => propertyCopy.SetValue(obj, ParseStringAsLong(@string));
+							tagInfo.ReadFromString = (obj, value) => propertyCopy.SetValue(obj, ParseStringAsLong(value));
 						}
 						else if (fieldType == typeof(string))
 						{
-							tagInfo._setFromString = (obj, @string) => propertyCopy.SetValue(obj, ParseString(@string));
+							tagInfo.ReadFromString = (obj, value) => propertyCopy.SetValue(obj, ParseString(value));
 						}
 						else if (fieldType == typeof(Utf8String))
 						{
-							tagInfo._setFromString = (obj, @string) => propertyCopy.SetValue(obj, @string.Clone());
+							tagInfo.ReadFromString = (obj, @string) => propertyCopy.SetValue(obj, @string.Clone());
 						}
 						else if (fieldType.IsEnum)
 						{
 							CachedEnumInfo enumInfo = GetCachedEnumInfo(fieldType);
-							tagInfo._setFromInteger = (obj, @int) => propertyCopy.SetValue(obj, enumInfo.ParseInteger(@int));
-							tagInfo._setFromString = (obj, @string) => propertyCopy.SetValue(obj, enumInfo.ParseString(@string));
+							tagInfo.ReadFromInteger = (obj, value) => propertyCopy.SetValue(obj, enumInfo.ParseInteger(value));
+							tagInfo.ReadFromString = (obj, value) => propertyCopy.SetValue(obj, enumInfo.ParseString(value));
 						}
 						else if (fieldType == typeof(DateTimeOffset?))
 						{
-							tagInfo._setFromString = (obj, @string) => propertyCopy.SetValue(obj, ParseStringAsNullableDateTimeOffset(@string));
+							tagInfo.ReadFromString = (obj, value) => propertyCopy.SetValue(obj, ParseStringAsNullableDateTimeOffset(value));
 						}
 						else if (fieldType == typeof(List<string>))
 						{
-							tagInfo._setFromString = (obj, @string) => ((List<string>)propertyCopy.GetValue(obj)!).Add(@string.ToString());
+							tagInfo.ReadFromString = (obj, value) => ((List<string>)propertyCopy.GetValue(obj)!).Add(value.ToString());
 						}
 						else if (fieldType == typeof(ReadOnlyMemory<byte>))
 						{
-							tagInfo._setFromString = (obj, @string) => propertyCopy.SetValue(obj, @string.Memory);
+							tagInfo.ReadFromString = (obj, value) => propertyCopy.SetValue(obj, value.Memory);
 						}
 						else
 						{
 							throw new PerforceException("Unsupported type of {0}.{1} for tag '{2}'", recordType.Name, fieldType.Name, tagName);
 						}
 
-						record._fields.Add(tagInfo);
+						record.Properties.Add(tagInfo);
 					}
 
-					record._nameToInfo = record._fields.ToDictionary(x => x._name, x => x);
+					record.NameToInfo = record.Properties.ToDictionary(x => x.Name, x => x);
 
 					PerforceRecordListAttribute? subElementAttribute = property.GetCustomAttribute<PerforceRecordListAttribute>();
 					if (subElementAttribute != null)
 					{
-						record._subElementProperty = property;
-						record._subElementType = property.PropertyType.GenericTypeArguments[0];
-						record._subElementRecordInfo = GetCachedRecordInfo(record._subElementType);
+						record.SubElementProperty = property;
+						record.SubElementType = property.PropertyType.GenericTypeArguments[0];
+						record.SubElementRecordInfo = GetCachedRecordInfo(record.SubElementType);
 					}
 				}
 
@@ -501,14 +501,14 @@ namespace EpicGames.Perforce
 			return record;
 		}
 
-		static object ParseString(Utf8String @string)
+		static object ParseString(Utf8String str)
 		{
-			return @string.ToString();
+			return str.ToString();
 		}
 
-		static object ParseStringAsDateTime(Utf8String @string)
+		static object ParseStringAsDateTime(Utf8String str)
 		{
-			string text = @string.ToString();
+			string text = str.ToString();
 
 			DateTime time;
 			if (DateTime.TryParse(text, out time))
@@ -517,60 +517,60 @@ namespace EpicGames.Perforce
 			}
 			else
 			{
-				return PerforceReflection.UnixEpoch + TimeSpan.FromSeconds(long.Parse(text));
+				return PerforceReflection.UnixEpoch + TimeSpan.FromSeconds(Int64.Parse(text));
 			}
 		}
 
-		static object ParseStringAsBool(Utf8String @string)
+		static object ParseStringAsBool(Utf8String str)
 		{
-			return @string.Length == 0 || @string == StringConstants.True;
+			return str.Length == 0 || str == StringConstants.True;
 		}
 
-		static object ParseStringAsNullableBool(Utf8String @string)
+		static object ParseStringAsNullableBool(Utf8String str)
 		{
-			return @string == StringConstants.True;
+			return str == StringConstants.True;
 		}
 
-		static object ParseStringAsInt(Utf8String @string)
+		static object ParseStringAsInt(Utf8String str)
 		{
 			int value;
 			int bytesConsumed;
-			if (Utf8Parser.TryParse(@string.Span, out value, out bytesConsumed) && bytesConsumed == @string.Length)
+			if (Utf8Parser.TryParse(str.Span, out value, out bytesConsumed) && bytesConsumed == str.Length)
 			{
 				return value;
 			}
-			else if (@string == StringConstants.New || @string == StringConstants.None)
+			else if (str == StringConstants.New || str == StringConstants.None)
 			{
 				return -1;
 			}
-			else if (@string.Length > 0 && @string[0] == '#')
+			else if (str.Length > 0 && str[0] == '#')
 			{
-				return ParseStringAsInt(@string.Slice(1));
+				return ParseStringAsInt(str.Slice(1));
 			}
-			else if (@string == StringConstants.Default)
+			else if (str == StringConstants.Default)
 			{
 				return DefaultChange;
 			}
 			else
 			{
-				throw new PerforceException($"Unable to parse {@string} as an integer");
+				throw new PerforceException($"Unable to parse {str} as an integer");
 			}
 		}
 
-		static object ParseStringAsLong(Utf8String @string)
+		static object ParseStringAsLong(Utf8String str)
 		{
 			long value;
 			int bytesConsumed;
-			if (!Utf8Parser.TryParse(@string.Span, out value, out bytesConsumed) || bytesConsumed != @string.Length)
+			if (!Utf8Parser.TryParse(str.Span, out value, out bytesConsumed) || bytesConsumed != str.Length)
 			{
-				throw new PerforceException($"Unable to parse {@string} as a long value");
+				throw new PerforceException($"Unable to parse {str} as a long value");
 			}
 			return value;
 		}
 
-		static object ParseStringAsNullableDateTimeOffset(Utf8String @string)
+		static object ParseStringAsNullableDateTimeOffset(Utf8String str)
 		{
-			string text = @string.ToString();
+			string text = str.ToString();
 			return DateTimeOffset.Parse(Regex.Replace(text, "[a-zA-Z ]*$", "")); // Strip timezone name (eg. "EST")
 		}
 	}
diff --git a/Engine/Source/Programs/Shared/EpicGames.Perforce/PerforceResponse.cs b/Engine/Source/Programs/Shared/EpicGames.Perforce/PerforceResponse.cs
--- a/Engine/Source/Programs/Shared/EpicGames.Perforce/PerforceResponse.cs
+++ b/Engine/Source/Programs/Shared/EpicGames.Perforce/PerforceResponse.cs
@@ -28,7 +28,7 @@ namespace EpicGames.Perforce
 		/// <summary>
 		/// True if the response is successful
 		/// </summary>
-		public bool Succeeded => !(InternalData is PerforceError);
+		public bool Succeeded => InternalData is not PerforceError;
 
 		/// <summary>
 		/// True if the response is an error
diff --git a/Engine/Source/Programs/Shared/EpicGames.Perforce/Records/ServerRecord.cs b/Engine/Source/Programs/Shared/EpicGames.Perforce/Records/ServerRecord.cs
--- a/Engine/Source/Programs/Shared/EpicGames.Perforce/Records/ServerRecord.cs
+++ b/Engine/Source/Programs/Shared/EpicGames.Perforce/Records/ServerRecord.cs
@@ -1,7 +1,6 @@
 // Copyright Epic Games, Inc. All Rights Reserved.
 
 using System;
-using System.Collections.Generic;
 
 namespace EpicGames.Perforce
 {
diff --git a/Engine/Source/Programs/Shared/EpicGames.Perforce/Records/SizesRecord.cs b/Engine/Source/Programs/Shared/EpicGames.Perforce/Records/SizesRecord.cs
--- a/Engine/Source/Programs/Shared/EpicGames.Perforce/Records/SizesRecord.cs
+++ b/Engine/Source/Programs/Shared/EpicGames.Perforce/Records/SizesRecord.cs
@@ -1,8 +1,5 @@
 // Copyright Epic Games, Inc. All Rights Reserved.
 
-using System;
-using System.Collections.Generic;
-
 namespace EpicGames.Perforce
 {
     /// <summary>
