EpicGames.Perforce: Support for serializing responses back into Perforce tagged record format for mocks.

diff --git a/Engine/Source/Programs/Shared/EpicGames.Perforce/IPerforceOutput.cs b/Engine/Source/Programs/Shared/EpicGames.Perforce/IPerforceOutput.cs
--- a/Engine/Source/Programs/Shared/EpicGames.Perforce/IPerforceOutput.cs
+++ b/Engine/Source/Programs/Shared/EpicGames.Perforce/IPerforceOutput.cs
@@ -1,6 +1,7 @@
 // Copyright Epic Games, Inc. All Rights Reserved.
 
 using System;
+using System.Buffers.Binary;
 using System.Buffers.Text;
 using System.Collections.Generic;
 using System.Diagnostics.CodeAnalysis;
@@ -38,9 +39,9 @@ namespace EpicGames.Perforce
 	}
 
 	/// <summary>
-	/// Wraps a call to a p4.exe child process, and allows reading data from it
+	/// Utility methods for IPerforceOutput
 	/// </summary>
-	public static class PerforceOutputExtensions
+	public static class PerforceOutput
 	{
 		/// <summary>
 		/// String constants for records
@@ -59,6 +60,58 @@ namespace EpicGames.Perforce
 		/// </summary>
 		static readonly byte[] s_recordPrefix = { (byte)'{', (byte)'s', 4, 0, 0, 0, (byte)'c', (byte)'o', (byte)'d', (byte)'e', (byte)'s' };
 
+		class BufferedPerforceOutput : IPerforceOutput
+		{
+			public ReadOnlyMemory<byte> Data { get; private set; }
+
+			public BufferedPerforceOutput(ReadOnlyMemory<byte> data)
+			{
+				Data = data;
+			}
+
+			public Task<bool> ReadAsync(CancellationToken cancellationToken) => Task.FromResult(Data.Length > 0);
+			public void Discard(int numBytes) => Data = Data.Slice(numBytes);
+			public ValueTask DisposeAsync() => new ValueTask();
+		}
+
+		/// <summary>
+		/// Constructs an <see cref="IPerforceOutput"/> object from a block of data
+		/// </summary>
+		/// <param name="data">Data to construct from</param>
+		/// <returns>Output object</returns>
+		public static IPerforceOutput FromData(ReadOnlyMemory<byte> data) => new BufferedPerforceOutput(data);
+
+		/// <summary>
+		/// Constructs an <see cref="IPerforceOutput"/> object from a response
+		/// </summary>
+		/// <param name="response">Response to construct from</param>
+		/// <returns>Output object</returns>
+		public static IPerforceOutput FromResponse(PerforceResponse response) => FromResponses(new[] { response });
+
+		/// <summary>
+		/// Constructs an <see cref="IPerforceOutput"/> object from a sequence of responses
+		/// </summary>
+		/// <param name="responses">Responses to construct from</param>
+		/// <returns>Output object</returns>
+		public static IPerforceOutput FromResponses(IEnumerable<PerforceResponse> responses)
+		{
+			using ChunkedMemoryWriter writer = new ChunkedMemoryWriter();
+			foreach (PerforceResponse response in responses)
+			{
+				writer.WriteFixedLengthBytes(s_recordPrefix);
+
+				Utf8String code = ReadOnlyUtf8StringConstants.Stat;
+
+				Span<byte> span = writer.GetSpanAndAdvance(code.Length + 4);
+				BinaryPrimitives.WriteInt32LittleEndian(span.Slice(0, 4), code.Length);
+				code.Span.CopyTo(span.Slice(4));
+
+				PerforceReflection.Serialize(response.Data, writer);
+				writer.WriteUInt8((byte)'0');
+			}
+			return FromData(writer.ToByteArray());
+		}
+
 		/// <summary>
 		/// Formats the current contents of the buffer to a string
 		/// </summary>
diff --git a/Engine/Source/Programs/Shared/EpicGames.Perforce/NativePerforceConnection.cs b/Engine/Source/Programs/Shared/EpicGames.Perforce/NativePerforceConnection.cs
--- a/Engine/Source/Programs/Shared/EpicGames.Perforce/NativePerforceConnection.cs
+++ b/Engine/Source/Programs/Shared/EpicGames.Perforce/NativePerforceConnection.cs
@@ -606,7 +606,7 @@ namespace EpicGames.Perforce
 			int pos = 0;
 
 			List<KeyValuePair<Utf8String, PerforceValue>> rows = new List<KeyValuePair<Utf8String, PerforceValue>>();
-			if (!PerforceOutputExtensions.ParseRecord(inputData, ref pos, rows))
+			if (!PerforceOutput.ParseRecord(inputData, ref pos, rows))
 			{
 				throw new PerforceException("Unable to parse input data as record");
 			}
diff --git a/Engine/Source/Programs/Shared/EpicGames.Perforce/PerforceReflection.cs b/Engine/Source/Programs/Shared/EpicGames.Perforce/PerforceReflection.cs
--- a/Engine/Source/Programs/Shared/EpicGames.Perforce/PerforceReflection.cs
+++ b/Engine/Source/Programs/Shared/EpicGames.Perforce/PerforceReflection.cs
@@ -1,12 +1,15 @@
 // Copyright Epic Games, Inc. All Rights Reserved.
 
 using System;
+using System.Buffers;
+using System.Buffers.Binary;
 using System.Buffers.Text;
 using System.Collections.Concurrent;
 using System.Collections.Generic;
 using System.Linq;
 using System.Reflection;
 using System.Reflection.Emit;
+using System.Text;
 using System.Text.RegularExpressions;
 using EpicGames.Core;
 
@@ -18,6 +21,7 @@ namespace EpicGames.Perforce
 	static class StringConstants
 	{
 		public static readonly Utf8String True = new Utf8String("true");
+		public static readonly Utf8String False = new Utf8String("false");
 		public static readonly Utf8String New = new Utf8String("new");
 		public static readonly Utf8String None = new Utf8String("none");
 		public static readonly Utf8String Default = new Utf8String("default");
@@ -43,6 +47,11 @@ namespace EpicGames.Perforce
 		/// </summary>
 		public PropertyInfo PropertyInfo { get; }
 
+		/// <summary>
+		/// Writes an instance of this field from an object
+		/// </summary>
+		public Action<IMemoryWriter, object> Write { get; set; }
+
 		/// <summary>
 		/// Parser for this field type
 		/// </summary>
@@ -71,6 +80,7 @@ namespace EpicGames.Perforce
 			Optional = optional;
 			PropertyInfo = propertyInfo;
 			RequiredTagBitMask = requiredTagBitMask;
+			Write = (obj, writer) => throw new PerforceException($"Field {name} does not have a serializer.");
 			ReadFromInteger = (obj, value) => throw new PerforceException($"Field {name} was not expecting an integer value.");
 			ReadFromString = (obj, str) => throw new PerforceException($"Field {name} was not expecting a string value.");
 		}
@@ -169,6 +179,11 @@ namespace EpicGames.Perforce
 		/// </summary>
 		public Dictionary<Utf8String, int> _nameToValue = new Dictionary<Utf8String, int>();
 
+		/// <summary>
+		/// Map of value to name
+		/// </summary>
+		public Dictionary<int, Utf8String> _valueToName = new Dictionary<int, Utf8String>();
+
 		/// <summary>
 		/// List of name/value pairs
 		/// </summary>
@@ -195,6 +210,7 @@ namespace EpicGames.Perforce
 					{
 						Utf8String name = new Utf8String(attribute.Name);
 						_nameToValue[name] = (int)value;
+						_valueToName[(int)value] = name;
 
 						_nameValuePairs.Add(new KeyValuePair<string, int>(attribute.Name, (int)value));
 					}
@@ -202,6 +218,13 @@ namespace EpicGames.Perforce
 			}
 		}
 
+		/// <summary>
+		/// Gets the name of a particular enum value
+		/// </summary>
+		/// <param name="value"></param>
+		/// <returns></returns>
+		public Utf8String GetName(int value) => _valueToName[value];
+
 		/// <summary>
 		/// Parses the given integer as an enum
 		/// </summary>
@@ -357,6 +380,49 @@ namespace EpicGames.Perforce
 		/// </summary>
 		public static CachedRecordInfo IoRecordInfo = GetCachedRecordInfo(typeof(PerforceIo));
 
+		/// <summary>
+		/// Serializes a sequence of objects to a stream
+		/// </summary>
+		/// <param name="obj">Object to serialize</param>
+		/// <param name="writer">Writer for output data</param>
+		public static void Serialize(object obj, IMemoryWriter writer)
+		{
+			CachedRecordInfo recordInfo = GetCachedRecordInfo(obj.GetType());
+			foreach (CachedTagInfo tagInfo in recordInfo.Properties)
+			{
+				object? value = tagInfo.PropertyInfo.GetValue(obj);
+				if (value != null)
+				{
+					WriteUtf8StringWithTag(writer, tagInfo.Name);
+					tagInfo.Write(writer, value!);
+				}
+			}
+		}
+
+		static void WriteIntegerWithTag(IMemoryWriter writer, int value)
+		{
+			Span<byte> span = writer.GetSpanAndAdvance(5);
+			span[0] = (byte)'i';
+			BinaryPrimitives.WriteInt32LittleEndian(span.Slice(1, 4), value);
+		}
+
+		static void WriteStringWithTag(IMemoryWriter writer, string str)
+		{
+			int length = Encoding.UTF8.GetByteCount(str);
+			Span<byte> span = writer.GetSpanAndAdvance(1 + length + 4);
+			span[0] = (byte)'s';
+			BinaryPrimitives.WriteInt32LittleEndian(span.Slice(1, 4), length);
+			Encoding.UTF8.GetBytes(str, span.Slice(5));
+		}
+
+		static void WriteUtf8StringWithTag(IMemoryWriter writer, Utf8String str)
+		{
+			Span<byte> span = writer.GetSpanAndAdvance(1 + str.Length + 4);
+			span[0] = (byte)'s';
+			BinaryPrimitives.WriteInt32LittleEndian(span.Slice(1, 4), str.Length);
+			str.Span.CopyTo(span.Slice(5));
+		}
+
 		/// <summary>
 		/// Gets a mapping of flags to enum values for the given type
 		/// </summary>
@@ -428,10 +494,12 @@ namespace EpicGames.Perforce
 						PropertyInfo propertyCopy = property;
 						if (fieldType == typeof(DateTime))
 						{
+							tagInfo.Write = (writer, value) => WriteUtf8StringWithTag(writer, ((long)((DateTime)value - PerforceReflection.UnixEpoch).TotalSeconds).ToString());
 							tagInfo.ReadFromString = (obj, value) => propertyCopy.SetValue(obj, ParseStringAsDateTime(value));
 						}
 						else if (fieldType == typeof(bool))
 						{
+							tagInfo.Write = (writer, value) => WriteUtf8StringWithTag(writer, ((bool)value) ? StringConstants.True : StringConstants.False);
 							tagInfo.ReadFromString = (obj, value) => propertyCopy.SetValue(obj, ParseStringAsBool(value));
 						}
 						else if (fieldType == typeof(Nullable<bool>))
@@ -440,24 +508,29 @@ namespace EpicGames.Perforce
 						}
 						else if (fieldType == typeof(int))
 						{
+							tagInfo.Write = (writer, value) => WriteIntegerWithTag(writer, (int)value);
 							tagInfo.ReadFromInteger = (obj, value) => propertyCopy.SetValue(obj, value);
 							tagInfo.ReadFromString = (obj, value) => propertyCopy.SetValue(obj, ParseStringAsInt(value));
 						}
 						else if (fieldType == typeof(long))
 						{
+							tagInfo.Write = (writer, value) => WriteUtf8StringWithTag(writer, ((long)value).ToString());
 							tagInfo.ReadFromString = (obj, value) => propertyCopy.SetValue(obj, ParseStringAsLong(value));
 						}
 						else if (fieldType == typeof(string))
 						{
+							tagInfo.Write = (writer, value) => WriteStringWithTag(writer, (string)value);
 							tagInfo.ReadFromString = (obj, value) => propertyCopy.SetValue(obj, ParseString(value));
 						}
 						else if (fieldType == typeof(Utf8String))
 						{
-							tagInfo.ReadFromString = (obj, @string) => propertyCopy.SetValue(obj, @string.Clone());
+							tagInfo.Write = (writer, value) => WriteUtf8StringWithTag(writer, (Utf8String)value);
+							tagInfo.ReadFromString = (obj, str) => propertyCopy.SetValue(obj, str.Clone());
 						}
 						else if (fieldType.IsEnum)
 						{
 							CachedEnumInfo enumInfo = GetCachedEnumInfo(fieldType);
+							tagInfo.Write = (writer, value) => WriteUtf8StringWithTag(writer, enumInfo.GetName((int)value));
 							tagInfo.ReadFromInteger = (obj, value) => propertyCopy.SetValue(obj, enumInfo.ParseInteger(value));
 							tagInfo.ReadFromString = (obj, value) => propertyCopy.SetValue(obj, enumInfo.ParseString(value));
 						}
