Fixed C# UHT issue where a class scoped delegate reference from another class would fail.

diff --git a/Engine/Source/Programs/Shared/EpicGames.UHT/Utils/UhtSession.cs b/Engine/Source/Programs/Shared/EpicGames.UHT/Utils/UhtSession.cs
--- a/Engine/Source/Programs/Shared/EpicGames.UHT/Utils/UhtSession.cs
+++ b/Engine/Source/Programs/Shared/EpicGames.UHT/Utils/UhtSession.cs
@@ -1310,7 +1310,14 @@ namespace EpicGames.UHT.Utils
 					}
 					if (type is UhtClass)
 					{
-						return FindTypeInternal(startingType, UhtFindOptions.DelegateFunction | UhtFindOptions.NoParents | (options & ~UhtFindOptions.TypesMask), secondName);
+						//TODO - Old UHT compatibility.  In UWidget, it references USlateAccessibleWidgetData::FGetText.  However, since UWidget has a FGetText, that 
+						// is returned first.
+						UhtType? compatType = FindTypeInternal(startingType, UhtFindOptions.DelegateFunction | UhtFindOptions.NoParents | (options & ~UhtFindOptions.TypesMask), secondName);
+						if (compatType != null)
+						{
+							return compatType;
+						}
+						return FindTypeInternal(type, UhtFindOptions.DelegateFunction | UhtFindOptions.NoParents | (options & ~UhtFindOptions.TypesMask), secondName);
 					}
 				}
 
