Fixed C# UHT to generate errors when UPROPERTIES are located in native interfaces.

diff --git a/Engine/Source/Programs/Shared/EpicGames.UHT/Parsers/UhtNativeInterfaceClassParser.cs b/Engine/Source/Programs/Shared/EpicGames.UHT/Parsers/UhtNativeInterfaceClassParser.cs
--- a/Engine/Source/Programs/Shared/EpicGames.UHT/Parsers/UhtNativeInterfaceClassParser.cs
+++ b/Engine/Source/Programs/Shared/EpicGames.UHT/Parsers/UhtNativeInterfaceClassParser.cs
@@ -127,10 +127,26 @@ namespace EpicGames.UHT.Parsers
 
 			classObj.SuperIdentifier = superName;
 
+			//TODO - C++ UHT compatibility - When we know for sure that we have a native interface, then we should error out.  Due to the lack of a symbol table,
+			// we can't do this 100% reliably.  However, if we find a U class, we can assume we have a native interface.
+			bool logUnhandledKeywords = false;
+			if (classObj.Outer != null) 
+			{
+				string interfaceName = "U" + classObj.EngineName;
+				foreach (UhtType outerChild in classObj.Outer.Children)
+				{
+					if (outerChild.SourceName == interfaceName && outerChild is UhtClass outerChildClass && outerChildClass.ClassType == UhtClassType.Interface)
+					{
+						logUnhandledKeywords = true;
+						break;
+					}
+				}
+			}
+
 			{
 				using UhtParsingScope topScope = new(parentScope, classObj, parentScope.Session.GetKeywordTable(UhtTableNames.NativeInterface), UhtAccessSpecifier.Private);
 				using UhtMessageContext tokenContext = new("native interface");
-				topScope.HeaderParser.ParseStatements('{', '}', false);
+				topScope.HeaderParser.ParseStatements('{', '}', logUnhandledKeywords);
 				tokenReader.Require(';');
 			}
 			return;
