Added -NoCompileUAT as an ignored parameter to allow it to be present in command lines where no commands are present.

diff --git a/Engine/Source/Programs/AutomationTool/Program.cs b/Engine/Source/Programs/AutomationTool/Program.cs
--- a/Engine/Source/Programs/AutomationTool/Program.cs
+++ b/Engine/Source/Programs/AutomationTool/Program.cs
@@ -261,7 +261,7 @@ namespace AutomationToolDriver
 					{"-BuildMachine", "" },
 					{"-WaitForUATMutex", "" }
 				},
-				new HashSet<string>(){ "-msbuild-verbose" }
+				new HashSet<string>(StringComparer.InvariantCultureIgnoreCase) { "-msbuild-verbose", "-NoCompileUAT" }
 			);
 
 			ParseProfile(ref Arguments);
diff --git a/Engine/Source/Programs/Shared/EpicGames.Build/Automation/CommandLine.cs b/Engine/Source/Programs/Shared/EpicGames.Build/Automation/CommandLine.cs
--- a/Engine/Source/Programs/Shared/EpicGames.Build/Automation/CommandLine.cs
+++ b/Engine/Source/Programs/Shared/EpicGames.Build/Automation/CommandLine.cs
@@ -57,7 +57,7 @@ namespace UnrealBuildBase
 		public ParsedCommandLine(Dictionary<string, string> GlobalParameters, HashSet<string> IgnoredGlobalParameters)
 		{
 			this.GlobalParameters = new Dictionary<string, string>(GlobalParameters, StringComparer.InvariantCultureIgnoreCase);
-			this.IgnoredGlobalParameters = new HashSet<string>(IgnoredGlobalParameters);
+			this.IgnoredGlobalParameters = new HashSet<string>(IgnoredGlobalParameters, StringComparer.InvariantCultureIgnoreCase);
 		}
 
 		HashSet<string> GlobalParameterValues = new HashSet<string>(StringComparer.InvariantCultureIgnoreCase);
