Horde: Fix static analysis warnings due to use of keyword as variable name.

diff --git a/Engine/Source/Programs/Shared/EpicGames.Horde/Storage/IStorageClient.cs b/Engine/Source/Programs/Shared/EpicGames.Horde/Storage/IStorageClient.cs
--- a/Engine/Source/Programs/Shared/EpicGames.Horde/Storage/IStorageClient.cs
+++ b/Engine/Source/Programs/Shared/EpicGames.Horde/Storage/IStorageClient.cs
@@ -153,10 +153,10 @@ namespace EpicGames.Horde.Storage
 		/// <summary>
 		/// Finds nodes with the given alias. Unlike refs, aliases do not serve as GC roots.
 		/// </summary>
-		/// <param name="alias">Alias for the node</param>
+		/// <param name="name">Alias for the node</param>
 		/// <param name="cancellationToken">Cancellation token for the operation</param>
 		/// <returns>Nodes matching the given handle</returns>
-		IAsyncEnumerable<NodeHandle> FindNodesAsync(Utf8String alias, CancellationToken cancellationToken = default);
+		IAsyncEnumerable<NodeHandle> FindNodesAsync(Utf8String name, CancellationToken cancellationToken = default);
 
 		#endregion
 
@@ -251,7 +251,7 @@ namespace EpicGames.Horde.Storage
 			#region Nodes
 
 			/// <inheritdoc/>
-			public IAsyncEnumerable<NodeHandle> FindNodesAsync(Utf8String alias, CancellationToken cancellationToken = default) => _inner.FindNodesAsync(alias, cancellationToken);
+			public IAsyncEnumerable<NodeHandle> FindNodesAsync(Utf8String name, CancellationToken cancellationToken = default) => _inner.FindNodesAsync(name, cancellationToken);
 
 			#endregion
 
diff --git a/Engine/Source/Programs/Shared/EpicGames.Horde/Storage/StorageClientBase.cs b/Engine/Source/Programs/Shared/EpicGames.Horde/Storage/StorageClientBase.cs
--- a/Engine/Source/Programs/Shared/EpicGames.Horde/Storage/StorageClientBase.cs
+++ b/Engine/Source/Programs/Shared/EpicGames.Horde/Storage/StorageClientBase.cs
@@ -37,7 +37,7 @@ namespace EpicGames.Horde.Storage
 		#region Nodes
 
 		/// <inheritdoc/>
-		public abstract IAsyncEnumerable<NodeHandle> FindNodesAsync(Utf8String alias, CancellationToken cancellationToken = default);
+		public abstract IAsyncEnumerable<NodeHandle> FindNodesAsync(Utf8String name, CancellationToken cancellationToken = default);
 
 		#endregion
 
