AutomationTool: Reset FileInfo cache when deleting files

#rnx

diff --git a/Engine/Source/Programs/AutomationTool/AutomationUtils/Utils.cs b/Engine/Source/Programs/AutomationTool/AutomationUtils/Utils.cs
--- a/Engine/Source/Programs/AutomationTool/AutomationUtils/Utils.cs
+++ b/Engine/Source/Programs/AutomationTool/AutomationUtils/Utils.cs
@@ -122,6 +122,7 @@ namespace AutomationTool
 							File.SetAttributes(Path, Attributes & ~FileAttributes.ReadOnly);
 						}
 						File.Delete(Path);
+						FileItem.ResetCachedInfo(Path);
 					}
 				}
 				catch (Exception Ex)
@@ -213,6 +214,7 @@ namespace AutomationTool
 					try
 					{
 						Directory.Delete(Path, true);
+						DirectoryItem.ResetCachedInfo(Path);
 					}
 					catch (Exception Ex)
 					{
diff --git a/Engine/Source/Programs/Shared/EpicGames.Build/System/DirectoryItem.cs b/Engine/Source/Programs/Shared/EpicGames.Build/System/DirectoryItem.cs
--- a/Engine/Source/Programs/Shared/EpicGames.Build/System/DirectoryItem.cs
+++ b/Engine/Source/Programs/Shared/EpicGames.Build/System/DirectoryItem.cs
@@ -180,6 +180,17 @@ namespace UnrealBuildBase
 			}
 		}
 
+		/// <summary>
+		/// Resets the cached info, if the DirectoryInfo is not found don't create a new entry
+		/// </summary>
+		public static void ResetCachedInfo(string Path)
+		{
+			if (LocationToItem.TryGetValue(new DirectoryReference(Path), out DirectoryItem? Result))
+			{
+				Result.ResetCachedInfo();
+			}
+		}
+
 		/// <summary>
 		/// Resets all cached directory info. Significantly reduces performance; do not use unless strictly necessary.
 		/// </summary>
diff --git a/Engine/Source/Programs/Shared/EpicGames.Build/System/FileItem.cs b/Engine/Source/Programs/Shared/EpicGames.Build/System/FileItem.cs
--- a/Engine/Source/Programs/Shared/EpicGames.Build/System/FileItem.cs
+++ b/Engine/Source/Programs/Shared/EpicGames.Build/System/FileItem.cs
@@ -237,6 +237,17 @@ namespace UnrealBuildBase
 			});
 		}
 
+		/// <summary>
+		/// Resets the cached info, if the FileInfo is not found don't create a new entry
+		/// </summary>
+		public static void ResetCachedInfo(string Path)
+		{
+			if (UniqueSourceFileMap.TryGetValue(new FileReference(Path), out FileItem? Result))
+			{
+				Result.ResetCachedInfo();
+			}
+		}
+
 		/// <summary>
 		/// Resets all cached file info. Significantly reduces performance; do not use unless strictly necessary.
 		/// </summary>
