EpicGames.Core: Handle long file paths on Windows for ForceDeleteDirectoryContents

diff --git a/Engine/Source/Programs/Shared/EpicGames.Core.Tests/FileUtilsTests.cs b/Engine/Source/Programs/Shared/EpicGames.Core.Tests/FileUtilsTests.cs
new file mode 100644
--- /dev/null
+++ b/Engine/Source/Programs/Shared/EpicGames.Core.Tests/FileUtilsTests.cs
@@ -0,0 +1,87 @@
+// Copyright Epic Games, Inc. All Rights Reserved.
+
+using System;
+using System.IO;
+using Microsoft.VisualStudio.TestTools.UnitTesting;
+
+namespace EpicGames.Core.Tests
+{
+	[TestClass]
+	public class FileUtilsTest
+	{
+		private readonly DirectoryReference _tempDir;
+
+		public FileUtilsTest()
+		{
+			_tempDir = CreateTempDir();
+		}
+
+		[TestMethod]
+		public void ForceDeleteDirectoryNormalPath()
+		{
+			string subDir = Path.Join(_tempDir.FullName, "Foo");
+			Directory.CreateDirectory(subDir);
+			FileUtils.ForceDeleteDirectory(subDir);
+			Assert.IsFalse(Directory.Exists(subDir));
+		}
+		
+		[TestMethod]
+		public void ForceDeleteDirectoryLongPath()
+		{
+			string subDir = CreateSubDirWithLongPath();
+			FileUtils.ForceDeleteDirectory(subDir);
+			Assert.IsFalse(Directory.Exists(subDir));
+		}
+		
+		[TestMethod]
+		public void ForceDeleteDirectoryContentsNormalPath()
+		{
+			string subDir = Path.Join(_tempDir.FullName, "Foo");
+			Directory.CreateDirectory(subDir);
+			string filePath = Path.Join(subDir, "file.txt");
+			File.WriteAllText(filePath, "placeholder");
+			
+			FileUtils.ForceDeleteDirectoryContents(subDir);
+			Assert.IsTrue(Directory.Exists(subDir));
+			Assert.IsFalse(File.Exists(filePath));
+		}
+		
+		[TestMethod]
+		public void ForceDeleteDirectoryContentsLongPath()
+		{
+			string subDir = CreateSubDirWithLongPath();
+			string filePath = Path.Join(subDir, "file.txt");
+			File.WriteAllText(filePath, "placeholder");
+
+			FileUtils.ForceDeleteDirectoryContents(subDir);
+			Assert.IsTrue(Directory.Exists(subDir));
+			Assert.IsFalse(File.Exists(filePath));
+		}
+
+		[TestCleanup]
+		public void RemoveTempDir()
+		{
+			if (Directory.Exists(_tempDir.FullName))
+			{
+				Directory.Delete(_tempDir.FullName, true);
+			}
+		}
+
+		private string CreateSubDirWithLongPath()
+		{
+			const string LongDirNameA = "ThisIsAVeryLongDirectoryName";
+			const string LongDirNameB = "AnotherLongDirectoryNameThatIsUsed";
+			string subDir = Path.Join(_tempDir.FullName, "Foo", LongDirNameA, LongDirNameB, LongDirNameA, LongDirNameB, LongDirNameA, LongDirNameB, LongDirNameA, LongDirNameB);
+			Assert.IsTrue(subDir.Length > 260, "Longer than Windows MAX_PATH");
+			Directory.CreateDirectory(subDir);
+			return subDir;
+		}
+		
+		private static DirectoryReference CreateTempDir()
+		{
+			string tempDir = Path.Join(Path.GetTempPath(), "epicgames-core-tests-" + Guid.NewGuid().ToString()[..8]);
+			Directory.CreateDirectory(tempDir);
+			return new DirectoryReference(tempDir);
+		}
+	}
+}
diff --git a/Engine/Source/Programs/Shared/EpicGames.Core/FileUtils.cs b/Engine/Source/Programs/Shared/EpicGames.Core/FileUtils.cs
--- a/Engine/Source/Programs/Shared/EpicGames.Core/FileUtils.cs
+++ b/Engine/Source/Programs/Shared/EpicGames.Core/FileUtils.cs
@@ -495,6 +495,12 @@ namespace EpicGames.Core
 		{
 			WIN32_FIND_DATA findData = new WIN32_FIND_DATA();
 
+			const string RawPathPrefix = "\\\\?\\";
+			if (!dirName.StartsWith(RawPathPrefix, StringComparison.Ordinal))
+			{
+				dirName = RawPathPrefix + dirName;				
+			}
+
 			IntPtr hFind = FindFirstFileW(dirName + "\\*", ref findData);
 			if (hFind != INVALID_HANDLE_VALUE)
 			{
