EpicGames.Build: Base script modules directories list should not contain duplicate paths

#rnx

diff --git a/Engine/Source/Programs/Shared/EpicGames.Build/System/CompileScriptModules.cs b/Engine/Source/Programs/Shared/EpicGames.Build/System/CompileScriptModules.cs
--- a/Engine/Source/Programs/Shared/EpicGames.Build/System/CompileScriptModules.cs
+++ b/Engine/Source/Programs/Shared/EpicGames.Build/System/CompileScriptModules.cs
@@ -156,6 +156,7 @@ namespace UnrealBuildBase
 			BaseDirectories.Add(Unreal.EngineDirectory);
 			BaseDirectories.AddRange(GameDirectories);
 			BaseDirectories.AddRange(AdditionalDirectories);
+			BaseDirectories = BaseDirectories.Distinct().ToList();
 
 			HashSet<FileReference> FoundProjects = new HashSet<FileReference>(
 				Rules.FindAllRulesSourceFiles(RulesFileType,
