Fix issue in C# UHT were circular base class/structs would hang UHT.  UHT will now generate an error showing the circular bases.

diff --git a/Engine/Source/Programs/Shared/EpicGames.UHT/Types/UhtClass.cs b/Engine/Source/Programs/Shared/EpicGames.UHT/Types/UhtClass.cs
--- a/Engine/Source/Programs/Shared/EpicGames.UHT/Types/UhtClass.cs
+++ b/Engine/Source/Programs/Shared/EpicGames.UHT/Types/UhtClass.cs
@@ -640,6 +640,14 @@ namespace EpicGames.UHT.Types
 			public bool Found { get; set; }
 		}
 
+		/// <inheritdoc/>
+		public override void BindSuperAndBases()
+		{
+			BindSuper(SuperIdentifier, UhtFindOptions.Class);
+			BindBases(BaseIdentifiers, UhtFindOptions.Class);
+			base.BindSuperAndBases();
+		}
+
 		/// <inheritdoc/>
 		protected override void ResolveSuper(UhtResolvePhase resolvePhase)
 		{
@@ -654,8 +662,15 @@ namespace EpicGames.UHT.Types
 			switch (resolvePhase)
 			{
 				case UhtResolvePhase.Bases:
-					BindAndResolveSuper(SuperIdentifier, UhtFindOptions.Class);
-					BindAndResolveBases(BaseIdentifiers, UhtFindOptions.Class);
+					UhtClass? superClass = SuperClass;
+					if (superClass != null)
+					{
+						superClass.Resolve(UhtResolvePhase.Bases);
+					}
+					foreach (UhtStruct baseStruct in Bases)
+					{
+						baseStruct.Resolve(UhtResolvePhase.Bases);
+					}
 
 					// Force the MatchedSerializers on for anything being exported
 					if (!ClassExportFlags.HasAnyFlags(UhtClassExportFlags.NoExport))
@@ -667,7 +682,6 @@ namespace EpicGames.UHT.Types
 					{
 						case UhtClassType.Class:
 							{
-								UhtClass? superClass = SuperClass;
 
 								// Merge the super class flags
 								if (superClass != null)
@@ -719,7 +733,6 @@ namespace EpicGames.UHT.Types
 
 						case UhtClassType.Interface:
 							{
-								UhtClass? superClass = SuperClass;
 								if (superClass != null)
 								{
 									ClassFlags |= superClass.ClassFlags & EClassFlags.ScriptInherit;
diff --git a/Engine/Source/Programs/Shared/EpicGames.UHT/Types/UhtScriptStruct.cs b/Engine/Source/Programs/Shared/EpicGames.UHT/Types/UhtScriptStruct.cs
--- a/Engine/Source/Programs/Shared/EpicGames.UHT/Types/UhtScriptStruct.cs
+++ b/Engine/Source/Programs/Shared/EpicGames.UHT/Types/UhtScriptStruct.cs
@@ -703,6 +703,14 @@ namespace EpicGames.UHT.Types
 		}
 
 		#region Resolution support
+
+		/// <inheritdoc/>
+		public override void BindSuperAndBases()
+		{
+			BindSuper(SuperIdentifier, UhtFindOptions.ScriptStruct);
+			base.BindSuperAndBases();
+		}
+
 		/// <inheritdoc/>
 		protected override void ResolveSuper(UhtResolvePhase resolvePhase)
 		{
@@ -711,12 +719,12 @@ namespace EpicGames.UHT.Types
 			switch (resolvePhase)
 			{
 				case UhtResolvePhase.Bases:
-					BindAndResolveSuper(SuperIdentifier, UhtFindOptions.ScriptStruct);
 
 					// if we have a base struct, propagate inherited struct flags now
 					UhtScriptStruct? superScriptStruct = SuperScriptStruct;
 					if (superScriptStruct != null)
 					{
+						superScriptStruct.Resolve(UhtResolvePhase.Bases);
 						ScriptStructFlags |= superScriptStruct.ScriptStructFlags & EStructFlags.Inherit;
 					}
 					break;
diff --git a/Engine/Source/Programs/Shared/EpicGames.UHT/Types/UhtStruct.cs b/Engine/Source/Programs/Shared/EpicGames.UHT/Types/UhtStruct.cs
--- a/Engine/Source/Programs/Shared/EpicGames.UHT/Types/UhtStruct.cs
+++ b/Engine/Source/Programs/Shared/EpicGames.UHT/Types/UhtStruct.cs
@@ -227,7 +227,7 @@ namespace EpicGames.UHT.Types
 		/// <param name="superIdentifier">Token that represent the super</param>
 		/// <param name="findOptions">Find options to restrict types</param>
 		/// <exception cref="UhtException">Thrown if super can not be found</exception>
-		public void BindAndResolveSuper(UhtToken superIdentifier, UhtFindOptions findOptions)
+		public void BindSuper(UhtToken superIdentifier, UhtFindOptions findOptions)
 		{
 			if (superIdentifier)
 			{
@@ -238,7 +238,6 @@ namespace EpicGames.UHT.Types
 				}
 				HeaderFile.AddReferencedHeader(Super);
 				MetaData.Parent = Super.MetaData;
-				Super.Resolve(UhtResolvePhase.Bases);
 			}
 		}
 
@@ -248,7 +247,7 @@ namespace EpicGames.UHT.Types
 		/// </summary>
 		/// <param name="baseIdentifiers">Collection of bases</param>
 		/// <param name="findOptions">Options to restrict types being searched</param>
-		public void BindAndResolveBases(List<UhtToken[]>? baseIdentifiers, UhtFindOptions findOptions)
+		public void BindBases(List<UhtToken[]>? baseIdentifiers, UhtFindOptions findOptions)
 		{
 			if (baseIdentifiers != null)
 			{
@@ -264,7 +263,6 @@ namespace EpicGames.UHT.Types
 						}
 						_bases.Add(baseStruct);
 						HeaderFile.AddReferencedHeader(baseStruct);
-						baseStruct.Resolve(UhtResolvePhase.Bases);
 					}
 				}
 			}
diff --git a/Engine/Source/Programs/Shared/EpicGames.UHT/Types/UhtType.cs b/Engine/Source/Programs/Shared/EpicGames.UHT/Types/UhtType.cs
--- a/Engine/Source/Programs/Shared/EpicGames.UHT/Types/UhtType.cs
+++ b/Engine/Source/Programs/Shared/EpicGames.UHT/Types/UhtType.cs
@@ -1008,6 +1008,20 @@ namespace EpicGames.UHT.Types
 				}
 			}
 		}
+
+		/// <summary>
+		/// Resolve all the super and base structures and classes
+		/// </summary>
+		public virtual void BindSuperAndBases()
+		{
+			if (_children != null)
+			{
+				foreach (UhtType child in _children)
+				{
+					child.BindSuperAndBases();
+				}
+			}
+		}
 		#endregion
 
 		#region FindType support
diff --git a/Engine/Source/Programs/Shared/EpicGames.UHT/Utils/UhtSession.cs b/Engine/Source/Programs/Shared/EpicGames.UHT/Utils/UhtSession.cs
--- a/Engine/Source/Programs/Shared/EpicGames.UHT/Utils/UhtSession.cs
+++ b/Engine/Source/Programs/Shared/EpicGames.UHT/Utils/UhtSession.cs
@@ -1040,6 +1040,8 @@ namespace EpicGames.UHT.Utils
 			StepParseHeaders();
 			StepPopulateTypeTable();
 			StepResolveInvalidCheck();
+			StepBindSuperAndBases();
+			RecursiveStructCheck();
 			StepResolveBases();
 			StepResolveProperties();
 			StepResolveFinal();
@@ -1986,6 +1988,11 @@ namespace EpicGames.UHT.Utils
 			});
 		}
 
+		private void StepBindSuperAndBases()
+		{
+			StepForAllHeaders(headerFile => headerFile.BindSuperAndBases());
+		}
+
 		private void StepResolveBases()
 		{
 			StepForAllHeaders(headerFile => Resolve(headerFile, UhtResolvePhase.Bases));
@@ -2152,7 +2159,7 @@ namespace EpicGames.UHT.Utils
 		}
 		#endregion
 
-		#region Topological sort of the header files
+		#region Topological testing of headers and structs
 		private enum TopologicalState
 		{
 			Unmarked,
@@ -2160,7 +2167,7 @@ namespace EpicGames.UHT.Utils
 			Permanent,
 		}
 
-		private void TopologicalVisit(List<TopologicalState> states, UhtHeaderFile visit, List<UhtHeaderFile> headerStack)
+		private void TopologicalHeaderVisit(List<TopologicalState> states, UhtHeaderFile visit, List<UhtHeaderFile> headerStack)
 		{
 			headerStack.Add(visit);
 			switch (states[visit.HeaderFileTypeIndex])
@@ -2169,7 +2176,7 @@ namespace EpicGames.UHT.Utils
 					states[visit.HeaderFileTypeIndex] = TopologicalState.Temporary;
 					foreach (UhtHeaderFile referenced in visit.ReferencedHeadersNoLock)
 					{
-						TopologicalVisit(states, referenced, headerStack);
+						TopologicalHeaderVisit(states, referenced, headerStack);
 					}
 					states[visit.HeaderFileTypeIndex] = TopologicalState.Permanent;
 					_sortedHeaderFiles.Add(visit);
@@ -2222,11 +2229,90 @@ namespace EpicGames.UHT.Utils
 				{
 					if (states[headerFile.HeaderFileTypeIndex] == TopologicalState.Unmarked)
 					{
-						TopologicalVisit(states, headerFile, headerStack);
+						TopologicalHeaderVisit(states, headerFile, headerStack);
 					}
 				}
 			});
 		}
+
+		private void TopologicalStructVisitChildren(List<TopologicalState> states, UhtType visit, List<UhtStruct> structStack)
+		{
+			foreach (UhtType child in visit.Children)
+			{
+				if (child is UhtStruct childStruct)
+				{
+					TopologicalStructVisit(states, childStruct, structStack);
+				}
+			}
+		}
+
+		private void TopologicalStructVisit(List<TopologicalState> states, UhtStruct visit, List<UhtStruct> structStack)
+		{
+			structStack.Add(visit);
+			switch (states[visit.ObjectTypeIndex])
+			{
+				case TopologicalState.Unmarked:
+					states[visit.ObjectTypeIndex] = TopologicalState.Temporary;
+					if (visit.Super != null)
+					{
+						TopologicalStructVisit(states, visit.Super, structStack);
+					}
+					foreach (UhtStruct baseStruct in visit.Bases)
+					{
+						TopologicalStructVisit(states, baseStruct, structStack);
+					}
+					TopologicalStructVisitChildren(states, visit, structStack);
+					states[visit.ObjectTypeIndex] = TopologicalState.Permanent;
+					break;
+
+				case TopologicalState.Temporary:
+					{
+						int index = structStack.IndexOf(visit);
+						if (index == -1 || index == structStack.Count - 1)
+						{
+							throw new UhtIceException("Error locating struct loop");
+						}
+						index++;
+						visit.LogError("Recursive class/struct definition:");
+						UhtStruct previous = visit;
+						for (int loopIndex = index; loopIndex < structStack.Count; loopIndex++)
+						{
+							UhtStruct next = structStack[loopIndex];
+							previous.LogError($"'{previous.SourceName}' inherits '{next.SourceName}'");
+							previous = next;
+						}
+					}
+					break;
+
+				case TopologicalState.Permanent:
+					break;
+
+				default:
+					throw new UhtIceException("Unknown topological state");
+			}
+			structStack.RemoveAt(structStack.Count - 1);
+		}
+
+		private void RecursiveStructCheck()
+		{
+			Try(null, () =>
+			{
+				Log.Logger.LogTrace("Step - Check for recursive structs");
+
+				// Initialize a scratch table for topological states
+				List<TopologicalState> states = new(ObjectTypeCount);
+				for (int index = 0; index < ObjectTypeCount; ++index)
+				{
+					states.Add(TopologicalState.Unmarked);
+				}
+
+				List<UhtStruct> structStack = new(32); // arbitrary capacity
+				foreach (UhtHeaderFile headerFile in HeaderFiles)
+				{
+					TopologicalStructVisitChildren(states, headerFile, structStack);
+				}
+			});
+		}
 		#endregion
 
 		#region Validation helpers
