Fixed C# UHT to properly detect an overridden function when they differ only by case.

diff --git a/Engine/Source/Programs/Shared/EpicGames.UHT/Types/UhtFunction.cs b/Engine/Source/Programs/Shared/EpicGames.UHT/Types/UhtFunction.cs
--- a/Engine/Source/Programs/Shared/EpicGames.UHT/Types/UhtFunction.cs
+++ b/Engine/Source/Programs/Shared/EpicGames.UHT/Types/UhtFunction.cs
@@ -638,7 +638,8 @@ namespace EpicGames.UHT.Types
 						UhtStruct? outerSuperStruct = outerStruct.SuperStruct;
 						if (outerSuperStruct != null)
 						{
-							UhtType? overriddenFunction = outerSuperStruct.FindType(UhtFindOptions.SourceName | UhtFindOptions.Function, SourceName);
+							// We use the engine name as the find option for a caseless compare.
+							UhtType? overriddenFunction = outerSuperStruct.FindType(UhtFindOptions.EngineName | UhtFindOptions.Function, SourceName);
 							if (overriddenFunction != null)
 							{
 								// Native function overrides should be done in CPP text, not in a UFUNCTION() declaration (you can't change flags, and it'd otherwise be a burden to keep them identical)
