Provided better messaging when DirectoryItem.CacheDirectories runs into duplicate keys

#rnx

diff --git a/Engine/Source/Programs/AutomationTool/Program.cs b/Engine/Source/Programs/AutomationTool/Program.cs
--- a/Engine/Source/Programs/AutomationTool/Program.cs
+++ b/Engine/Source/Programs/AutomationTool/Program.cs
@@ -420,7 +420,7 @@ namespace AutomationToolDriver
 			}
 			catch (Exception Ex)
             {
-				Logger.LogError(Ex, "Unhandled exception: {Message}", Ex.Message);
+				Logger.LogError(Ex, "Unhandled exception: {Message}", ExceptionUtils.FormatException(Ex));
             }
             finally
             {
diff --git a/Engine/Source/Programs/Shared/EpicGames.Build/System/DirectoryItem.cs b/Engine/Source/Programs/Shared/EpicGames.Build/System/DirectoryItem.cs
--- a/Engine/Source/Programs/Shared/EpicGames.Build/System/DirectoryItem.cs
+++ b/Engine/Source/Programs/Shared/EpicGames.Build/System/DirectoryItem.cs
@@ -224,6 +224,11 @@ namespace UnrealBuildBase
 					NewDirectories = new Dictionary<string, DirectoryItem>(Directories.Length, DirectoryReference.Comparer);
 					foreach (DirectoryInfo SubDirectoryInfo in Directories)
 					{
+						if (NewDirectories.ContainsKey(SubDirectoryInfo.Name))
+						{
+							throw new Exception($"Trying to add {SubDirectoryInfo.FullName} as '{SubDirectoryInfo.Name}' yet exists as {NewDirectories[SubDirectoryInfo.Name].FullName}");
+						}
+
 						NewDirectories.Add(SubDirectoryInfo.Name, DirectoryItem.GetItemByDirectoryInfo(SubDirectoryInfo));
 					}
 				}
