EpicGames.Build: Fix XML comment for NativeProjectBase.FindProjectFile()

#rnx

diff --git a/Engine/Source/Programs/Shared/EpicGames.Build/System/NativeProjectsBase.cs b/Engine/Source/Programs/Shared/EpicGames.Build/System/NativeProjectsBase.cs
--- a/Engine/Source/Programs/Shared/EpicGames.Build/System/NativeProjectsBase.cs
+++ b/Engine/Source/Programs/Shared/EpicGames.Build/System/NativeProjectsBase.cs
@@ -154,7 +154,8 @@ namespace UnrealBuildBase
 		/// <summary>
 		/// Takes a project name (e.g "ShooterGame") or path and attempt to find the existing uproject file in the base directories
 		/// </summary>
-		/// <param name="FilePath"></param>
+		/// <param name="Project">Project to search, either a name or a .uproject path</param>
+		/// <param name="Logger"></param>
 		/// <returns>A FileReference to an existing .uproject file, otherwise null</returns>
 		public static FileReference? FindProjectFile(string Project, ILogger Logger)
 		{
