Fixed C# UHT to property catch the error when a base class can not be found.  C# UHT will print the exception as an error and continue on.

diff --git a/Engine/Source/Programs/Shared/EpicGames.UHT/Utils/UhtSession.cs b/Engine/Source/Programs/Shared/EpicGames.UHT/Utils/UhtSession.cs
--- a/Engine/Source/Programs/Shared/EpicGames.UHT/Utils/UhtSession.cs
+++ b/Engine/Source/Programs/Shared/EpicGames.UHT/Utils/UhtSession.cs
@@ -1990,7 +1990,7 @@ namespace EpicGames.UHT.Utils
 
 		private void StepBindSuperAndBases()
 		{
-			StepForAllHeaders(headerFile => headerFile.BindSuperAndBases());
+			StepForAllHeaders(headerFile => BindSuperAndBases(headerFile));
 		}
 
 		private void StepResolveBases()
@@ -2045,6 +2045,18 @@ namespace EpicGames.UHT.Utils
 			}
 		}
 
+		private void BindSuperAndBases(UhtHeaderFile headerFile)
+		{
+			try
+			{
+				headerFile.BindSuperAndBases();
+			}
+			catch (Exception e)
+			{
+				HandleException(headerFile.MessageSource, e);
+			}
+		}
+
 		private delegate void StepDelegate(UhtHeaderFile headerFile);
 
 		private void StepForAllHeaders(StepDelegate stepDelegate)
