Horde: Fix enums in JSON schemas.

diff --git a/Engine/Source/Programs/Shared/EpicGames.Core/JsonSchema.cs b/Engine/Source/Programs/Shared/EpicGames.Core/JsonSchema.cs
--- a/Engine/Source/Programs/Shared/EpicGames.Core/JsonSchema.cs
+++ b/Engine/Source/Programs/Shared/EpicGames.Core/JsonSchema.cs
@@ -334,7 +334,14 @@ namespace EpicGames.Core
 				case TypeCode.UInt32:
 				case TypeCode.Int64:
 				case TypeCode.UInt64:
-					return new JsonSchemaInteger();
+					if (type.IsEnum)
+					{
+						return new JsonSchemaEnum(Enum.GetNames(type)) { Name = type.Name };
+					}
+					else
+					{
+						return new JsonSchemaInteger();
+					}
 				case TypeCode.Single:
 				case TypeCode.Double:
 					return new JsonSchemaNumber();
@@ -349,10 +356,6 @@ namespace EpicGames.Core
 					return new JsonSchemaString(str.Format);
 			}
 
-			if (type.IsEnum)
-			{
-				return new JsonSchemaEnum(Enum.GetNames(type)) { Name = type.Name };
-			}
 			if (type.IsGenericType && type.GetGenericTypeDefinition() == typeof(Nullable<>))
 			{
 				return CreateSchemaType(type.GetGenericArguments()[0], typeCache, xmlDoc);
