Fix issue where C# UHT fails to generate an error when a property in a structure references a deprecated class and the property has not been marked as deprecated.

diff --git a/Engine/Source/Programs/Shared/EpicGames.UHT/Types/UhtScriptStruct.cs b/Engine/Source/Programs/Shared/EpicGames.UHT/Types/UhtScriptStruct.cs
--- a/Engine/Source/Programs/Shared/EpicGames.UHT/Types/UhtScriptStruct.cs
+++ b/Engine/Source/Programs/Shared/EpicGames.UHT/Types/UhtScriptStruct.cs
@@ -860,7 +860,7 @@ namespace EpicGames.UHT.Types
 
 			ValidateProperties();
 
-			return options |= UhtValidationOptions.Shadowing;
+			return options |= UhtValidationOptions.Shadowing | UhtValidationOptions.Deprecated;
 		}
 
 		void ValidateProperties()
