Horde: Add prototype for new compute backend.

* Clients register requests for helpers with the server, matching a set of requirements, and listen on a specific port for responses.
* Agents connect to clients directly, and exchange messages over a socket.
* Messages are compact binary objects deriving from MessageBase.

diff --git a/Engine/Source/Programs/Horde/Horde.Agent/Commands/Compute/ComputeCommandV2.cs b/Engine/Source/Programs/Horde/Horde.Agent/Commands/Compute/ComputeCommandV2.cs
new file mode 100644
--- /dev/null
+++ b/Engine/Source/Programs/Horde/Horde.Agent/Commands/Compute/ComputeCommandV2.cs
@@ -0,0 +1,117 @@
+// Copyright Epic Games, Inc. All Rights Reserved.
+
+using System;
+using System.Collections.Generic;
+using System.Linq;
+using System.Net;
+using System.Net.Http;
+using System.Net.Sockets;
+using System.Threading;
+using System.Threading.Tasks;
+using EpicGames.Core;
+using EpicGames.Horde.Compute;
+using Microsoft.Extensions.Logging;
+using Microsoft.Extensions.Options;
+
+namespace Horde.Agent.Commands
+{
+	/// <summary>
+	/// Installs the agent as a service
+	/// </summary>
+	[Command("ComputeV2", "Executes a command through the Horde Compute API")]
+	class ComputeCommandV2 : Command
+	{
+		public class AddComputeTaskRequest
+		{
+			public Requirements? Requirements { get; set; }
+			public string RemoteIp { get; set; } = String.Empty;
+			public int RemotePort { get; set; } = 4000;
+		}
+
+		/// <summary>
+		/// Condition for 
+		/// </summary>
+		[CommandLine("-Condition")]
+		public string? Condition { get; set; }
+
+		readonly AgentSettings _settings;
+
+		/// <summary>
+		/// Constructor
+		/// </summary>
+		public ComputeCommandV2(IOptions<AgentSettings> settings)
+		{
+			_settings = settings.Value;
+		}
+
+		/// <inheritdoc/>
+		public override async Task<int> ExecuteAsync(ILogger logger)
+		{
+			ServerProfile profile = _settings.GetCurrentServerProfile();
+			logger.LogInformation("Server: {Server}", profile.Name);
+
+			int port = 9001;
+
+			using CancellationTokenSource source = new CancellationTokenSource();
+			Task task = Task.Run(() => ListenerProcess(port, logger, source.Token));
+
+			using (HttpClient client = new HttpClient())
+			{
+				AddComputeTaskRequest request = new AddComputeTaskRequest();
+				request.RemoteIp = "127.0.0.1";
+				request.RemotePort = port;
+
+				Uri uri = new Uri(profile.Url, "api/v2/compute");
+				using (HttpResponseMessage response = await client.PostAsync(uri.AbsoluteUri, request, source.Token))
+				{
+					response.EnsureSuccessStatusCode();
+				}
+			}
+
+			await task;
+			return 0;
+		}
+
+		static async Task ListenerProcess(int port, ILogger logger, CancellationToken cancellationToken)
+		{
+			TcpListener listener = new TcpListener(IPAddress.Loopback, port);
+			listener.Start();
+
+			List<Task> tasks = new List<Task>();
+			try
+			{
+				for (; ; )
+				{
+					TcpClient client = await listener.AcceptTcpClientAsync(cancellationToken);
+					logger.LogInformation("Received connection");
+
+					Task task = ServerProcess(client, logger, cancellationToken);
+					tasks.Add(task);
+				}
+			}
+			finally
+			{
+				await Task.WhenAll(tasks);
+				listener.Stop();
+			}
+		}
+
+		static async Task ServerProcess(TcpClient client, ILogger logger, CancellationToken cancellationToken)
+		{
+			MessageChannel channel = new MessageChannel(client.Client);
+			try
+			{
+				await channel.WriteAsync(new XorRequestMessage { Value = 123, Payload = new byte[] { 1, 2, 3, 4, 5 } }, cancellationToken);
+
+				XorResponseMessage response = await channel.ReadAsync<XorResponseMessage>(cancellationToken);
+				logger.LogInformation("Received response: {Message}", String.Join(", ", response.Payload.Select(x => x.ToString())));
+
+				await channel.WriteAsync(new CloseMessage(), cancellationToken);
+			}
+			finally
+			{
+				client.Dispose();
+			}
+		}
+	}
+}
diff --git a/Engine/Source/Programs/Horde/Horde.Agent/Horde.Agent.csproj b/Engine/Source/Programs/Horde/Horde.Agent/Horde.Agent.csproj
--- a/Engine/Source/Programs/Horde/Horde.Agent/Horde.Agent.csproj
+++ b/Engine/Source/Programs/Horde/Horde.Agent/Horde.Agent.csproj
@@ -154,6 +154,7 @@
     <Protobuf Include="..\Horde.Build\Protos\horde\messages\lease.proto" ProtoRoot="..\Horde.Build\Protos" Access="Public" />
     <Protobuf Include="..\Horde.Build\Protos\horde\messages\aws_lambda.proto" ProtoRoot="..\Horde.Build\Protos" Access="Public" />
     <Protobuf Include="..\Horde.Build\Protos\horde\tasks\compute_task.proto" ProtoRoot="..\Horde.Build\Protos" Access="Public" />
+    <Protobuf Include="..\Horde.Build\Protos\horde\tasks\compute_task_v2.proto" ProtoRoot="..\Horde.Build\Protos" Access="Public" />
     <Protobuf Include="..\Horde.Build\Protos\horde\tasks\conform_task.proto" ProtoRoot="..\Horde.Build\Protos" Access="Public" />
     <Protobuf Include="..\Horde.Build\Protos\horde\tasks\job_task.proto" ProtoRoot="..\Horde.Build\Protos" Access="Public" />
     <Protobuf Include="..\Horde.Build\Protos\horde\tasks\restart_task.proto" ProtoRoot="..\Horde.Build\Protos" Access="Public" />
diff --git a/Engine/Source/Programs/Horde/Horde.Agent/Leases/Handlers/ComputeHandlerV2.cs b/Engine/Source/Programs/Horde/Horde.Agent/Leases/Handlers/ComputeHandlerV2.cs
new file mode 100644
--- /dev/null
+++ b/Engine/Source/Programs/Horde/Horde.Agent/Leases/Handlers/ComputeHandlerV2.cs
@@ -0,0 +1,74 @@
+// Copyright Epic Games, Inc. All Rights Reserved.
+
+using System;
+using System.Net.Sockets;
+using System.Threading;
+using System.Threading.Tasks;
+using EpicGames.Horde.Compute;
+using Horde.Agent.Services;
+using HordeCommon.Rpc.Tasks;
+using Microsoft.Extensions.Logging;
+
+namespace Horde.Agent.Leases.Handlers
+{
+	class ComputeHandlerV2 : LeaseHandler<ComputeTaskMessageV2>
+	{
+		readonly ILogger _logger;
+
+		/// <summary>
+		/// Constructor
+		/// </summary>
+		public ComputeHandlerV2(ILogger<ComputeHandlerV2> logger)
+		{
+			_logger = logger;
+		}
+
+		/// <inheritdoc/>
+		public override async Task<LeaseResult> ExecuteAsync(ISession session, string leaseId, ComputeTaskMessageV2 computeTask, CancellationToken cancellationToken)
+		{
+			_logger.LogInformation("Starting compute task (lease {LeaseId})", leaseId);
+
+			try
+			{
+				await ConnectAsync(computeTask, cancellationToken);
+			}
+			catch (Exception ex)
+			{
+				_logger.LogError(ex, "Exception while executing compute task");
+			}
+
+			return LeaseResult.Success;
+		}
+
+		async Task ConnectAsync(ComputeTaskMessageV2 computeTask, CancellationToken cancellationToken)
+		{
+			using TcpClient tcpClient = new TcpClient();
+			await tcpClient.ConnectAsync(computeTask.RemoteIp, computeTask.RemotePort);
+
+			MessageChannel channel = new MessageChannel(tcpClient.Client);
+			for (; ; )
+			{
+				MessageBase request = await channel.ReadAsync(cancellationToken);
+				_logger.LogInformation("Received message {Type}", request.Type);
+
+				switch (request)
+				{
+					case CloseMessage _:
+						return;
+					case XorRequestMessage xorRequest:
+						{
+							XorResponseMessage response = new XorResponseMessage();
+							response.Payload = new byte[xorRequest.Payload.Length];
+							for (int idx = 0; idx < xorRequest.Payload.Length; idx++)
+							{
+								response.Payload[idx] = (byte)(xorRequest.Payload[idx] ^ xorRequest.Value);
+							}
+							await channel.WriteAsync(response, cancellationToken);
+						}
+						break;
+				}
+			}
+		}
+	}
+}
+
diff --git a/Engine/Source/Programs/Horde/Horde.Agent/Program.cs b/Engine/Source/Programs/Horde/Horde.Agent/Program.cs
--- a/Engine/Source/Programs/Horde/Horde.Agent/Program.cs
+++ b/Engine/Source/Programs/Horde/Horde.Agent/Program.cs
@@ -186,6 +186,7 @@ namespace Horde.Agent
 			services.AddSingleton<JobExecutorFactory, TestExecutorFactory>();
 
 			services.AddSingleton<LeaseHandler, ComputeHandler>();
+			services.AddSingleton<LeaseHandler, ComputeHandlerV2>();
 			services.AddSingleton<LeaseHandler, ConformHandler>();
 			services.AddSingleton<LeaseHandler, JobHandler>();
 			services.AddSingleton<LeaseHandler, RestartHandler>();
diff --git a/Engine/Source/Programs/Horde/Horde.Build/Compute/V2/ComputeControllerV2.cs b/Engine/Source/Programs/Horde/Horde.Build/Compute/V2/ComputeControllerV2.cs
new file mode 100644
--- /dev/null
+++ b/Engine/Source/Programs/Horde/Horde.Build/Compute/V2/ComputeControllerV2.cs
@@ -0,0 +1,80 @@
+// Copyright Epic Games, Inc. All Rights Reserved.
+
+using System;
+using System.Collections.Generic;
+using System.Threading;
+using System.Threading.Tasks;
+using EpicGames.Horde.Common;
+using EpicGames.Horde.Compute;
+using EpicGames.Horde.Compute.Impl;
+using Horde.Build.Acls;
+using Horde.Build.Server;
+using Horde.Build.Utilities;
+using Microsoft.AspNetCore.Authorization;
+using Microsoft.AspNetCore.Mvc;
+using Microsoft.Extensions.Options;
+
+namespace Horde.Build.Compute.V2
+{
+	/// <summary>
+	/// Request a machine to execute compute requests
+	/// </summary>
+	public class AddComputeTaskRequest
+	{
+		/// <summary>
+		/// Condition to identify machines that can execute the request
+		/// </summary>
+		public Requirements? Requirements { get; set; }
+
+		/// <summary>
+		/// IP address to contact the initiator on.
+		/// </summary>
+		public string RemoteIp { get; set; } = String.Empty;
+
+		/// <summary>
+		/// Port to connect on
+		/// </summary>
+		public int RemotePort { get; set; } = 4000;
+	}
+
+	/// <summary>
+	/// Controller for the /api/v2/compute endpoint
+	/// </summary>
+	[ApiController]
+	[Authorize]
+	[Route("[controller]")]
+	public class ComputeControllerV2 : HordeControllerBase
+	{
+		readonly ComputeServiceV2 _computeService;
+		readonly IOptionsSnapshot<GlobalConfig> _globalConfig;
+
+		/// <summary>
+		/// Constructor
+		/// </summary>
+		public ComputeControllerV2(ComputeServiceV2 computeService, IOptionsSnapshot<GlobalConfig> globalConfig)
+		{
+			_computeService = computeService;
+			_globalConfig = globalConfig;
+		}
+
+		/// <summary>
+		/// Add tasks to be executed remotely
+		/// </summary>
+		/// <param name="request">The request parameters</param>
+		/// <returns></returns>
+		[HttpPost]
+		[Authorize]
+		[Route("/api/v2/compute")]
+		public ActionResult AddTasksAsync([FromBody] AddComputeTaskRequest request)
+		{
+			if(!_globalConfig.Value.Authorize(AclAction.AddComputeTasks, User))
+			{
+				return Forbid(AclAction.AddComputeTasks);
+			}
+
+			Requirements requirements = request.Requirements ?? new Requirements();
+			_computeService.AddRequest(requirements, request.RemoteIp, request.RemotePort);
+			return Ok();
+		}
+	}
+}
diff --git a/Engine/Source/Programs/Horde/Horde.Build/Compute/V2/ComputeServiceV2.cs b/Engine/Source/Programs/Horde/Horde.Build/Compute/V2/ComputeServiceV2.cs
new file mode 100644
--- /dev/null
+++ b/Engine/Source/Programs/Horde/Horde.Build/Compute/V2/ComputeServiceV2.cs
@@ -0,0 +1,141 @@
+// Copyright Epic Games, Inc. All Rights Reserved.
+
+using System;
+using System.Collections.Generic;
+using System.Threading;
+using System.Threading.Tasks;
+using EpicGames.Core;
+using EpicGames.Horde.Compute;
+using EpicGames.Serialization;
+using Google.Protobuf;
+using Google.Protobuf.WellKnownTypes;
+using Horde.Build.Agents;
+using Horde.Build.Agents.Leases;
+using Horde.Build.Agents.Pools;
+using Horde.Build.Tasks;
+using Horde.Build.Utilities;
+using HordeCommon;
+using HordeCommon.Rpc.Tasks;
+using Microsoft.Extensions.Logging;
+
+namespace Horde.Build.Compute.V2
+{
+	using LeaseId = ObjectId<ILease>;
+
+	/// <summary>
+	/// Dispatches remote actions. Does not implement any cross-pod communication to satisfy leases; only agents connected to this server instance will be stored.
+	/// </summary>
+	public sealed class ComputeServiceV2 : TaskSourceBase<ComputeTaskMessageV2>
+	{
+		class Request
+		{
+			public string Ip { get; }
+			public int Port { get; }
+
+			public Request(string ip, int port)
+			{
+				Ip = ip;
+				Port = port;
+			}
+		}
+
+		class RequestQueue
+		{
+			public IoHash Hash { get; }
+			public Requirements Requirements { get; }
+			public Queue<Request> Requests { get; } = new Queue<Request>();
+
+			public RequestQueue(IoHash hash, Requirements requirements)
+			{
+				Hash = hash;
+				Requirements = requirements;
+			}
+		}
+
+		readonly object _lockObject = new object();
+		readonly Dictionary<IoHash, RequestQueue> _queues = new Dictionary<IoHash, RequestQueue>();
+		readonly ILogger _logger;
+
+		/// <inheritdoc/>
+		public override string Type => "ComputeV2";
+
+		/// <inheritdoc/>
+		public override TaskSourceFlags Flags => TaskSourceFlags.None;
+
+		/// <summary>
+		/// Constructor
+		/// </summary>
+		/// <param name="logger"></param>
+		public ComputeServiceV2(ILogger<ComputeServiceV2> logger)
+		{
+			_logger = logger;
+		}
+
+		/// <summary>
+		/// 
+		/// </summary>
+		/// <param name="requirements"></param>
+		/// <param name="ip"></param>
+		/// <param name="port"></param>
+		public void AddRequest(Requirements requirements, string ip, int port)
+		{
+			CbObject obj = CbSerializer.Serialize(requirements);
+			IoHash hash = IoHash.Compute(obj.GetView().Span);
+			lock (_lockObject)
+			{
+				RequestQueue? queue;
+				if (!_queues.TryGetValue(hash, out queue))
+				{
+					queue = new RequestQueue(hash, requirements);
+					_queues.Add(hash, queue);
+				}
+				queue.Requests.Enqueue(new Request(ip, port));
+			}
+		}
+
+		/// <inheritdoc/>
+		public override async Task<Task<AgentLease?>> AssignLeaseAsync(IAgent agent, CancellationToken cancellationToken)
+		{
+			while (!cancellationToken.IsCancellationRequested)
+			{
+				Request? request = null;
+				RequestQueue? requestQueue = null;
+
+				lock (_lockObject)
+				{
+					foreach ((IoHash hash, RequestQueue queue) in _queues)
+					{
+						if (agent.MeetsRequirements(queue.Requirements))
+						{
+							request = queue.Requests.Dequeue();
+							requestQueue = queue;
+
+							if (queue.Requests.Count == 0)
+							{
+								_queues.Remove(hash);
+							}
+							break;
+						}
+					}
+				}
+
+				if (request != null)
+				{
+					ComputeTaskMessageV2 computeTask = new ComputeTaskMessageV2();
+					computeTask.RemoteIp = request.Ip;
+					computeTask.RemotePort = request.Port;
+
+					string leaseName = $"Compute lease";
+					byte[] payload = Any.Pack(computeTask).ToByteArray();
+
+					AgentLease lease = new AgentLease(LeaseId.GenerateNewId(), leaseName, null, null, null, LeaseState.Pending, requestQueue!.Requirements.Resources, requestQueue.Requirements.Exclusive, payload);
+					_logger.LogInformation("Created compute lease for agent {AgentId} and remote {RemoteIp}:{RemotePort}", agent.Id, request.Ip, request.Port);
+					return Task.FromResult<AgentLease?>(lease);
+				}
+
+				await AsyncUtils.DelayNoThrow(TimeSpan.FromSeconds(5.0), cancellationToken);
+			}
+			return Task.FromResult<AgentLease?>(null);
+		}
+	}
+}
diff --git a/Engine/Source/Programs/Horde/Horde.Build/Horde.Build.csproj b/Engine/Source/Programs/Horde/Horde.Build/Horde.Build.csproj
--- a/Engine/Source/Programs/Horde/Horde.Build/Horde.Build.csproj
+++ b/Engine/Source/Programs/Horde/Horde.Build/Horde.Build.csproj
@@ -94,6 +94,7 @@
     <Protobuf Include="Protos\horde\messages\agent_workspace.proto" ProtoRoot="Protos" />
     <Protobuf Include="Protos\horde\messages\lease.proto" ProtoRoot="Protos" />
     <Protobuf Include="Protos\horde\tasks\compute_task.proto" ProtoRoot="Protos" />
+    <Protobuf Include="Protos\horde\tasks\compute_task_v2.proto" ProtoRoot="Protos" />
     <Protobuf Include="Protos\horde\tasks\conform_task.proto" ProtoRoot="Protos" />
     <Protobuf Include="Protos\horde\tasks\job_task.proto" ProtoRoot="Protos" />
     <Protobuf Include="Protos\horde\tasks\restart_task.proto" ProtoRoot="Protos" />
diff --git a/Engine/Source/Programs/Horde/Horde.Build/Startup.cs b/Engine/Source/Programs/Horde/Horde.Build/Startup.cs
--- a/Engine/Source/Programs/Horde/Horde.Build/Startup.cs
+++ b/Engine/Source/Programs/Horde/Horde.Build/Startup.cs
@@ -685,6 +685,7 @@ namespace Horde.Build
 			services.AddSingleton<Compute.V1.ComputeService>();
 			services.AddSingleton<Compute.V1.IComputeService, Compute.V1.ComputeService>();
 			services.AddHostedService(provider => provider.GetRequiredService<Compute.V1.ComputeService>());
+			services.AddSingleton<Compute.V2.ComputeServiceV2>();
 
 			services.AddSingleton<ITaskSource, UpgradeTaskSource>();
 			services.AddSingleton<ITaskSource, ShutdownTaskSource>();
@@ -692,6 +693,7 @@ namespace Horde.Build
 			services.AddSingleton<ITaskSource, ConformTaskSource>(provider => provider.GetRequiredService<ConformTaskSource>());
 			services.AddSingleton<ITaskSource, JobTaskSource>(provider => provider.GetRequiredService<JobTaskSource>());
 			services.AddSingleton<ITaskSource, Compute.V1.ComputeService>();
+			services.AddSingleton<ITaskSource, Compute.V2.ComputeServiceV2>(provider => provider.GetRequiredService<Compute.V2.ComputeServiceV2>());
 
 			services.AddHostedService(provider => provider.GetRequiredService<ConformTaskSource>());
 
diff --git a/Engine/Source/Programs/Shared/EpicGames.Core/VarInt.cs b/Engine/Source/Programs/Shared/EpicGames.Core/VarInt.cs
--- a/Engine/Source/Programs/Shared/EpicGames.Core/VarInt.cs
+++ b/Engine/Source/Programs/Shared/EpicGames.Core/VarInt.cs
@@ -95,6 +95,17 @@ namespace EpicGames.Core
 			return BitOperations.LeadingZeroCount(b) - 23;
 		}
 
+		/// <summary>
+		/// Measure the length in bytes (1-9) of an encoded variable-length integer.
+		/// </summary>
+		/// <param name="b">First byte of the encoded integer.</param>
+		/// <returns>The number of bytes used to encode the integer, in the range 1-9.</returns>
+		public static int Measure(byte b)
+		{
+			b = (byte)~b;
+			return BitOperations.LeadingZeroCount(b) - 23;
+		}
+
 		/// <summary>
 		/// Measure the number of bytes required to encode the input.
 		/// </summary>
diff --git a/Engine/Source/Programs/Shared/EpicGames.Horde/Common/Condition.cs b/Engine/Source/Programs/Shared/EpicGames.Horde/Common/Condition.cs
--- a/Engine/Source/Programs/Shared/EpicGames.Horde/Common/Condition.cs
+++ b/Engine/Source/Programs/Shared/EpicGames.Horde/Common/Condition.cs
@@ -632,13 +632,39 @@ namespace EpicGames.Horde.Common
 	public class ConditionCbConverter : CbConverterBase<Condition>
 	{
 		/// <inheritdoc/>
-		public override Condition Read(CbField field) => Condition.TryParse(field.AsUtf8String().ToString());
+		public override Condition Read(CbField field)
+		{
+			if (field.IsNull())
+			{
+				return null!;
+			}
+			else
+			{
+				return Condition.TryParse(field.AsUtf8String().ToString());
+			}
+		}
 
 		/// <inheritdoc/>
-		public override void Write(CbWriter writer, Condition value) => writer.WriteUtf8StringValue(value.Text);
+		public override void Write(CbWriter writer, Condition value)
+		{
+			if (value == null)
+			{
+				writer.WriteNullValue();
+			}
+			else
+			{
+				writer.WriteUtf8StringValue(value.Text);
+			}
+		}
 
 		/// <inheritdoc/>
-		public override void WriteNamed(CbWriter writer, Utf8String name, Condition value) => writer.WriteUtf8String(name, value.Text);
+		public override void WriteNamed(CbWriter writer, Utf8String name, Condition value)
+		{
+			if (value != null)
+			{
+				writer.WriteUtf8String(name, value.Text);
+			}
+		}
 	}
 
 	/// <summary>
diff --git a/Engine/Source/Programs/Shared/EpicGames.Horde/Compute/Message.cs b/Engine/Source/Programs/Shared/EpicGames.Horde/Compute/Message.cs
new file mode 100644
--- /dev/null
+++ b/Engine/Source/Programs/Shared/EpicGames.Horde/Compute/Message.cs
@@ -0,0 +1,41 @@
+// Copyright Epic Games, Inc. All Rights Reserved.
+
+using EpicGames.Serialization;
+using System;
+
+#pragma warning disable CS1591
+#pragma warning disable CA1819
+
+namespace EpicGames.Horde.Compute
+{
+	/// <summary>
+	/// Terminate the current connection
+	/// </summary>
+	[Message("close"), CbObject]
+	public class CloseMessage : MessageBase
+	{
+	}
+
+	/// <summary>
+	/// XOR a block of data with a value
+	/// </summary>
+	[Message("xor-req")]
+	public class XorRequestMessage : MessageBase
+	{
+		[CbField]
+		public byte Value { get; set; }
+
+		[CbField]
+		public byte[] Payload { get; set; } = Array.Empty<byte>();
+	}
+
+	/// <summary>
+	/// Response from an XOR message
+	/// </summary>
+	[Message("xor-rsp")]
+	public class XorResponseMessage : MessageBase
+	{
+		[CbField]
+		public byte[] Payload { get; set; } = Array.Empty<byte>();
+	}
+}
diff --git a/Engine/Source/Programs/Shared/EpicGames.Horde/Compute/MessageChannel.cs b/Engine/Source/Programs/Shared/EpicGames.Horde/Compute/MessageChannel.cs
new file mode 100644
--- /dev/null
+++ b/Engine/Source/Programs/Shared/EpicGames.Horde/Compute/MessageChannel.cs
@@ -0,0 +1,176 @@
+// Copyright Epic Games, Inc. All Rights Reserved.
+
+using EpicGames.Core;
+using EpicGames.Serialization;
+using System;
+using System.Collections.Generic;
+using System.Net.Sockets;
+using System.Reflection;
+using System.Threading;
+using System.Threading.Tasks;
+
+namespace EpicGames.Horde.Compute
+{
+	/// <summary>
+	/// Base class for messages
+	/// </summary>
+	public abstract class MessageBase
+	{
+		/// <summary>
+		/// Type of the message
+		/// </summary>
+		public Utf8String Type { get; }
+
+		/// <summary>
+		/// Constructor
+		/// </summary>
+		protected MessageBase()
+		{
+			Type = GetType().GetCustomAttribute<MessageAttribute>()!.Name;
+		}
+	}
+
+	/// <summary>
+	/// Attribute used to annotate a message type with an identifier string
+	/// </summary>
+	[AttributeUsage(AttributeTargets.Class)]
+	public sealed class MessageAttribute : Attribute
+	{
+		/// <summary>
+		/// Message type
+		/// </summary>
+		public string Name { get; }
+
+		/// <summary>
+		/// Constructor
+		/// </summary>
+		public MessageAttribute(string name)
+		{
+			Name = name;
+		}
+	}
+
+	/// <summary>
+	/// Reads messages from a socket
+	/// </summary>
+	public class MessageChannel
+	{
+		static readonly Utf8String s_typeField = "name";
+		static readonly Utf8String s_dataField = "data";
+
+		static readonly Type[] s_messageTypes =
+		{
+			typeof(CloseMessage),
+			typeof(XorRequestMessage),
+			typeof(XorResponseMessage)
+		};
+
+		static readonly Dictionary<Utf8String, ICbConverter> s_typeToConverter = CreateTypeToConverterMap();
+
+		readonly Socket _socket;
+		readonly CbWriter _writer;
+
+		byte[] _buffer = new byte[64];
+
+		/// <summary>
+		/// Constructor
+		/// </summary>
+		/// <param name="socket"></param>
+		public MessageChannel(Socket socket)
+		{
+			_socket = socket;
+			_writer = new CbWriter();
+		}
+
+		/// <summary>
+		/// Build a lookup of all known message types
+		/// </summary>
+		static Dictionary<Utf8String, ICbConverter> CreateTypeToConverterMap()
+		{
+			Dictionary<Utf8String, ICbConverter> typeToConverter = new Dictionary<Utf8String, ICbConverter>();
+			foreach (Type messageType in s_messageTypes)
+			{
+				string name = messageType.GetCustomAttribute<MessageAttribute>()!.Name;
+				ICbConverter converter = CbConverter.GetConverter(messageType);
+				typeToConverter.Add(name, converter);
+			}
+			return typeToConverter;
+		}
+
+		/// <summary>
+		/// Receives a message from the remote
+		/// </summary>
+		/// <param name="cancellationToken"></param>
+		/// <returns></returns>
+		public async Task<T> ReadAsync<T>(CancellationToken cancellationToken) where T : MessageBase
+		{
+			return (T)await ReadAsync(cancellationToken);
+		}
+
+		/// <summary>
+		/// Receives a message from the remote
+		/// </summary>
+		/// <param name="cancellationToken"></param>
+		/// <returns></returns>
+		public async Task<MessageBase> ReadAsync(CancellationToken cancellationToken)
+		{
+			int length = 0;
+			for (; ; )
+			{
+				int read = await _socket.ReceiveAsync(_buffer.AsMemory(length), SocketFlags.None, cancellationToken);
+				length += read;
+
+				if (length >= 2)
+				{
+					int requiredLength = 1 + VarInt.Measure(_buffer[1]);
+					if (length > requiredLength)
+					{
+						requiredLength += (int)VarInt.ReadUnsigned(_buffer.AsSpan(1));
+					}
+					if (length >= requiredLength)
+					{
+						break;
+					}
+
+					if (requiredLength > _buffer.Length)
+					{
+						Array.Resize(ref _buffer, requiredLength);
+					}
+				}
+			}
+
+			CbObject obj = new CbObject(_buffer.AsMemory(0, length));
+
+			CbField typeField = obj.Find(s_typeField);
+			Utf8String type = typeField.AsUtf8String();
+
+			ICbConverter converter = s_typeToConverter[type];
+
+			CbField dataField = obj.Find(s_dataField);
+			return (MessageBase)converter.ReadObject(dataField)!;
+		}
+
+		/// <summary>
+		/// Sends a message to the remote
+		/// </summary>
+		/// <param name="message">Message to be sent</param>
+		/// <param name="cancellationToken">Cancellation token</param>
+		/// <returns></returns>
+		public async Task WriteAsync(MessageBase message, CancellationToken cancellationToken)
+		{
+			Utf8String name = message.Type;
+		
+			_writer.Clear();
+			_writer.BeginObject();
+			_writer.WriteUtf8String(s_typeField, name);
+
+			ICbConverter converter = CbConverter.GetConverter(message.GetType());
+			converter.WriteNamedObject(_writer, s_dataField, message);
+
+			_writer.EndObject();
+
+			byte[] packet = _writer.ToByteArray();
+			await _socket.SendAsync(packet, SocketFlags.None, cancellationToken);
+		}
+	}
+}
