OIDC: Add default email scope to ensure compatibility with Horde

diff --git a/Engine/Source/Programs/Shared/EpicGames.OIDC/OidcTokenManager.cs b/Engine/Source/Programs/Shared/EpicGames.OIDC/OidcTokenManager.cs
--- a/Engine/Source/Programs/Shared/EpicGames.OIDC/OidcTokenManager.cs
+++ b/Engine/Source/Programs/Shared/EpicGames.OIDC/OidcTokenManager.cs
@@ -511,6 +511,6 @@ namespace EpicGames.OIDC
 		[Required] public Uri RedirectUri { get; set; } = null!;
 		[Required] public bool LoadClaimsFromUserProfile { get; set; } = false;
 
-		public string Scopes { get; set; } = "openid profile offline_access";
+		public string Scopes { get; set; } = "openid profile offline_access email";
 	}
 }
\ No newline at end of file
