EpicGames.Build: Validate all existing build records, not just the records associated with found projects

#rnx

diff --git a/Engine/Source/Programs/Shared/EpicGames.Build/System/CompileScriptModules.cs b/Engine/Source/Programs/Shared/EpicGames.Build/System/CompileScriptModules.cs
--- a/Engine/Source/Programs/Shared/EpicGames.Build/System/CompileScriptModules.cs
+++ b/Engine/Source/Programs/Shared/EpicGames.Build/System/CompileScriptModules.cs
@@ -286,13 +286,30 @@ namespace UnrealBuildBase
 			// Only validate the build records if requested
 			if (BuildFlags.HasFlag(BuildFlags.UseBuildRecords))
 			{
-				foreach (FileReference Project in FoundProjects)
+				if (Unreal.IsEngineInstalled())
 				{
-					ValidateBuildRecordRecursively(BuildRecords, Project, Hook, Logger);
+					// Validate build records only for projects to be compiled
+					foreach (FileReference Project in FoundProjects)
+					{
+						ValidateBuildRecordRecursively(BuildRecords, Project, Hook, Logger);
+					}
+				}
+				else
+				{
+					// Validate build records for all existing records
+					foreach (FileReference Project in BuildRecords.Keys)
+					{
+						ValidateBuildRecordRecursively(BuildRecords, Project, Hook, Logger);
+					}
 				}
 			}
 			else
 			{
+				// If not using build records, mark all existing records as invalid and force compile
+				foreach (CsProjBuildRecordEntry Record in BuildRecords.Values)
+				{
+					Record.Status = CsProjBuildRecordStatus.Invalid;
+				}
 				BuildFlags |= BuildFlags.ForceCompile;
 			}
 
