EpicGames.Build: Fix UserSettingDirectory not being statically initialized in the correct order, which is preventing WritableEngineDirectory from being set correctly for installed builds.

#rnx

diff --git a/Engine/Source/Programs/Shared/EpicGames.Build/Unreal.cs b/Engine/Source/Programs/Shared/EpicGames.Build/Unreal.cs
--- a/Engine/Source/Programs/Shared/EpicGames.Build/Unreal.cs
+++ b/Engine/Source/Programs/Shared/EpicGames.Build/Unreal.cs
@@ -143,15 +143,20 @@ namespace UnrealBuildBase
 		/// </summary>
 		public static readonly DirectoryReference EngineSourceDirectory = DirectoryReference.Combine(EngineDirectory, "Source");
 
+		/// <summary>
+		/// Returns the User Settings Directory path. This matches FPlatformProcess::UserSettingsDir().
+		/// </summary>
+		static public readonly DirectoryReference UserSettingDirectory = GetUserSettingDirectory();
+
 		/// <summary>
 		/// Writable engine directory. Uses the user's settings folder for installed builds.
 		/// </summary>
-		public static readonly DirectoryReference WritableEngineDirectory = IsEngineInstalled() && UserSettingDirectory != null ? DirectoryReference.Combine(UserSettingDirectory, "UnrealEngine") : EngineDirectory;
+		public static readonly DirectoryReference WritableEngineDirectory = IsEngineInstalled() ? DirectoryReference.Combine(UserSettingDirectory, "UnrealEngine") : EngineDirectory;
 
 		/// <summary>
 		/// The engine saved programs directory
 		/// </summary>
-		public static readonly DirectoryReference EngineProgramSavedDirectory = IsEngineInstalled() && UserSettingDirectory != null ? UserSettingDirectory : DirectoryReference.Combine(EngineDirectory, "Programs");
+		public static readonly DirectoryReference EngineProgramSavedDirectory = IsEngineInstalled() ? UserSettingDirectory : DirectoryReference.Combine(EngineDirectory, "Programs");
 
 		/// <summary>
 		/// The path to UBT
@@ -174,13 +179,6 @@ namespace UnrealBuildBase
 		/// </summary>
 		static public readonly FileReference DotnetPath = FileReference.Combine(DotnetDirectory, "dotnet" + RuntimePlatform.ExeExtension);
 
-		/// <summary>
-		/// Returns the User Settings Directory path. This matches FPlatformProcess::UserSettingsDir().
-		/// NOTE: This function may return null. Some accounts (eg. the SYSTEM account on Windows) do not have a personal folder, and Jenkins
-		/// runs using this account by default.
-		/// </summary>
-		static public readonly DirectoryReference? UserSettingDirectory = GetUserSettingDirectory();
-
 		/// <summary>
 		/// Whether we're running with engine installed
 		/// </summary>
@@ -342,10 +340,8 @@ namespace UnrealBuildBase
 
 		/// <summary>
 		/// Returns the User Settings Directory path. This matches FPlatformProcess::UserSettingsDir().
-		/// NOTE: This function may return null. Some accounts (eg. the SYSTEM account on Windows) do not have a personal folder, and Jenkins
-		/// runs using this account by default.
 		/// </summary>
-		private static DirectoryReference? GetUserSettingDirectory()
+		private static DirectoryReference GetUserSettingDirectory()
 		{
 			if (RuntimePlatform.IsMac)
 			{
@@ -358,16 +354,21 @@ namespace UnrealBuildBase
 			else
 			{
 				// Not all user accounts have a local application data directory (eg. SYSTEM, used by Jenkins for builds).
-				string DirectoryName = Environment.GetFolderPath(Environment.SpecialFolder.LocalApplicationData);
-				if (String.IsNullOrEmpty(DirectoryName))
-				{
-					return DirectoryReference.Combine(EngineDirectory, "Saved");
-				}
-				else
+				List<Environment.SpecialFolder> DataFolders = new List<Environment.SpecialFolder>() {
+					Environment.SpecialFolder.LocalApplicationData,
+					Environment.SpecialFolder.CommonApplicationData
+				};
+				foreach (Environment.SpecialFolder DataFolder in DataFolders)
 				{
-					return new DirectoryReference(DirectoryName);
+					string DirectoryName = Environment.GetFolderPath(DataFolder);
+					if (!String.IsNullOrEmpty(DirectoryName))
+					{
+						return new DirectoryReference(DirectoryName);
+					}
 				}
 			}
+
+			return DirectoryReference.Combine(EngineDirectory, "Saved");
 		}
 	}
 }
