Horde: Various improvements for remote execution.

* All remote execution requests are now created with a nonce used to identify the replying agent.
* ComputeClient class in EpicGames.Horde wraps functionality to schedule and initiate requests and route them to the appropriate handler.
* Traffic is encrypted using AES, with a key passed to the server in the initial request message.

diff --git a/Engine/Source/Programs/Horde/Horde.Agent/Commands/Compute/ComputeCommandV2.cs b/Engine/Source/Programs/Horde/Horde.Agent/Commands/Compute/ComputeCommandV2.cs
--- a/Engine/Source/Programs/Horde/Horde.Agent/Commands/Compute/ComputeCommandV2.cs
+++ b/Engine/Source/Programs/Horde/Horde.Agent/Commands/Compute/ComputeCommandV2.cs
@@ -2,10 +2,12 @@
 
 using System;
 using System.Collections.Generic;
+using System.IO;
 using System.Linq;
 using System.Net;
 using System.Net.Http;
 using System.Net.Sockets;
+using System.Security.Cryptography;
 using System.Threading;
 using System.Threading.Tasks;
 using EpicGames.Core;
@@ -21,13 +23,6 @@ namespace Horde.Agent.Commands
 	[Command("ComputeV2", "Executes a command through the Horde Compute API")]
 	class ComputeCommandV2 : Command
 	{
-		public class AddComputeTaskRequest
-		{
-			public Requirements? Requirements { get; set; }
-			public string RemoteIp { get; set; } = String.Empty;
-			public int RemotePort { get; set; } = 4000;
-		}
-
 		/// <summary>
 		/// Condition for 
 		/// </summary>
@@ -50,68 +45,30 @@ namespace Horde.Agent.Commands
 			ServerProfile profile = _settings.GetCurrentServerProfile();
 			logger.LogInformation("Server: {Server}", profile.Name);
 
-			int port = 9001;
-
-			using CancellationTokenSource source = new CancellationTokenSource();
-			Task task = Task.Run(() => ListenerProcess(port, logger, source.Token));
+			await using ComputeClient client = new ComputeClient(CreateHttpClient, logger);
+			client.Start();
 
-			using (HttpClient client = new HttpClient())
-			{
-				AddComputeTaskRequest request = new AddComputeTaskRequest();
-				request.RemoteIp = "127.0.0.1";
-				request.RemotePort = port;
+			IComputeRequest<object?> request = await client.AddRequestAsync(null, TestCommandAsync);
+			await request.Result;
 
-				Uri uri = new Uri(profile.Url, "api/v2/compute");
-				using (HttpResponseMessage response = await client.PostAsync(uri.AbsoluteUri, request, source.Token))
-				{
-					response.EnsureSuccessStatusCode();
-				}
-			}
-
-			await task;
 			return 0;
 		}
 
-		static async Task ListenerProcess(int port, ILogger logger, CancellationToken cancellationToken)
+		HttpClient CreateHttpClient()
 		{
-			TcpListener listener = new TcpListener(IPAddress.Loopback, port);
-			listener.Start();
-
-			List<Task> tasks = new List<Task>();
-			try
-			{
-				for (; ; )
-				{
-					TcpClient client = await listener.AcceptTcpClientAsync(cancellationToken);
-					logger.LogInformation("Received connection");
-
-					Task task = ServerProcess(client, logger, cancellationToken);
-					tasks.Add(task);
-				}
-			}
-			finally
-			{
-				await Task.WhenAll(tasks);
-				listener.Stop();
-			}
+			HttpClient client = new HttpClient();
+			client.BaseAddress = _settings.GetCurrentServerProfile().Url;
+			return client;
 		}
 
-		static async Task ServerProcess(TcpClient client, ILogger logger, CancellationToken cancellationToken)
+		async Task<object?> TestCommandAsync(ComputeChannel channel, CancellationToken cancellationToken)
 		{
-			MessageChannel channel = new MessageChannel(client.Client);
-			try
-			{
-				await channel.WriteAsync(new XorRequestMessage { Value = 123, Payload = new byte[] { 1, 2, 3, 4, 5 } }, cancellationToken);
+			await channel.WriteAsync(new XorRequestMessage { Value = 123, Payload = new byte[] { 1, 2, 3, 4, 5 } }, cancellationToken);
 
-				XorResponseMessage response = await channel.ReadAsync<XorResponseMessage>(cancellationToken);
-				logger.LogInformation("Received response: {Message}", String.Join(", ", response.Payload.Select(x => x.ToString())));
+			XorResponseMessage response = await channel.ReadAsync<XorResponseMessage>(cancellationToken);
 
-				await channel.WriteAsync(new CloseMessage(), cancellationToken);
-			}
-			finally
-			{
-				client.Dispose();
-			}
+			await channel.WriteAsync(new CloseMessage(), cancellationToken);
+			return null;
 		}
 	}
 }
diff --git a/Engine/Source/Programs/Horde/Horde.Agent/Leases/Handlers/ComputeHandlerV2.cs b/Engine/Source/Programs/Horde/Horde.Agent/Leases/Handlers/ComputeHandlerV2.cs
--- a/Engine/Source/Programs/Horde/Horde.Agent/Leases/Handlers/ComputeHandlerV2.cs
+++ b/Engine/Source/Programs/Horde/Horde.Agent/Leases/Handlers/ComputeHandlerV2.cs
@@ -45,7 +45,10 @@ namespace Horde.Agent.Leases.Handlers
 			using TcpClient tcpClient = new TcpClient();
 			await tcpClient.ConnectAsync(computeTask.RemoteIp, computeTask.RemotePort);
 
-			MessageChannel channel = new MessageChannel(tcpClient.Client);
+			Socket socket = tcpClient.Client;
+			await socket.SendAsync(computeTask.Nonce.Memory, SocketFlags.None, cancellationToken);
+
+			ComputeChannel channel = new ComputeChannel(socket, computeTask.AesKey.Memory, computeTask.AesIv.Memory);
 			for (; ; )
 			{
 				MessageBase request = await channel.ReadAsync(cancellationToken);
diff --git a/Engine/Source/Programs/Horde/Horde.Build/Compute/V2/ComputeControllerV2.cs b/Engine/Source/Programs/Horde/Horde.Build/Compute/V2/ComputeControllerV2.cs
--- a/Engine/Source/Programs/Horde/Horde.Build/Compute/V2/ComputeControllerV2.cs
+++ b/Engine/Source/Programs/Horde/Horde.Build/Compute/V2/ComputeControllerV2.cs
@@ -2,8 +2,10 @@
 
 using System;
 using System.Collections.Generic;
+using System.Net;
 using System.Threading;
 using System.Threading.Tasks;
+using EpicGames.Core;
 using EpicGames.Horde.Common;
 using EpicGames.Horde.Compute;
 using EpicGames.Horde.Compute.Impl;
@@ -27,14 +29,24 @@ namespace Horde.Build.Compute.V2
 		public Requirements? Requirements { get; set; }
 
 		/// <summary>
-		/// IP address to contact the initiator on.
+		/// Port to connect on
 		/// </summary>
-		public string RemoteIp { get; set; } = String.Empty;
+		public int RemotePort { get; set; } = 4000;
 
 		/// <summary>
-		/// Port to connect on
+		/// Base-64 encoded cryptographic nonce to identify the request
 		/// </summary>
-		public int RemotePort { get; set; } = 4000;
+		public string Nonce { get; set; } = String.Empty;
+
+		/// <summary>
+		/// Base-64 encoded AES key for the channel
+		/// </summary>
+		public string AesKey { get; set; } = String.Empty;
+
+		/// <summary>
+		/// Base-64 encoded AES IV for the channel
+		/// </summary>
+		public string AesIv { get; set; } = String.Empty;
 	}
 
 	/// <summary>
@@ -73,7 +85,17 @@ namespace Horde.Build.Compute.V2
 			}
 
 			Requirements requirements = request.Requirements ?? new Requirements();
-			_computeService.AddRequest(requirements, request.RemoteIp, request.RemotePort);
+			byte[] nonce = StringUtils.ParseHexString(request.Nonce);
+			byte[] aesKey = StringUtils.ParseHexString(request.AesKey);
+			byte[] aesIv = StringUtils.ParseHexString(request.AesIv);
+
+			IPAddress? remoteIp = HttpContext.Connection.RemoteIpAddress;
+			if (remoteIp == null)
+			{
+				return BadRequest("Missing remote IP address");
+			}
+
+			_computeService.AddRequest(requirements, remoteIp.ToString(), request.RemotePort, nonce, aesKey, aesIv);
 			return Ok();
 		}
 	}
diff --git a/Engine/Source/Programs/Horde/Horde.Build/Compute/V2/ComputeServiceV2.cs b/Engine/Source/Programs/Horde/Horde.Build/Compute/V2/ComputeServiceV2.cs
--- a/Engine/Source/Programs/Horde/Horde.Build/Compute/V2/ComputeServiceV2.cs
+++ b/Engine/Source/Programs/Horde/Horde.Build/Compute/V2/ComputeServiceV2.cs
@@ -16,6 +16,7 @@ using Horde.Build.Tasks;
 using Horde.Build.Utilities;
 using HordeCommon;
 using HordeCommon.Rpc.Tasks;
+using Microsoft.AspNetCore.Http;
 using Microsoft.Extensions.Logging;
 
 namespace Horde.Build.Compute.V2
@@ -31,11 +32,17 @@ namespace Horde.Build.Compute.V2
 		{
 			public string Ip { get; }
 			public int Port { get; }
+			public ReadOnlyMemory<byte> Nonce { get; }
+			public ReadOnlyMemory<byte> AesKey { get; }
+			public ReadOnlyMemory<byte> AesIv { get; }
 
-			public Request(string ip, int port)
+			public Request(string ip, int port, ReadOnlyMemory<byte> nonce, ReadOnlyMemory<byte> aesKey, ReadOnlyMemory<byte> aesIv)
 			{
 				Ip = ip;
 				Port = port;
+				Nonce = nonce;
+				AesKey = aesKey;
+				AesIv = aesIv;
 			}
 		}
 
@@ -72,12 +79,15 @@ namespace Horde.Build.Compute.V2
 		}
 
 		/// <summary>
-		/// 
+		/// Adds a new compute request
 		/// </summary>
-		/// <param name="requirements"></param>
-		/// <param name="ip"></param>
-		/// <param name="port"></param>
-		public void AddRequest(Requirements requirements, string ip, int port)
+		/// <param name="requirements">Requirements for the agent to serve the request</param>
+		/// <param name="ip">IP address to connect to</param>
+		/// <param name="port">Remote port to connect to</param>
+		/// <param name="nonce">Cryptographic nonce used to identify the request</param>
+		/// <param name="aesKey">Key for AES encryption on the compute channel</param>
+		/// <param name="aesIv">Initialization vector for AES encryption</param>
+		public void AddRequest(Requirements requirements, string ip, int port, ReadOnlyMemory<byte> nonce, ReadOnlyMemory<byte> aesKey, ReadOnlyMemory<byte> aesIv)
 		{
 			CbObject obj = CbSerializer.Serialize(requirements);
 			IoHash hash = IoHash.Compute(obj.GetView().Span);
@@ -89,7 +99,7 @@ namespace Horde.Build.Compute.V2
 					queue = new RequestQueue(hash, requirements);
 					_queues.Add(hash, queue);
 				}
-				queue.Requests.Enqueue(new Request(ip, port));
+				queue.Requests.Enqueue(new Request(ip, port, nonce, aesKey, aesIv));
 			}
 		}
 
@@ -124,6 +134,9 @@ namespace Horde.Build.Compute.V2
 					ComputeTaskMessageV2 computeTask = new ComputeTaskMessageV2();
 					computeTask.RemoteIp = request.Ip;
 					computeTask.RemotePort = request.Port;
+					computeTask.Nonce = ByteString.CopyFrom(request.Nonce.Span);
+					computeTask.AesKey = ByteString.CopyFrom(request.AesKey.Span);
+					computeTask.AesIv = ByteString.CopyFrom(request.AesIv.Span);
 
 					string leaseName = $"Compute lease";
 					byte[] payload = Any.Pack(computeTask).ToByteArray();
diff --git a/Engine/Source/Programs/Shared/EpicGames.Horde/Compute/ComputeChannel.cs b/Engine/Source/Programs/Shared/EpicGames.Horde/Compute/ComputeChannel.cs
new file mode 100644
--- /dev/null
+++ b/Engine/Source/Programs/Shared/EpicGames.Horde/Compute/ComputeChannel.cs
@@ -0,0 +1,256 @@
+// Copyright Epic Games, Inc. All Rights Reserved.
+
+using EpicGames.Core;
+using EpicGames.Serialization;
+using Microsoft.Extensions.Caching.Memory;
+using System;
+using System.Collections.Generic;
+using System.Diagnostics;
+using System.IO;
+using System.Net.Http.Headers;
+using System.Net.Sockets;
+using System.Reflection;
+using System.Security.Cryptography;
+using System.Threading;
+using System.Threading.Tasks;
+
+namespace EpicGames.Horde.Compute
+{
+	/// <summary>
+	/// Base class for messages
+	/// </summary>
+	public abstract class MessageBase
+	{
+		/// <summary>
+		/// Type of the message
+		/// </summary>
+		public Utf8String Type { get; }
+
+		/// <summary>
+		/// Constructor
+		/// </summary>
+		protected MessageBase()
+		{
+			Type = GetType().GetCustomAttribute<MessageAttribute>()!.Name;
+		}
+	}
+
+	/// <summary>
+	/// Attribute used to annotate a message type with an identifier string
+	/// </summary>
+	[AttributeUsage(AttributeTargets.Class)]
+	public sealed class MessageAttribute : Attribute
+	{
+		/// <summary>
+		/// Message type
+		/// </summary>
+		public string Name { get; }
+
+		/// <summary>
+		/// Constructor
+		/// </summary>
+		public MessageAttribute(string name)
+		{
+			Name = name;
+		}
+	}
+
+	/// <summary>
+	/// Reads messages from a socket
+	/// </summary>
+	public class ComputeChannel
+	{
+		static readonly Utf8String s_typeField = "name";
+		static readonly Utf8String s_dataField = "data";
+
+		static readonly Type[] s_messageTypes =
+		{
+			typeof(CloseMessage),
+			typeof(XorRequestMessage),
+			typeof(XorResponseMessage)
+		};
+
+		static readonly Dictionary<Utf8String, ICbConverter> s_typeToConverter = CreateTypeToConverterMap();
+
+		readonly Socket _socket;
+		readonly int _blockSize;
+		readonly ICryptoTransform _decryptor;
+		readonly ICryptoTransform _encryptor;
+
+		readonly CbWriter _writer;
+
+		byte[] _readBuffer = new byte[128];
+		int _decryptedLength;
+		int _encryptedLength;
+		int _paddedMessageLength;
+
+		byte[] _writeBuffer = new byte[128];
+
+		/// <summary>
+		/// Constructor
+		/// </summary>
+		/// <param name="socket">Socket for communication</param>
+		/// <param name="aesKey">AES encryption key</param>
+		/// <param name="aesIv">AES initialization vector</param>
+		public ComputeChannel(Socket socket, ReadOnlyMemory<byte> aesKey, ReadOnlyMemory<byte> aesIv)
+		{
+			using Aes aes = Aes.Create();
+			aes.Key = aesKey.ToArray();
+			aes.IV = aesIv.ToArray();
+			aes.Padding = PaddingMode.None;
+
+			_socket = socket;
+			_blockSize = aes.BlockSize / 8;
+			_decryptor = aes.CreateDecryptor();
+			_encryptor = aes.CreateEncryptor();
+
+			_writer = new CbWriter();
+		}
+
+		/// <summary>
+		/// Build a lookup of all known message types
+		/// </summary>
+		static Dictionary<Utf8String, ICbConverter> CreateTypeToConverterMap()
+		{
+			Dictionary<Utf8String, ICbConverter> typeToConverter = new Dictionary<Utf8String, ICbConverter>();
+			foreach (Type messageType in s_messageTypes)
+			{
+				string name = messageType.GetCustomAttribute<MessageAttribute>()!.Name;
+				ICbConverter converter = CbConverter.GetConverter(messageType);
+				typeToConverter.Add(name, converter);
+			}
+			return typeToConverter;
+		}
+
+		/// <summary>
+		/// Receives a message from the remote
+		/// </summary>
+		/// <param name="cancellationToken"></param>
+		/// <returns></returns>
+		public async Task<T> ReadAsync<T>(CancellationToken cancellationToken) where T : MessageBase
+		{
+			return (T)await ReadAsync(cancellationToken);
+		}
+
+		/// <summary>
+		/// Receives a message from the remote
+		/// </summary>
+		/// <param name="cancellationToken"></param>
+		/// <returns></returns>
+		public async Task<MessageBase> ReadAsync(CancellationToken cancellationToken)
+		{
+			// Discard any data from the previous message
+			if (_paddedMessageLength > 0)
+			{
+				Buffer.BlockCopy(_readBuffer, _paddedMessageLength, _readBuffer, 0, _encryptedLength - _paddedMessageLength);
+				_encryptedLength -= _paddedMessageLength;
+				_decryptedLength -= _paddedMessageLength;
+			}
+
+			// Read the next message data
+			for (; ; )
+			{
+				// Check the object in the buffer
+				if (_decryptedLength >= 2)
+				{
+					int messageLength = 1 + VarInt.Measure(_readBuffer[1]);
+					if (_decryptedLength > messageLength)
+					{
+						messageLength += (int)VarInt.ReadUnsigned(_readBuffer.AsSpan(1));
+					}
+
+					_paddedMessageLength = messageLength + GetPadding(messageLength);
+
+					if (_decryptedLength >= _paddedMessageLength)
+					{
+						CbObject obj = new CbObject(_readBuffer.AsMemory(0, messageLength));
+
+						CbField typeField = obj.Find(s_typeField);
+						Utf8String type = typeField.AsUtf8String();
+
+						ICbConverter converter = s_typeToConverter[type];
+
+						CbField dataField = obj.Find(s_dataField);
+						MessageBase message = (MessageBase)converter.ReadObject(dataField)!;
+
+						return message;
+					}
+
+					if (_readBuffer.Length < _paddedMessageLength)
+					{
+						Array.Resize(ref _readBuffer, _paddedMessageLength);
+					}
+				}
+
+				// Read the next chunk of data from the socket
+				int read = await _socket.ReceiveAsync(_readBuffer.AsMemory(_encryptedLength), SocketFlags.Partial, cancellationToken);
+				if (read == 0)
+				{
+					throw new EndOfStreamException();
+				}
+				_encryptedLength += read;
+
+				// Decrypt any new full blocks that have been received
+				int nextTransformSize = _encryptedLength - _decryptedLength;
+				if (nextTransformSize >= _blockSize)
+				{
+					nextTransformSize -= nextTransformSize % _blockSize;
+					_decryptedLength += _decryptor.TransformBlock(_readBuffer, _decryptedLength, nextTransformSize, _readBuffer, _decryptedLength);
+				}
+			}
+		}
+
+		/// <summary>
+		/// Sends a message to the remote
+		/// </summary>
+		/// <param name="message">Message to be sent</param>
+		/// <param name="cancellationToken">Cancellation token</param>
+		/// <returns></returns>
+		public async Task WriteAsync(MessageBase message, CancellationToken cancellationToken)
+		{
+			// Format the output object
+			Utf8String name = message.Type;
+		
+			_writer.Clear();
+			_writer.BeginObject();
+			_writer.WriteUtf8String(s_typeField, name);
+
+			ICbConverter converter = CbConverter.GetConverter(message.GetType());
+			converter.WriteNamedObject(_writer, s_dataField, message);
+
+			_writer.EndObject();
+
+			// Serialize it to the write buffer
+			int length = _writer.GetSize();
+			length += GetPadding(length);
+
+			if (length > _writeBuffer.Length)
+			{
+				_writeBuffer = new byte[length];
+			}
+
+			_writer.CopyTo(_writeBuffer);
+
+			// Encrypt the data and send it
+			int transformed = _encryptor.TransformBlock(_writeBuffer, 0, length, _writeBuffer, 0);
+			if (transformed != length)
+			{
+				throw new InvalidOperationException();
+			}
+			await _socket.SendAsync(_writeBuffer.AsMemory(0, length), SocketFlags.Partial, cancellationToken);
+		}
+
+		int GetPadding(int length)
+		{
+			int modulo = length % _blockSize;
+			if (modulo == 0)
+			{
+				return 0;
+			}
+			else
+			{
+				return _blockSize - modulo;
+			}
+		}
+	}
+}
diff --git a/Engine/Source/Programs/Shared/EpicGames.Horde/Compute/ComputeClient.cs b/Engine/Source/Programs/Shared/EpicGames.Horde/Compute/ComputeClient.cs
new file mode 100644
--- /dev/null
+++ b/Engine/Source/Programs/Shared/EpicGames.Horde/Compute/ComputeClient.cs
@@ -0,0 +1,294 @@
+// Copyright Epic Games, Inc. All Rights Reserved.
+
+using EpicGames.Core;
+using Microsoft.Extensions.Logging;
+using System;
+using System.Collections.Generic;
+using System.Net;
+using System.Net.Http;
+using System.Net.Sockets;
+using System.Security.Cryptography;
+using System.Threading;
+using System.Threading.Tasks;
+
+namespace EpicGames.Horde.Compute
+{
+	/// <summary>
+	/// Interface for a compute request
+	/// </summary>
+	public interface IComputeRequest<TResult>
+	{
+		/// <summary>
+		/// Task to await completion of this request
+		/// </summary>
+		Task<TResult> Result { get; }
+
+		/// <summary>
+		/// Cancel the current request
+		/// </summary>
+		void Cancel();
+	}
+
+	/// <summary>
+	/// Helper class to enlist remote resources to perform compute-intensive tasks.
+	/// </summary>
+	public sealed class ComputeClient : IAsyncDisposable, IDisposable
+	{
+		const int NonceLength = 64;
+
+		class AddComputeTaskRequest
+		{
+			public Requirements? Requirements { get; set; }
+			public int RemotePort { get; set; }
+			public string Nonce { get; set; } = String.Empty;
+			public string AesKey { get; set; } = String.Empty;
+			public string AesIv { get; set; } = String.Empty;
+		}
+
+		abstract class RequestInfo : IDisposable
+		{
+			public ReadOnlyMemory<byte> Nonce { get; }
+			public ReadOnlyMemory<byte> AesKey { get; }
+			public ReadOnlyMemory<byte> AesIv { get; }
+
+			protected readonly TaskCompletionSource<ComputeChannel> _computeChannelSource;
+			protected readonly CancellationTokenSource _cancellationSource;
+
+			public RequestInfo(CancellationToken cancellationToken)
+			{
+				Nonce = RandomNumberGenerator.GetBytes(NonceLength);
+
+				using (Aes aes = Aes.Create())
+				{
+					AesKey = aes.Key;
+					AesIv = aes.IV;
+				}
+
+				_computeChannelSource = new TaskCompletionSource<ComputeChannel>();
+				_cancellationSource = CancellationTokenSource.CreateLinkedTokenSource(cancellationToken);
+			}
+
+			public abstract Task RunAsync(ComputeChannel channel);
+
+			public void Cancel()
+			{
+				_computeChannelSource.TrySetCanceled();
+				_cancellationSource.Cancel();
+			}
+
+			public void Dispose()
+			{
+				_cancellationSource.Dispose();
+			}
+		}
+
+		class RequestInfo<TResult> : RequestInfo, IComputeRequest<TResult>
+		{
+			readonly Func<ComputeChannel, CancellationToken, Task<TResult>> _handler;
+
+			public Task<TResult> Result { get; }
+
+			public RequestInfo(Func<ComputeChannel, CancellationToken, Task<TResult>> handler, CancellationToken cancellationToken)
+				: base(cancellationToken)
+			{
+				_handler = handler;
+
+				Result = RunWrapperAsync();
+			}
+
+			public override Task RunAsync(ComputeChannel channel)
+			{
+				_computeChannelSource.SetResult(channel);
+				return Result;
+			}
+
+			async Task<TResult> RunWrapperAsync()
+			{
+				ComputeChannel channel = await _computeChannelSource.Task;
+				return await _handler(channel, _cancellationSource.Token);
+			}
+		}
+
+		class MemoryComparer : IEqualityComparer<ReadOnlyMemory<byte>>
+		{
+			public static MemoryComparer Instance { get; } = new MemoryComparer();
+
+			public bool Equals(ReadOnlyMemory<byte> x, ReadOnlyMemory<byte> y) => x.Span.SequenceEqual(y.Span);
+
+			public int GetHashCode(ReadOnlyMemory<byte> memory)
+			{
+				HashCode hasher = new HashCode();
+				hasher.AddBytes(memory.Span);
+				return hasher.ToHashCode();
+			}
+		}
+
+		readonly object _lockObject = new object();
+		readonly Func<HttpClient> _createHttpClient;
+		readonly CancellationTokenSource _cancellationSource = new CancellationTokenSource();
+		readonly Dictionary<ReadOnlyMemory<byte>, RequestInfo> _requests = new Dictionary<ReadOnlyMemory<byte>, RequestInfo>(MemoryComparer.Instance);
+		TcpListener? _listener;
+		Task _listenTask = Task.CompletedTask;
+		readonly ILogger _logger;
+
+		/// <summary>
+		/// Constructor
+		/// </summary>
+		/// <param name="createHttpClient">Creates an HTTP client with the correct base address for the server</param>
+		/// <param name="logger">Logger for diagnostic messages</param>
+		public ComputeClient(Func<HttpClient> createHttpClient, ILogger logger)
+		{
+			_createHttpClient = createHttpClient;
+			_logger = logger;
+		}
+
+		/// <inheritdoc/>
+		public async ValueTask DisposeAsync()
+		{
+			await StopAsync();
+			Dispose();
+		}
+
+		/// <inheritdoc/>
+		public void Dispose()
+		{
+			_cancellationSource.Dispose();
+		}
+
+		/// <summary>
+		/// Starts the client
+		/// </summary>
+		public void Start()
+		{
+			_listener = new TcpListener(IPAddress.Any, 0);
+			_listener.Start();
+			_listenTask = Task.Run(() => ListenAsync(_listener, _cancellationSource.Token));
+		}
+
+		/// <summary>
+		/// Cancels all running tasks and waits for them to complete
+		/// </summary>
+		public async Task StopAsync()
+		{
+			_cancellationSource.Cancel();
+			try
+			{
+				await _listenTask;
+			}
+			catch (OperationCanceledException)
+			{
+			}
+		}
+
+		/// <summary>
+		/// Adds a new remote request
+		/// </summary>
+		/// <param name="requirements">Requirements for the agent</param>
+		/// <param name="handler">Handler for the connection</param>
+		public async Task<IComputeRequest<TResult>> AddRequestAsync<TResult>(Requirements? requirements, Func<ComputeChannel, CancellationToken, Task<TResult>> handler)
+		{
+			RequestInfo<TResult> requestInfo = new RequestInfo<TResult>(handler, _cancellationSource.Token);
+			await AddRequestAsync(requirements, requestInfo);
+			return requestInfo;
+		}
+
+		async Task AddRequestAsync(Requirements? requirements, RequestInfo requestInfo)
+		{
+			try
+			{
+				HttpClient client = _createHttpClient();
+
+				AddComputeTaskRequest request = new AddComputeTaskRequest();
+				request.Requirements = requirements;
+				request.RemotePort = ((IPEndPoint)_listener!.LocalEndpoint).Port;
+				request.Nonce = StringUtils.FormatHexString(requestInfo.Nonce.Span);
+				request.AesKey = StringUtils.FormatHexString(requestInfo.AesKey.Span);
+				request.AesIv = StringUtils.FormatHexString(requestInfo.AesIv.Span);
+
+				using (HttpResponseMessage response = await client.PostAsync("api/v2/compute", request, _cancellationSource.Token))
+				{
+					response.EnsureSuccessStatusCode();
+					_requests.Add(requestInfo.Nonce, requestInfo);
+				}
+			}
+			catch
+			{
+				requestInfo.Dispose();
+			}
+		}
+
+		async Task ListenAsync(TcpListener listener, CancellationToken cancellationToken)
+		{
+			List<Task> tasks = new List<Task>();
+			try
+			{
+				for (; ; )
+				{
+					TcpClient client = await listener.AcceptTcpClientAsync(cancellationToken);
+					_logger.LogInformation("Received connection from {Remote}", client.Client.RemoteEndPoint);
+
+					Task task = HandleConnectionGuardedAsync(client, cancellationToken);
+					tasks.Add(task);
+				}
+			}
+			finally
+			{
+				await Task.WhenAll(tasks);
+				listener.Stop();
+			}
+		}
+
+		async Task HandleConnectionGuardedAsync(TcpClient client, CancellationToken cancellationToken)
+		{
+			try
+			{
+				await HandleConnectionAsync(client, cancellationToken);
+			}
+			catch (Exception ex)
+			{
+				_logger.LogInformation(ex, "Exception while handling request: {Message}", ex.Message);
+			}
+		}
+
+		async Task HandleConnectionAsync(TcpClient client, CancellationToken cancellationToken)
+		{
+			using TcpClient disposeClient = client;
+			Socket socket = client.Client;
+
+			// Read the nonce for this connection
+			byte[] nonceBuffer = new byte[NonceLength];
+			for (int nonceLength = 0; nonceLength < nonceBuffer.Length;)
+			{
+				int read = await socket.ReceiveAsync(nonceBuffer.AsMemory(nonceLength), SocketFlags.Partial, cancellationToken);
+				if (read == 0)
+				{
+					return;
+				}
+				nonceLength += read;
+			}
+
+			// Dispatch the connection to the appropriate handler
+			RequestInfo? requestInfo = null;
+			try
+			{
+				// Find the matching request
+				lock (_lockObject)
+				{
+					if (!_requests.Remove(nonceBuffer, out requestInfo))
+					{
+						_logger.LogInformation("Invalid nonce for connection; ignoring.");
+						return;
+					}
+				}
+
+				// Create a crypto stream for communication, using the predetermined key for this request
+				ComputeChannel channel = new ComputeChannel(socket, requestInfo.AesKey, requestInfo.AesIv);
+				await requestInfo.RunAsync(channel);
+			}
+			finally
+			{
+				requestInfo?.Dispose();
+			}
+		}
+	}
+}
diff --git a/Engine/Source/Programs/Shared/EpicGames.Horde/Compute/MessageChannel.cs b/Engine/Source/Programs/Shared/EpicGames.Horde/Compute/MessageChannel.cs
deleted file mode 100644
