Horde: Abstract compute channel implementation behind an interface, allowing different implementations in the future (eg. websockets, NNG).

diff --git a/Engine/Source/Programs/Horde/Horde.Agent/Commands/Compute/ComputeCommandV2.cs b/Engine/Source/Programs/Horde/Horde.Agent/Commands/Compute/ComputeCommandV2.cs
--- a/Engine/Source/Programs/Horde/Horde.Agent/Commands/Compute/ComputeCommandV2.cs
+++ b/Engine/Source/Programs/Horde/Horde.Agent/Commands/Compute/ComputeCommandV2.cs
@@ -61,7 +61,7 @@ namespace Horde.Agent.Commands
 			return client;
 		}
 
-		async Task<object?> TestCommandAsync(ComputeChannel channel, CancellationToken cancellationToken)
+		async Task<object?> TestCommandAsync(IComputeChannel channel, CancellationToken cancellationToken)
 		{
 			await channel.WriteAsync(new XorRequestMessage { Value = 123, Payload = new byte[] { 1, 2, 3, 4, 5 } }, cancellationToken);
 
diff --git a/Engine/Source/Programs/Horde/Horde.Agent/Leases/Handlers/ComputeHandlerV2.cs b/Engine/Source/Programs/Horde/Horde.Agent/Leases/Handlers/ComputeHandlerV2.cs
--- a/Engine/Source/Programs/Horde/Horde.Agent/Leases/Handlers/ComputeHandlerV2.cs
+++ b/Engine/Source/Programs/Horde/Horde.Agent/Leases/Handlers/ComputeHandlerV2.cs
@@ -48,11 +48,11 @@ namespace Horde.Agent.Leases.Handlers
 			Socket socket = tcpClient.Client;
 			await socket.SendAsync(computeTask.Nonce.Memory, SocketFlags.None, cancellationToken);
 
-			ComputeChannel channel = new ComputeChannel(socket, computeTask.AesKey.Memory, computeTask.AesIv.Memory);
+			IComputeChannel channel = new SocketComputeChannel(socket, computeTask.AesKey.Memory, computeTask.AesIv.Memory);
 			for (; ; )
 			{
-				MessageBase request = await channel.ReadAsync(cancellationToken);
-				_logger.LogInformation("Received message {Type}", request.Type);
+				object request = await channel.ReadAsync(cancellationToken);
+				_logger.LogInformation("Received message {Type}", request.GetType().Name);
 
 				switch (request)
 				{
diff --git a/Engine/Source/Programs/Shared/EpicGames.Horde/Compute/ComputeChannel.cs b/Engine/Source/Programs/Shared/EpicGames.Horde/Compute/ComputeChannel.cs
deleted file mode 100644
diff --git a/Engine/Source/Programs/Shared/EpicGames.Horde/Compute/ComputeClient.cs b/Engine/Source/Programs/Shared/EpicGames.Horde/Compute/ComputeClient.cs
--- a/Engine/Source/Programs/Shared/EpicGames.Horde/Compute/ComputeClient.cs
+++ b/Engine/Source/Programs/Shared/EpicGames.Horde/Compute/ComputeClient.cs
@@ -58,7 +58,7 @@ namespace EpicGames.Horde.Compute
 			public ReadOnlyMemory<byte> AesKey { get; }
 			public ReadOnlyMemory<byte> AesIv { get; }
 
-			protected readonly TaskCompletionSource<ComputeChannel> _computeChannelSource;
+			protected readonly TaskCompletionSource<IComputeChannel> _computeChannelSource;
 			protected readonly CancellationTokenSource _cancellationSource;
 
 			public RequestInfo(CancellationToken cancellationToken)
@@ -71,11 +71,11 @@ namespace EpicGames.Horde.Compute
 					AesIv = aes.IV;
 				}
 
-				_computeChannelSource = new TaskCompletionSource<ComputeChannel>();
+				_computeChannelSource = new TaskCompletionSource<IComputeChannel>();
 				_cancellationSource = CancellationTokenSource.CreateLinkedTokenSource(cancellationToken);
 			}
 
-			public abstract Task RunAsync(ComputeChannel channel);
+			public abstract Task RunAsync(IComputeChannel channel);
 
 			public void Cancel()
 			{
@@ -91,11 +91,11 @@ namespace EpicGames.Horde.Compute
 
 		class RequestInfo<TResult> : RequestInfo, IComputeRequest<TResult>
 		{
-			readonly Func<ComputeChannel, CancellationToken, Task<TResult>> _handler;
+			readonly Func<IComputeChannel, CancellationToken, Task<TResult>> _handler;
 
 			public Task<TResult> Result { get; }
 
-			public RequestInfo(Func<ComputeChannel, CancellationToken, Task<TResult>> handler, CancellationToken cancellationToken)
+			public RequestInfo(Func<IComputeChannel, CancellationToken, Task<TResult>> handler, CancellationToken cancellationToken)
 				: base(cancellationToken)
 			{
 				_handler = handler;
@@ -103,7 +103,7 @@ namespace EpicGames.Horde.Compute
 				Result = RunWrapperAsync();
 			}
 
-			public override Task RunAsync(ComputeChannel channel)
+			public override Task RunAsync(IComputeChannel channel)
 			{
 				_computeChannelSource.SetResult(channel);
 				return Result;
@@ -111,7 +111,7 @@ namespace EpicGames.Horde.Compute
 
 			async Task<TResult> RunWrapperAsync()
 			{
-				ComputeChannel channel = await _computeChannelSource.Task;
+				IComputeChannel channel = await _computeChannelSource.Task;
 				return await _handler(channel, _cancellationSource.Token);
 			}
 		}
@@ -193,7 +193,7 @@ namespace EpicGames.Horde.Compute
 		/// <param name="clusterId">Cluster to execute the request</param>
 		/// <param name="requirements">Requirements for the agent</param>
 		/// <param name="handler">Handler for the connection</param>
-		public async Task<IComputeRequest<TResult>> AddRequestAsync<TResult>(ClusterId clusterId, Requirements? requirements, Func<ComputeChannel, CancellationToken, Task<TResult>> handler)
+		public async Task<IComputeRequest<TResult>> AddRequestAsync<TResult>(ClusterId clusterId, Requirements? requirements, Func<IComputeChannel, CancellationToken, Task<TResult>> handler)
 		{
 			RequestInfo<TResult> requestInfo = new RequestInfo<TResult>(handler, _cancellationSource.Token);
 			await AddRequestAsync(clusterId, requirements, requestInfo);
@@ -293,7 +293,7 @@ namespace EpicGames.Horde.Compute
 				}
 
 				// Create a crypto stream for communication, using the predetermined key for this request
-				ComputeChannel channel = new ComputeChannel(socket, requestInfo.AesKey, requestInfo.AesIv);
+				IComputeChannel channel = new SocketComputeChannel(socket, requestInfo.AesKey, requestInfo.AesIv);
 				await requestInfo.RunAsync(channel);
 			}
 			finally
diff --git a/Engine/Source/Programs/Shared/EpicGames.Horde/Compute/ComputeMessage.cs b/Engine/Source/Programs/Shared/EpicGames.Horde/Compute/ComputeMessage.cs
new file mode 100644
--- /dev/null
+++ b/Engine/Source/Programs/Shared/EpicGames.Horde/Compute/ComputeMessage.cs
@@ -0,0 +1,125 @@
+// Copyright Epic Games, Inc. All Rights Reserved.
+
+using System;
+using System.Collections.Generic;
+using System.Reflection;
+using EpicGames.Core;
+using EpicGames.Serialization;
+
+namespace EpicGames.Horde.Compute
+{
+	/// <summary>
+	/// Attribute used to annotate a message type with an identifier string
+	/// </summary>
+	[AttributeUsage(AttributeTargets.Class)]
+	public sealed class ComputeMessageAttribute : Attribute
+	{
+		/// <summary>
+		/// Message type
+		/// </summary>
+		public string Name { get; }
+
+		/// <summary>
+		/// Constructor
+		/// </summary>
+		public ComputeMessageAttribute(string name)
+		{
+			Name = name;
+		}
+	}
+
+	/// <summary>
+	/// Utility methods for compute messages
+	/// </summary>
+	public static class ComputeMessage
+	{
+		static readonly Utf8String s_typeField = "name";
+		static readonly Utf8String s_dataField = "data";
+
+		static readonly Type[] s_messageTypes =
+		{
+			typeof(CloseMessage),
+			typeof(XorRequestMessage),
+			typeof(XorResponseMessage)
+		};
+
+		static readonly Dictionary<Utf8String, ICbConverter> s_typeToConverter = CreateTypeToConverterMap();
+
+		/// <summary>
+		/// Build a lookup of all known message types
+		/// </summary>
+		static Dictionary<Utf8String, ICbConverter> CreateTypeToConverterMap()
+		{
+			Dictionary<Utf8String, ICbConverter> typeToConverter = new Dictionary<Utf8String, ICbConverter>();
+			foreach (Type messageType in s_messageTypes)
+			{
+				string name = messageType.GetCustomAttribute<ComputeMessageAttribute>()!.Name;
+				ICbConverter converter = CbConverter.GetConverter(messageType);
+				typeToConverter.Add(name, converter);
+			}
+			return typeToConverter;
+		}
+
+		static class TagCache<T>
+		{
+			public static string Tag { get; } = GetTag(typeof(T));
+		}
+
+		/// <summary>
+		/// Gets the tag for a particular message class
+		/// </summary>
+		/// <typeparam name="T">Type of the message</typeparam>
+		/// <returns>Tag name for the message</returns>
+		public static string GetTag<T>()
+		{
+			return TagCache<T>.Tag;
+		}
+
+		/// <summary>
+		/// Gets the tag for a particular message class
+		/// </summary>
+		/// <param name="type">Type of the message</param>
+		/// <returns>Tag for the type</returns>
+		public static string GetTag(Type type)
+		{
+			return type.GetCustomAttribute<ComputeMessageAttribute>()!.Name;
+		}
+
+		/// <summary>
+		/// Serializes a message
+		/// </summary>
+		/// <param name="message">Message to serialize</param>
+		/// <param name="writer">Writer for the message</param>
+		public static void Serialize(object message, CbWriter writer)
+		{
+			Utf8String name = GetTag(message.GetType());
+
+			writer.Clear();
+			writer.BeginObject();
+			writer.WriteUtf8String(s_typeField, name);
+
+			ICbConverter converter = CbConverter.GetConverter(message.GetType());
+			converter.WriteNamedObject(writer, s_dataField, message);
+
+			writer.EndObject();
+		}
+
+		/// <summary>
+		/// Deserialize a message from a block of memory
+		/// </summary>
+		/// <param name="memory"></param>
+		/// <returns></returns>
+		public static object Deserialize(ReadOnlyMemory<byte> memory)
+		{
+			CbObject obj = new CbObject(memory);
+
+			CbField typeField = obj.Find(s_typeField);
+			Utf8String type = typeField.AsUtf8String();
+
+			ICbConverter converter = s_typeToConverter[type];
+
+			CbField dataField = obj.Find(s_dataField);
+			return converter.ReadObject(dataField)!;
+		}
+	}
+}
diff --git a/Engine/Source/Programs/Shared/EpicGames.Horde/Compute/ComputeMessageTypes.cs b/Engine/Source/Programs/Shared/EpicGames.Horde/Compute/ComputeMessageTypes.cs
new file mode 100644
--- /dev/null
+++ b/Engine/Source/Programs/Shared/EpicGames.Horde/Compute/ComputeMessageTypes.cs
@@ -0,0 +1,41 @@
+// Copyright Epic Games, Inc. All Rights Reserved.
+
+using EpicGames.Serialization;
+using System;
+
+#pragma warning disable CS1591
+#pragma warning disable CA1819
+
+namespace EpicGames.Horde.Compute
+{
+	/// <summary>
+	/// Terminate the current connection
+	/// </summary>
+	[ComputeMessage("close"), CbObject]
+	public class CloseMessage
+	{
+	}
+
+	/// <summary>
+	/// XOR a block of data with a value
+	/// </summary>
+	[ComputeMessage("xor-req")]
+	public class XorRequestMessage
+	{
+		[CbField]
+		public byte Value { get; set; }
+
+		[CbField]
+		public byte[] Payload { get; set; } = Array.Empty<byte>();
+	}
+
+	/// <summary>
+	/// Response from an XOR message
+	/// </summary>
+	[ComputeMessage("xor-rsp")]
+	public class XorResponseMessage
+	{
+		[CbField]
+		public byte[] Payload { get; set; } = Array.Empty<byte>();
+	}
+}
diff --git a/Engine/Source/Programs/Shared/EpicGames.Horde/Compute/IComputeChannel.cs b/Engine/Source/Programs/Shared/EpicGames.Horde/Compute/IComputeChannel.cs
new file mode 100644
--- /dev/null
+++ b/Engine/Source/Programs/Shared/EpicGames.Horde/Compute/IComputeChannel.cs
@@ -0,0 +1,45 @@
+// Copyright Epic Games, Inc. All Rights Reserved.
+
+using System.Threading;
+using System.Threading.Tasks;
+
+namespace EpicGames.Horde.Compute
+{
+	/// <summary>
+	/// Full-duplex channel for sending and reciving messages
+	/// </summary>
+	public interface IComputeChannel
+	{
+		/// <summary>
+		/// Receives a message from the remote
+		/// </summary>
+		/// <param name="cancellationToken"></param>
+		/// <returns></returns>
+		Task<object> ReadAsync(CancellationToken cancellationToken);
+
+		/// <summary>
+		/// Sends a message to the remote
+		/// </summary>
+		/// <param name="message">Message to be sent</param>
+		/// <param name="cancellationToken">Cancellation token</param>
+		/// <returns></returns>
+		Task WriteAsync(object message, CancellationToken cancellationToken);
+	}
+
+	/// <summary>
+	/// Extension methods for compute channels
+	/// </summary>
+	public static class ComputeChannelExtensions
+	{
+		/// <summary>
+		/// Receives a message from the remote
+		/// </summary>
+		/// <param name="channel">Channel to read from</param>
+		/// <param name="cancellationToken"></param>
+		/// <returns></returns>
+		public static async Task<T> ReadAsync<T>(this IComputeChannel channel, CancellationToken cancellationToken)
+		{
+			return (T)await channel.ReadAsync(cancellationToken);
+		}
+	}
+}
diff --git a/Engine/Source/Programs/Shared/EpicGames.Horde/Compute/SocketComputeChannel.cs b/Engine/Source/Programs/Shared/EpicGames.Horde/Compute/SocketComputeChannel.cs
new file mode 100644
--- /dev/null
+++ b/Engine/Source/Programs/Shared/EpicGames.Horde/Compute/SocketComputeChannel.cs
@@ -0,0 +1,158 @@
+// Copyright Epic Games, Inc. All Rights Reserved.
+
+using EpicGames.Core;
+using EpicGames.Serialization;
+using System;
+using System.IO;
+using System.Net.Sockets;
+using System.Security.Cryptography;
+using System.Threading;
+using System.Threading.Tasks;
+
+#pragma warning disable CA5401
+
+namespace EpicGames.Horde.Compute
+{
+	/// <summary>
+	/// Transfers messages via a AES-encrypted socket
+	/// </summary>
+	public class SocketComputeChannel : IComputeChannel
+	{
+		readonly Socket _socket;
+		readonly int _blockSize;
+		readonly ICryptoTransform _decryptor;
+		readonly ICryptoTransform _encryptor;
+
+		readonly CbWriter _writer;
+
+		byte[] _readBuffer = new byte[128];
+		int _decryptedLength;
+		int _encryptedLength;
+		int _paddedMessageLength;
+
+		byte[] _writeBuffer = new byte[128];
+
+		/// <summary>
+		/// Constructor
+		/// </summary>
+		/// <param name="socket">Socket for communication</param>
+		/// <param name="aesKey">AES encryption key</param>
+		/// <param name="aesIv">AES initialization vector</param>
+		public SocketComputeChannel(Socket socket, ReadOnlyMemory<byte> aesKey, ReadOnlyMemory<byte> aesIv)
+		{
+			using Aes aes = Aes.Create();
+			aes.Key = aesKey.ToArray();
+			aes.IV = aesIv.ToArray();
+			aes.Padding = PaddingMode.None;
+
+			_socket = socket;
+			_blockSize = aes.BlockSize / 8;
+			_decryptor = aes.CreateDecryptor();
+			_encryptor = aes.CreateEncryptor();
+
+			_writer = new CbWriter();
+		}
+
+		/// <summary>
+		/// Receives a message from the remote
+		/// </summary>
+		/// <param name="cancellationToken"></param>
+		/// <returns></returns>
+		public async Task<object> ReadAsync(CancellationToken cancellationToken)
+		{
+			// Discard any data from the previous message
+			if (_paddedMessageLength > 0)
+			{
+				Buffer.BlockCopy(_readBuffer, _paddedMessageLength, _readBuffer, 0, _encryptedLength - _paddedMessageLength);
+				_encryptedLength -= _paddedMessageLength;
+				_decryptedLength -= _paddedMessageLength;
+			}
+
+			// Read the next message data
+			for (; ; )
+			{
+				// Check the object in the buffer
+				if (_decryptedLength >= 2)
+				{
+					int messageLength = 1 + VarInt.Measure(_readBuffer[1]);
+					if (_decryptedLength > messageLength)
+					{
+						messageLength += (int)VarInt.ReadUnsigned(_readBuffer.AsSpan(1));
+					}
+
+					_paddedMessageLength = messageLength + GetPadding(messageLength);
+
+					if (_decryptedLength >= _paddedMessageLength)
+					{
+						return ComputeMessage.Deserialize(_readBuffer.AsMemory(0, messageLength));
+					}
+
+					if (_readBuffer.Length < _paddedMessageLength)
+					{
+						Array.Resize(ref _readBuffer, _paddedMessageLength);
+					}
+				}
+
+				// Read the next chunk of data from the socket
+				int read = await _socket.ReceiveAsync(_readBuffer.AsMemory(_encryptedLength), SocketFlags.Partial, cancellationToken);
+				if (read == 0)
+				{
+					throw new EndOfStreamException();
+				}
+				_encryptedLength += read;
+
+				// Decrypt any new full blocks that have been received
+				int nextTransformSize = _encryptedLength - _decryptedLength;
+				if (nextTransformSize >= _blockSize)
+				{
+					nextTransformSize -= nextTransformSize % _blockSize;
+					_decryptedLength += _decryptor.TransformBlock(_readBuffer, _decryptedLength, nextTransformSize, _readBuffer, _decryptedLength);
+				}
+			}
+		}
+
+		/// <summary>
+		/// Sends a message to the remote
+		/// </summary>
+		/// <param name="message">Message to be sent</param>
+		/// <param name="cancellationToken">Cancellation token</param>
+		/// <returns></returns>
+		public async Task WriteAsync(object message, CancellationToken cancellationToken)
+		{
+			// Format the output object
+			ComputeMessage.Serialize(message, _writer);
+
+			// Serialize it to the write buffer
+			int length = _writer.GetSize();
+			length += GetPadding(length);
+
+			if (length > _writeBuffer.Length)
+			{
+				_writeBuffer = new byte[length];
+			}
+
+			_writer.CopyTo(_writeBuffer);
+
+			// Encrypt the data and send it
+			int transformed = _encryptor.TransformBlock(_writeBuffer, 0, length, _writeBuffer, 0);
+			if (transformed != length)
+			{
+				throw new InvalidOperationException();
+			}
+			await _socket.SendAsync(_writeBuffer.AsMemory(0, length), SocketFlags.Partial, cancellationToken);
+		}
+
+		int GetPadding(int length)
+		{
+			int modulo = length % _blockSize;
+			if (modulo == 0)
+			{
+				return 0;
+			}
+			else
+			{
+				return _blockSize - modulo;
+			}
+		}
+	}
+}
