Fixed issue where a class might not be marked as containing instanced items.

diff --git a/Engine/Source/Programs/Shared/EpicGames.UHT/Types/Properties/UhtStructProperty.cs b/Engine/Source/Programs/Shared/EpicGames.UHT/Types/Properties/UhtStructProperty.cs
--- a/Engine/Source/Programs/Shared/EpicGames.UHT/Types/Properties/UhtStructProperty.cs
+++ b/Engine/Source/Programs/Shared/EpicGames.UHT/Types/Properties/UhtStructProperty.cs
@@ -58,6 +58,7 @@ namespace EpicGames.UHT.Types
 			switch (phase)
 			{
 				case UhtResolvePhase.Final:
+					ScriptStruct.Resolve(phase);
 					if (ScanForInstancedReferenced(true))
 					{
 						PropertyFlags |= EPropertyFlags.ContainsInstancedReference;
