* Updated AssemblyUtils to add an event callback to AssemblyResolve when calling AddFileToAssemblyCache. You had to before explicitly call InstallRecursiveAssemblyResolver to get the assembly cache resolver to work.
* Updated RulesAssembly to add the parent assembles to the assembly cache so the types in the other assemblies can be used when calling the constructor.

diff --git a/Engine/Source/Programs/ShaderCompileWorker/ShaderCompileWorker.Build.cs b/Engine/Source/Programs/ShaderCompileWorker/ShaderCompileWorker.Build.cs
--- a/Engine/Source/Programs/ShaderCompileWorker/ShaderCompileWorker.Build.cs
+++ b/Engine/Source/Programs/ShaderCompileWorker/ShaderCompileWorker.Build.cs
@@ -29,11 +29,9 @@ public class ShaderCompileWorker : ModuleRules
 		PrivateIncludePaths.Add("Runtime/Launch/Private");      // For LaunchEngineLoop.cpp include
 
 		// Include D3D compiler binaries
-		string EngineDir = Path.GetFullPath(Target.RelativeEnginePath);
-
 		if (Target.Platform == UnrealTargetPlatform.Win64)
 		{
-			RuntimeDependencies.Add(EngineDir + "Binaries/ThirdParty/Windows/DirectX/x64/d3dcompiler_47.dll");
+			RuntimeDependencies.Add(DirectX.GetDllDir(Target) + "d3dcompiler_47.dll");
 		}
 	}
 }
diff --git a/Engine/Source/Programs/Shared/EpicGames.Core/AssemblyUtils.cs b/Engine/Source/Programs/Shared/EpicGames.Core/AssemblyUtils.cs
--- a/Engine/Source/Programs/Shared/EpicGames.Core/AssemblyUtils.cs
+++ b/Engine/Source/Programs/Shared/EpicGames.Core/AssemblyUtils.cs
@@ -64,24 +64,10 @@ namespace EpicGames.Core
 			{
 				// Name is fully qualified assembly definition - e.g. "p4dn, Version=1.0.0.0, Culture=neutral, PublicKeyToken=ff968dc1933aba6f"
 				string assemblyName = args.Name!.Split(',')[0];
-				if (s_assemblyLocationCache.TryGetValue(assemblyName, out string? assemblyLocation))
-				{
-					// We have this assembly in our folder.					
-					if (File.Exists(assemblyLocation))
-					{
-						// The assembly still exists, so load it.
-						return Assembly.LoadFile(assemblyLocation);
-					}
-					else
-					{
-						// The assembly no longer exists on disk, so remove it from our cache.
-						s_assemblyLocationCache.Remove(assemblyName);
-					}
-				}
 
-				// The assembly wasn't found, though may have been compiled or copied as a dependency
+				// The assembly wasn't found by other resolvers, though may have been compiled or copied as a dependency
 				RefreshAssemblyCache(rootDirectory, String.Format("{0}.dll", assemblyName));
-				if (s_assemblyLocationCache.TryGetValue(assemblyName, out assemblyLocation))
+				if (s_assemblyLocationCache.TryGetValue(assemblyName, out string? assemblyLocation))
 				{
 					return Assembly.LoadFile(assemblyLocation);
 				}
@@ -129,12 +115,41 @@ namespace EpicGames.Core
 				s_assemblyLocationCache.Add(assemblyName, assemblyPath);
 				s_assemblyWriteTimes.Add(assemblyName, assemblyLastWriteTime);
 			}
+
+			if (!s_addedToAssemblyResolver)
+			{
+				AppDomain.CurrentDomain.AssemblyResolve += (sender, args) =>
+				{
+					// Name is fully qualified assembly definition - e.g. "p4dn, Version=1.0.0.0, Culture=neutral, PublicKeyToken=ff968dc1933aba6f"
+					string assemblyName = args.Name!.Split(',')[0];
+					if (s_assemblyLocationCache.TryGetValue(assemblyName, out string? assemblyLocation))
+					{
+						// We have this assembly in our folder.					
+						if (File.Exists(assemblyLocation))
+						{
+							// The assembly still exists, so load it.
+							return Assembly.LoadFile(assemblyLocation);
+						}
+						else
+						{
+							// The assembly no longer exists on disk, so remove it from our cache.
+							s_assemblyLocationCache.Remove(assemblyName);
+						}
+					}
+
+					return null;
+				};
+
+				s_addedToAssemblyResolver = true;
+			}
 		}
 
 		// Map of assembly name to path on disk
 		private static readonly Dictionary<string, string> s_assemblyLocationCache = new Dictionary<string, string>(StringComparer.OrdinalIgnoreCase);
 		// Track last modified date of each assembly, so we can ensure we always reference the latest one in the case of stale assemblies on disk.
 		private static readonly Dictionary<string, DateTime> s_assemblyWriteTimes = new Dictionary<string, DateTime>(StringComparer.OrdinalIgnoreCase);
+		// Flag used to make sure we don't redundantly add resolvers
+		private static bool s_addedToAssemblyResolver = false;
 
 	}
 }
diff --git a/Engine/Source/Programs/UnrealBuildTool/System/RulesAssembly.cs b/Engine/Source/Programs/UnrealBuildTool/System/RulesAssembly.cs
--- a/Engine/Source/Programs/UnrealBuildTool/System/RulesAssembly.cs
+++ b/Engine/Source/Programs/UnrealBuildTool/System/RulesAssembly.cs
@@ -526,6 +526,15 @@ namespace UnrealBuildTool
 				{
 					throw new BuildException("No valid constructor found for {0}.", ModuleName);
 				}
+
+				// Add the parent assemblies to the assembly cache so the types in them can be used when the constructor is called
+				var CurrentParent = Parent;
+				while (CurrentParent != null && CurrentParent.CompiledAssembly != null)
+				{
+					EpicGames.Core.AssemblyUtils.AddFileToAssemblyCache(CurrentParent.CompiledAssembly.Location);
+					CurrentParent = CurrentParent.Parent;
+				}
+
 				Constructor.Invoke(RulesObject, new object[] { Target });
 
 				if (Target.IsTestTarget && !RulesObject.IsTestModule)
