EpicGames.Core: Add missing XML doc

#skipci

diff --git a/Engine/Source/Programs/Shared/EpicGames.Core/LogEvent.cs b/Engine/Source/Programs/Shared/EpicGames.Core/LogEvent.cs
--- a/Engine/Source/Programs/Shared/EpicGames.Core/LogEvent.cs
+++ b/Engine/Source/Programs/Shared/EpicGames.Core/LogEvent.cs
@@ -605,17 +605,27 @@ namespace EpicGames.Core
 	/// </summary>
 	public class FilteringEventSink : ILogEventSink
 	{
-		private IReadOnlyList<LogEvent> LogEvents => _logEvents;
+		/// <summary>
+		/// Log events received
+		/// </summary>
+		public IReadOnlyList<LogEvent> LogEvents => _logEvents;
+		
 		private readonly List<LogEvent> _logEvents = new();
 		private readonly int[] _includeEventIds;
 		private readonly Action<LogEvent>? _eventCallback;
 		
+		/// <summary>
+		/// Constructor
+		/// </summary>
+		/// <param name="includeEventIds">Event IDs to include</param>
+		/// <param name="eventCallback">Optional callback function for each event received</param>
 		public FilteringEventSink(int[] includeEventIds, Action<LogEvent>? eventCallback = null)
 		{
 			_includeEventIds = includeEventIds;
 			_eventCallback = eventCallback;
 		}
 
+		/// <inheritdoc/>
 		public void ProcessEvent(LogEvent logEvent)
 		{
 			if (_includeEventIds.Any(x => x == logEvent.Id.Id))
