Fixing two instances to allow project to build with C# language version 8.

diff --git a/Engine/Source/Programs/Shared/EpicGames.Core/LogEvent.cs b/Engine/Source/Programs/Shared/EpicGames.Core/LogEvent.cs
--- a/Engine/Source/Programs/Shared/EpicGames.Core/LogEvent.cs
+++ b/Engine/Source/Programs/Shared/EpicGames.Core/LogEvent.cs
@@ -610,7 +610,7 @@ namespace EpicGames.Core
 		/// </summary>
 		public IReadOnlyList<LogEvent> LogEvents => _logEvents;
 		
-		private readonly List<LogEvent> _logEvents = new();
+		private readonly List<LogEvent> _logEvents = new List<LogEvent>();
 		private readonly int[] _includeEventIds;
 		private readonly Action<LogEvent>? _eventCallback;
 		
diff --git a/Engine/Source/Programs/Shared/EpicGames.Core/LogEventParser.cs b/Engine/Source/Programs/Shared/EpicGames.Core/LogEventParser.cs
--- a/Engine/Source/Programs/Shared/EpicGames.Core/LogEventParser.cs
+++ b/Engine/Source/Programs/Shared/EpicGames.Core/LogEventParser.cs
@@ -144,7 +144,7 @@ namespace EpicGames.Core
 		/// <summary>
 		/// Log events sinks in addition to <see cref="_logger" />
 		/// </summary>
-		List<ILogEventSink> _logEventSinks = new();
+		List<ILogEventSink> _logEventSinks = new List<ILogEventSink>();
 
 		/// <summary>
 		/// Timer for the parser being active
