Horde: Support treating a JsonLogEvent like any other state object passed into ILogger (support enumeration of properties, etc...).

diff --git a/Engine/Source/Programs/Shared/EpicGames.Core/JsonLogEvent.cs b/Engine/Source/Programs/Shared/EpicGames.Core/JsonLogEvent.cs
--- a/Engine/Source/Programs/Shared/EpicGames.Core/JsonLogEvent.cs
+++ b/Engine/Source/Programs/Shared/EpicGames.Core/JsonLogEvent.cs
@@ -1,6 +1,8 @@
 // Copyright Epic Games, Inc. All Rights Reserved.
 
 using System;
+using System.Collections;
+using System.Collections.Generic;
 using System.Text;
 using System.Text.Json;
 using Microsoft.Extensions.Logging;
@@ -10,7 +12,7 @@ namespace EpicGames.Core
 	/// <summary>
 	/// Defines a preformatted Json log event, which can pass through raw Json data directly or format it as a regular string
 	/// </summary>
-	public struct JsonLogEvent
+	public struct JsonLogEvent : IEnumerable<KeyValuePair<string, object?>>
 	{
 		/// <summary>
 		/// The log level
@@ -244,5 +246,27 @@ namespace EpicGames.Core
 
 		/// <inheritdoc/>
 		public override string ToString() => Encoding.UTF8.GetString(Data.ToArray());
+
+		/// <inheritdoc/>
+		public IEnumerator<KeyValuePair<string, object?>> GetEnumerator() => LogEvent.Read(Data.Span).GetEnumerator();
+
+		/// <inheritdoc/>
+		IEnumerator IEnumerable.GetEnumerator() => GetEnumerator();
+	}
+
+	/// <summary>
+	/// Extension methods for <see cref="JsonLogEvent"/>
+	/// </summary>
+	public static class JsonLogEventExtensions
+	{
+		/// <summary>
+		/// Logs a <see cref="JsonLogEvent"/> to the given logger
+		/// </summary>
+		/// <param name="logger">Logger to write to</param>
+		/// <param name="jsonLogEvent">Json log event to write</param>
+		public static void LogJsonLogEvent(this ILogger logger, JsonLogEvent jsonLogEvent)
+		{
+			logger.Log(jsonLogEvent.Level, jsonLogEvent.EventId, jsonLogEvent, null, JsonLogEvent.Format);
+		}
 	}
 }
diff --git a/Engine/Source/Programs/Shared/EpicGames.Core/LogEventParser.cs b/Engine/Source/Programs/Shared/EpicGames.Core/LogEventParser.cs
--- a/Engine/Source/Programs/Shared/EpicGames.Core/LogEventParser.cs
+++ b/Engine/Source/Programs/Shared/EpicGames.Core/LogEventParser.cs
@@ -402,7 +402,7 @@ namespace EpicGames.Core
 				if (JsonLogEvent.TryParse(data, out jsonEvent))
 				{
 					ProcessData(true);
-					_logger.Log(jsonEvent.Level, jsonEvent.EventId, jsonEvent, null, JsonLogEvent.Format);
+					_logger.LogJsonLogEvent(jsonEvent);
 					return;
 				}
 			}
