Horde: Fix static analyzer warning.

diff --git a/Engine/Source/Programs/Shared/EpicGames.Core/LogEventParser.cs b/Engine/Source/Programs/Shared/EpicGames.Core/LogEventParser.cs
--- a/Engine/Source/Programs/Shared/EpicGames.Core/LogEventParser.cs
+++ b/Engine/Source/Programs/Shared/EpicGames.Core/LogEventParser.cs
@@ -144,7 +144,7 @@ namespace EpicGames.Core
 		/// <summary>
 		/// Log events sinks in addition to <see cref="_logger" />
 		/// </summary>
-		List<ILogEventSink> _logEventSinks = new List<ILogEventSink>();
+		readonly List<ILogEventSink> _logEventSinks = new List<ILogEventSink>();
 
 		/// <summary>
 		/// Timer for the parser being active
