Horde: Faster path for constructing LogEvents from JsonLogEvents (also preserves timestamp in original message).

diff --git a/Engine/Source/Programs/Shared/EpicGames.Core/LogEvent.cs b/Engine/Source/Programs/Shared/EpicGames.Core/LogEvent.cs
--- a/Engine/Source/Programs/Shared/EpicGames.Core/LogEvent.cs
+++ b/Engine/Source/Programs/Shared/EpicGames.Core/LogEvent.cs
@@ -443,10 +443,14 @@ namespace EpicGames.Core
 		/// </summary>
 		public static LogEvent FromState<TState>(LogLevel level, EventId eventId, TState state, Exception? exception, Func<TState, Exception?, string> formatter)
 		{
-			if(state is LogEvent logEvent)
+			if (state is LogEvent logEvent)
 			{
 				return logEvent;
 			}
+			if (state is JsonLogEvent jsonLogEvent)
+			{
+				return Read(jsonLogEvent.Data.Span);
+			}
 
 			DateTime time = DateTime.UtcNow;
 
