Horde: Fix schema generation for cluster ids.

diff --git a/Engine/Source/Programs/Shared/EpicGames.Horde/Compute/ClusterId.cs b/Engine/Source/Programs/Shared/EpicGames.Horde/Compute/ClusterId.cs
--- a/Engine/Source/Programs/Shared/EpicGames.Horde/Compute/ClusterId.cs
+++ b/Engine/Source/Programs/Shared/EpicGames.Horde/Compute/ClusterId.cs
@@ -13,6 +13,7 @@ namespace EpicGames.Horde.Compute
 	/// <summary>
 	/// Identifier for a compute cluster
 	/// </summary>
+	[JsonSchemaString]
 	[CbConverter(typeof(ClusterIdCbConverter))]
 	[JsonConverter(typeof(ClusterIdJsonConverter))]
 	[TypeConverter(typeof(ClusterIdTypeConverter))]
