EpicGames.Core: Fix exception parsing null properties from log events.

diff --git a/Engine/Source/Programs/Shared/EpicGames.Core/LogEvent.cs b/Engine/Source/Programs/Shared/EpicGames.Core/LogEvent.cs
--- a/Engine/Source/Programs/Shared/EpicGames.Core/LogEvent.cs
+++ b/Engine/Source/Programs/Shared/EpicGames.Core/LogEvent.cs
@@ -276,10 +276,12 @@ namespace EpicGames.Core
 		{
 			switch (reader.TokenType)
 			{
+				case JsonTokenType.Null:
+					return null!;
 				case JsonTokenType.True:
 					return true;
 				case JsonTokenType.False:
-					return true;
+					return false;
 				case JsonTokenType.StartObject:
 					return ReadStructuredPropertyValue(ref reader);
 				case JsonTokenType.String:
