Updated the UHT header file parser to parse code in WITH_ENGINE so generated includes can be found and optionally throw an error or warning if the generated header is not included.

diff --git a/Engine/Source/Programs/Shared/EpicGames.UHT/Parsers/UhtHeaderFileParser.cs b/Engine/Source/Programs/Shared/EpicGames.UHT/Parsers/UhtHeaderFileParser.cs
--- a/Engine/Source/Programs/Shared/EpicGames.UHT/Parsers/UhtHeaderFileParser.cs
+++ b/Engine/Source/Programs/Shared/EpicGames.UHT/Parsers/UhtHeaderFileParser.cs
@@ -82,10 +82,15 @@ namespace EpicGames.UHT.Parsers
 		/// </summary>
 		WithHotReload = 1 << 6,
 
+		/// <summary>
+		/// This indicates we are in a "#if WITH_ENGINE" block
+		/// </summary>
+		WithEngine = 1 << 7,
+
 		/// <summary>
 		/// This directive is unrecognized and does not change the code generation at all
 		/// </summary>
-		Unrecognized = 1 << 7,
+		Unrecognized = 1 << 8,
 	}
 
 	/// <summary>
@@ -244,6 +249,7 @@ namespace EpicGames.UHT.Parsers
 			if (!headerParser.SpottedAutogeneratedHeaderInclude && headerParser.HeaderFile.Data.Length > 0)
 			{
 				bool noExportClassesOnly = true;
+				bool missingGeneratedHeader = false;
 				foreach (UhtType type in headerParser.HeaderFile.Children)
 				{
 					if (type is UhtClass classObj)
@@ -254,11 +260,36 @@ namespace EpicGames.UHT.Parsers
 							break;
 						}
 					}
+					else if (!headerParser.HeaderFile.IsNoExportTypes)
+					{
+						missingGeneratedHeader = true;
+					}
 				}
 
+				string logMessage = $"Expected an include at the top of the header the follows all other includes: '#include \"{headerParser.HeaderFile.GeneratedHeaderFileName}\"'";
 				if (!noExportClassesOnly)
 				{
-					headerParser.HeaderFile.LogError($"Expected an include at the top of the header the follows all other includes: '#include \"{headerParser.HeaderFile.GeneratedHeaderFileName}\"'");
+					headerParser.HeaderFile.LogError(logMessage);
+				}
+
+				if (missingGeneratedHeader)
+				{
+					UhtIssueBehavior missingGeneratedHeaderBehavior = headerFile.Package.IsPartOfEngine ? headerFile.Session.Config!.EngineMissingGeneratedHeaderIncludeBehavior
+						: headerFile.Session.Config!.NonEngineMissingGeneratedHeaderIncludeBehavior;
+
+					switch (missingGeneratedHeaderBehavior) 
+					{
+						case UhtIssueBehavior.AllowSilently:
+							break;
+
+						case UhtIssueBehavior.AllowAndLog:
+							headerParser.HeaderFile.LogTrace(logMessage);
+							break;
+
+						default:
+							headerParser.HeaderFile.LogError(logMessage);
+							break;
+					}
 				}
 			}
 			return headerParser;
@@ -748,6 +779,9 @@ namespace EpicGames.UHT.Parsers
 					case UhtCompilerDirective.CPPBlock:
 						PushCompilerDirective(UhtCompilerDirective.NotCPPBlock);
 						break;
+					case UhtCompilerDirective.WithEngine:
+						PushCompilerDirective(UhtCompilerDirective.Unrecognized);
+						break;
 					case UhtCompilerDirective.WithHotReload:
 						throw new UhtException(TokenReader, directive.InputLine, "Can not use WITH_HOT_RELOAD with an #else clause");
 					default:
@@ -837,6 +871,10 @@ namespace EpicGames.UHT.Parsers
 					{
 						return UhtCompilerDirective.WithHotReload;
 					}
+					else if (define.IsValue("WITH_ENGINE"))
+					{
+						return notPresent ? UhtCompilerDirective.Unrecognized : UhtCompilerDirective.WithEngine;
+					}
 					else if (define.IsValue("CPP"))
 					{
 						return notPresent ? UhtCompilerDirective.NotCPPBlock : UhtCompilerDirective.CPPBlock;
@@ -914,6 +952,7 @@ namespace EpicGames.UHT.Parsers
 
 				case UhtCompilerDirective.WithEditor:
 				case UhtCompilerDirective.WithEditorOnlyData:
+				case UhtCompilerDirective.WithEngine:
 				case UhtCompilerDirective.WithHotReload:
 					return true;
 
@@ -935,6 +974,7 @@ namespace EpicGames.UHT.Parsers
 
 				case UhtCompilerDirective.WithEditor:
 				case UhtCompilerDirective.WithEditorOnlyData:
+				case UhtCompilerDirective.WithEngine:
 				case UhtCompilerDirective.WithHotReload:
 					return true;
 
@@ -948,7 +988,8 @@ namespace EpicGames.UHT.Parsers
 			return
 				compilerDirective == UhtCompilerDirective.WithEditor ||
 				compilerDirective == UhtCompilerDirective.WithEditorOnlyData ||
-				compilerDirective == UhtCompilerDirective.WithHotReload;
+				compilerDirective == UhtCompilerDirective.WithHotReload ||
+				compilerDirective == UhtCompilerDirective.WithEngine;
 		}
 
 		private static string GetCompilerDirectiveText(UhtCompilerDirective compilerDirective)
@@ -962,6 +1003,7 @@ namespace EpicGames.UHT.Parsers
 				case UhtCompilerDirective.WithHotReload: return "WITH_HOT_RELOAD";
 				case UhtCompilerDirective.WithEditor: return "WITH_EDITOR";
 				case UhtCompilerDirective.WithEditorOnlyData: return "WITH_EDITORONLY_DATA";
+				case UhtCompilerDirective.WithEngine: return "WITH_ENGINE";
 				default: return "<unrecognized>";
 			}
 		}
