OidcToken Fixes

Support multiple redirect uris (so that we can attempt to use multiple ports)
Handle invalid json in state file, ignoring the state file in those cases
Also made sure to delete oidc token result file after consuming it from the editor

diff --git a/Engine/Source/Developer/DesktopPlatform/Private/DesktopPlatformBase.cpp b/Engine/Source/Developer/DesktopPlatform/Private/DesktopPlatformBase.cpp
--- a/Engine/Source/Developer/DesktopPlatform/Private/DesktopPlatformBase.cpp
+++ b/Engine/Source/Developer/DesktopPlatform/Private/DesktopPlatformBase.cpp
@@ -868,6 +868,9 @@ bool FDesktopPlatformBase::GetOidcAccessToken(const FString& RootDir, const FStr
 
 				FDateTime::ParseIso8601(*ExpiresAt, OutTokenExpiresAt);
 
+				// Remove the output file if its still around
+				IFileManager::Get().Delete(*ResultFilePath, true, false, true);
+
 				return true;
 			}
 		}
@@ -875,6 +878,8 @@ bool FDesktopPlatformBase::GetOidcAccessToken(const FString& RootDir, const FStr
 
 	UE_LOG(LogDesktopPlatform, Warning, TEXT("Failed to run OidcToken (project file is '%s', exe path is '%s'). No result file found at '%s', closed with exit code: %d"), *ProjectFileName, *GetOidcTokenExecutableFilename(RootDir), *ResultFilePath, ExitCode);
 
+	// Remove the output file if its still around
+	IFileManager::Get().Delete(*ResultFilePath, true, false, true);
 	return false;
 }
 
@@ -913,11 +918,16 @@ bool FDesktopPlatformBase::GetOidcTokenStatus(const FString& RootDir, const FStr
 			{
 				OutStatus = Status;
 
+				// Remove the output file if its still around
+				IFileManager::Get().Delete(*ResultFilePath, true, false, true);
+
 				return true;
 			}
 		}
 	}
 
+	// Remove the output file if its still around
+	IFileManager::Get().Delete(*ResultFilePath, true, false, true);
 	return false;
 }
 
diff --git a/Engine/Source/Programs/OidcToken/TokenService.cs b/Engine/Source/Programs/OidcToken/TokenService.cs
--- a/Engine/Source/Programs/OidcToken/TokenService.cs
+++ b/Engine/Source/Programs/OidcToken/TokenService.cs
@@ -52,6 +52,11 @@ namespace OidcToken
 					Logger.LogWarning("Was unable to allocate a token");
 					ExitCode = 10;
 				}
+				catch (HttpServerException e)
+				{
+					Logger.LogWarning("Unable to start http server:" + e.Message);
+					ExitCode = 2;
+				}
 				catch (Exception ex)
 				{
 					Logger.LogError(ex, "Unhandled exception!");
diff --git a/Engine/Source/Programs/OidcToken/readme.md b/Engine/Source/Programs/OidcToken/readme.md
--- a/Engine/Source/Programs/OidcToken/readme.md
+++ b/Engine/Source/Programs/OidcToken/readme.md
@@ -15,6 +15,10 @@ The configuration file can look like this:
 				"ClientId": "<unique-id-from-provider-usually-guid>",
 				"DisplayName": "MyOwnProvider",
 				"RedirectUri": "http://localhost:6556/callback", // this needs to match what is configured as the redirect uri for your IdP
+                "PossibleRedirectUri": [
+					"http://localhost:6556/callback",
+					"http://localhost:6557/callback",
+				], // set of redirect uris that can be used, ports can be in use so it is a good idea to configure a few alternatives. these needs to match configuration in IdP
 				"Scopes": "openid profile offline_access" // these scopes are the basic ones you will need, some system may require more and they may be named differently
 			}
 		}
diff --git a/Engine/Source/Programs/Shared/EpicGames.OIDC/OidcTokenManager.cs b/Engine/Source/Programs/Shared/EpicGames.OIDC/OidcTokenManager.cs
--- a/Engine/Source/Programs/Shared/EpicGames.OIDC/OidcTokenManager.cs
+++ b/Engine/Source/Programs/Shared/EpicGames.OIDC/OidcTokenManager.cs
@@ -156,18 +156,14 @@ namespace EpicGames.OIDC
 		private readonly ProviderInfo _providerInfo;
 		private readonly Uri _authorityUri;
 		private readonly string _clientId;
-		private readonly Uri _redirectUri;
 		private readonly string _scopes;
 
-		private OidcClient? _oidcClient;
-
 		private string? _refreshToken;
 		private string? _accessToken;
 		private DateTimeOffset _tokenExpiry;
 
-		private bool _initialized = false;
-
 		private readonly ILogger<OidcTokenClient>? _logger;
+		private readonly List<Uri> _redirectUris;
 
 		public OidcTokenClient(ILogger<OidcTokenClient> logger, ProviderInfo providerInfo)
 		{
@@ -176,7 +172,14 @@ namespace EpicGames.OIDC
 
 			_authorityUri = providerInfo.ServerUri;
 			_clientId = providerInfo.ClientId;
-			_redirectUri = providerInfo.RedirectUri;
+
+			List<Uri> possibleRedirectUris = new List<Uri> { providerInfo.RedirectUri };
+			if (providerInfo.PossibleRedirectUri != null)
+			{
+				possibleRedirectUris.AddRange(providerInfo.PossibleRedirectUri);
+			}
+			_redirectUris = possibleRedirectUris;
+
 			_scopes = providerInfo.Scopes;
 		}
 
@@ -187,17 +190,18 @@ namespace EpicGames.OIDC
 
 			_authorityUri = providerInfo.ServerUri;
 			_clientId = providerInfo.ClientId;
-			_redirectUri = providerInfo.RedirectUri;
+			List<Uri> possibleRedirectUris = new List<Uri> { providerInfo.RedirectUri };
+			if (providerInfo.PossibleRedirectUri != null)
+			{
+				possibleRedirectUris.AddRange(providerInfo.PossibleRedirectUri);
+			}
+
+			_redirectUris = possibleRedirectUris;
 			_scopes = providerInfo.Scopes;
 		}
 
-		private async Task Initialize()
+		private async Task<OidcClientOptions> BuildClientOptions(Uri redirectUri)
 		{
-			if (_initialized)
-			{
-				return;
-			}
-
 			OidcClientOptions options = new OidcClientOptions
 			{
 				Authority = _authorityUri.ToString(),
@@ -205,7 +209,7 @@ namespace EpicGames.OIDC
 				ClientId = _clientId,
 				Scope = _scopes,
 				FilterClaims = false,
-				RedirectUri = _redirectUri.ToString(),
+				RedirectUri = redirectUri.ToString(),
 				LoadProfile = _providerInfo.LoadClaimsFromUserProfile,
 				//Flow = OidcClientOptions.AuthenticationFlow.AuthorizationCode,
 				//ResponseMode = OidcClientOptions.AuthorizeResponseMode.Redirect
@@ -226,31 +230,46 @@ namespace EpicGames.OIDC
 				TokenEndPointAuthenticationMethods = discoveryDocument.TokenEndpointAuthenticationMethodsSupported
 			};
 
-			_oidcClient = new OidcClient(options);
-			_initialized = true;
+			return options;
 		}
 
 		public async Task<OidcTokenInfo> Login()
 		{
-			await Initialize();
-
 			// setup a local http server to listen for the result of the login
-			using HttpListener http = new HttpListener();
-			Uri uri = _redirectUri;
-			// build the url the server should be hosted at
-			string prefix = $"{uri.Scheme}{Uri.SchemeDelimiter}{uri.Authority}/";
-			http.Prefixes.Add(prefix);
-			http.Start();
+			LoginResult? loginResult = null;
+			foreach (Uri uri in _redirectUris)
+			{
+				try
+				{
+					using HttpListener http = new HttpListener();
 
-			// generate the appropriate codes we need to login
-			AuthorizeState loginState = await _oidcClient!.PrepareLoginAsync();
-			// start the user browser
-			OpenBrowser(loginState.StartUrl);
+					// build the url the server should be hosted at
+					string prefix = $"{uri.Scheme}{Uri.SchemeDelimiter}{uri.Authority}/";
+					http.Prefixes.Add(prefix);
+					http.Start();
 
-			LoginResult loginResult = await ProcessHttpRequest(http, loginState);
+					OidcClientOptions options = await BuildClientOptions(uri);
+					OidcClient oidcClient = new OidcClient(options);
+					// generate the appropriate codes we need to login
+					AuthorizeState loginState = await oidcClient!.PrepareLoginAsync();
+					// start the user browser
+					OpenBrowser(loginState.StartUrl);
 
-			http.Stop();
+					loginResult = await ProcessHttpRequest(http, loginState, oidcClient);
 
+					http.Stop();
+					break;
+				}
+				catch (HttpListenerException)
+				{
+					continue;
+				}
+			}
+
+			if (loginResult == null)
+			{
+				throw new HttpServerException("Unable to login as none of the possible redirect uris were successful. Uris used: " + String.Join(' ', _redirectUris));
+			}
 			if (loginResult.IsError)
 			{
 				throw new LoginFailedException("Failed to login due to error: " + loginResult.Error, loginResult.ErrorDescription);
@@ -269,7 +288,7 @@ namespace EpicGames.OIDC
 			};
 		}
 
-		private async Task<LoginResult> ProcessHttpRequest(HttpListener http, AuthorizeState loginState)
+		private async Task<LoginResult> ProcessHttpRequest(HttpListener http, AuthorizeState loginState, OidcClient oidcClient)
 		{
 			LoginResult loginResult;
 			HttpListenerContext context = await http.GetContextAsync();
@@ -280,7 +299,7 @@ namespace EpicGames.OIDC
 					responseData = context.Request.RawUrl;
 
 					// parse the returned url for the tokens needed to complete the login
-					loginResult = await _oidcClient!.ProcessResponseAsync(responseData, loginState);
+					loginResult = await oidcClient!.ProcessResponseAsync(responseData, loginState);
 					break;
 				case "POST":
 				{
@@ -305,7 +324,7 @@ namespace EpicGames.OIDC
 					using StreamReader reader = new StreamReader(body, request.ContentEncoding);
 					responseData = await reader.ReadToEndAsync();
 
-					loginResult = await _oidcClient!.ProcessResponseAsync(responseData, loginState);
+					loginResult = await oidcClient!.ProcessResponseAsync(responseData, loginState);
 					break;
 				}
 				default:
@@ -380,10 +399,12 @@ namespace EpicGames.OIDC
 
 		private async Task<OidcTokenInfo> DoRefreshToken(string inRefreshToken)
 		{
-			await Initialize();
+			// redirect uri is not used for refrehs tokens so we can just pick one of them to configure the client
+			OidcClientOptions options = await BuildClientOptions(_redirectUris.First());
+			OidcClient oidcClient = new OidcClient(options);
 
 			// use the refresh token to acquire a new access token
-			RefreshTokenResult refreshTokenResult = await _oidcClient!.RefreshTokenAsync(inRefreshToken);
+			RefreshTokenResult refreshTokenResult = await oidcClient!.RefreshTokenAsync(inRefreshToken);
 
 			if (refreshTokenResult.IsError)
 			{
@@ -491,6 +512,13 @@ namespace EpicGames.OIDC
 		public string ErrorDescription { get;  }
 	}
 
+	public class HttpServerException : Exception
+	{
+		public HttpServerException(string message) : base(message)
+		{
+		}
+	}
+
 	public class NotLoggedInException : Exception
 	{
 	}
@@ -508,7 +536,8 @@ namespace EpicGames.OIDC
 
 		[Required] public string DisplayName { get; set; } = null!;
 
-		[Required] public Uri RedirectUri { get; set; } = null!;
+		public Uri RedirectUri { get; set; } = null!;
+		public List<Uri>? PossibleRedirectUri { get; set; } = null!;
 		[Required] public bool LoadClaimsFromUserProfile { get; set; } = false;
 
 		public string Scopes { get; set; } = "openid profile offline_access email";
diff --git a/Engine/Source/Programs/Shared/EpicGames.OIDC/TokenStore.cs b/Engine/Source/Programs/Shared/EpicGames.OIDC/TokenStore.cs
--- a/Engine/Source/Programs/Shared/EpicGames.OIDC/TokenStore.cs
+++ b/Engine/Source/Programs/Shared/EpicGames.OIDC/TokenStore.cs
@@ -62,7 +62,15 @@ namespace EpicGames.OIDC
 
 			using FileStream fs = fi.OpenRead();
 			using TextReader tr = new StreamReader(fs);
-			TokenStoreState? state = JsonSerializer.Deserialize<TokenStoreState>(tr.ReadToEnd());
+			TokenStoreState? state;
+			try
+			{
+				state = JsonSerializer.Deserialize<TokenStoreState>(tr.ReadToEnd());
+			}
+			catch (JsonException)
+			{
+				state = null;
+			}
 
 			if (state == null)
 			{
diff --git a/Engine/Source/Programs/Shared/EpicGames.OIDC/WindowsTokenStore.cs b/Engine/Source/Programs/Shared/EpicGames.OIDC/WindowsTokenStore.cs
--- a/Engine/Source/Programs/Shared/EpicGames.OIDC/WindowsTokenStore.cs
+++ b/Engine/Source/Programs/Shared/EpicGames.OIDC/WindowsTokenStore.cs
@@ -55,7 +55,16 @@ namespace EpicGames.OIDC
 
 			using FileStream fs = fi.OpenRead();
 			using TextReader tr = new StreamReader(fs);
-			TokenStoreState? state = JsonSerializer.Deserialize<TokenStoreState>(tr.ReadToEnd());
+
+			TokenStoreState? state;
+			try
+			{
+				state = JsonSerializer.Deserialize<TokenStoreState>(tr.ReadToEnd());
+			}
+			catch (JsonException)
+			{
+				state = null;
+			}
 
 			if (state == null)
 			{
