Updated UHT to only include the necessary headers needed for enums when only enums are being generated. This fixes a compile issue CoreNetTypes was having after include its generated.h file.

diff --git a/Engine/Source/Programs/Shared/EpicGames.UHT/Exporters/CodeGen/UhtHeaderCodeGeneratorHFile.cs b/Engine/Source/Programs/Shared/EpicGames.UHT/Exporters/CodeGen/UhtHeaderCodeGeneratorHFile.cs
--- a/Engine/Source/Programs/Shared/EpicGames.UHT/Exporters/CodeGen/UhtHeaderCodeGeneratorHFile.cs
+++ b/Engine/Source/Programs/Shared/EpicGames.UHT/Exporters/CodeGen/UhtHeaderCodeGeneratorHFile.cs
@@ -42,8 +42,23 @@ namespace EpicGames.UHT.Exporters.CodeGen
 
 				builder.Append(HeaderCopyright);
 				builder.Append("// IWYU pragma: private, include \"").Append(this.HeaderFile.IncludeFilePath).Append("\"\r\n");
-				builder.Append("#include \"UObject/ObjectMacros.h\"\r\n");
-				builder.Append("#include \"UObject/ScriptMacros.h\"\r\n");
+
+				// Attempt to limit the headers included. This is needed in the lower level engine code
+				// to get around circular header include issues.
+				{
+					if (HeaderFile.References.ExportTypes.Count > 0 && HeaderFile.References.ExportTypes.Find(x => x is not UhtEnum) == null)
+					{
+						builder.Append("#include \"Templates/IsUEnumClass.h\"\r\n");
+						builder.Append("#include \"UObject/ObjectMacros.h\"\r\n");
+						builder.Append("#include \"UObject/ReflectedTypeAccessors.h\"\r\n");
+					}
+					else
+					{
+						builder.Append("#include \"UObject/ObjectMacros.h\"\r\n");
+						builder.Append("#include \"UObject/ScriptMacros.h\"\r\n");
+					}
+				}
+				
 				if (headerInfo.NeedsPushModelHeaders)
 				{
 					builder.Append("#include \"Net/Core/PushModel/PushModelMacros.h\"\r\n");
diff --git a/Engine/Source/Runtime/CoreUObject/Public/UObject/CoreNetTypes.h b/Engine/Source/Runtime/CoreUObject/Public/UObject/CoreNetTypes.h
--- a/Engine/Source/Runtime/CoreUObject/Public/UObject/CoreNetTypes.h
+++ b/Engine/Source/Runtime/CoreUObject/Public/UObject/CoreNetTypes.h
@@ -4,6 +4,8 @@
 
 #include "UObject/ObjectMacros.h"
 
+#include "CoreNetTypes.generated.h"
+
 /** Whether or not caching of actor/subobject names to the stack should be enabled, for async demo recording crashdumps */
 #ifndef UE_NET_REPACTOR_NAME_DEBUG
 	#define UE_NET_REPACTOR_NAME_DEBUG 0
