UnrealBuildTool: Fix .NET project generation for .net5. UDN suggestion

#rnx

diff --git a/Engine/Source/Programs/Shared/EpicGames.Core/CsProjectInfo.cs b/Engine/Source/Programs/Shared/EpicGames.Core/CsProjectInfo.cs
--- a/Engine/Source/Programs/Shared/EpicGames.Core/CsProjectInfo.cs
+++ b/Engine/Source/Programs/Shared/EpicGames.Core/CsProjectInfo.cs
@@ -282,15 +282,23 @@ namespace EpicGames.Core
 		/// <returns>True if the project is a .NET core project</returns>
 		public bool IsDotNETCoreProject()
 		{
-			bool bIsDotNetCoreProject = false;
-
-			string? targetFramework;
-			if (Properties.TryGetValue("TargetFramework", out targetFramework))
+			if (Properties.TryGetValue("TargetFramework", out string? targetFramework))
 			{
-				bIsDotNetCoreProject = targetFramework.ToLower().Contains("netstandard", StringComparison.Ordinal) || targetFramework.ToLower().Contains("netcoreapp", StringComparison.Ordinal);
+				if (targetFramework.ToLower().Contains("netstandard", StringComparison.Ordinal) || targetFramework.ToLower().Contains("netcoreapp", StringComparison.Ordinal))
+				{
+					return true;
+				}
+				else if (targetFramework.StartsWith("net", StringComparison.OrdinalIgnoreCase))
+				{
+					string[] versionSplit = targetFramework.Substring(3).Split('.');
+					if (versionSplit.Length >= 1 && Int32.TryParse(versionSplit[0], out int majorVersion))
+					{
+						return majorVersion >= 5;
+					}
+				}
 			}
 
-			return bIsDotNetCoreProject;
+			return false;
 		}
 
 		/// <summary>
