OidcToken - Move state store into place avoiding race condition if multiple tools tries to use oidc token at the same time
OidcToken - Replace token in keychain on mac if it already exists

diff --git a/Engine/Source/Programs/Shared/EpicGames.OIDC/MacOSTokenStore.cs b/Engine/Source/Programs/Shared/EpicGames.OIDC/MacOSTokenStore.cs
--- a/Engine/Source/Programs/Shared/EpicGames.OIDC/MacOSTokenStore.cs
+++ b/Engine/Source/Programs/Shared/EpicGames.OIDC/MacOSTokenStore.cs
@@ -95,6 +95,7 @@ namespace EpicGames.OIDC
 				itemIdentifier,
 				"-s",
 				OidcTokenIdentifier,
+				"-U",
 				"-w",
 				dataToEncrypt
 			});
diff --git a/Engine/Source/Programs/Shared/EpicGames.OIDC/TokenStore.cs b/Engine/Source/Programs/Shared/EpicGames.OIDC/TokenStore.cs
--- a/Engine/Source/Programs/Shared/EpicGames.OIDC/TokenStore.cs
+++ b/Engine/Source/Programs/Shared/EpicGames.OIDC/TokenStore.cs
@@ -95,9 +95,14 @@ namespace EpicGames.OIDC
 				Directory.CreateDirectory(fi.Directory!.FullName);
 			}
 
-			using FileStream fs = fi.Open(FileMode.Create, FileAccess.Write);
-			using Utf8JsonWriter writer = new Utf8JsonWriter(fs);
-			JsonSerializer.Serialize<TokenStoreState>(writer, new TokenStoreState(_crypt.IV, _providerToRefreshToken));
+			string tempFile = Path.GetTempFileName();
+			{
+				using FileStream fs = new FileStream(tempFile, FileMode.Create, FileAccess.Write);
+				using Utf8JsonWriter writer = new Utf8JsonWriter(fs);
+				JsonSerializer.Serialize<TokenStoreState>(writer, new TokenStoreState(_crypt.IV, _providerToRefreshToken));
+			}
+
+			File.Move(tempFile, fi.FullName, true);
 		}
 
 		public bool TryGetRefreshToken(string oidcProvider, out string refreshToken)
diff --git a/Engine/Source/Programs/Shared/EpicGames.OIDC/WindowsTokenStore.cs b/Engine/Source/Programs/Shared/EpicGames.OIDC/WindowsTokenStore.cs
--- a/Engine/Source/Programs/Shared/EpicGames.OIDC/WindowsTokenStore.cs
+++ b/Engine/Source/Programs/Shared/EpicGames.OIDC/WindowsTokenStore.cs
@@ -87,9 +87,14 @@ namespace EpicGames.OIDC
 				Directory.CreateDirectory(fi.Directory!.FullName);
 			}
 
-			using FileStream fs = fi.Open(FileMode.Create, FileAccess.Write);
-			using Utf8JsonWriter writer = new Utf8JsonWriter(fs);
-			JsonSerializer.Serialize<TokenStoreState>(writer, new TokenStoreState(_providerToRefreshToken));
+			string tempFile = Path.GetTempFileName();
+			{
+				using FileStream fs = new FileStream(tempFile, FileMode.Create, FileAccess.Write);
+				using Utf8JsonWriter writer = new Utf8JsonWriter(fs);
+				JsonSerializer.Serialize<TokenStoreState>(writer, new TokenStoreState(_providerToRefreshToken));
+			}
+
+			File.Move(tempFile, fi.FullName, true);
 		}
 
 		public bool TryGetRefreshToken(string oidcProvider, out string refreshToken)
