Horde: Add additional checks to validate data returned from log flush.

diff --git a/Engine/Source/Programs/Shared/EpicGames.Horde/Storage/TreeNodeRef.cs b/Engine/Source/Programs/Shared/EpicGames.Horde/Storage/TreeNodeRef.cs
--- a/Engine/Source/Programs/Shared/EpicGames.Horde/Storage/TreeNodeRef.cs
+++ b/Engine/Source/Programs/Shared/EpicGames.Horde/Storage/TreeNodeRef.cs
@@ -391,6 +391,16 @@ namespace EpicGames.Horde.Storage
 			TreeNodeRef rootRef = new TreeNodeRef(root);
 			await writer.WriteAsync(rootRef, cancellationToken);
 			await writer.FlushAsync(cancellationToken);
+
+			if (rootRef.Handle == null)
+			{
+				throw new InvalidOperationException("Invalid handle in root ref during flush");
+			}
+			else if (!rootRef.Handle.Locator.IsValid())
+			{
+				throw new InvalidOperationException("Missing locator in root ref after flush");
+			}
+
 			return rootRef.Handle!;
 		}
 	}
