Horde: Additional logging to track down log write failures.

diff --git a/Engine/Source/Programs/Shared/EpicGames.Horde/Storage/TreeNodeRef.cs b/Engine/Source/Programs/Shared/EpicGames.Horde/Storage/TreeNodeRef.cs
--- a/Engine/Source/Programs/Shared/EpicGames.Horde/Storage/TreeNodeRef.cs
+++ b/Engine/Source/Programs/Shared/EpicGames.Horde/Storage/TreeNodeRef.cs
@@ -396,6 +396,10 @@ namespace EpicGames.Horde.Storage
 			{
 				throw new InvalidOperationException("Invalid handle in root ref during flush");
 			}
+			else if (rootRef.IsDirty())
+			{
+				throw new InvalidOperationException("Root node is still dirty");
+			}
 			else if (!rootRef.Handle.Locator.IsValid())
 			{
 				throw new InvalidOperationException("Missing locator in root ref after flush");
