Horde: More logging for tree node ref flushes.

diff --git a/Engine/Source/Programs/Shared/EpicGames.Horde/Storage/TreeNodeRef.cs b/Engine/Source/Programs/Shared/EpicGames.Horde/Storage/TreeNodeRef.cs
--- a/Engine/Source/Programs/Shared/EpicGames.Horde/Storage/TreeNodeRef.cs
+++ b/Engine/Source/Programs/Shared/EpicGames.Horde/Storage/TreeNodeRef.cs
@@ -400,6 +400,10 @@ namespace EpicGames.Horde.Storage
 			{
 				throw new InvalidOperationException("Root node is still dirty");
 			}
+			else if (rootRef.Handle.PendingBundle != null)
+			{
+				throw new InvalidOperationException("Handle still has a pending bundle");
+			}
 			else if (!rootRef.Handle.Locator.IsValid())
 			{
 				throw new InvalidOperationException("Missing locator in root ref after flush");
