Horde: Add an optional trace logger for TreeWriter.

diff --git a/Engine/Source/Programs/Horde/Horde.Agent/Utility/JsonRpcLogSink.cs b/Engine/Source/Programs/Horde/Horde.Agent/Utility/JsonRpcLogSink.cs
--- a/Engine/Source/Programs/Horde/Horde.Agent/Utility/JsonRpcLogSink.cs
+++ b/Engine/Source/Programs/Horde/Horde.Agent/Utility/JsonRpcLogSink.cs
@@ -100,7 +100,7 @@ namespace Horde.Agent.Utility
 			_logId = logId;
 			_builder = new LogBuilder(LogFormat.Json, logger);
 			_inner = inner;
-			_writer = new TreeWriter(store);
+			_writer = new TreeWriter(store, traceLogger: logger);
 			_logger = logger;
 
 			_tailTaskStop = new AsyncEvent();
diff --git a/Engine/Source/Programs/Shared/EpicGames.Horde/Storage/TreeNodeRef.cs b/Engine/Source/Programs/Shared/EpicGames.Horde/Storage/TreeNodeRef.cs
--- a/Engine/Source/Programs/Shared/EpicGames.Horde/Storage/TreeNodeRef.cs
+++ b/Engine/Source/Programs/Shared/EpicGames.Horde/Storage/TreeNodeRef.cs
@@ -7,6 +7,7 @@ using System.Linq;
 using System.Threading;
 using System.Threading.Tasks;
 using EpicGames.Core;
+using Microsoft.Extensions.Logging;
 
 namespace EpicGames.Horde.Storage
 {
@@ -389,25 +390,12 @@ namespace EpicGames.Horde.Storage
 		public static async Task<NodeHandle> FlushAsync(this TreeWriter writer, TreeNode root, CancellationToken cancellationToken)
 		{
 			TreeNodeRef rootRef = new TreeNodeRef(root);
-			await writer.WriteAsync(rootRef, cancellationToken);
-			await writer.FlushAsync(cancellationToken);
 
-			if (rootRef.Handle == null)
-			{
-				throw new InvalidOperationException("Invalid handle in root ref during flush");
-			}
-			else if (rootRef.IsDirty())
-			{
-				throw new InvalidOperationException("Root node is still dirty");
-			}
-			else if (rootRef.Handle.PendingBundle != null)
-			{
-				throw new InvalidOperationException("Handle still has a pending bundle");
-			}
-			else if (!rootRef.Handle.Locator.IsValid())
-			{
-				throw new InvalidOperationException("Missing locator in root ref after flush");
-			}
+			NodeHandle handle = await writer.WriteAsync(rootRef, cancellationToken);
+			writer._traceLogger?.LogInformation("Written root node {Handle}", handle);
+
+			await writer.FlushAsync(cancellationToken);
+			writer._traceLogger?.LogInformation("Flushed root node {Handle}", handle);
 
 			return rootRef.Handle!;
 		}
diff --git a/Engine/Source/Programs/Shared/EpicGames.Horde/Storage/TreeWriter.cs b/Engine/Source/Programs/Shared/EpicGames.Horde/Storage/TreeWriter.cs
--- a/Engine/Source/Programs/Shared/EpicGames.Horde/Storage/TreeWriter.cs
+++ b/Engine/Source/Programs/Shared/EpicGames.Horde/Storage/TreeWriter.cs
@@ -11,6 +11,7 @@ using System.Threading.Tasks;
 using EpicGames.Core;
 using Microsoft.CodeAnalysis;
 using Microsoft.Extensions.Caching.Memory;
+using Microsoft.Extensions.Logging;
 using static EpicGames.Horde.Storage.TreeWriter;
 
 namespace EpicGames.Horde.Storage
@@ -276,6 +277,9 @@ namespace EpicGames.Horde.Storage
 
 			static readonly WriteCallbackSentinel s_callbackSentinel = new WriteCallbackSentinel();
 
+			static int s_nextUniqueId = 1;
+			public int BundleId { get; } = Interlocked.Increment(ref s_nextUniqueId);
+
 			readonly int _maxPacketSize;
 			readonly BundleCompressionFormat _compressionFormat;
 
@@ -454,18 +458,19 @@ namespace EpicGames.Horde.Storage
 			}
 
 			// Mark the bundle as complete
-			public void MarkAsComplete(IStorageClient store, Utf8String prefix)
+			public void MarkAsComplete(IStorageClient store, Utf8String prefix, ILogger? traceLogger)
 			{
 				if (!IsReadOnly)
 				{
+					traceLogger?.LogInformation("Marking bundle {BundleId} as complete ({NumNodes} nodes); adding to write queue.", BundleId, _queue.Count);
 					FlushPacket();
 					Task prevWriteTask = _writeTask;
-					_writeTask = Task.Run(() => CompleteAsync(prevWriteTask, store, prefix));
+					_writeTask = Task.Run(() => CompleteAsync(prevWriteTask, store, prefix, traceLogger));
 					IsReadOnly = true;
 				}
 			}
 
-			async Task CompleteAsync(Task prevWriteTask, IStorageClient store, Utf8String prefix)
+			async Task CompleteAsync(Task prevWriteTask, IStorageClient store, Utf8String prefix, ILogger? traceLogger)
 			{
 				try
 				{
@@ -474,10 +479,12 @@ namespace EpicGames.Horde.Storage
 
 					Bundle bundle = CreateBundle();
 					BlobLocator locator = await store.WriteBundleAsync(bundle, prefix);
+					traceLogger?.LogInformation("Written bundle {BundleId} as {Locator}", BundleId, locator);
 
 					for (int idx = 0; idx < _queue.Count; idx++)
 					{
 						NodeLocator nodeLocator = new NodeLocator(locator, idx);
+						traceLogger?.LogInformation("Updated pending node {Hash} with locator {Locator}", _queue[idx].Hash, nodeLocator);
 						_queue[idx].MarkAsWritten(nodeLocator);
 					}
 
@@ -613,6 +620,8 @@ namespace EpicGames.Horde.Storage
 		PendingBundle? _currentBundle;
 		int _memoryFootprint;
 
+		internal readonly ILogger? _traceLogger;
+
 		/// <summary>
 		/// Accessor for the store backing this writer
 		/// </summary>
@@ -630,12 +639,14 @@ namespace EpicGames.Horde.Storage
 		/// <param name="options">Options for the writer</param>
 		/// <param name="prefix">Prefix for blobs written to the store</param>
 		/// <param name="nodeCache">Cache of nodes for deduplication</param>
-		public TreeWriter(IStorageClient store, TreeOptions? options = null, Utf8String prefix = default, NodeCache? nodeCache = null)
+		/// <param name="traceLogger">Optional logger for trace information</param>
+		public TreeWriter(IStorageClient store, TreeOptions? options = null, Utf8String prefix = default, NodeCache? nodeCache = null, ILogger? traceLogger = null)
 		{
 			_store = store;
 			_options = options ?? s_defaultOptions;
 			_prefix = prefix;
 			_nodeCache = nodeCache ?? new NodeCache(_options.NodeCacheSize);
+			_traceLogger = traceLogger;
 		}
 
 		/// <summary>
@@ -644,7 +655,8 @@ namespace EpicGames.Horde.Storage
 		/// <param name="store">Store to write data to</param>
 		/// <param name="refName">Ref being written. Will be used as a prefix for storing blobs.</param>
 		/// <param name="options">Options for the writer</param>
-		public TreeWriter(IStorageClient store, RefName refName, TreeOptions? options = null)
+		/// <param name="traceLogger">Optional logger for trace information</param>
+		public TreeWriter(IStorageClient store, RefName refName, TreeOptions? options = null, ILogger? traceLogger = null)
 			: this(store, options, refName.Text)
 		{
 		}
@@ -654,7 +666,7 @@ namespace EpicGames.Horde.Storage
 		/// </summary>
 		/// <param name="other"></param>
 		public TreeWriter(TreeWriter other)
-			: this(other._store, other._options, other._prefix, other._nodeCache)
+			: this(other._store, other._options, other._prefix, other._nodeCache, other._traceLogger)
 		{
 		}
 
@@ -675,7 +687,7 @@ namespace EpicGames.Horde.Storage
 		{
 			if (_currentBundle != null)
 			{
-				_currentBundle.MarkAsComplete(_store, _prefix);
+				_currentBundle.MarkAsComplete(_store, _prefix, _traceLogger);
 				_writeQueue.Enqueue(_currentBundle);
 				_currentBundle = null;
 			}
@@ -722,6 +734,7 @@ namespace EpicGames.Horde.Storage
 			// Check if we have a matching node already in storage
 			if (_nodeCache.TryGetNode(nodeKey, out NodeHandle? handle))
 			{
+				_traceLogger?.LogInformation("Returning cached handle for {NodeKey} -> {Handle}", nodeKey, handle);
 				return handle;
 			}
 
@@ -729,6 +742,7 @@ namespace EpicGames.Horde.Storage
 			PendingNode pendingNode = currentBundle.WriteNode(nodeKey, size, references);
 			_memoryFootprint += size;
 			_nodeCache.Add(nodeKey, pendingNode);
+			_traceLogger?.LogInformation("Added new node for {NodeKey} in bundle {BundleId}", nodeKey, currentBundle.BundleId);
 
 			// Add dependencies on all bundles containing a dependent node
 			foreach (NodeHandle reference in references)
@@ -767,6 +781,7 @@ namespace EpicGames.Horde.Storage
 		async Task WaitForWriteAsync(CancellationToken cancellationToken)
 		{
 			PendingBundle writtenBundle = _writeQueue.Dequeue();
+			_traceLogger?.LogInformation("Waiting for bundle {BundleId} to finish writing", writtenBundle.BundleId);
 			await await Task.WhenAny(writtenBundle.CompleteTask, Task.Delay(-1, cancellationToken));
 			_memoryFootprint -= writtenBundle.UncompressedLength;
 			writtenBundle.Dispose();
