Horde: Fix incorrect property returned from HTTP blob uploads.

diff --git a/Engine/Source/Programs/Shared/EpicGames.Horde/Storage/Backends/HttpStorageClient.cs b/Engine/Source/Programs/Shared/EpicGames.Horde/Storage/Backends/HttpStorageClient.cs
--- a/Engine/Source/Programs/Shared/EpicGames.Horde/Storage/Backends/HttpStorageClient.cs
+++ b/Engine/Source/Programs/Shared/EpicGames.Horde/Storage/Backends/HttpStorageClient.cs
@@ -32,7 +32,7 @@ namespace EpicGames.Horde.Storage.Backends
 
 		class WriteBlobResponse
 		{
-			public BlobLocator Locator { get; set; }
+			public BlobLocator Blob { get; set; }
 			public Uri? UploadUrl { get; set; }
 			public bool? SupportsRedirects { get; set; }
 		}
@@ -153,16 +153,16 @@ namespace EpicGames.Horde.Storage.Backends
 								throw new StorageException($"Unable to upload data to redirected URL: {body}", null);
 							}
 						}
-						_logger.LogDebug("Written {Locator} (using redirect)", redirectResponse.Locator);
-						return redirectResponse.Locator;
+						_logger.LogDebug("Written {Locator} (using redirect)", redirectResponse.Blob);
+						return redirectResponse.Blob;
 					}
 				}
 			}
 
 			WriteBlobResponse response = await SendWriteRequestAsync(streamContent, prefix, cancellationToken);
 			_supportsUploadRedirects = response.SupportsRedirects ?? false;
-			_logger.LogDebug("Written {Locator} (direct)", response.Locator);
-			return response.Locator;
+			_logger.LogDebug("Written {Locator} (direct)", response.Blob);
+			return response.Blob;
 		}
 
 		async Task<WriteBlobResponse> SendWriteRequestAsync(StreamContent? streamContent, Utf8String prefix = default, CancellationToken cancellationToken = default)
diff --git a/Engine/Source/Programs/Shared/EpicGames.Horde/Storage/TreeWriter.cs b/Engine/Source/Programs/Shared/EpicGames.Horde/Storage/TreeWriter.cs
--- a/Engine/Source/Programs/Shared/EpicGames.Horde/Storage/TreeWriter.cs
+++ b/Engine/Source/Programs/Shared/EpicGames.Horde/Storage/TreeWriter.cs
@@ -277,8 +277,8 @@ namespace EpicGames.Horde.Storage
 
 			static readonly WriteCallbackSentinel s_callbackSentinel = new WriteCallbackSentinel();
 
-			static int s_nextUniqueId = 1;
-			public int BundleId { get; } = Interlocked.Increment(ref s_nextUniqueId);
+			static int s_lastBundleId = 0;
+			public int BundleId { get; } = Interlocked.Increment(ref s_lastBundleId);
 
 			readonly int _maxPacketSize;
 			readonly BundleCompressionFormat _compressionFormat;
