UnrealBuildTool: Fix typo in DirectoryItem docstring

#rnx

diff --git a/Engine/Source/Programs/Shared/EpicGames.Build/System/DirectoryItem.cs b/Engine/Source/Programs/Shared/EpicGames.Build/System/DirectoryItem.cs
--- a/Engine/Source/Programs/Shared/EpicGames.Build/System/DirectoryItem.cs
+++ b/Engine/Source/Programs/Shared/EpicGames.Build/System/DirectoryItem.cs
@@ -457,8 +457,8 @@ namespace UnrealBuildBase
 		/// <summary>
 		/// Writes a directory reference  to a binary archive
 		/// </summary>
-		/// <param name="writer">The writer to output data to</param>
-		/// <param name="directory">The item to write</param>
+		/// <param name="Writer">The writer to output data to</param>
+		/// <param name="Directory">The item to write</param>
 		public static void WriteCompactDirectoryReference(this BinaryArchiveWriter Writer, DirectoryReference Directory)
 		{
 			DirectoryItem Item = DirectoryItem.GetItemByDirectoryReference(Directory);
@@ -468,7 +468,7 @@ namespace UnrealBuildBase
 		/// <summary>
 		/// Reads a directory reference from a binary archive
 		/// </summary>
-		/// <param name="reader">Reader to serialize data from</param>
+		/// <param name="Reader">Reader to serialize data from</param>
 		/// <returns>New directory reference instance</returns>
 		public static DirectoryReference ReadCompactDirectoryReference(this BinaryArchiveReader Reader)
 		{
