Horde: Include the job and step it in artifact refs.

diff --git a/Engine/Source/Programs/Horde/Horde.Agent/Execution/JobExecutor.cs b/Engine/Source/Programs/Horde/Horde.Agent/Execution/JobExecutor.cs
--- a/Engine/Source/Programs/Horde/Horde.Agent/Execution/JobExecutor.cs
+++ b/Engine/Source/Programs/Horde/Horde.Agent/Execution/JobExecutor.cs
@@ -886,7 +886,7 @@ namespace Horde.Agent.Execution
 
 				using IRpcClientRef<JobRpc.JobRpcClient> jobRpc = await RpcConnection.GetClientRefAsync<JobRpc.JobRpcClient>(cancellationToken);
 				CreateJobArtifactResponse artifact = await jobRpc.Client.CreateArtifactAsync(new CreateJobArtifactRequest { JobId = _jobId, StepId = step.StepId, Name = "Output", Type = JobArtifactType.Output }, cancellationToken: cancellationToken);
-				_logger.LogInformation("Created artifact {ArtifactId} with ref {RefName} in ns {Namespace}", artifact.Id, artifact.RefName, artifact.NamespaceId);
+				logger.LogInformation("Created output artifact {ArtifactId} with ref {RefName} in namespace {Namespace}", artifact.Id, artifact.RefName, artifact.NamespaceId);
 
 				IStorageClient storage = _storageFactory.CreateStorageClient(_session, new NamespaceId(artifact.NamespaceId), artifact.Token);
 				using TreeWriter writer = new TreeWriter(storage, new RefName(artifact.RefName));
diff --git a/Engine/Source/Programs/Horde/Horde.Build/Artifacts/ArtifactCollection.cs b/Engine/Source/Programs/Horde/Horde.Build/Artifacts/ArtifactCollection.cs
--- a/Engine/Source/Programs/Horde/Horde.Build/Artifacts/ArtifactCollection.cs
+++ b/Engine/Source/Programs/Horde/Horde.Build/Artifacts/ArtifactCollection.cs
@@ -54,14 +54,14 @@ namespace Horde.Build.Artifacts
 			{
 			}
 
-			public Artifact(string name, ArtifactType type, IEnumerable<string> keys, NamespaceId namespaceId, RefName? refName, AclScopeName scopeName)
+			public Artifact(string name, ArtifactType type, IEnumerable<string> keys, NamespaceId namespaceId, string? refPrefix, AclScopeName scopeName)
 			{
 				Id = ArtifactId.GenerateNewId();
 				Name = name;
 				Type = type;
 				Keys.AddRange(keys);
 				NamespaceId = namespaceId;
-				RefName = refName ?? new RefName($"af-{Id}");
+				RefName = new RefName($"{refPrefix ?? ""}{Type}-{Id}");
 				AclScope = scopeName;
 			}
 		}
@@ -78,9 +78,9 @@ namespace Horde.Build.Artifacts
 		}
 
 		/// <inheritdoc/>
-		public async Task<IArtifact> AddAsync(string name, ArtifactType type, IEnumerable<string> keys, NamespaceId namespaceId, RefName? refName, AclScopeName scopeName, CancellationToken cancellationToken)
+		public async Task<IArtifact> AddAsync(string name, ArtifactType type, IEnumerable<string> keys, NamespaceId namespaceId, string? refPrefix, AclScopeName scopeName, CancellationToken cancellationToken)
 		{
-			Artifact artifact = new Artifact(name, type, keys, namespaceId, refName, scopeName);
+			Artifact artifact = new Artifact(name, type, keys, namespaceId, refPrefix, scopeName);
 			await _artifacts.InsertOneAsync(artifact, null, cancellationToken);
 			return artifact;
 		}
diff --git a/Engine/Source/Programs/Horde/Horde.Build/Artifacts/IArtifactCollection.cs b/Engine/Source/Programs/Horde/Horde.Build/Artifacts/IArtifactCollection.cs
--- a/Engine/Source/Programs/Horde/Horde.Build/Artifacts/IArtifactCollection.cs
+++ b/Engine/Source/Programs/Horde/Horde.Build/Artifacts/IArtifactCollection.cs
@@ -21,11 +21,11 @@ namespace Horde.Build.Artifacts
 		/// <param name="type">Type identifier for the artifact</param>
 		/// <param name="keys">Keys for the artifact</param>
 		/// <param name="namespaceId">Namespace containing the data</param>
-		/// <param name="refName">Name of the root reference. Will be initialized to the object id if not specified.</param>
+		/// <param name="refPrefix">Prefix for the ref name. Object id will be appended to this string.</param>
 		/// <param name="scopeName">Inherited scope used for permissions</param>
 		/// <param name="cancellationToken">Cancellation token for the operation</param>
 		/// <returns>The new log file document</returns>
-		Task<IArtifact> AddAsync(string name, ArtifactType type, IEnumerable<string> keys, NamespaceId namespaceId, RefName? refName, AclScopeName scopeName, CancellationToken cancellationToken = default);
+		Task<IArtifact> AddAsync(string name, ArtifactType type, IEnumerable<string> keys, NamespaceId namespaceId, string? refPrefix, AclScopeName scopeName, CancellationToken cancellationToken = default);
 
 		/// <summary>
 		/// Finds artifacts with the given keys
diff --git a/Engine/Source/Programs/Horde/Horde.Build/Jobs/JobRpcService.cs b/Engine/Source/Programs/Horde/Horde.Build/Jobs/JobRpcService.cs
--- a/Engine/Source/Programs/Horde/Horde.Build/Jobs/JobRpcService.cs
+++ b/Engine/Source/Programs/Horde/Horde.Build/Jobs/JobRpcService.cs
@@ -63,7 +63,7 @@ namespace Horde.Build.Jobs
 			}
 
 			NamespaceId namespaceId = Namespace.Artifacts;
-			IArtifact artifact = await _artifactCollection.AddAsync(request.Name, type, keys, namespaceId, null, templateConfig.ScopeName, context.CancellationToken);
+			IArtifact artifact = await _artifactCollection.AddAsync(request.Name, type, keys, namespaceId, $"job-{job.Id}/step-{step.Id}/", templateConfig.ScopeName, context.CancellationToken);
 
 			List<AclClaimConfig> claims = new List<AclClaimConfig>();
 			claims.Add(new AclClaimConfig(HordeClaimTypes.WriteNamespace, namespaceId.ToString()));
diff --git a/Engine/Source/Programs/Shared/EpicGames.Horde/Storage/Nodes/DirectoryNode.cs b/Engine/Source/Programs/Shared/EpicGames.Horde/Storage/Nodes/DirectoryNode.cs
--- a/Engine/Source/Programs/Shared/EpicGames.Horde/Storage/Nodes/DirectoryNode.cs
+++ b/Engine/Source/Programs/Shared/EpicGames.Horde/Storage/Nodes/DirectoryNode.cs
@@ -659,7 +659,7 @@ namespace EpicGames.Horde.Storage.Nodes
 			List<(DirectoryNode, FileInfo)> groupedFiles = new List<(DirectoryNode, FileInfo)>();
 			foreach (FileReference file in files)
 			{
-				DirectoryNode? node = FindOrAddDirectory(baseDir, file.Directory, dirToNode);
+				DirectoryNode? node = FindOrAddDirectory(file.Directory, baseDir, dirToNode);
 				if (node == null)
 				{
 					throw new InvalidOperationException($"File {file} is not under base directory {baseDir}");
@@ -670,7 +670,7 @@ namespace EpicGames.Horde.Storage.Nodes
 			await CopyFromDirectoryAsync(groupedFiles, options, writer, cancellationToken);
 		}
 
-		DirectoryNode? FindOrAddDirectory(DirectoryReference baseDir, DirectoryReference dir, Dictionary<DirectoryReference, DirectoryNode> dirToNode)
+		static DirectoryNode? FindOrAddDirectory(DirectoryReference dir, DirectoryReference baseDir, Dictionary<DirectoryReference, DirectoryNode> dirToNode)
 		{
 			DirectoryNode? node;
 			if (!dirToNode.TryGetValue(dir, out node))
@@ -678,7 +678,7 @@ namespace EpicGames.Horde.Storage.Nodes
 				DirectoryReference? parentDir = dir.ParentDirectory;
 				if (parentDir != null)
 				{
-					DirectoryNode? parentNode = FindOrAddDirectory(baseDir, parentDir, dirToNode);
+					DirectoryNode? parentNode = FindOrAddDirectory(parentDir, baseDir, dirToNode);
 					if (parentNode != null)
 					{
 						node = parentNode.AddDirectory(dir.GetDirectoryName());
