Updated UHT to support parsing "#if WITH_COREUOBJECT" statements to resolve issues with generated.h includes not being found.

#[fyi] tim.smith

diff --git a/Engine/Source/Programs/Shared/EpicGames.UHT/Parsers/UhtHeaderFileParser.cs b/Engine/Source/Programs/Shared/EpicGames.UHT/Parsers/UhtHeaderFileParser.cs
--- a/Engine/Source/Programs/Shared/EpicGames.UHT/Parsers/UhtHeaderFileParser.cs
+++ b/Engine/Source/Programs/Shared/EpicGames.UHT/Parsers/UhtHeaderFileParser.cs
@@ -87,10 +87,15 @@ namespace EpicGames.UHT.Parsers
 		/// </summary>
 		WithEngine = 1 << 7,
 
+		/// <summary>
+		/// This indicates we are in a "#if WITH_COREUOBJECT" block
+		/// </summary>
+		WithCoreUObject = 1 << 8,
+
 		/// <summary>
 		/// This directive is unrecognized and does not change the code generation at all
 		/// </summary>
-		Unrecognized = 1 << 8,
+		Unrecognized = 1 << 9,
 	}
 
 	/// <summary>
@@ -782,6 +787,9 @@ namespace EpicGames.UHT.Parsers
 					case UhtCompilerDirective.WithEngine:
 						PushCompilerDirective(UhtCompilerDirective.Unrecognized);
 						break;
+					case UhtCompilerDirective.WithCoreUObject:
+						PushCompilerDirective(UhtCompilerDirective.Unrecognized);
+						break;
 					case UhtCompilerDirective.WithHotReload:
 						throw new UhtException(TokenReader, directive.InputLine, "Can not use WITH_HOT_RELOAD with an #else clause");
 					default:
@@ -875,6 +883,10 @@ namespace EpicGames.UHT.Parsers
 					{
 						return notPresent ? UhtCompilerDirective.Unrecognized : UhtCompilerDirective.WithEngine;
 					}
+					else if (define.IsValue("WITH_COREUOBJECT"))
+					{
+						return notPresent ? UhtCompilerDirective.Unrecognized : UhtCompilerDirective.WithCoreUObject;
+					}
 					else if (define.IsValue("CPP"))
 					{
 						return notPresent ? UhtCompilerDirective.NotCPPBlock : UhtCompilerDirective.CPPBlock;
@@ -953,6 +965,7 @@ namespace EpicGames.UHT.Parsers
 				case UhtCompilerDirective.WithEditor:
 				case UhtCompilerDirective.WithEditorOnlyData:
 				case UhtCompilerDirective.WithEngine:
+				case UhtCompilerDirective.WithCoreUObject:
 				case UhtCompilerDirective.WithHotReload:
 					return true;
 
@@ -975,6 +988,7 @@ namespace EpicGames.UHT.Parsers
 				case UhtCompilerDirective.WithEditor:
 				case UhtCompilerDirective.WithEditorOnlyData:
 				case UhtCompilerDirective.WithEngine:
+				case UhtCompilerDirective.WithCoreUObject:
 				case UhtCompilerDirective.WithHotReload:
 					return true;
 
@@ -989,7 +1003,8 @@ namespace EpicGames.UHT.Parsers
 				compilerDirective == UhtCompilerDirective.WithEditor ||
 				compilerDirective == UhtCompilerDirective.WithEditorOnlyData ||
 				compilerDirective == UhtCompilerDirective.WithHotReload ||
-				compilerDirective == UhtCompilerDirective.WithEngine;
+				compilerDirective == UhtCompilerDirective.WithEngine ||
+				compilerDirective == UhtCompilerDirective.WithCoreUObject;
 		}
 
 		private static string GetCompilerDirectiveText(UhtCompilerDirective compilerDirective)
@@ -1004,6 +1019,7 @@ namespace EpicGames.UHT.Parsers
 				case UhtCompilerDirective.WithEditor: return "WITH_EDITOR";
 				case UhtCompilerDirective.WithEditorOnlyData: return "WITH_EDITORONLY_DATA";
 				case UhtCompilerDirective.WithEngine: return "WITH_ENGINE";
+				case UhtCompilerDirective.WithCoreUObject: return "WITH_COREUOBJECT";
 				default: return "<unrecognized>";
 			}
 		}
diff --git a/Engine/Source/Runtime/Datasmith/DatasmithCore/Public/DatasmithDefinitions.h b/Engine/Source/Runtime/Datasmith/DatasmithCore/Public/DatasmithDefinitions.h
--- a/Engine/Source/Runtime/Datasmith/DatasmithCore/Public/DatasmithDefinitions.h
+++ b/Engine/Source/Runtime/Datasmith/DatasmithCore/Public/DatasmithDefinitions.h
@@ -5,8 +5,12 @@
 #include "Misc/CoreMiscDefines.h"
 #include "Misc/EnumClassFlags.h"
 
-#ifdef WITH_COREUOBJECT
-	#include "UObject/ObjectMacros.h"
+#ifndef WITH_COREUOBJECT
+	#define WITH_COREUOBJECT 0
+#endif
+
+#if WITH_COREUOBJECT
+	#include "DatasmithDefinitions.generated.h"
 #else
 	#define UENUM(...)
 #endif // WITH_COREUOBJECT
