Fixed issue in C# UHT where if a structure referenced itself, it would end up in an infinite loop.

diff --git a/Engine/Source/Programs/Shared/EpicGames.UHT/Types/UhtClass.cs b/Engine/Source/Programs/Shared/EpicGames.UHT/Types/UhtClass.cs
--- a/Engine/Source/Programs/Shared/EpicGames.UHT/Types/UhtClass.cs
+++ b/Engine/Source/Programs/Shared/EpicGames.UHT/Types/UhtClass.cs
@@ -663,14 +663,6 @@ namespace EpicGames.UHT.Types
 			{
 				case UhtResolvePhase.Bases:
 					UhtClass? superClass = SuperClass;
-					if (superClass != null)
-					{
-						superClass.Resolve(UhtResolvePhase.Bases);
-					}
-					foreach (UhtStruct baseStruct in Bases)
-					{
-						baseStruct.Resolve(UhtResolvePhase.Bases);
-					}
 
 					// Force the MatchedSerializers on for anything being exported
 					if (!ClassExportFlags.HasAnyFlags(UhtClassExportFlags.NoExport))
diff --git a/Engine/Source/Programs/Shared/EpicGames.UHT/Types/UhtProperty.cs b/Engine/Source/Programs/Shared/EpicGames.UHT/Types/UhtProperty.cs
--- a/Engine/Source/Programs/Shared/EpicGames.UHT/Types/UhtProperty.cs
+++ b/Engine/Source/Programs/Shared/EpicGames.UHT/Types/UhtProperty.cs
@@ -2218,7 +2218,7 @@ namespace EpicGames.UHT.Types
 			if (hash == 0)
 			{
 				string type = referingType is UhtProperty ? "property" : "object";
-				referingType.LogWarning($"The {type} \"{referingType.SourceName}\" references type \"{obj.SourceName}\" but the code generation hash is zero");
+				referingType.LogError($"The {type} \"{referingType.SourceName}\" references type \"{obj.SourceName}\" but the code generation hash is zero.  Check for circular dependencies or missing includes.");
 			}
 			builder.Append(context.GetTypeHash(obj));
 			return builder;
diff --git a/Engine/Source/Programs/Shared/EpicGames.UHT/Types/UhtScriptStruct.cs b/Engine/Source/Programs/Shared/EpicGames.UHT/Types/UhtScriptStruct.cs
--- a/Engine/Source/Programs/Shared/EpicGames.UHT/Types/UhtScriptStruct.cs
+++ b/Engine/Source/Programs/Shared/EpicGames.UHT/Types/UhtScriptStruct.cs
@@ -724,7 +724,6 @@ namespace EpicGames.UHT.Types
 					UhtScriptStruct? superScriptStruct = SuperScriptStruct;
 					if (superScriptStruct != null)
 					{
-						superScriptStruct.Resolve(UhtResolvePhase.Bases);
 						ScriptStructFlags |= superScriptStruct.ScriptStructFlags & EStructFlags.Inherit;
 					}
 					break;
diff --git a/Engine/Source/Programs/Shared/EpicGames.UHT/Types/UhtType.cs b/Engine/Source/Programs/Shared/EpicGames.UHT/Types/UhtType.cs
--- a/Engine/Source/Programs/Shared/EpicGames.UHT/Types/UhtType.cs
+++ b/Engine/Source/Programs/Shared/EpicGames.UHT/Types/UhtType.cs
@@ -582,6 +582,29 @@ namespace EpicGames.UHT.Types
 		Final,
 	}
 
+	/// <summary>
+	/// Extension methods of UhtResolvePhase
+	/// </summary>
+	public static class UhtResolvePhaseExtensions
+	{
+
+		/// <summary>
+		/// Returns true if the phase should execute multi-threaded.
+		/// </summary>
+		/// <param name="resolvePhase">Phase to be tested</param>
+		/// <returns>True if the phase should execute multi-threaded</returns>
+		public static bool IsMultiThreadedResolvePhase(this UhtResolvePhase resolvePhase)
+		{
+			switch (resolvePhase)
+			{
+				case UhtResolvePhase.Final:
+					return false;
+				default:
+					return true;
+			}
+		}
+	}
+
 	/// <summary>
 	/// Flags for required documentation
 	/// </summary>
@@ -933,11 +956,14 @@ namespace EpicGames.UHT.Types
 			}
 			else if (currentState == initState)
 			{
-				do
+				if (phase.IsMultiThreadedResolvePhase())
 				{
-					Thread.Yield();
-					currentState = Interlocked.Add(ref _resolveState, 0);
-				} while (currentState == initState);
+					do
+					{
+						Thread.Yield();
+						currentState = Interlocked.Add(ref _resolveState, 0);
+					} while (currentState == initState);
+				}
 				return true;
 			}
 			else
diff --git a/Engine/Source/Programs/Shared/EpicGames.UHT/Utils/UhtSession.cs b/Engine/Source/Programs/Shared/EpicGames.UHT/Utils/UhtSession.cs
--- a/Engine/Source/Programs/Shared/EpicGames.UHT/Utils/UhtSession.cs
+++ b/Engine/Source/Programs/Shared/EpicGames.UHT/Utils/UhtSession.cs
@@ -1056,6 +1056,7 @@ namespace EpicGames.UHT.Utils
 				_referenceDeleteTask.Wait();
 			}
 
+			Log.Logger.LogTrace("Step - Starting exporters.");
 			StepExport();
 		}
 
@@ -1990,36 +1991,43 @@ namespace EpicGames.UHT.Utils
 
 		private void StepBindSuperAndBases()
 		{
-			StepForAllHeaders(headerFile => BindSuperAndBases(headerFile));
+			Log.Logger.LogTrace("Step - Bind super and bases");
+			StepForAllHeaders(headerFile => BindSuperAndBases(headerFile), true);
 		}
 
 		private void StepResolveBases()
 		{
-			StepForAllHeaders(headerFile => Resolve(headerFile, UhtResolvePhase.Bases));
+			Log.Logger.LogTrace("Step - Resolve bases");
+			ResolveAllHeaders(UhtResolvePhase.Bases);
 		}
 
 		private void StepResolveInvalidCheck()
 		{
-			StepForAllHeaders(headerFile => Resolve(headerFile, UhtResolvePhase.InvalidCheck));
+			Log.Logger.LogTrace("Step - Resolve invalid check");
+			ResolveAllHeaders(UhtResolvePhase.InvalidCheck);
 		}
 
 		private void StepResolveProperties()
 		{
-			StepForAllHeaders(headerFile => Resolve(headerFile, UhtResolvePhase.Properties));
+			Log.Logger.LogTrace("Step - Resolve properties");
+			ResolveAllHeaders(UhtResolvePhase.Properties);
 		}
 
 		private void StepResolveFinal()
 		{
-			StepForAllHeaders(headerFile => Resolve(headerFile, UhtResolvePhase.Final));
+			Log.Logger.LogTrace("Step - Resolve final");
+			ResolveAllHeaders(UhtResolvePhase.Final);
 		}
 
 		private void StepResolveValidate()
 		{
-			StepForAllHeaders(headerFile => UhtType.ValidateType(headerFile, UhtValidationOptions.None));
+			Log.Logger.LogTrace("Step - Resolve validate");
+			StepForAllHeaders(headerFile => UhtType.ValidateType(headerFile, UhtValidationOptions.None), true);
 		}
 
 		private void StepCollectReferences()
 		{
+			Log.Logger.LogTrace("Step - Collect references");
 			StepForAllHeaders(headerFile =>
 			{
 				foreach (UhtType child in headerFile.Children)
@@ -2030,7 +2038,12 @@ namespace EpicGames.UHT.Utils
 				{
 					headerFile.AddReferencedHeader(refHeaderFile);
 				}
-			});
+			}, true);
+		}
+
+		private void ResolveAllHeaders(UhtResolvePhase resolvePhase)
+		{
+			StepForAllHeaders(headerFile => Resolve(headerFile, resolvePhase), resolvePhase.IsMultiThreadedResolvePhase());
 		}
 
 		private void Resolve(UhtHeaderFile headerFile, UhtResolvePhase resolvePhase)
@@ -2059,14 +2072,14 @@ namespace EpicGames.UHT.Utils
 
 		private delegate void StepDelegate(UhtHeaderFile headerFile);
 
-		private void StepForAllHeaders(StepDelegate stepDelegate)
+		private void StepForAllHeaders(StepDelegate stepDelegate, bool allowGoWide)
 		{
 			if (HasErrors)
 			{
 				return;
 			}
 
-			if (GoWide)
+			if (GoWide && allowGoWide)
 			{
 				Parallel.ForEach(_headerFiles, headerFile =>
 				{
