EpicGames.Perforce: Fix parsing of records with multidimensional arrays (eg. filelog commands). This was breaking UGS in cases where a precompiled binaries zip file had been branched.

diff --git a/Engine/Source/Programs/Shared/EpicGames.Perforce/IPerforceOutput.cs b/Engine/Source/Programs/Shared/EpicGames.Perforce/IPerforceOutput.cs
--- a/Engine/Source/Programs/Shared/EpicGames.Perforce/IPerforceOutput.cs
+++ b/Engine/Source/Programs/Shared/EpicGames.Perforce/IPerforceOutput.cs
@@ -48,11 +48,12 @@ namespace EpicGames.Perforce
 		/// </summary>
 		static class ReadOnlyUtf8StringConstants
 		{
-			public static Utf8String Code = "code";
-			public static Utf8String Stat = "stat";
-			public static Utf8String Info = "info";
-			public static Utf8String Error = "error";
-			public static Utf8String Io = "io";
+			public static readonly Utf8String Code = "code";
+			public static readonly Utf8String Stat = "stat";
+			public static readonly Utf8String Info = "info";
+			public static readonly Utf8String Error = "error";
+			public static readonly Utf8String Io = "io";
+			public static readonly Utf8String Func = "func";
 		}
 
 		/// <summary>
@@ -630,10 +631,7 @@ namespace EpicGames.Perforce
 						}
 
 						// Check the suffix matches the index of the next element
-						if (!IsCorrectIndex(suffix, requiredSuffix, list.Count))
-						{
-							throw new PerforceException("Subobject element received out of order: got {0}", suffix);
-						}
+						ReadIndex(tag, suffix, requiredSuffix, list.Count);
 
 						// Add it to the list
 						if (!TryReadValue(buffer, ref bufferPos, newRecord, tagInfo))
@@ -654,15 +652,12 @@ namespace EpicGames.Perforce
 							throw new PerforceException($"Invalid field for {recordInfo.SubElementProperty.Name}");
 						}
 
-						// Check the suffix matches the index of the next element
-						if (!IsCorrectIndex(suffix, requiredSuffix, list.Count))
-						{
-							throw new PerforceException("Subobject element received out of order: got {0}", suffix);
-						}
+						// Find the next index
+						int subElementSuffixLength = ReadIndex(tag, suffix, requiredSuffix, list.Count);
 
 						// Parse the subobject and add it to the list
 						object? subRecord;
-						if (!TryReadTypedRecord(buffer, ref bufferPos, suffix, recordInfo.SubElementRecordInfo!, out subRecord))
+						if (!TryReadTypedRecord(buffer, ref bufferPos, suffix.Substring(0, subElementSuffixLength), recordInfo.SubElementRecordInfo!, out subRecord))
 						{
 							record = null;
 							return false;
@@ -679,6 +674,15 @@ namespace EpicGames.Perforce
 						}
 					}
 				}
+				else if (tag == ReadOnlyUtf8StringConstants.Func)
+				{
+					// Not sure why these fields are in the client output, but they are peppered into filelog results without an element index breaking the parser.
+					if (!TryReadValue(buffer, ref bufferPos, newRecord, null))
+					{
+						record = null;
+						return false;
+					}
+				}
 				else
 				{
 					// Roll back
@@ -721,21 +725,21 @@ namespace EpicGames.Perforce
 			// Parse the appropriate value
 			if (valueType == 's')
 			{
-				Utf8String @string;
-				if (!TryReadString(buffer, ref bufferPos, out @string))
+				Utf8String stringValue;
+				if (!TryReadString(buffer, ref bufferPos, out stringValue))
 				{
 					return false;
 				}
-				tagInfo?.ReadFromString(newRecord, @string);
+				tagInfo?.ReadFromString(newRecord, stringValue);
 			}
 			else if (valueType == 'i')
 			{
-				int integer;
-				if (!TryReadInt(bufferSpan, ref bufferPos, out integer))
+				int integerValue;
+				if (!TryReadInt(bufferSpan, ref bufferPos, out integerValue))
 				{
 					return false;
 				}
-				tagInfo?.ReadFromInteger(newRecord, integer);
+				tagInfo?.ReadFromInteger(newRecord, integerValue);
 			}
 			else
 			{
@@ -813,35 +817,33 @@ namespace EpicGames.Perforce
 		}
 
 		/// <summary>
-		/// Determines if the given text contains the expected prefix followed by an array index
+		/// Checks that the given field suffix starts with the given required suffix and index.
 		/// </summary>
-		/// <param name="text">The text to check</param>
-		/// <param name="prefix">The required prefix</param>
+		/// <param name="tag">The tag name</param>
+		/// <param name="suffix">The text to check</param>
+		/// <param name="requiredSuffix">The required prefix</param>
 		/// <param name="index">The required index</param>
 		/// <returns>True if the index is correct</returns>
-		static bool IsCorrectIndex(Utf8String text, Utf8String prefix, int index)
+		static int ReadIndex(Utf8String tag, Utf8String suffix, Utf8String requiredSuffix, int index)
 		{
-			if (prefix.Length > 0)
+			int valueStart = 0;
+			if (requiredSuffix.Length > 0)
 			{
-				return text.StartsWith(prefix) && text.Length > prefix.Length && text[prefix.Length] == (byte)',' && IsCorrectIndex(text.Span.Slice(prefix.Length + 1), index);
+				valueStart = requiredSuffix.Length + 1;
 			}
-			else
+
+			int value;
+			int valueLength;
+			if (Utf8Parser.TryParse(suffix.Substring(valueStart), out value, out valueLength))
 			{
-				return IsCorrectIndex(text.Span, index);
+				if (value == index && (valueStart + valueLength == suffix.Length || suffix[valueStart + valueLength] == (byte)','))
+				{
+					return valueStart + valueLength;
+				}
 			}
-		}
 
-		/// <summary>
-		/// Determines if the given text matches the expected array index
-		/// </summary>
-		/// <param name="span">The text to check</param>
-		/// <param name="expectedIndex">The expected array index</param>
-		/// <returns>True if the span matches</returns>
-		static bool IsCorrectIndex(ReadOnlySpan<byte> span, int expectedIndex)
-		{
-			int index;
-			int bytesConsumed;
-			return Utf8Parser.TryParse(span, out index, out bytesConsumed) && bytesConsumed == span.Length && index == expectedIndex;
+			string expectedSuffix = (requiredSuffix.Length > 0) ? $"{requiredSuffix},{index}" : $"{index}";
+			throw new PerforceException("Subobject element received out of order: got {0}{1}, expected {0}{2}", tag, suffix, expectedSuffix);
 		}
 	}
 }
