Fixed issue where C# UHT failed to generate error when interface function was marked as BlueprintPure.

diff --git a/Engine/Source/Programs/Shared/EpicGames.UHT/Types/UhtFunction.cs b/Engine/Source/Programs/Shared/EpicGames.UHT/Types/UhtFunction.cs
--- a/Engine/Source/Programs/Shared/EpicGames.UHT/Types/UhtFunction.cs
+++ b/Engine/Source/Programs/Shared/EpicGames.UHT/Types/UhtFunction.cs
@@ -483,7 +483,9 @@ namespace EpicGames.UHT.Types
 							}
 						}
 
-						if (FunctionType == UhtFunctionType.Function && FunctionExportFlags.HasAnyFlags(UhtFunctionExportFlags.DeclaredConst))
+						// If the function has already been marked as blueprint pure, don't bother.  This is important to being
+						// able to detect interfaces where the value has been specified.
+						if (FunctionType == UhtFunctionType.Function && FunctionExportFlags.HasAnyFlags(UhtFunctionExportFlags.DeclaredConst) && !FunctionFlags.HasAnyFlags(EFunctionFlags.BlueprintPure))
 						{
 							// @todo: the presence of const and one or more outputs does not guarantee that there are
 							// no side effects. On GCC and clang we could use __attribure__((pure)) or __attribute__((const))
