Fixed issue where C# UHT meta data did not issue value change errors in meta data.

diff --git a/Engine/Source/Programs/Shared/EpicGames.UHT/Parsers/UhtEnumParser.cs b/Engine/Source/Programs/Shared/EpicGames.UHT/Parsers/UhtEnumParser.cs
--- a/Engine/Source/Programs/Shared/EpicGames.UHT/Parsers/UhtEnumParser.cs
+++ b/Engine/Source/Programs/Shared/EpicGames.UHT/Parsers/UhtEnumParser.cs
@@ -176,6 +176,9 @@ namespace EpicGames.UHT.Parsers
 						}
 
 						// Save the new tag with a default value.  This will be replaced later
+						//COMPATIBILITY-TODO: If a enum value has a comment and a tooltip, it will generate an error.
+						// This is the reverse of all other parsing.  This code should be modified to process the comment
+						// after the UMETA.
 						int enumIndex = enumObject.AddEnumValue(tagToken.Value.ToString(), 0);
 						topScope.AddFormattedCommentsAsTooltipMetaData(enumIndex);
 
diff --git a/Engine/Source/Programs/Shared/EpicGames.UHT/Parsers/UhtSpecifierParser.cs b/Engine/Source/Programs/Shared/EpicGames.UHT/Parsers/UhtSpecifierParser.cs
--- a/Engine/Source/Programs/Shared/EpicGames.UHT/Parsers/UhtSpecifierParser.cs
+++ b/Engine/Source/Programs/Shared/EpicGames.UHT/Parsers/UhtSpecifierParser.cs
@@ -218,7 +218,7 @@ namespace EpicGames.UHT.Parsers
 			}
 
 			++_umetaElementsParsed;
-			_specifierContext.MetaData.Add(key.Value.ToString(), _specifierContext.MetaNameIndex, builder.ToString());
+			_specifierContext.MetaData.CheckedAdd(key.Value.ToString(), _specifierContext.MetaNameIndex, builder.ToString());
 		}
 
 		private void ParseKVPValueInternal()
diff --git a/Engine/Source/Programs/Shared/EpicGames.UHT/Specifiers/UhtDefaultSpecifiers.cs b/Engine/Source/Programs/Shared/EpicGames.UHT/Specifiers/UhtDefaultSpecifiers.cs
--- a/Engine/Source/Programs/Shared/EpicGames.UHT/Specifiers/UhtDefaultSpecifiers.cs
+++ b/Engine/Source/Programs/Shared/EpicGames.UHT/Specifiers/UhtDefaultSpecifiers.cs
@@ -61,19 +61,7 @@ namespace EpicGames.UHT.Parsers
 			// Get the trimmed string
 			value = new StringView(value, start, end - start);
 
-			// Make sure this isn't a duplicate assignment
-			string keyAsString = key.ToString();
-			if (specifierContext.MetaData.TryGetValue(keyAsString, out string? existingValue))
-			{
-				if (StringViewComparer.OrdinalIgnoreCase.Compare(existingValue, value) != 0)
-				{
-					specifierContext.MessageSite.LogError($"Metadata key '{key}' first seen with value '{existingValue}' then '{value}'");
-				}
-			}
-			else
-			{
-				specifierContext.MetaData.Add(keyAsString, specifierContext.MetaNameIndex, value.ToString());
-			}
+			specifierContext.MetaData.CheckedAdd(key.ToString(), specifierContext.MetaNameIndex, value.ToString());
 		}
 
 		private static void SetMetaData(UhtSpecifierContext specifierContext, StringView key, bool value)
diff --git a/Engine/Source/Programs/Shared/EpicGames.UHT/Types/UhtMetaData.cs b/Engine/Source/Programs/Shared/EpicGames.UHT/Types/UhtMetaData.cs
--- a/Engine/Source/Programs/Shared/EpicGames.UHT/Types/UhtMetaData.cs
+++ b/Engine/Source/Programs/Shared/EpicGames.UHT/Types/UhtMetaData.cs
@@ -421,17 +421,18 @@ namespace EpicGames.UHT.Types
 		/// <param name="value">Value of the meta data</param>
 		public void Add(string name, int nameIndex, string value)
 		{
-			if (Config != null)
-			{
-				if (Config.RedirectMetaDataKey(name, out string remappedName))
-				{
-					if (MessageSite != null)
-					{
-						MessageSite.LogWarning(LineNumber, $"Remapping old metadata key '{name}' to new key '{remappedName}', please update the declaration.");
-					}
-				}
-			}
-			GetDictionary()[new UhtMetaDataKey(name, nameIndex)] = value;
+			AddInternalWithRedirect(name, nameIndex, value, false);
+		}
+
+		/// <summary>
+		/// Add new meta data
+		/// </summary>
+		/// <param name="name">Name of the meta data key</param>
+		/// <param name="nameIndex">Index of the meta data key</param>
+		/// <param name="value">Value of the meta data</param>
+		public void CheckedAdd(string name, int nameIndex, string value)
+		{
+			AddInternalWithRedirect(name, nameIndex, value, true);
 		}
 
 		/// <summary>
@@ -481,9 +482,56 @@ namespace EpicGames.UHT.Types
 				SortedList<UhtMetaDataKey, string> dictionary = GetDictionary();
 				foreach (KeyValuePair<UhtMetaDataKey, string> kvp in metaData.Dictionary)
 				{
-					dictionary[kvp.Key] = kvp.Value;
+					AddInternal(dictionary, kvp.Key, kvp.Value, true);
+				}
+			}
+		}
+
+		/// <summary>
+		/// Add the given meta data but test for redirection of meta data name
+		/// </summary>
+		/// <param name="name">Name of the meta data</param>
+		/// <param name="nameIndex">Index of the meta data</param>
+		/// <param name="value">Value of the meta data</param>
+		/// <param name="checkValueChange">If true, verifies that the meta data value is not changing</param>
+		private void AddInternalWithRedirect(string name, int nameIndex, string value, bool checkValueChange)
+		{
+			if (Config != null)
+			{
+				if (Config.RedirectMetaDataKey(name, out string remappedName))
+				{
+					if (MessageSite != null)
+					{
+						MessageSite.LogWarning(LineNumber, $"Remapping old metadata key '{name}' to new key '{remappedName}', please update the declaration.");
+					}
+				}
+			}
+			AddInternal(GetDictionary(), new UhtMetaDataKey(name, nameIndex), value, checkValueChange);
+		}
+
+		/// <summary>
+		/// Add the value to the meta data
+		/// </summary>
+		/// <param name="dictionary">Destination dictionary</param>
+		/// <param name="key">Key of the meta data</param>
+		/// <param name="value">Value of the meta data</param>
+		/// <param name="checkValueChange">If true, verifies that the meta data value is not changing</param>
+		private void AddInternal(SortedList<UhtMetaDataKey, string> dictionary, UhtMetaDataKey key, string value, bool checkValueChange)
+		{
+			if (checkValueChange && dictionary.TryGetValue(key, out string? oldValue))
+			{
+				if (!oldValue.Equals(value, StringComparison.OrdinalIgnoreCase))
+				{
+					if (MessageSite != null)
+					{
+						MessageSite.LogError(LineNumber, $"Metadata key '{key.Name}' first seen with value '{oldValue}' then '{value}'");
+					}
 				}
 			}
+			else
+			{
+				dictionary[key] = value;
+			}
 		}
 
 		/// <summary>
