Horde: Include hash in directory nodes.

diff --git a/Engine/Source/Programs/Shared/EpicGames.Horde/Storage/Nodes/DirectoryNode.cs b/Engine/Source/Programs/Shared/EpicGames.Horde/Storage/Nodes/DirectoryNode.cs
--- a/Engine/Source/Programs/Shared/EpicGames.Horde/Storage/Nodes/DirectoryNode.cs
+++ b/Engine/Source/Programs/Shared/EpicGames.Horde/Storage/Nodes/DirectoryNode.cs
@@ -196,6 +196,7 @@ namespace EpicGames.Horde.Storage.Nodes
 			: base(reader)
 		{
 			_cachedLength = (long)reader.ReadUnsignedVarInt();
+			Hash = reader.ReadIoHash();
 		}
 
 		/// <summary>
@@ -207,6 +208,7 @@ namespace EpicGames.Horde.Storage.Nodes
 			base.Serialize(writer);
 
 			writer.WriteUnsignedVarInt((ulong)Length);
+			writer.WriteIoHash(Hash);
 		}
 
 		/// <inheritdoc/>
