Horde: Add a feedback reporting mechanism when adding large numbers of files to storage.

diff --git a/Engine/Source/Programs/Horde/Horde.Agent/Commands/Bundles/BundleCreate.cs b/Engine/Source/Programs/Horde/Horde.Agent/Commands/Bundles/BundleCreate.cs
--- a/Engine/Source/Programs/Horde/Horde.Agent/Commands/Bundles/BundleCreate.cs
+++ b/Engine/Source/Programs/Horde/Horde.Agent/Commands/Bundles/BundleCreate.cs
@@ -40,7 +40,7 @@ namespace Horde.Agent.Commands.Bundles
 			Stopwatch timer = Stopwatch.StartNew();
 
 			ChunkingOptions options = new ChunkingOptions();
-			await node.CopyFromDirectoryAsync(InputDir.ToDirectoryInfo(), options, writer, CancellationToken.None);
+			await node.CopyFromDirectoryAsync(InputDir.ToDirectoryInfo(), options, writer, null, CancellationToken.None);
 
 			await writer.WriteAsync(RefName, node);
 
diff --git a/Engine/Source/Programs/Horde/Horde.Agent/Commands/Vcs/VcsCommit.cs b/Engine/Source/Programs/Horde/Horde.Agent/Commands/Vcs/VcsCommit.cs
--- a/Engine/Source/Programs/Horde/Horde.Agent/Commands/Vcs/VcsCommit.cs
+++ b/Engine/Source/Programs/Horde/Horde.Agent/Commands/Vcs/VcsCommit.cs
@@ -75,7 +75,7 @@ namespace Horde.Agent.Commands.Vcs
 			await UpdateTreeAsync(reader, rootRef, rootDir, oldState, newState, files, directories);
 
 			using TreeWriter writer = new TreeWriter(store);
-			await DirectoryNode.CopyFromDirectoryAsync(files.ConvertAll(x => (x.Item1, x.Item2)), new ChunkingOptions(), writer, CancellationToken.None);
+			await DirectoryNode.CopyFromDirectoryAsync(files.ConvertAll(x => (x.Item1, x.Item2)), new ChunkingOptions(), writer, null, CancellationToken.None);
 
 			CommitNode newTip;
 			if (tip == null)
diff --git a/Engine/Source/Programs/Horde/Horde.Agent/Execution/JobExecutor.cs b/Engine/Source/Programs/Horde/Horde.Agent/Execution/JobExecutor.cs
--- a/Engine/Source/Programs/Horde/Horde.Agent/Execution/JobExecutor.cs
+++ b/Engine/Source/Programs/Horde/Horde.Agent/Execution/JobExecutor.cs
@@ -758,7 +758,7 @@ namespace Horde.Agent.Execution
 		{
 			if (_jobOptions.UseNewTempStorage ?? false)
 			{
-				await CreateArtifactAsync(stepId, type, baseDir, files.Select(x => x.Item2), cancellationToken);
+				await CreateArtifactAsync(stepId, type, baseDir, files.Select(x => x.Item2), logger, cancellationToken);
 			}
 			else
 			{
@@ -766,7 +766,7 @@ namespace Horde.Agent.Execution
 			}
 		}
 
-		protected async Task CreateArtifactAsync(string stepId, JobArtifactType type, DirectoryReference baseDir, IEnumerable<FileReference> files, CancellationToken cancellationToken)
+		protected async Task CreateArtifactAsync(string stepId, JobArtifactType type, DirectoryReference baseDir, IEnumerable<FileReference> files, ILogger logger, CancellationToken cancellationToken)
 		{
 			try
 			{
@@ -779,7 +779,7 @@ namespace Horde.Agent.Execution
 				using TreeWriter writer = new TreeWriter(storage, new RefName(artifact.RefName));
 
 				DirectoryNode dir = new DirectoryNode();
-				await dir.CopyFilesAsync(baseDir, files, new ChunkingOptions(), writer, cancellationToken);
+				await dir.CopyFilesAsync(baseDir, files, new ChunkingOptions(), writer, new CopyStatsLogger(logger), cancellationToken);
 
 				await writer.WriteAsync(new RefName(artifact.RefName), dir, cancellationToken: cancellationToken);
 			}
@@ -958,7 +958,7 @@ namespace Horde.Agent.Execution
 				RefName refName = TempStorage.GetRefNameForNode(_storagePrefix, step.Name);
 
 				TreeOptions treeOptions = new TreeOptions();
-				using TreeWriter treeWriter = new TreeWriter(storage, treeOptions, refName.Text, traceLogger: logger);
+				using TreeWriter treeWriter = new TreeWriter(storage, treeOptions, refName.Text);
 
 				DirectoryNode outputNode = new DirectoryNode();
 
@@ -1400,7 +1400,7 @@ namespace Horde.Agent.Execution
 				List<FileReference> artifactFiles = DirectoryReference.EnumerateFiles(logDir, "*", SearchOption.AllDirectories).ToList();
 				if (_jobOptions.UseNewTempStorage ?? false)
 				{
-					await CreateArtifactAsync(step.StepId, JobArtifactType.Saved, workspaceDir, artifactFiles, cancellationToken);
+					await CreateArtifactAsync(step.StepId, JobArtifactType.Saved, workspaceDir, artifactFiles, jobLogger, cancellationToken);
 				}
 				else
 				{
diff --git a/Engine/Source/Programs/Horde/Horde.Agent/Utility/TempStorage.cs b/Engine/Source/Programs/Horde/Horde.Agent/Utility/TempStorage.cs
--- a/Engine/Source/Programs/Horde/Horde.Agent/Utility/TempStorage.cs
+++ b/Engine/Source/Programs/Horde/Horde.Agent/Utility/TempStorage.cs
@@ -632,7 +632,7 @@ namespace Horde.Storage.Utility
 
 			// Create the file tree
 			DirectoryNode rootNode = new DirectoryNode();
-			await rootNode.CopyFilesAsync(workspaceDir, Enumerable.Concat(files, new[] { manifestLocation }), new ChunkingOptions(), writer, cancellationToken);
+			await rootNode.CopyFilesAsync(workspaceDir, Enumerable.Concat(files, new[] { manifestLocation }), new ChunkingOptions(), writer, new CopyStatsLogger(logger), cancellationToken);
 
 			return new DirectoryEntry(blockName, rootNode);
 		}
diff --git a/Engine/Source/Programs/Horde/Horde.Build/Commands/Bundles/CreateCommand.cs b/Engine/Source/Programs/Horde/Horde.Build/Commands/Bundles/CreateCommand.cs
--- a/Engine/Source/Programs/Horde/Horde.Build/Commands/Bundles/CreateCommand.cs
+++ b/Engine/Source/Programs/Horde/Horde.Build/Commands/Bundles/CreateCommand.cs
@@ -38,7 +38,7 @@ namespace Horde.Build.Commands.Bundles
 			using TreeWriter writer = new TreeWriter(store, prefix: RefName.Text);
 
 			DirectoryNode node = new DirectoryNode(DirectoryFlags.None);
-			await node.CopyFromDirectoryAsync(InputDir.ToDirectoryInfo(), new ChunkingOptions(), writer, CancellationToken.None);
+			await node.CopyFromDirectoryAsync(InputDir.ToDirectoryInfo(), new ChunkingOptions(), writer, new CopyStatsLogger(logger), CancellationToken.None);
 
 			await writer.WriteAsync(RefName, node);
 			return 0;
diff --git a/Engine/Source/Programs/Shared/EpicGames.Horde/Storage/Nodes/DirectoryNode.cs b/Engine/Source/Programs/Shared/EpicGames.Horde/Storage/Nodes/DirectoryNode.cs
--- a/Engine/Source/Programs/Shared/EpicGames.Horde/Storage/Nodes/DirectoryNode.cs
+++ b/Engine/Source/Programs/Shared/EpicGames.Horde/Storage/Nodes/DirectoryNode.cs
@@ -4,6 +4,7 @@ using EpicGames.Core;
 using Microsoft.Extensions.Logging;
 using System;
 using System.Collections.Generic;
+using System.Diagnostics;
 using System.Diagnostics.CodeAnalysis;
 using System.IO;
 using System.Linq;
@@ -273,6 +274,51 @@ namespace EpicGames.Horde.Storage.Nodes
 		public override string ToString() => Name.ToString();
 	}
 
+	/// <summary>
+	/// Stats reported for copy operations
+	/// </summary>
+	public interface ICopyStats
+	{
+		/// <summary>
+		/// Number of files that have been copied
+		/// </summary>
+		int CopiedCount { get; }
+
+		/// <summary>
+		/// Total size of data to be copied
+		/// </summary>
+		long CopiedSize { get; }
+
+		/// <summary>
+		/// Total number of files to copy
+		/// </summary>
+		int TotalCount { get; }
+
+		/// <summary>
+		/// Total size of data to copy
+		/// </summary>
+		long TotalSize { get; }
+	}
+
+	/// <summary>
+	/// Progress logger for writing copy stats
+	/// </summary>
+	public class CopyStatsLogger : IProgress<ICopyStats>
+	{
+		readonly ILogger _logger;
+
+		/// <summary>
+		/// Constructor
+		/// </summary>
+		public CopyStatsLogger(ILogger logger) => _logger = logger;
+
+		/// <inheritdoc/>
+		public void Report(ICopyStats stats)
+		{
+			_logger.LogInformation("Copied {NumFiles:n0}/{TotalFiles:n0} ({Size:n1}/{TotalSize:n1}mb, {Pct}%)", stats.CopiedCount, stats.TotalCount, stats.CopiedSize / 1024.0, stats.TotalSize / 1024.0, (int)((Math.Max(stats.CopiedCount, 1) * 100) / Math.Max(stats.TotalCount, 1)));
+		}
+	}
+
 	/// <summary>
 	/// Flags for a directory node
 	/// </summary>
@@ -651,6 +697,51 @@ namespace EpicGames.Horde.Storage.Nodes
 
 		#endregion
 
+		/// <summary>
+		/// Reports progress info back to callers
+		/// </summary>
+		class CopyStats : ICopyStats
+		{
+			readonly object _lockObject = new object();
+			readonly Stopwatch _timer = Stopwatch.StartNew();
+			readonly IProgress<ICopyStats> _progress;
+
+			public int CopiedCount { get; set; }
+			public long CopiedSize { get; set; }
+			public int TotalCount { get; }
+			public long TotalSize { get; }
+
+			public CopyStats(int totalCount, long totalSize, IProgress<ICopyStats> progress)
+			{
+				TotalCount = totalCount;
+				TotalSize = totalSize;
+				_progress = progress;
+			}
+
+			public void Update(int count, long size)
+			{
+				lock (_lockObject)
+				{
+					CopiedCount += count;
+					CopiedSize += size;
+					if (_timer.Elapsed > TimeSpan.FromSeconds(10.0))
+					{
+						_progress.Report(this);
+						_timer.Restart();
+					}
+				}
+			}
+
+			public void Flush()
+			{
+				lock (_lockObject)
+				{
+					_progress.Report(this);
+					_timer.Restart();
+				}
+			}
+		}
+
 		/// <summary>
 		/// Adds files from a flat list of paths
 		/// </summary>
@@ -658,8 +749,9 @@ namespace EpicGames.Horde.Storage.Nodes
 		/// <param name="files">Files to add</param>
 		/// <param name="options">Options for chunking file content</param>
 		/// <param name="writer">Writer for new node data</param>
+		/// <param name="progress">Feedback interface for progress updates</param>
 		/// <param name="cancellationToken">Cancellation token for the operation</param>
-		public async Task CopyFilesAsync(DirectoryReference baseDir, IEnumerable<FileReference> files, ChunkingOptions options, TreeWriter writer, CancellationToken cancellationToken)
+		public async Task CopyFilesAsync(DirectoryReference baseDir, IEnumerable<FileReference> files, ChunkingOptions options, TreeWriter writer, IProgress<ICopyStats>? progress, CancellationToken cancellationToken)
 		{
 			Dictionary<DirectoryReference, DirectoryNode> dirToNode = new Dictionary<DirectoryReference, DirectoryNode>();
 			dirToNode.Add(baseDir, this);
@@ -675,7 +767,7 @@ namespace EpicGames.Horde.Storage.Nodes
 				groupedFiles.Add((node, file.ToFileInfo()));
 			}
 
-			await CopyFromDirectoryAsync(groupedFiles, options, writer, cancellationToken);
+			await CopyFromDirectoryAsync(groupedFiles, options, writer, progress, cancellationToken);
 		}
 
 		static DirectoryNode? FindOrAddDirectory(DirectoryReference dir, DirectoryReference baseDir, Dictionary<DirectoryReference, DirectoryNode> dirToNode)
@@ -703,14 +795,15 @@ namespace EpicGames.Horde.Storage.Nodes
 		/// <param name="directoryInfo"></param>
 		/// <param name="options">Options for chunking file content</param>
 		/// <param name="writer">Writer for new node data</param>
+		/// <param name="progress">Feedback interface for progress updates</param>
 		/// <param name="cancellationToken">Cancellation token for the operation</param>
 		/// <returns></returns>
-		public async Task CopyFromDirectoryAsync(DirectoryInfo directoryInfo, ChunkingOptions options, TreeWriter writer, CancellationToken cancellationToken)
+		public async Task CopyFromDirectoryAsync(DirectoryInfo directoryInfo, ChunkingOptions options, TreeWriter writer, IProgress<ICopyStats>? progress, CancellationToken cancellationToken)
 		{
 			// Enumerate all the files below this directory
 			List<(DirectoryNode DirectoryNode, FileInfo FileInfo)> files = new List<(DirectoryNode, FileInfo)>();
 			FindFilesToCopy(directoryInfo, files);
-			await CopyFromDirectoryAsync(files, options, writer, cancellationToken);
+			await CopyFromDirectoryAsync(files, options, writer, progress, cancellationToken);
 		}
 
 		/// <summary>
@@ -719,9 +812,10 @@ namespace EpicGames.Horde.Storage.Nodes
 		/// <param name="files"></param>
 		/// <param name="options">Options for chunking file content</param>
 		/// <param name="writer">Writer for new node data</param>
+		/// <param name="progress">Progress notification object</param>
 		/// <param name="cancellationToken">Cancellation token for the operation</param>
 		/// <returns></returns>
-		public static async Task CopyFromDirectoryAsync(List<(DirectoryNode DirectoryNode, FileInfo FileInfo)> files, ChunkingOptions options, TreeWriter writer, CancellationToken cancellationToken)
+		public static async Task CopyFromDirectoryAsync(List<(DirectoryNode DirectoryNode, FileInfo FileInfo)> files, ChunkingOptions options, TreeWriter writer, IProgress<ICopyStats>? progress, CancellationToken cancellationToken)
 		{
 			const int MaxWriters = 32;
 			const long MinSizePerWriter = 1024 * 1024;
@@ -730,6 +824,13 @@ namespace EpicGames.Horde.Storage.Nodes
 			long totalSize = files.Sum(x => x.Item2.Length);
 			long chunkSize = Math.Max(MinSizePerWriter, totalSize / MaxWriters);
 
+			// Create the progress reporting object
+			CopyStats? copyStats = null;
+			if (progress != null)
+			{
+				copyStats = new CopyStats(files.Count, totalSize, progress);
+			}
+
 			List<Task> tasks = new List<Task>();
 			long currentSize = 0;
 			long targetSize = chunkSize;
@@ -748,7 +849,7 @@ namespace EpicGames.Horde.Storage.Nodes
 				}
 
 				int minIdxCopy = minIdx;
-				tasks.Add(Task.Run(() => CopyFilesAsync(files, minIdxCopy, maxIdx, fileEntries, options, writer, cancellationToken), cancellationToken));
+				tasks.Add(Task.Run(() => CopyFilesAsync(files, minIdxCopy, maxIdx, fileEntries, options, writer, copyStats, cancellationToken), cancellationToken));
 
 				targetSize += chunkSize;
 				minIdx = maxIdx;
@@ -762,6 +863,9 @@ namespace EpicGames.Horde.Storage.Nodes
 			{
 				files[idx].DirectoryNode.AddFile(fileEntries[idx]);
 			}
+
+			// Write the final stats
+			copyStats?.Flush();
 		}
 
 		void FindFilesToCopy(DirectoryInfo directoryInfo, List<(DirectoryNode, FileInfo)> files)
@@ -776,7 +880,7 @@ namespace EpicGames.Horde.Storage.Nodes
 			}
 		}
 
-		static async Task CopyFilesAsync(List<(DirectoryNode DirectoryNode, FileInfo FileInfo)> files, int minIdx, int maxIdx, FileEntry[] entries, ChunkingOptions options, TreeWriter baseWriter, CancellationToken cancellationToken)
+		static async Task CopyFilesAsync(List<(DirectoryNode DirectoryNode, FileInfo FileInfo)> files, int minIdx, int maxIdx, FileEntry[] entries, ChunkingOptions options, TreeWriter baseWriter, CopyStats? copyStats, CancellationToken cancellationToken)
 		{
 			using TreeWriter writer = new TreeWriter(baseWriter);
 
@@ -786,6 +890,7 @@ namespace EpicGames.Horde.Storage.Nodes
 				FileInfo fileInfo = files[idx].FileInfo;
 				NodeHandle handle = await fileNodeWriter.CreateAsync(fileInfo, cancellationToken);
 				entries[idx] = new FileEntry(fileInfo.Name, FileEntryFlags.None, fileNodeWriter.Length, handle);
+				copyStats?.Update(1, fileNodeWriter.Length);
 			}
 			await writer.FlushAsync(cancellationToken);
 		}
