Horde: Fix incorrect size reporting when copying data to bundles.

diff --git a/Engine/Source/Programs/Shared/EpicGames.Horde/Storage/Nodes/DirectoryNode.cs b/Engine/Source/Programs/Shared/EpicGames.Horde/Storage/Nodes/DirectoryNode.cs
--- a/Engine/Source/Programs/Shared/EpicGames.Horde/Storage/Nodes/DirectoryNode.cs
+++ b/Engine/Source/Programs/Shared/EpicGames.Horde/Storage/Nodes/DirectoryNode.cs
@@ -315,7 +315,7 @@ namespace EpicGames.Horde.Storage.Nodes
 		/// <inheritdoc/>
 		public void Report(ICopyStats stats)
 		{
-			_logger.LogInformation("Copied {NumFiles:n0}/{TotalFiles:n0} ({Size:n1}/{TotalSize:n1}mb, {Pct}%)", stats.CopiedCount, stats.TotalCount, stats.CopiedSize / 1024.0, stats.TotalSize / 1024.0, (int)((Math.Max(stats.CopiedCount, 1) * 100) / Math.Max(stats.TotalCount, 1)));
+			_logger.LogInformation("Copied {NumFiles:n0}/{TotalFiles:n0} ({Size:n1}/{TotalSize:n1}mb, {Pct}%)", stats.CopiedCount, stats.TotalCount, stats.CopiedSize / (1024.0 * 1024.0), stats.TotalSize / (1024.0 * 1024.0), (int)((Math.Max(stats.CopiedCount, 1) * 100) / Math.Max(stats.TotalCount, 1)));
 		}
 	}
 
