Added the "Build" directory as locations UBT will scan for UBT plugins.

diff --git a/Engine/Source/Programs/Shared/EpicGames.Build/System/EnumeratePlugins.cs b/Engine/Source/Programs/Shared/EpicGames.Build/System/EnumeratePlugins.cs
--- a/Engine/Source/Programs/Shared/EpicGames.Build/System/EnumeratePlugins.cs
+++ b/Engine/Source/Programs/Shared/EpicGames.Build/System/EnumeratePlugins.cs
@@ -119,10 +119,12 @@ namespace UnrealBuildBase
 		public static List<FileReference> EnumerateUbtPlugins(FileReference? ProjectFile)
 		{
 			List<DirectoryReference> PluginDirectories = new List<DirectoryReference>();
+			PluginDirectories.AddRange(Unreal.GetExtensionDirs(Unreal.EngineDirectory, "Build"));
 			PluginDirectories.AddRange(Unreal.GetExtensionDirs(Unreal.EngineDirectory, "Source"));
 			PluginDirectories.AddRange(Unreal.GetExtensionDirs(Unreal.EngineDirectory, "Plugins"));
 			if (ProjectFile != null)
 			{
+				PluginDirectories.AddRange(Unreal.GetExtensionDirs(ProjectFile.Directory, "Build"));
 				PluginDirectories.AddRange(Unreal.GetExtensionDirs(ProjectFile.Directory, "Source"));
 				PluginDirectories.AddRange(Unreal.GetExtensionDirs(ProjectFile.Directory, "Plugins"));
 				PluginDirectories.AddRange(Unreal.GetExtensionDirs(ProjectFile.Directory, "Mods"));
