EpicGames.Core: Use .NET function to split reprocessed command line args, rather than doing so manually, so results match how the system processes quotes

#rnx

diff --git a/Engine/Source/Programs/Shared/EpicGames.Core/CommandLineArguments.cs b/Engine/Source/Programs/Shared/EpicGames.Core/CommandLineArguments.cs
--- a/Engine/Source/Programs/Shared/EpicGames.Core/CommandLineArguments.cs
+++ b/Engine/Source/Programs/Shared/EpicGames.Core/CommandLineArguments.cs
@@ -9,6 +9,7 @@ using System.Diagnostics.CodeAnalysis;
 using System.Linq;
 using System.Reflection;
 using System.Text;
+using Microsoft.CodeAnalysis;
 using Microsoft.Extensions.Logging;
 
 #pragma warning disable CA1710 // Identifiers should have correct suffix
@@ -966,33 +967,7 @@ namespace EpicGames.Core
 		/// <returns>Array of arguments</returns>
 		public static string[] Split(string commandLine)
 		{
-			StringBuilder argument = new StringBuilder();
-
-			List<string> arguments = new List<string>();
-			for(int idx = 0; idx < commandLine.Length; idx++)
-			{
-				if(!Char.IsWhiteSpace(commandLine[idx]))
-				{
-					argument.Clear();
-					for(bool bInQuotes = false; idx < commandLine.Length; idx++)
-					{
-						if(commandLine[idx] == '\"')
-						{
-							bInQuotes ^= true;
-						}
-						else if(!bInQuotes && Char.IsWhiteSpace(commandLine[idx]))
-						{
-							break;
-						}
-						else
-						{
-							argument.Append(commandLine[idx]);
-						}
-					}
-					arguments.Add(argument.ToString());
-				}
-			}
-			return arguments.ToArray();
+			return CommandLineParser.SplitCommandLineIntoArguments(commandLine, true).ToArray();
 		}
 
 		/// <summary>
diff --git a/Engine/Source/Programs/Shared/EpicGames.Core/EpicGames.Core.csproj b/Engine/Source/Programs/Shared/EpicGames.Core/EpicGames.Core.csproj
--- a/Engine/Source/Programs/Shared/EpicGames.Core/EpicGames.Core.csproj
+++ b/Engine/Source/Programs/Shared/EpicGames.Core/EpicGames.Core.csproj
@@ -49,6 +49,7 @@
   </Choose>
   <ItemGroup>
     <PackageReference Include="JetBrains.Annotations" Version="2022.1.0" />
+    <PackageReference Include="Microsoft.CodeAnalysis.Common" Version="3.11.0" />
     <PackageReference Include="Microsoft.Extensions.Logging" Version="6.0.0" />
     <PackageReference Include="Microsoft.Extensions.Logging.Abstractions" Version="6.0.1" />
     <PackageReference Include="Microsoft.Extensions.ObjectPool" Version="6.0.8" />
