Suppress warning that was inadvertently enabled until it can be addressed

#rnx

diff --git a/Engine/Source/Programs/Shared/EpicGames.BuildGraph/BgCompiler.cs b/Engine/Source/Programs/Shared/EpicGames.BuildGraph/BgCompiler.cs
--- a/Engine/Source/Programs/Shared/EpicGames.BuildGraph/BgCompiler.cs
+++ b/Engine/Source/Programs/Shared/EpicGames.BuildGraph/BgCompiler.cs
@@ -34,6 +34,7 @@ namespace EpicGames.BuildGraph
 
 			public new bool Equals([AllowNull] object x, [AllowNull] object y) => ReferenceEquals(x, y);
 
+			[SuppressMessage("MicrosoftCodeAnalysisCorrectness", "RS1024:Compare symbols correctly", Justification = "<Pending>")]
 			public int GetHashCode([DisallowNull] object obj) => RuntimeHelpers.GetHashCode(obj);
 		}
 
