UnrealBuildTool: Special case -define: command line args to not remove quotes if only the value is quoted

#rnx

diff --git a/Engine/Source/Programs/Shared/EpicGames.Core/CommandLineArguments.cs b/Engine/Source/Programs/Shared/EpicGames.Core/CommandLineArguments.cs
--- a/Engine/Source/Programs/Shared/EpicGames.Core/CommandLineArguments.cs
+++ b/Engine/Source/Programs/Shared/EpicGames.Core/CommandLineArguments.cs
@@ -969,6 +969,7 @@ namespace EpicGames.Core
 			StringBuilder argument = new StringBuilder();
 
 			List<string> arguments = new List<string>();
+			// First do a pass leaving all quotes in the arguments, they will be removed later
 			for(int idx = 0; idx < commandLine.Length; idx++)
 			{
 				if(!Char.IsWhiteSpace(commandLine[idx]))
@@ -976,7 +977,7 @@ namespace EpicGames.Core
 					argument.Clear();
 					for(bool bInQuotes = false; idx < commandLine.Length; idx++)
 					{
-						if(commandLine[idx] == '\"')
+						if (commandLine[idx] == '\"')
 						{
 							bInQuotes ^= true;
 						}
@@ -984,14 +985,22 @@ namespace EpicGames.Core
 						{
 							break;
 						}
-						else
-						{
-							argument.Append(commandLine[idx]);
-						}
+						argument.Append(commandLine[idx]);
 					}
 					arguments.Add(argument.ToString());
 				}
 			}
+
+			// Remove quotes from arguments except where only the value is quoted in -Define (-Define:KEY="VALUE")
+			for (int idx = 0; idx < arguments.Count; idx++)
+			{
+				string arg = arguments[idx];
+				if (arg.StartsWith("-Define:", StringComparison.OrdinalIgnoreCase) && !arg.StartsWith("-Define:\"", StringComparison.OrdinalIgnoreCase))
+				{
+					continue;
+				}
+				arguments[idx] = arg.Replace("\"", "");
+			}
 			return arguments.ToArray();
 		}
 
