Minor static analyzer fixes.

diff --git a/Engine/Source/Programs/Shared/EpicGames.Core/BuzHash.cs b/Engine/Source/Programs/Shared/EpicGames.Core/BuzHash.cs
--- a/Engine/Source/Programs/Shared/EpicGames.Core/BuzHash.cs
+++ b/Engine/Source/Programs/Shared/EpicGames.Core/BuzHash.cs
@@ -1,9 +1,7 @@
 // Copyright Epic Games, Inc. All Rights Reserved.
 
 using System;
-using System.Buffers;
 using System.Diagnostics;
-using System.Runtime.InteropServices;
 
 namespace EpicGames.Core
 {
diff --git a/Engine/Source/Programs/Shared/EpicGames.Core/Log.cs b/Engine/Source/Programs/Shared/EpicGames.Core/Log.cs
--- a/Engine/Source/Programs/Shared/EpicGames.Core/Log.cs
+++ b/Engine/Source/Programs/Shared/EpicGames.Core/Log.cs
@@ -1,7 +1,6 @@
 // Copyright Epic Games, Inc. All Rights Reserved.
 
 using System;
-using System.Buffers;
 using System.Collections.Generic;
 using System.Diagnostics;
 using System.IO;
diff --git a/Engine/Source/Programs/Shared/EpicGames.Jupiter/JupierUtils.cs b/Engine/Source/Programs/Shared/EpicGames.Jupiter/JupierUtils.cs
deleted file mode 100644
diff --git a/Engine/Source/Programs/Shared/EpicGames.Jupiter/JupiterTree.cs b/Engine/Source/Programs/Shared/EpicGames.Jupiter/JupiterTree.cs
--- a/Engine/Source/Programs/Shared/EpicGames.Jupiter/JupiterTree.cs
+++ b/Engine/Source/Programs/Shared/EpicGames.Jupiter/JupiterTree.cs
@@ -29,7 +29,6 @@ namespace EpicGames.Jupiter
 			LastTreeHash = null;
 		}
 
-
 		public List<JupiterTree> Trees { get; } = new List<JupiterTree>();
 
 		public IEnumerable<string> ContentHashes
@@ -57,7 +56,9 @@ namespace EpicGames.Jupiter
 		public byte[] CalculateTreeHash()
 		{
 			if (LastTreeHash != null)
+			{
 				return LastTreeHash;
+			}
 
 			List<byte> Bytes = new List<byte>();
 			foreach (JupiterTree Tree in Trees)
@@ -190,7 +191,9 @@ namespace EpicGames.Jupiter
 						{
 							// check if blob is already present
 							if (FoundBlobs.Contains(BlobHash))
+							{
 								continue;
+							}
 
 							SubmitTasks.Add(Task.Run(async () =>
 								{
@@ -278,7 +281,6 @@ namespace EpicGames.Jupiter
 							throw new Exception(string.Format("Failed to download tree {0}. Response: {1}", Tree, GetTreeResultString));
 						}
 
-
 						TreeContents TreeContents = JsonSerializer.Deserialize<TreeContents>(GetTreeResultString);
 						TreeMapping.Add(Tree, TreeContents);
 					}
@@ -365,10 +367,8 @@ namespace EpicGames.Jupiter
 
 			OnContentChanged?.Invoke(this, new EventArgs());
 		}
-
 	}
 
-
 	public class JupiterStreamedContentProvider : JupiterTreeContentProvider
 	{
 		private class ChunkDescription
@@ -390,7 +390,6 @@ namespace EpicGames.Jupiter
 			OnContentChanged?.Invoke(this, new EventArgs());
 		}
 
-
 		public override Task<byte[]> GetContent(string Sha1)
 		{
 			ChunkDescription Chunk;
@@ -427,7 +426,6 @@ namespace EpicGames.Jupiter
 
 			return NewChunks;
 		}
-
 	}
 
 	public class JupiterInMemoryContentProvider : JupiterTreeContentProvider
@@ -437,7 +435,9 @@ namespace EpicGames.Jupiter
 		public void AddContent(byte[] Sha1, byte[] Content)
 		{
 			if (Sha1.Length != 20)
+			{
 				throw new ArgumentException("Sha1 argument was not 20 bytes, did you really specify a hash of the contents?", nameof(Sha1));
+			}
 
 			string Hash = SHA1Utils.FormatAsHexString(Sha1);
 			ContentBlobs.Add(Hash, Content);
@@ -464,7 +464,6 @@ namespace EpicGames.Jupiter
 		public abstract event EventHandler OnContentChanged;
 	}
 
-
 	/// <summary>
 	/// A merkel tree built from a set of files ready to be uploaded to Jupiter
 	/// </summary>
@@ -527,7 +526,6 @@ namespace EpicGames.Jupiter
 				FileToChunkMapping.Add(File, ContentProvider.GetHashes().ToList());
 			}
 
-
 			// The manifest is the first node in the tree so that we can find it later
 			JupiterTopTree.AddTree(Manifest.AsTree());
 			JupiterTopTree.AddTrees(FileTrees);
diff --git a/Engine/Source/Programs/Shared/EpicGames.Jupiter/JupiterUtils.cs b/Engine/Source/Programs/Shared/EpicGames.Jupiter/JupiterUtils.cs
new file mode 100644
--- /dev/null
+++ b/Engine/Source/Programs/Shared/EpicGames.Jupiter/JupiterUtils.cs
@@ -0,0 +1,111 @@
+// Copyright Epic Games, Inc. All Rights Reserved.
+
+using EpicGames.Core;
+using System;
+using System.Collections.Generic;
+using System.IO;
+using System.Security.Cryptography;
+using System.Text;
+
+namespace EpicGames.Jupiter
+{
+
+	public static class JupiterFileUtils
+	{
+		public static List<Tuple<byte[], byte[]>> ChunkFile(FileReference File)
+		{
+			List<Tuple<byte[], byte[]>> Chunks = new List<Tuple<byte[], byte[]>>();
+			const int ChunkSize = 10 * 1024 * 1024; // chunk size in bytes
+			const int BlockSize = 4 * 1024; // the size of each block we read
+			byte[] FileBlock = new byte[BlockSize];
+
+			using (FileStream FileStream = new FileStream(File.FullName, FileMode.Open))
+			{
+				bool MoreDataAvailable = true;
+				while (MoreDataAvailable)
+				{
+					int TotalReadCountOfBytes = 0;
+					long RemainingFileSize = FileStream.Length - FileStream.Position;
+					if (RemainingFileSize == 0)
+					{
+						MoreDataAvailable = false;
+						continue;
+					}
+
+					int CurrentChunkSize = Math.Min((int)RemainingFileSize, ChunkSize);
+					byte[] Chunk = new byte[CurrentChunkSize];
+
+					while (TotalReadCountOfBytes < CurrentChunkSize)
+					{
+						int ReadCountOfBytes = FileStream.Read(FileBlock, 0, BlockSize);
+						if (ReadCountOfBytes == 0)
+						{
+							MoreDataAvailable = false;
+							break;
+						}
+						Array.Copy(FileBlock, 0, Chunk, TotalReadCountOfBytes, ReadCountOfBytes);
+
+						TotalReadCountOfBytes += ReadCountOfBytes;
+					}
+
+					byte[] Hash = SHA1Utils.GetSHA1(Chunk);
+					Chunks.Add(new Tuple<byte[], byte[]>(Hash, Chunk));
+				}
+			}
+
+			return Chunks;
+		}
+	}
+
+	public static class SHA1Utils
+	{
+		public static byte[] GetSHA1(byte[] payload)
+		{
+			using (SHA1 sha1 = SHA1.Create())
+			{
+				return sha1.ComputeHash(payload);
+			}
+		}
+
+		public static string GetSHA1AsString(byte[] payload)
+		{
+			byte[] sha1 = GetSHA1(payload);
+
+			return FormatAsHexString(sha1);
+		}
+
+		private static int GetHexVal(char hex)
+		{
+			return hex - (hex < 58 ? 48 : hex < 97 ? 55 : 87);
+		}
+
+		public static byte[] ToSha1FromString(string hex)
+		{
+			if (hex.Length % 2 == 1)
+			{
+				throw new Exception("The binary key cannot have an odd number of digits");
+			}
+
+			byte[] arr = new byte[hex.Length >> 1];
+
+			for (int i = 0; i < hex.Length >> 1; ++i)
+			{
+				arr[i] = (byte)((GetHexVal(hex[i << 1]) << 4) + GetHexVal(hex[(i << 1) + 1]));
+			}
+
+			return arr;
+		}
+
+		public static string FormatAsHexString(byte[] sha1Hash)
+		{
+			StringBuilder sb = new StringBuilder();
+			foreach (byte b in sha1Hash)
+			{
+				string hex = b.ToString("X2");
+				sb.Append(hex);
+			}
+
+			return sb.ToString();
+		}
+	}
+}
