Fix static analysis warning

#rnx

diff --git a/Engine/Source/Programs/Shared/EpicGames.Core/CommandLineArguments.cs b/Engine/Source/Programs/Shared/EpicGames.Core/CommandLineArguments.cs
--- a/Engine/Source/Programs/Shared/EpicGames.Core/CommandLineArguments.cs
+++ b/Engine/Source/Programs/Shared/EpicGames.Core/CommandLineArguments.cs
@@ -1002,7 +1002,7 @@ namespace EpicGames.Core
 				{
 					continue;
 				}
-				arguments[idx] = arg.Replace("\"", "");
+				arguments[idx] = arg.Replace("\"", "", StringComparison.OrdinalIgnoreCase);
 			}
 			return arguments.ToArray();
 		}
