UGS: First stab at OIDC login.

* Enabled via a OidcProvider setting in a project-specific UnrealGameSync.ini file, which must match a defined OIDC provider in oidc-configuration.json.
* When enabled, displays a "Connect to Horde" button next to the project file in the UGS status panel. No notifications on expiry etc... yet.
* Removed legacy OIDC/Jupiter support code to avoid having to refactor it. Can bring it back later as and when we are ready to productize it.

diff --git a/Engine/Source/Programs/Shared/EpicGames.OIDC/OidcTokenManager.cs b/Engine/Source/Programs/Shared/EpicGames.OIDC/OidcTokenManager.cs
--- a/Engine/Source/Programs/Shared/EpicGames.OIDC/OidcTokenManager.cs
+++ b/Engine/Source/Programs/Shared/EpicGames.OIDC/OidcTokenManager.cs
@@ -4,12 +4,14 @@ using System;
 using System.Collections.Generic;
 using System.ComponentModel.DataAnnotations;
 using System.Diagnostics;
+using System.Diagnostics.CodeAnalysis;
 using System.IO;
 using System.Linq;
 using System.Net;
 using System.Net.Http;
 using System.Runtime.InteropServices;
 using System.Text;
+using System.Threading;
 using System.Threading.Tasks;
 using IdentityModel.Client;
 using IdentityModel.OidcClient;
@@ -18,18 +20,34 @@ using Microsoft.Extensions.Configuration;
 using Microsoft.Extensions.DependencyInjection;
 using Microsoft.Extensions.Logging;
 using Microsoft.Extensions.Options;
+using Microsoft.Win32;
 
 namespace EpicGames.OIDC
 {
 	public class OidcTokenManager
 	{
-		private readonly Dictionary<string, OidcTokenClient> _tokenClients = new Dictionary<string, OidcTokenClient>();
+		private readonly Dictionary<string, OidcTokenClient> _tokenClients = new Dictionary<string, OidcTokenClient>(StringComparer.OrdinalIgnoreCase);
 		private readonly ITokenStore _tokenStore;
 
-		public OidcTokenManager(IServiceProvider provider, IOptionsMonitor<OidcTokenOptions> settings, ITokenStore tokenStore, List<string>? allowedProviders = null)
+		public OidcTokenManager(ITokenStore tokenStore)
 		{
 			_tokenStore = tokenStore;
+		}
+
+		public OidcTokenClient FindOrAddClient(string name, ProviderInfo providerInfo)
+		{
+			OidcTokenClient? client;
+			if (!_tokenClients.TryGetValue(name, out client))
+			{
+				client = new OidcTokenClient(name, providerInfo, _tokenStore);
+				_tokenClients.Add(name, client);
+			}
+			return client;
+		}
 
+		public OidcTokenManager(IOptionsMonitor<OidcTokenOptions> settings, ITokenStore tokenStore, List<string>? allowedProviders = null)
+			: this(tokenStore)
+		{
 			Dictionary<string, string> refreshTokens = new Dictionary<string, string>();
 			List<string> providerKeys = settings.CurrentValue.Providers.Keys.Where(k => allowedProviders == null || allowedProviders.Contains(k)).ToList();
 			foreach (string oidcProvider in providerKeys)
@@ -40,7 +58,6 @@ namespace EpicGames.OIDC
 				}
 			}
 
-			Providers = settings.CurrentValue.Providers;
 			foreach ((string key, ProviderInfo providerInfo) in settings.CurrentValue.Providers)
 			{
 				if (providerKeys != null && !providerKeys.Contains(key))
@@ -48,7 +65,7 @@ namespace EpicGames.OIDC
 					continue;
 				}
 
-				OidcTokenClient tokenClient = ActivatorUtilities.CreateInstance<OidcTokenClient>(provider, providerInfo);
+				OidcTokenClient tokenClient = new OidcTokenClient(key, providerInfo, _tokenStore);
 
 				if (refreshTokens.TryGetValue(key, out string? refreshToken))
 				{
@@ -65,9 +82,8 @@ namespace EpicGames.OIDC
 		}
 
 		private OidcTokenManager(IConfiguration providerConfiguration, ITokenStore tokenStore, List<string>? allowedProviders = null)
+			: this(tokenStore)
 		{
-			_tokenStore = tokenStore;
-
 			Dictionary<string, string> refreshTokens = new Dictionary<string, string>();
 
 			OidcTokenOptions options = new OidcTokenOptions();
@@ -83,15 +99,14 @@ namespace EpicGames.OIDC
 				}
 			}
 
-			Providers = options.Providers;
-			foreach ((string key, ProviderInfo providerInfo) in Providers)
+			foreach ((string key, ProviderInfo providerInfo) in options.Providers)
 			{
 				if (providerKeys != null && !providerKeys.Contains(key))
 				{
 					continue;
 				}
 
-				OidcTokenClient tokenClient = new OidcTokenClient(providerInfo);
+				OidcTokenClient tokenClient = new OidcTokenClient(key, providerInfo, _tokenStore);
 
 				if (refreshTokens.TryGetValue(key, out string? refreshToken))
 				{
@@ -102,19 +117,17 @@ namespace EpicGames.OIDC
 			}
 		}
 
-		public Dictionary<string, ProviderInfo> Providers { get; }
-
 		public bool HasUnfinishedLogin()
 		{
-			return _tokenClients.Any(pair => pair.Value.GetStatusForProvider() == OidcStatus.NotLoggedIn);
+			return _tokenClients.Any(pair => pair.Value.GetStatus() == OidcStatus.NotLoggedIn);
 		}
 
-		public async Task<OidcTokenInfo> Login(string providerIdentifier)
+		public async Task<OidcTokenInfo> Login(string providerIdentifier, CancellationToken cancellationToken = default)
 		{
 			OidcTokenClient tokenClient = _tokenClients[providerIdentifier];
 
-			OidcTokenInfo tokenInfo = await tokenClient.Login();
-			
+			OidcTokenInfo tokenInfo = await tokenClient.LoginAsync(cancellationToken);
+
 			if (!String.IsNullOrEmpty(tokenInfo.RefreshToken))
 			{
 				_tokenStore.AddRefreshToken(providerIdentifier, tokenInfo.RefreshToken);
@@ -123,14 +136,14 @@ namespace EpicGames.OIDC
 			return tokenInfo;
 		}
 
-		public Task<OidcTokenInfo> GetAccessToken(string providerIdentifier)
+		public Task<OidcTokenInfo> GetAccessToken(string providerIdentifier, CancellationToken cancellationToken = default)
 		{
-			return _tokenClients[providerIdentifier].GetAccessToken();
+			return _tokenClients[providerIdentifier].GetAccessTokenAsync(cancellationToken);
 		}
 
 		public OidcStatus GetStatusForProvider(string providerIdentifier)
 		{
-			return _tokenClients[providerIdentifier].GetStatusForProvider();
+			return _tokenClients[providerIdentifier].GetStatus();
 		}
 	}
 
@@ -151,9 +164,14 @@ namespace EpicGames.OIDC
 		public bool IsValid => IdentityToken != null && RefreshToken != null && AccessToken != null;
 	};
 
-	class OidcTokenClient
+	public class OidcTokenClient
 	{
+		public string Name => _name;
+		public string DisplayName => _providerInfo.DisplayName;
+
+		private readonly string _name;
 		private readonly ProviderInfo _providerInfo;
+		private readonly ITokenStore _tokenStore;
 		private readonly Uri _authorityUri;
 		private readonly string _clientId;
 		private readonly string _scopes;
@@ -162,13 +180,13 @@ namespace EpicGames.OIDC
 		private string? _accessToken;
 		private DateTimeOffset _tokenExpiry;
 
-		private readonly ILogger<OidcTokenClient>? _logger;
 		private readonly List<Uri> _redirectUris;
 
-		public OidcTokenClient(ILogger<OidcTokenClient> logger, ProviderInfo providerInfo)
+		public OidcTokenClient(string name, ProviderInfo providerInfo, ITokenStore tokenStore)
 		{
-			_logger = logger;
+			_name = name;
 			_providerInfo = providerInfo;
+			_tokenStore = tokenStore;
 
 			_authorityUri = providerInfo.ServerUri;
 			_clientId = providerInfo.ClientId;
@@ -181,26 +199,11 @@ namespace EpicGames.OIDC
 			_redirectUris = possibleRedirectUris;
 
 			_scopes = providerInfo.Scopes;
-		}
-
-		public OidcTokenClient(ProviderInfo providerInfo)
-		{
-			_providerInfo = providerInfo;
-			_logger = null;
 
-			_authorityUri = providerInfo.ServerUri;
-			_clientId = providerInfo.ClientId;
-			List<Uri> possibleRedirectUris = new List<Uri> { providerInfo.RedirectUri };
-			if (providerInfo.PossibleRedirectUri != null)
-			{
-				possibleRedirectUris.AddRange(providerInfo.PossibleRedirectUri);
-			}
-
-			_redirectUris = possibleRedirectUris;
-			_scopes = providerInfo.Scopes;
+			_tokenStore.TryGetRefreshToken(name, out _refreshToken);
 		}
 
-		private async Task<OidcClientOptions> BuildClientOptions(Uri redirectUri)
+		private async Task<OidcClientOptions> BuildClientOptions(Uri redirectUri, CancellationToken cancellationToken)
 		{
 			OidcClientOptions options = new OidcClientOptions
 			{
@@ -217,7 +220,7 @@ namespace EpicGames.OIDC
 
 			// we need to fetch the discovery document ourselves to support OIDC Authorities which have a subresource for it
 			// with Okta has for authorization servers for instance.
-			DiscoveryDocumentResponse discoveryDocument = await GetDiscoveryDocument();
+			DiscoveryDocumentResponse discoveryDocument = await GetDiscoveryDocument(cancellationToken);
 			options.ProviderInformation = new ProviderInformation
 			{
 				IssuerName = discoveryDocument.Issuer,
@@ -233,7 +236,7 @@ namespace EpicGames.OIDC
 			return options;
 		}
 
-		public async Task<OidcTokenInfo> Login()
+		public async Task<OidcTokenInfo> LoginAsync(CancellationToken cancellationToken)
 		{
 			// setup a local http server to listen for the result of the login
 			LoginResult? loginResult = null;
@@ -248,16 +251,28 @@ namespace EpicGames.OIDC
 					http.Prefixes.Add(prefix);
 					http.Start();
 
-					OidcClientOptions options = await BuildClientOptions(uri);
+					OidcClientOptions options = await BuildClientOptions(uri, cancellationToken);
 					OidcClient oidcClient = new OidcClient(options);
 					// generate the appropriate codes we need to login
-					AuthorizeState loginState = await oidcClient!.PrepareLoginAsync();
-					// start the user browser
-					OpenBrowser(loginState.StartUrl);
+					AuthorizeState loginState = await oidcClient!.PrepareLoginAsync(cancellationToken: cancellationToken);
 
-					loginResult = await ProcessHttpRequest(http, loginState, oidcClient);
+					// start the user browser
+					using (Process process = OpenBrowser(loginState.StartUrl))
+					{
+						using (IDisposable _ = cancellationToken.Register(() => http.Stop()))
+						{
+							try
+							{
+								loginResult = await ProcessHttpRequest(http, loginState, oidcClient);
+								http.Stop();
+							}
+							catch when (cancellationToken.IsCancellationRequested)
+							{
+								loginResult = new LoginResult("Operation cancelled");
+							}
+						}
+					}
 
-					http.Stop();
 					break;
 				}
 				catch (HttpListenerException)
@@ -279,6 +294,9 @@ namespace EpicGames.OIDC
 			_accessToken = loginResult.AccessToken;
 			_tokenExpiry = loginResult.AccessTokenExpiration;
 
+			_tokenStore.AddRefreshToken(_name, _refreshToken);
+			_tokenStore.Save();
+
 			return new OidcTokenInfo
 			{
 				IdentityToken = loginResult.IdentityToken,
@@ -347,7 +365,7 @@ namespace EpicGames.OIDC
 
 		private string GetSuccessPage()
 		{
-			string httpPageSuccess = String.Format("<html><head></head><body style=\"background-color: black; color: white; \"><div style=\"text-align: center; \"><img style=\"object-fit:contain; width: 40 %; height: 30 %; opacity: 0.5; margin-top: 64px;\" alt=\"Unreal Engine Logo\" src=\"{0}\" /></div><br /><br /><h1 style=\"margin:1%; text-align: center;font-family: sans-serif;font-style: normal;font-weight: 400;font-size: 25.236px;line-height: 36px;color: #FFFFFF;\"><img style=\"object-fit:contain; vertical-align:middle;\" alt=\"\" src=\"{1}\" />You have been successfully logged in.<h1></body><p style=\"text-align: center; font-family: sans-serif;font-style: normal;font-weight: 400;font-size: 16px;line-height: 30px;letter-spacing: 0.371118px;color: #FFFFFF;opacity: 0.5; \">You may safely close this page to continue.</p></html>"
+			string httpPageSuccess = String.Format("<html><head></head><body style=\"background-color: black; color: white; \"><div style=\"text-align: center; \"><img style=\"object-fit:contain; width: 40 %; height: 30 %; opacity: 0.5; margin-top: 64px;\" alt=\"Unreal Engine Logo\" src=\"{0}\" /></div><br /><br /><h1 style=\"margin:1%; text-align: center;font-family: sans-serif;font-style: normal;font-weight: 400;font-size: 25.236px;line-height: 36px;color: #FFFFFF;\"><img style=\"object-fit:contain; vertical-align:middle;\" alt=\"\" src=\"{1}\" />You have been successfully logged in.<h1></body><p style=\"text-align: center; font-family: sans-serif;font-style: normal;font-weight: 400;font-size: 16px;line-height: 30px;letter-spacing: 0.371118px;color: #FFFFFF;opacity: 0.5; \">You may safely close this page to continue.</p><script>setTimeout(function() {{window.close()}}, 1000);</script></html>"
 				, "data:image/png; base64,iVBORw0KGgoAAAANSUhEUgAAAqUAAAIhCAMAAABAJzPOAAAAM1BMVEUAAAD///////////////////////////////////////////////////////////////+3leKCAAAAEHRSTlMAQMCAEPCgYNAw4CCQcFCwJwL4owAAHPRJREFUeNrs3dl62jAQBeDRvlnL+z9t268UbIxpAsjRcv77AI4PkmZkGwIAAAAAAAAAAAAA+KjAdgigAY4ZsXBuyxHFuRYCgYWfEJjQXJVv4FkYhBVOwqLm5VU+CyYJoJ5gFlve53PEqAo1uJhV+SAukFT4JGm0LxXw6AjgA1y05X8853wRdzjnX/hLnQjgLW7x5ZjlQiTm6BnJWBT5WT9AaYOCCipElC+RBfoWZgQ/fL2MERVeEKI/6NOLFOhVkol88Loa1RR8j+GPF5HG0ftkWh6/vAgE8DVhUQ8T+tEMpcVi5odXJX5Wy0garco9H1FKwXNyvxpVuWYJzpb9G2pM/HAsCLVLTKIDNVsJGZUUPBb0LixnNTLdbiXMkVPYC/pnF4gp3+fUEMCzjCrt6GwheoyncEjqRurslJFTeEje1Uw80dmOe7Uc9T78Zrax+PE2kBS+rGn0T6fHbIOtSrPJqRIEMwt5G4dmhq1tTj2WpxMTqs2M7nOamxjj4XzMNrUefZ5TFQnmI5fmx6pt88HiHqnpMN9DW3L7VUIVNZfN2fctb0UGjuF0UpuBtK2iaS95DKczEn3t7cjOPi98QLDlSrU82d843t1Hho9tiObGJ/ubqLBlOo/11U+qpzvhAkcRNYtgexxI/4oKs/4UkupzIN1/xRaCQYneS+XVAdjOZgL4Gpn7bzuuGr0Ki9MBreZL1eqG6P/JjMXpwJy6zfZdT5ax/xkBDhg1zLllt0PRBAMxPdf296QdZFqADT1YQ1yj1B+PHm7siYN97UDqAddxSaEjNZLVKm6kW4ccYjoQaQftL0qLmI5i4HM58KFNZuwzqUc+uHlIO3YpjJgOQNrR24qIaffGDyli2j89fkiJFsS0a1OEdHVltBr6MAc1SUgnOtABmXnOHWLaq4lCSqTHu05hCtddbtXlXXiI6QyCmqzwzeVirGsVhibtZCFdtYb7vfFwNhOOLPN9MXsn5gspkVMTlYsDSHOWEq5cZILmOTXp2TIj3pIwKGmnnfkWVFC90OdWESxq3szz7zJ29PtgykWi2oLRf8Zt3k5RfZ1HOEHDnDrnMTsyad/g83ycGuQpQ0OT9ozKKcTc6o0qBkvT9v2rH7ykWmRs+qHn1/9Ac58MLli5cFRJyq3/4qKdslncEanqNgyD8OWPpre2rhfa9P9swTHlqhUu06W0H1Iig3ZUy1LN08N46SOkRHnSvbcuyIpTHePl50IaksiRvk56zPnN0quyoV5GT94rdzH7b78bw5zfKlarBRNy2dNUn4tZvTZuL6jz2yR9nXa2FOUBK6mypP3td39a+WfAm0SdXz5MquzUv+UvabV6M/fGxELQkFBlKRZ4eShSRW65RfTla7uWhvcd5sVrlLVRlYc4VSONLeX9kJJUM93o3YlUIT+BlwOBKgmLKm+F9CahafqLvXtdchOGoQAs38EGovd/2nZIOr0kgRALo8L5/u/uDDkrLBkTdbJ85UiB3xhoHyXyE0PfsmiglHHiqzB/4yf7TiFL5GepfqXeEajgg3RH22fmilIqlFFOVMFp38m9mEG6dUqB3yskzjt+JVZOkPGkqSJFunWKvOBG4lLYZZSQMI1SJMoWOd/xkomE9R3zPvtbFtv5ahTZTes+85JAsvzAT6ROahsUUzWs6BSqD8wNb/gmL9VsFNOzkC0YiVeMJGnYd5JQ+E7Nay0uS7ReJF5j9liRPutkl+wopsfqJUtp4lVEu7b2wq8PKliZqhAFS2niVZmk+MjvjRovD2goFonXWRJSOn7Pqrw+oKBWJG6YUhP4vVBUXiD4kpcrFYkbpjTxEkdyCnbzD+fEKkUfGqZ04CWZJEWcLTla4NnQKqRsZYKzyJCkHofzD5b4rkjv3VfMMav/1o1kWbx491id1A6+5U/tHlIuJMvwHU5AHcNIvQdy4I+VvUMaSVrGKySOFIXuZYk/N+0W0t1KXsIw6kCe71KbzmlLo1aMGd1oTL85pJHE+YBh1HFGnoU2ndND3vhAXheT3/KnCsmLOKd3nE5mDDXwJmb7GdRb+jikkXZQ+E7ZN1VcQi9TfSbexn4T/jzntNvwP6Dx/xm+roG2etG20bhtQfKQDcWKcaxM/0TQmkxPcOPNEj2UTVtM3YZfLcwH7D8dYxKZr0z8hVjmjLpAyyxv5WkfESPTj+i88D7zVzprO2aWXk1E2onByPQYIjexgaush2ObifaSMTI9wiTxQFrPdWhN5E1oNwPeE3mEKDFcsVxHbIKwe4Z6viNoKQgMqieuRKtcxYxLVEaX/wmFN/zMlYTHsYX2M6DLb28QuOGPXItEi2mgHfUY7H9G2Q3MhxYpLSqWpT9lvMB8nb7K4LiabIfm6AVd7SZ8cbOOx5ZSll1XGNrThMf3WrP1Q2rHbVJaVDRPP+H4U2sC296hUUopHD/Tn92w/dSWqb97JW6VUltxCkDZKgm2cPWdQGYBoksLS/vq8ZKTtmx1u2H4cimlgIVpU/ULuZu+lA60s4iFaUumuvgU1pdSRztLWJi25Ko/VnfFlPaYmLZ0q94ezVdMKQV8FcR7+q624Uum9Iat/HZK9Z0r8t9yHI0xYwznTqnDO/bbmXgWpfadcqJfUj5zSg2OlbTjap9snxbSMZw4pR5z/RWaZvqRfwuG/pbOm1LKaJ+aCbUz/bz4JXXpvClF+9SMr22e+sWQEt1Om1K0T4tU9QBu5cFjH86aUoPdp1bG2g/Vrh0tHs6a0oKvK2llqNx58qvH4cpZU0p4d0Qrlmd99Rwqe3ojnzWlFk1+I7myHrj103C3s6Y0oslvpHY0bdef53RnTanDI6Zt9LUdAN8Ff8GUThhFtWF4NtROSxNdMKUYRS1RNIhKPLN0xZR6jKLacDybKgdZ5pIpJTxvskBRm2of1eSaKe0wMG3C8qzUPatSLppSi5S+ou06+0f7oC+lkd7CU1H/m8Czui63KEyppQYcUtpE5Uea7hv4F08pXq//D2UpdfyTuWxKR4z139Oz9TTMP3zZlBqkdImWrSfLzAkpJfiDustsmTNdPqXYIv2HspQG5vHCKcXT+q8pSykz+wunlJDSRUqOQTJHQkqR0pfUjKWZDVKKlP5LV0oNZ0JKkdJ/aUvpiJQipU+0pbQgpUjpE2UptYSUIqVPlKU0IaVI6TNdKe09UoqUPtOVUiKkFCl9ASlFSq8GKUVK9UNKkVL9Xh+JQEqRUk3eHIlASpFSRQxSipT+YO9esNyEYSiACix/sPlo/6vtOXDaSUNmQmsBgry7hORhW7LB5iGlVTqk9DWk1FBK8d7Tty734g5SCva/bYiUgv3l/31T6mQ2EOzq9c15SKmhTRH45ooSpBQptaRFSjUuy4JHl/ydT0mp0GuXe8bhmzkLKd3II6U/utLFWvdNKbaefnSpi7Vum1J8JuqdC20+3Taljcx6grWrbT7dNqX4FPRRZObpL0ipnTPk8N3FWp+UUh7of7Vo6v/sSlcWmU4pJ2f5CYfvWlGflFJO0uAaUvPiiwrgc1LKSQQn9e1r1j/156SUk8hk+QGH74v8T0kpJxEZLC+W4O8SgOnLh6R0DqkUlPgX8KLI/4yULiHNOGtyBetrSj4jpZwqtzeLzBLB/hqZTfTHR6S0S7Una0dcnHccXpdPH5DSzlc3O3u8mnegJPJURdw/pZ2XxWD7NQf4LawaKrdPaZRZ3bMpKJ4OFFfrq7un1Mlvk+ntEPjSrVZo904pB/ljtNwbgUf+efK7dUqX4r7+yWxxuPRQ0/PC9M4pHb3MqgdCWTDBIYbnhel9U8pBHviKiDXo6b91qYWpakon2YhprcmyqH8ue7yZd7D81Pmzm1JuZavmzUBauRJPOGpysFDR4VZMKdM7XZbN0lMG2XkRtaG0oFt6tFixxlJMqduygN7OD/SlOC/PSv1PNhEchSvqVcWUeq6a7dfyUJaIxknWAlWYsIl/uFTRi1JJ6YaRafTyH3zbtvJaUXiuC8FhhlVKjk7pu6GptKKspwoj+lDHK7tO+U42i/QSOy/KPFOFgD7UCfKeU76TypjGLOoGquFxau8E/Z5TvqubiGMWfZlqjPhexBm6iilfN6WSIz0ovZc9NFRjwoS/2VWmfCf/JvdLhrjpk+xjohosmPBP0e94preV/9A+BdRQ6UQRE/45yo4NwCTWRKrSYsI/Sdrt6HkRa1q7TzT8aNhtFnNijGeq0qOlfxaWxUjKihdjRqqTsYd/mmmflimbW5UGqjPiXZLzjEqrrWEaOvojmhtJE1OdFl/eOVFWqp9aEd/2rmmaIWSxxnfVxSBO6Z/IKX0wiu1F80GkSgHN0jMVrf+xMzfLf+mpEnvUTqcKWh2WKFa1Wh07j9rpJI0smvu1SDUqp0VG7XSypLaZH8QiX6hWxL7T2aIsujvu3c/lfbWML+2dLsss3LGZL9JYWhRBbWkg5Y4xjVSvxRb++dh/DaZ3a5tGlfoSF5EZ4B4G03u1TSMpaNHRt0BzMKXO0GgaCUPpfThZdBc4D+UPDillDKU2sJdZa//YXuB0VAtqETGUWuFk0ViPaSDq/NaQYii9F/aPzRa7u1Bhc4GWCqkYMJTa4WQRTb/2FGg2ylsT23x8QaPPmZl0RC/qwubjVz0p6bHtZEl8/DSzzY5U2DpU+0hKimAH35Skfe6HJ1EVti58U0daWnx1x5ZGf9gYRFHYerdOz6RlFJwrNWbSfze/S6IlbOx2+ZHUcMYRfWuKLDKTGu71QropphOTnl5wN645TrdAXjRZP6SLzqsOpGudoKFvD+c9agV2KiHdEtPApCmhC2VRI7u0sEurHtJX3f3caE8tuIPMpGmnhdiYlUO67u77gXR1gtLJJvZ79QedVw7pwq3aT2qS4EMRRg0Kr7ArX94UNhxrCYW0Oew62dU+1PkmchroZ9Mqo7rzPV7BN6j4HQtbdlk7pMSukD7OmO8tG/YtGmKqCOlxAuZ729qd/6AuePMhjZjvjWO/92THsbUd0s7jgL51o+y/51KGZDeknNDPt68/pJ9dhslmSCnIIqOfbxino2qHsU/2QhoF+/dX0Hk57DZDblzrLYW0E5zX+8XOnSC5DQJRGDarAG3v/qdNUklsVeyZeMaLBPzfAabGpWfR3WDqMOi97cNoXLThECH1gSFULaL2+MFPNoNzk/1pOu3EJ4rSeqROJ4ZR/B6vHued0tTVO2USk9KaGHVYnw27l8Xgif2HkbhvpzLxcp1NH3Kgc6qP7atIO4c00DlVxKeeYuoDe05V8qGfmPb1lWxKPsd0ObXNJzZGq5XVR63mU38TjYYMXcSUkFauh5j61OMWRlMGtd5W+KQut4Ob0npMCWkTXNMxzYWQNmGSmr1EIQdC2ohh/zP0L2IIaTs2MW3qYQ4ipA1p83E6/WUb+lQdu8S0NDM4jc0WMt1agv4IbWzq+0RI25NDU61+DuKASYNyaaiHGtT0FLhjmzUyVV6cRjVWv+DMbx5uzW+gMTXYC+LM6WyqdtVfQpNzNZwNofZV308HuI4KL7XtoUKNvX5OampWgZv8qoq3bJzUwblunE5OtXbIo635G4YvMUFna0UPew5ilN8Pb1XfTCpv/2muhujBrAtbwyWn3onVvjs56cId/qmborpnE3h47qhy7C5qXKXa57z4HlPqWPa9C6Jt6pafJB3+oNRQeJH2zSRdhCOWp6aIirR7LujAUyljVUdVgtcarTbKkXJqrOrp8PBaSzlkTo2VqpqW4bVc0EY5QiD+yejKYg8ftRXczqEYirYsG6K4Lk+luF8wvCs6ZA2C/RkrbaXBn3Zgosgo7s9piPn0Xn4uko4+w8WernOqNL8xJcsq6XB9HI7nKqda37Py5xjEWo87jVFvD2qeiv6RyCg+M7pwHdR5PL3IEoOkw4wYUI0h6Uqanh+dPK/65VjjWtQix6Br1j0vqeMQi65Zlnrczw9Wt1i3+NODzLwG3VB4jeKrxrnoprK65Xt58maOSTcFqlE82n5fs9Et5v6/ZNxkgz4QIifz8IDskj5j7eRmY8wH2TSDc/Z3PIkoXmicV90l2Iuiu6SJHzPhWcyU9GwlDrRLeC6/TJaEogJmXoseE6xbOEiCF/NmjlbfkVa38ArF+3gzu9UG3SXZyRn6JOxlNMY5t1prk7aK/WlybjGM6wEAAAAAAAAAAAAAAAAAAAAAAAAAAAAAAIAf7NFBCsMwEANAe1NSTB3w/1/bcyCUddnjzAMkkAAAAAAAAAAAAAAAAACAf8x4Vhg12l0kzLZjJJtHVDiP4+j9lRuk3mybPvHLVf90vVjPCqN6u1sZ77ahJ5v7qhTnNRKDVIu26ctO3a04EMJgGK6j8bd2vvu/2t0934x2QiGheU8FkfiQjKtofOCnv0UpNfVK/6L80q504LrsSm8rRU0WlAKgmFQrzVg0XOltpchGlAJUFCsN6wtd6X2lmFaUAmdSq/TAsuBK7ytFM6MUNSlVGnZudKUCpT2ZUYqadCqt2OjlSu8rxWlHKbJKpQU7dVcqUIppRymCRqUdWxVXKlCKYEdpV6h0Yq+eXKlAKSUzShHUKU2EzaIrFSjFYUdpVqc0YjdKrlSgFNGMUtKmNBG2i65UohQvK0rRlCnNeKPhSiVKaVhROnUpHXin7EolSlGtKH3qUprxVs2VSpTiaUTpoUppw/+d3L2uVKQU5cNKc1hWYoznQbiq7w5EVJN9R+ef5UpFSql9Vml87NZmB59gtsL2Z1LYk8OVrpV28NWkROlv04ZSdpXy7yqudKk0ZvBlPUovmDY9SgsrkR9Dd6VrpamCr+hR+qhgCnqUdhbi1TJ1pUulj0Fgo6ZH6dSvlF2ll3Og5Ep/2LeXJDthIIiipT9IfHL/q3V45klCyUZW2aG7AqJ1XvbrBl6VygVejWaUOvNK2b3R+nJlYSl9VyoBvNuM0mheaXj5lb7TMV1K35WKB++0olSsK6VT+npf6lhKFUpjAm9fSnUdr38dNZDaUvquVHbwUlxKFVGCVXGTPy+lCqVygueXUppCYFFIhltKFUolgxeW0j+YUhVlv5RqlMYNPGda6finTdofnENRWYZbShVKZU+gpfYfKFXXf6Xq+598TLelVKNUCnibAaWRfYRMKPXKN3MaSGUp1SiVA7xjvlL66JsFpRe/OOWY1qVUpVQ28K7pSk/2AbKgtKoFNZDOpVSlND5/NTX6TFQxoLTwKVWPaYpLqUapOPC2OFdpAanNVxprB6AGUlhKVUrlBC/PVMqRVpmvNPAp7RrTpVR3ADd4ZZ7Si9M75yuNqctPAykvpboDiBW0tI97h5QXQvbgpThfad+U8jFFW0p1B7An0LZo7338INOVNpBc7/bmpVR5AAW825zSGucr7b83H6jrpVR5ABm805jStMt0pYopVY+pX0q1B7CBt9tSWmS+0lshTj+mS6nyAFoCrUZDStMu85U6BTj9mNZ/QWm0oPSRlbej9I5iQKnXTKl+TMs/oNSZUCoBvGBE6e1EDChVTql+TJdS9QF48JwBpf5s/eP2+xWhVd2U6ocgLKX8h9Lxz/04XekWrn6lMqIC0h2eO/iP939UWkcolR08P13pz+rhDCit+LxgXmmR7jBEqRTwDgtKAdQyW2nA96VmXWkwo1QyeJcNpYCPU5XGhAFl60qz9OZGKY0baKkZUYot9im1P6UAmnGlVXo7RymVlkDbrCiFn6i0JQzpNq4UTTq7hymVC7xsRSmueUozBuWMKw3SV8M4pXKAV6wo9dOUNozKG1eaopB6ITW1Up4HLe1GlKJ1KP03phS4TCjluLL0tIMmHyiNj49E/6nS6l/aoKlMUuowrmpCaQDt7EGaxiqVHbz8N957iu4KHo/lSUo9BlYsKD3Bu5soO5+mrldp/4X6j5XyLg+eVyv9Z6YUqNGAUoen8qXZ0VDxkP9GqdzgDVDKuhJYaY7SDUMLBpRGvORDcTvz6Uq4E547PlIaNxNKn77caO/r/WGN3EEeUxr3QpfedoKy5H+tQlvpVcp5mFAqF/qVDjzWCt4/8XKsRumNwe1fKZViQyk/pDZB6YnhtflKT4wtyWdK5bChtIDk/r7SmDC8PF9pw9jyh0plM6E0GlIaQKq+tw2sNl2peAzNfak0JgtKpZpRGhPGP8YOP1/phZFV+VKpOBNKvRmlmcPqr4Dlpiv9wd6d7TYSQlEUvcxQ4/3/r+10px9MYjBUgQzyWY8ZkdlCiHIIee7Ita2U1AiV2lEqXXNf03IxfX+lgvvx1LhSOgeoVI1SaW4pbbqYHm+vlDbuJjSvVC+olP4LuS9pupgu769Ue+7EUfNKKRhU+uoXSKLGi6l7e6XdMt2oQ6XkUOmrrZro8Hbjt1faKdOdulRKFpX2WErz4av3V9pjb2oO6lQpeVRacHDUejEdoFISnpuymrpVuhpUmj836rCYbiNUSuR8w0ZXSmqQyoFKeyyl+UGvQ1RKFLaFG/C7powWqaiPr1SbzFLaZTG1g1T6JSjJdyzWrZTXJBX56ZWqdrNeOuowTKV/Baek51pGWiU0RfpVqpfPrjS7lHZaTOVQlX7TQjilNvmX4ae8/HIqpQ4RCAAAAAAAAAAAAAAAAAAAAAAAAAAAAAAAAAAAAAAAAAA+mPimCT4d15K5e0eWdFIFl5xooZSUMvmPL61ygX5z6lH+7hBTPUAt6uBejy9jV8rn5UoPu3CJ0+nshUv5SlnWDlDUv0AweKW8X6pUK8PFzKavV8oKlc6ndaUsLlR6GK5iwvVK+UCl02leqdHVlTquZcL1Ss2KSmfTvFKWtZU6rrdcr5Q9Kp1N+0p5q6t0NXyBu14pW1Q6mQ6V8lFVqeUr/I1K2aHSufSo1KwVlWq+Zr1RqQmodCo9KmVfUanja8SNStlrVDoTfrSo11xBpWzLKz2fNCQ3Fdnkwj+pO5WyRaUzufwCK85xxZUajp0HPaOdaVkp79cqFQQlxqxUccyEwgg0xw5KWT0/snWV7vxDQKXzaFSpUBzz+koEvFFaHLSsq1TsHFs0Kp1Gq0rp5Jgti8DF6VCOvFMpWY6dqHQazSrVC8f2oghU/PGOlWrPMYVKZ9GsUgqGY6G+0oNynHrgSirND1Cg0km0q5RcYueXr/RyBNWV0sExo1HpHLiOzOVlOSYLIjj5ke5aKW0c89WVWvWCIyAauVLK7fzKDjv7VkqSY1tBpTjTf7+mla6GY+JOpYqzxIVKteHYgUpn0LRSEhwzeqxKKXDMrKh0Aq0qTZTlB6uUdo55jUrH17DS5zu/wSolyzGLSsfXulK9cOwYrFLtOeZQ6fAaVpra+Y1V6a/DfRNQ6eiaV0qOY16PVemTAaLSwfEjI1/ZMpUmd36DVUobx05UOriGT0jzO784gkTSa8dKH3iO7cWV7rglqsj4ldL6bOdX/xx/FT+1qlQbjgU8xx9aj0rp4JjXhZU6ympQaerpAyodWZdKSXHsLKzUdq80cbgvUenI+lRKkmN7KoKDI4EyQrtK6eSYQqUD61SpXjgWiiJgrynNxjneqlR7jglUOq7GlaYP9wv/hvRcKWXjBpWmD/dR6biqb41QoqhSchxLREAL/3Aq8btUcWwLR5b6SvMDrDiJylsJGuNqqqxSsmWVWk6Q//nEp+9WSlthpbUUwTSVal80kQdfo25XSh6VTqJfpbSaook0fEm4X6k2qHQOHSulo2giVc/7S/MDFKh0Dj0rpa1kIvXCF4gWldKOSqfQtVKSJRMZDFez1KRSOlHpDPpWqk3qZ9zL1FKjSvWCSifQt1IKRRO5ytohtKqUgkGl4+tcKe1lEylOLmW2ldpVSg6Vjq93pWQLJ1I76/kluR0UuV0pWVQ6PFFtffYeZU0puuYhYhC7Uqf8Z+Fvi/xHqV2ExDc9SvxiXfcSrKnx4wkpAAAAAAAAAMCf9uCABAAAAEDQ/9ftCFQAAICjAKAG6mKYnP9DAAAAAElFTkSuQmCC"
 				, "data:image/png; base64,iVBORw0KGgoAAAANSUhEUgAAADEAAAAwCAYAAAC4wJK5AAAACXBIWXMAAAsTAAALEwEAmpwYAAAAAXNSR0IArs4c6QAAAARnQU1BAACxjwv8YQUAAAV7SURBVHgB7VlNcts2FH6AJGfqdKbaplUd+gR1ThB5l8SdqXICOyewfALbJ7BzAts3UGdSuzurJ7BzgsCuO96yi6iNKOHlARRJgCREipKcLPKtSIIA3gfg/QLgG74OMFgwvA8dD2qfNoCzJo3+NG5AuAHkAkb1a7He82GBmJsECd2ER6NtQNlG5G2StlncC68ZsGtAdiZ+fteHOVGZhHdPKz4OdhHYTjnBnRAkxIFo/XEGFTEzCb3y9dEBMtx1/OJrwRgTgGgeGw8BFNkNR7/KZGYi4f39axs5ntD59lJNPpE65hz74snFX1PH+NBuwsrqcwnQYah2MSPQKQSNQ9IbASVRmgQRUKu/b3+ls13DbpHgzjHvX3gg+TbpUjd1JAULGptliZQikUPARy4Pbn+8eAsLgCIjJd9P7YzPADdF6/y6qH8hiSwBpFWStEp/ClgwvLutLunNkfGpFBE226B0fIIBEegv1M7bc3Y2EEaXxvEqPFrcORiZ0IcmoCBaPfIh9U1a32geDxvBybQ+ThIog0vjTR2h18smECEkIveMT23v7mXX9X/ucfLuXu2QE4vZs2C8vgwdKMLa3dYRCRgJ77Pg43reQubuBDK2nzzjwZcgoMDJXwBDMXltypXVXAebIaF2IXFmKPhQVg4HyoCs3/HTuy0kI3KUblOBIkM4jN7JBHe1s0whQ8LeBThd5i5MzLdeXUyOjQUyr6fmbsCj77bT/1gklHkzQ4pl7oK2fob/QcBT1786FIn+k7yTbrd3go+eG2/9Ze2CCiLT1o8Hgz1nh+GKGRm000fKJoHQjh4Zlz1YFlaCI3PHwwjAbb51EpUcKVABpNlukUBEL35hUBizVIGOAhB24jlLWj9KovrGqxXOpxU7afz/v/ewYIRRQGI4aHV7tz+dH5bqzFiyqMg8s4knE7ywGpbhnUM9iGIiigKGcq90Zyn/jccBO5+p5/cwzt+CoPyBeVwpD9kRrRkMhyoykNPIQx1mFeb2ZQd5jUIS9MsmLmEYk6SzSg9uKiZSeXAEgMxz9qjVtidHQkWXl7pEMwWhP2BmQtUvrQeWSNJzNcUkxBN7a/Pcu8YYTAdYSETrgRHGUDD5BqqA8x/iZ2YVIFI7Ydri+ve5VQmx9q5H2ZYpiJOICissfyBxr7IDRYzlIQ9umX+LBAOeNNbGv4ADKp4pIqJ1xwwrlD9Yu6jsQMm3mIvqJkEs+nGnnBjFxDQiWg9q3MwKRSU9iAa+14uTkBhKy4fZJGTdtBhtp15M4CJClcETWw/kJsyD8bBtvGViOm4LpdLCZDdcSYjdJ4cIGDEYsLmTKtO6mRFthKyJ5Sw+t64kJI0cIuHkII/nqbEq2EkaIRhn/EuWxKf6mVFpaELj8T6UQJYI6UHronxYkQMdsptJGuUcebuaIaFTQsDDpCN0VQ0WSsAg4s+tBwpUuDZ2weeBzDUOzuIZ5b1XkFgEUs6Pz8oGhSqYTDvPWZGuuCgT7bJwzD1IuhL3MMUz19w3rfNnrv+dxbPQUqHBnG1gY/WyjKLPgxwCunA3rQ+f1kiW5ZjRNiZfFJHHV0VBX1WEtd/gyqjDkn7C6yITXa60/8+rfUR2YHyiweWhaF0cwwKgb59WRieUb3TsORZU2o8nyhJRmOu+bXJ11qX61q59yTLb9UFpEnpSfd0lT9I5LoRkevT996LbUC14Y9ihcspvqKsrqUtLyrvZcPBmFgMyEwkthLrVGdd2maNiF4IsGQNBpHxDuCYqk51dgKiP0KF6hUh3ZhIRHFdUFYBClUv5cPC2qvmuTCKCNrmN1cnxwDaE17wFs5LgiD1eg55YQK49N4k0yNNuQA0plUR1M+qFs3Cq4I18GPP3EAzEQ13WfMND4zMnXf+pw+Ue3wAAAABJRU5ErkJggg=="
 				);
@@ -366,41 +384,59 @@ namespace EpicGames.OIDC
 			return httpPageFailure;
 		}
 
-		private void OpenBrowser(string url)
+		private Process OpenBrowser(string url)
 		{
-			_logger?.LogInformation("Opening system browser to {Url}", url);
+			if (RuntimeInformation.IsOSPlatform(OSPlatform.Windows))
+			{
+				// Special handling for Chrome (maybe other browsers in future?) that creates a new window to ensure we get focus when the window closes.
+				if (TryReadRegistryString(Registry.CurrentUser, @"Software\Microsoft\Windows\Shell\Associations\UrlAssociations\http\UserChoice", "ProgId", out string? progId))
+				{
+					if (TryReadRegistryString(Registry.ClassesRoot, $"{progId}\\shell\\open\\command", null, out string? command))
+					{
+						if (progId.Equals("ChromeHTML", StringComparison.Ordinal) || progId.Equals("MSEdgeHTM", StringComparison.Ordinal))
+						{
+							string exe = command.Replace("--single-argument", "").Replace("%1", "").TrimEnd();
+							exe = $"{exe} --new-window \"{url}\"";
+							return Process.Start(new ProcessStartInfo(exe));
+						}
+					}
+				}
 
-			try
+				return Process.Start(new ProcessStartInfo(url) { UseShellExecute = true });
+			}
+			else if (RuntimeInformation.IsOSPlatform(OSPlatform.Linux))
 			{
-				Process.Start(url);
+				return Process.Start("xdg-open", url);
 			}
-			catch
+			else if (RuntimeInformation.IsOSPlatform(OSPlatform.OSX))
 			{
-				// hack because of this: [REDACTED_PROFILE_URL]/corefx/issues/10361
-				if (RuntimeInformation.IsOSPlatform(OSPlatform.Windows))
-				{
-					url = url.Replace("&", "^&", StringComparison.InvariantCultureIgnoreCase);
-					Process.Start(new ProcessStartInfo("cmd", $"/c start {url}") { CreateNoWindow = true });
-				}
-				else if (RuntimeInformation.IsOSPlatform(OSPlatform.Linux))
-				{
-					Process.Start("xdg-open", url);
-				}
-				else if (RuntimeInformation.IsOSPlatform(OSPlatform.OSX))
-				{
-					Process.Start("open", url);
-				}
-				else
+				return Process.Start("open", url);
+			}
+			else
+			{
+				throw new NotImplementedException();
+			}
+		}
+
+		static bool TryReadRegistryString(RegistryKey rootKey, string path, string? name, [NotNullWhen(true)] out string? value)
+		{
+			using (RegistryKey? registryKey = rootKey.OpenSubKey(path))
+			{
+				if (registryKey != null)
 				{
-					throw;
+					value = registryKey.GetValue(name) as string;
+					return value != null;
 				}
 			}
+
+			value = null;
+			return false;
 		}
 
-		private async Task<OidcTokenInfo> DoRefreshToken(string inRefreshToken)
+		private async Task<OidcTokenInfo> DoRefreshToken(string inRefreshToken, CancellationToken cancellationToken)
 		{
 			// redirect uri is not used for refrehs tokens so we can just pick one of them to configure the client
-			OidcClientOptions options = await BuildClientOptions(_redirectUris.First());
+			OidcClientOptions options = await BuildClientOptions(_redirectUris.First(), cancellationToken);
 			OidcClient oidcClient = new OidcClient(options);
 
 			// use the refresh token to acquire a new access token
@@ -433,7 +469,7 @@ namespace EpicGames.OIDC
 			};
 		}
 
-		private async Task<DiscoveryDocumentResponse> GetDiscoveryDocument()
+		private async Task<DiscoveryDocumentResponse> GetDiscoveryDocument(CancellationToken cancellationToken)
 		{
 			string discoUrl = $"{_authorityUri}/.well-known/openid-configuration";
 
@@ -446,7 +482,7 @@ namespace EpicGames.OIDC
 					ValidateEndpoints = false
 				}
 			};
-			DiscoveryDocumentResponse disco = await client.GetDiscoveryDocumentAsync(doc);
+			DiscoveryDocumentResponse disco = await client.GetDiscoveryDocumentAsync(doc, cancellationToken);
 
 			if (disco.IsError)
 			{
@@ -455,7 +491,7 @@ namespace EpicGames.OIDC
 
 			return disco;
 		}
-		public async Task<OidcTokenInfo> GetAccessToken()
+		public async Task<OidcTokenInfo> GetAccessTokenAsync(CancellationToken cancellationToken)
 		{
 			if (String.IsNullOrEmpty(_refreshToken))
 			{
@@ -474,10 +510,10 @@ namespace EpicGames.OIDC
 				};
 			}
 
-			return await DoRefreshToken(_refreshToken);
+			return await DoRefreshToken(_refreshToken, cancellationToken);
 		}
 
-		public OidcStatus GetStatusForProvider()
+		public OidcStatus GetStatus()
 		{
 			if (String.IsNullOrEmpty(_refreshToken))
 			{
@@ -526,6 +562,13 @@ namespace EpicGames.OIDC
 	public class OidcTokenOptions
 	{
 		public Dictionary<string, ProviderInfo> Providers { get; set; } = new Dictionary<string, ProviderInfo>();
+
+		public static OidcTokenOptions Bind(IConfiguration config)
+		{
+			OidcTokenOptions options = new OidcTokenOptions();
+			config.GetSection("OidcToken").Bind(options);
+			return options;
+		}
 	}
 
 	public class ProviderInfo
diff --git a/Engine/Source/Programs/Shared/EpicGames.OIDC/ProviderConfigurationFactory.cs b/Engine/Source/Programs/Shared/EpicGames.OIDC/ProviderConfigurationFactory.cs
--- a/Engine/Source/Programs/Shared/EpicGames.OIDC/ProviderConfigurationFactory.cs
+++ b/Engine/Source/Programs/Shared/EpicGames.OIDC/ProviderConfigurationFactory.cs
@@ -3,15 +3,25 @@
 using System;
 using System.Collections.Generic;
 using System.IO;
+using System.Linq;
+using System.Net;
 using Microsoft.Extensions.Configuration;
 
 namespace EpicGames.OIDC
 {
 	public class ProviderConfigurationFactory
 	{
+		const string ConfigFileName = "oidc-configuration.json";
+
+		public static IReadOnlyList<string> ConfigPaths { get; } = new string[]
+		{
+			$"/Programs/OidcToken/{ConfigFileName}",
+			$"/Restricted/NoRedist/Programs/OidcToken/{ConfigFileName}",
+			$"/Restricted/NotForLicensees/Programs/OidcToken/{ConfigFileName}"
+		};
+
 		public static IConfiguration ReadConfiguration(DirectoryInfo engineDir, DirectoryInfo? gameDir)
 		{
-			string configFileName = "oidc-configuration.json";
 
 			if (!engineDir.Exists)
 			{
@@ -20,15 +30,15 @@ namespace EpicGames.OIDC
 
 			ConfigurationBuilder configBuilder = new ConfigurationBuilder();
 			configBuilder
-				.AddJsonFile($"{engineDir}/Programs/OidcToken/{configFileName}", true, false)
-				.AddJsonFile($"{engineDir}/Restricted/NoRedist/Programs/OidcToken/{configFileName}", true, false)
-				.AddJsonFile($"{engineDir}/Restricted/NotForLicensees/Programs/OidcToken/{configFileName}", true, false);
+				.AddJsonFile($"{engineDir}/Programs/OidcToken/{ConfigFileName}", true, false)
+				.AddJsonFile($"{engineDir}/Restricted/NoRedist/Programs/OidcToken/{ConfigFileName}", true, false)
+				.AddJsonFile($"{engineDir}/Restricted/NotForLicensees/Programs/OidcToken/{ConfigFileName}", true, false);
 
 			if (gameDir?.Exists ?? false)
 			{
-				configBuilder.AddJsonFile($"{gameDir}/Programs/OidcToken/{configFileName}", true, false)
-					.AddJsonFile($"{gameDir}/Restricted/NoRedist/Programs/OidcToken/{configFileName}", true, false)
-					.AddJsonFile($"{gameDir}/Restricted/NotForLicensees/Programs/OidcToken/{configFileName}", true, false);
+				configBuilder.AddJsonFile($"{gameDir}/Programs/OidcToken/{ConfigFileName}", true, false)
+					.AddJsonFile($"{gameDir}/Restricted/NoRedist/Programs/OidcToken/{ConfigFileName}", true, false)
+					.AddJsonFile($"{gameDir}/Restricted/NotForLicensees/Programs/OidcToken/{ConfigFileName}", true, false);
 			}
 
 			IConfiguration config = configBuilder.Build();
diff --git a/Engine/Source/Programs/UnrealGameSync/Shared/EpicGames.OIDC/EpicGames.OIDC.csproj b/Engine/Source/Programs/UnrealGameSync/Shared/EpicGames.OIDC/EpicGames.OIDC.csproj
--- a/Engine/Source/Programs/UnrealGameSync/Shared/EpicGames.OIDC/EpicGames.OIDC.csproj
+++ b/Engine/Source/Programs/UnrealGameSync/Shared/EpicGames.OIDC/EpicGames.OIDC.csproj
@@ -35,6 +35,7 @@
     <PackageReference Include="Microsoft.Extensions.Configuration" Version="5.0.0" />
     <PackageReference Include="Microsoft.Extensions.Configuration.Binder" Version="5.0.0" />
     <PackageReference Include="Microsoft.Extensions.Configuration.Json" Version="5.0.0" />
+    <PackageReference Include="Microsoft.Win32.Registry" Version="5.0.0" />
   </ItemGroup>
 
 </Project>
diff --git a/Engine/Source/Programs/UnrealGameSync/UnrealGameSync/Controls/OIDCControl.cs b/Engine/Source/Programs/UnrealGameSync/UnrealGameSync/Controls/OIDCControl.cs
deleted file mode 100644
diff --git a/Engine/Source/Programs/UnrealGameSync/UnrealGameSync/Controls/WorkspaceControl.cs b/Engine/Source/Programs/UnrealGameSync/UnrealGameSync/Controls/WorkspaceControl.cs
--- a/Engine/Source/Programs/UnrealGameSync/UnrealGameSync/Controls/WorkspaceControl.cs
+++ b/Engine/Source/Programs/UnrealGameSync/UnrealGameSync/Controls/WorkspaceControl.cs
@@ -224,7 +224,7 @@ namespace UnrealGameSync
 		readonly SynchronizationContext _mainThreadSynchronizationContext;
 		bool _isDisposing;
 
-		JupiterMonitor _jupiterMonitor;
+//		JupiterMonitor _jupiterMonitor;
 		PerforceMonitor _perforceMonitor;
 		Workspace _workspace;
 		IssueMonitor _issueMonitor;
@@ -288,23 +288,23 @@ namespace UnrealGameSync
 		// Placeholder text that is in the control and cleared when the user starts editing.
 		const string AuthorFilterPlaceholderText = "<username>";
 
-		public WorkspaceControl(IWorkspaceControlOwner inOwner, DirectoryReference inAppDataFolder, string? inApiUrl, OpenProjectInfo openProjectInfo, IServiceProvider inServiceProvider, UserSettings inSettings, OidcTokenManager? inOidcTokenManager)
+		public WorkspaceControl(IWorkspaceControlOwner owner, DirectoryReference appDataFolder, string? apiUrl, OpenProjectInfo openProjectInfo, IServiceProvider serviceProvider, UserSettings settings)
 		{
 			InitializeComponent();
 
 			_mainThreadSynchronizationContext = SynchronizationContext.Current!;
 
-			_owner = inOwner;
-			_appDataFolder = inAppDataFolder;
-			_apiUrl = inApiUrl;
+			_owner = owner;
+			_appDataFolder = appDataFolder;
+			_apiUrl = apiUrl;
 			_workspaceDataFolder = openProjectInfo.ProjectInfo.DataFolder;
-			_serviceProvider = inServiceProvider;
-			_logger = inServiceProvider.GetRequiredService<ILogger<WorkspaceControl>>();
+			_serviceProvider = serviceProvider;
+			_logger = serviceProvider.GetRequiredService<ILogger<WorkspaceControl>>();
 			_perforceSettings = openProjectInfo.PerforceSettings;
 			ProjectInfo = openProjectInfo.ProjectInfo;
-			_settings = inSettings;
+			_settings = settings;
 			_workspaceSettings = openProjectInfo.WorkspaceSettings;
-			_projectSettings = inSettings.FindOrAddProjectSettings(openProjectInfo.ProjectInfo, openProjectInfo.WorkspaceSettings, _logger);
+			_projectSettings = settings.FindOrAddProjectSettings(openProjectInfo.ProjectInfo, openProjectInfo.WorkspaceSettings, _logger);
 
 			DesiredTaskbarState = Tuple.Create(TaskbarState.NoProgress, 0.0f);
 
@@ -368,7 +368,7 @@ namespace UnrealGameSync
 			FileReference projectLogBaseName = FileReference.Combine(_workspaceDataFolder, "sync.log");
 
 			ILogger perforceLogger = _serviceProvider.GetRequiredService<ILogger<PerforceMonitor>>();
-			_perforceMonitor = new PerforceMonitor(perforceClientSettings, openProjectInfo.ProjectInfo, openProjectInfo.LatestProjectConfigFile, openProjectInfo.ProjectInfo.CacheFolder, openProjectInfo.LocalConfigFiles, _serviceProvider);
+			_perforceMonitor = new PerforceMonitor(perforceClientSettings, openProjectInfo.ProjectInfo, openProjectInfo.LatestProjectConfigFile, openProjectInfo.ProjectInfo.CacheFolder, openProjectInfo.LocalConfigFiles, openProjectInfo.OidcTokenClient, _serviceProvider);
 			_perforceMonitor.OnUpdate += UpdateBuildListCallback;
 			_perforceMonitor.OnUpdateMetadata += UpdateBuildMetadataCallback;
 			_perforceMonitor.OnStreamChange += StreamChanged;
@@ -378,8 +378,8 @@ namespace UnrealGameSync
 			_eventMonitor = new EventMonitor(_apiUrl, PerforceUtils.GetClientOrDepotDirectoryName(SelectedProjectIdentifier), openProjectInfo.PerforceSettings.UserName, _serviceProvider);
 			_eventMonitor.OnUpdatesReady += UpdateReviewsCallback;
 
-			ILogger<JupiterMonitor> jupiterLogger = _serviceProvider.GetRequiredService<ILogger<JupiterMonitor>>();
-			_jupiterMonitor = JupiterMonitor.CreateFromConfigFile(inOidcTokenManager, jupiterLogger, openProjectInfo.LatestProjectConfigFile, SelectedProjectIdentifier);
+//			ILogger<JupiterMonitor> jupiterLogger = _serviceProvider.GetRequiredService<ILogger<JupiterMonitor>>();
+//			_jupiterMonitor = JupiterMonitor.CreateFromConfigFile(inOidcTokenManager, jupiterLogger, openProjectInfo.LatestProjectConfigFile, SelectedProjectIdentifier);
 
 			UpdateColumnSettings(true);
 
@@ -407,7 +407,7 @@ namespace UnrealGameSync
 			_startupCallbacks = new List<WorkspaceStartupCallback>();
 
 			string? issuesApiUrl = GetIssuesApiUrl();
-			_issueMonitor = inOwner.CreateIssueMonitor(issuesApiUrl, openProjectInfo.PerforceSettings.UserName);
+			_issueMonitor = owner.CreateIssueMonitor(issuesApiUrl, openProjectInfo.PerforceSettings.UserName);
 			_issueMonitor.OnIssuesChanged += IssueMonitor_OnIssuesChangedAsync;
 
 			if (SelectedFileName.HasExtension(".uproject"))
@@ -758,11 +758,11 @@ namespace UnrealGameSync
 				_badgeFont.Dispose();
 				_badgeFont = null!;
 			}
-			if (_jupiterMonitor != null)
-			{
-				_jupiterMonitor.Dispose();
-				_jupiterMonitor = null!;
-			}
+//			if (_jupiterMonitor != null)
+//			{
+//				_jupiterMonitor.Dispose();
+//				_jupiterMonitor = null!;
+//			}
 
 			base.Dispose(disposing);
 		}
@@ -3342,6 +3342,23 @@ namespace UnrealGameSync
 				StatusLine projectLine = new StatusLine();
 				projectLine.AddText(String.Format("Opened "));
 				projectLine.AddLink(SelectedFileName.FullName + " \u25BE", FontStyle.Regular, (p, r) => { SelectRecentProject(r); });
+
+				OidcTokenClient? oidcClient = _perforceMonitor.LatestOidcTokenClient;
+				if (oidcClient != null)
+				{
+				projectLine.AddText("  |  ");
+
+					StatusLine hordeLine = new StatusLine();
+					if (oidcClient.GetStatus() == OidcStatus.Connected)
+					{
+						projectLine.AddText("Connected to Horde.");
+					}
+					else
+					{
+						projectLine.AddLink("Connect to Horde", FontStyle.Bold | FontStyle.Underline, (p, r) => DoOidcLogin());
+					}
+				}
+
 				projectLine.AddText("  |  ");
 				projectLine.AddLink("Settings...", FontStyle.Regular, (p, r) => { _owner.EditSelectedProject(this); });
 				lines.Add(projectLine);
@@ -3616,6 +3633,26 @@ namespace UnrealGameSync
 			StatusPanel.Set(lines, caption, alert, tintColor);
 		}
 
+		private void DoOidcLogin()
+		{
+			OidcTokenClient? oidcClient = _perforceMonitor.LatestOidcTokenClient;
+			if (oidcClient != null)
+			{
+				ModalTask.Execute(this, "OIDC Login", "Opening OIDC login page (check your browser).", async cancellationToken =>
+				{
+					try
+					{
+						await oidcClient.LoginAsync(cancellationToken);
+					}
+					catch when (cancellationToken.IsCancellationRequested)
+					{
+						// nop
+					}
+				});
+				UpdateStatusPanel();
+			}
+		}
+
 		private void ShowBuildHealthMenu(Rectangle bounds)
 		{
 			int minSeparatorIdx = BuildHealthContextMenu.Items.IndexOf(BuildHealthContextMenu_MinSeparator);
@@ -4482,11 +4519,11 @@ namespace UnrealGameSync
 
 		private IReadOnlyList<IArchiveInfo> GetArchives()
 		{
-			IReadOnlyList<IArchiveInfo>? availableArchives = _jupiterMonitor?.AvailableArchives;
-			if (availableArchives != null && availableArchives.Count != 0)
-			{
-				return availableArchives;
-			}
+//			IReadOnlyList<IArchiveInfo>? availableArchives = _jupiterMonitor?.AvailableArchives;
+//			if (availableArchives != null && availableArchives.Count != 0)
+//			{
+//				return availableArchives;
+//			}
 
 			// if jupiter had no archives we fallback to the perforce monitor
 			if (_perforceMonitor != null)
diff --git a/Engine/Source/Programs/UnrealGameSync/UnrealGameSync/Forms/MainWindow.cs b/Engine/Source/Programs/UnrealGameSync/UnrealGameSync/Forms/MainWindow.cs
--- a/Engine/Source/Programs/UnrealGameSync/UnrealGameSync/Forms/MainWindow.cs
+++ b/Engine/Source/Programs/UnrealGameSync/UnrealGameSync/Forms/MainWindow.cs
@@ -103,7 +103,7 @@ namespace UnrealGameSync
 
 		bool _allowClose = false;
 
-		readonly OidcTokenManager? _oidcTokenManager;
+		readonly OidcTokenManager _oidcTokenManager;
 
 		System.Threading.Timer? _scheduleTimer;
 		System.Threading.Timer? _scheduleSettledTimer;
@@ -122,26 +122,26 @@ namespace UnrealGameSync
 
 		public ToolUpdateMonitor ToolUpdateMonitor { get; private set; }
 
-		public MainWindow(UpdateMonitor inUpdateMonitor, string? inApiUrl, DirectoryReference inDataFolder, string inOriginalExecutableFileName, bool inPreview, List<(UserSelectedProjectSettings, ModalTask<OpenProjectInfo>)> startupTasks, IPerforceSettings inDefaultPerforceSettings, IServiceProvider inServiceProvider, UserSettings inSettings, string? inUri, OidcTokenManager? inOidcTokenManager)
+		public MainWindow(UpdateMonitor inUpdateMonitor, string? apiUrl, DirectoryReference dataFolder, string originalExecutableFileName, bool preview, List<(UserSelectedProjectSettings, ModalTask<OpenProjectInfo>)> startupTasks, IPerforceSettings defaultPerforceSettings, IServiceProvider serviceProvider, UserSettings settings, string? uri)
 		{
-			_serviceProvider = inServiceProvider;
+			_serviceProvider = serviceProvider;
 			_logger = _serviceProvider.GetRequiredService<ILogger<MainWindow>>();
 
-			_logger.LogInformation("Opening Main Window for {NumProject} projects. Last Project {LastProject}", startupTasks.Count, inSettings.LastProject);
+			_logger.LogInformation("Opening Main Window for {NumProject} projects. Last Project {LastProject}", startupTasks.Count, settings.LastProject);
 
 			InitializeComponent();
 
 			_updateMonitor = inUpdateMonitor;
 			_mainThreadSynchronizationContext = SynchronizationContext.Current!;
-			_apiUrl = inApiUrl;
-			_dataFolder = inDataFolder;
-			_originalExecutableFileName = inOriginalExecutableFileName;
-			_preview = inPreview;
-			_defaultPerforceSettings = inDefaultPerforceSettings;
-			ToolUpdateMonitor = new ToolUpdateMonitor(_defaultPerforceSettings, _dataFolder, inSettings, _serviceProvider);
+			_apiUrl = apiUrl;
+			_dataFolder = dataFolder;
+			_originalExecutableFileName = originalExecutableFileName;
+			_preview = preview;
+			_defaultPerforceSettings = defaultPerforceSettings;
+			ToolUpdateMonitor = new ToolUpdateMonitor(_defaultPerforceSettings, _dataFolder, settings, _serviceProvider);
 
-			_settings = inSettings;
-			_oidcTokenManager = inOidcTokenManager;
+			_settings = settings;
+			_oidcTokenManager = _serviceProvider.GetRequiredService<OidcTokenManager>();
 
 			// While creating tab controls during startup, we need to prevent layout calls resulting in the window handle being created too early. Disable layout calls here.
 			SuspendLayout();
@@ -194,7 +194,7 @@ namespace UnrealGameSync
 			}
 
 			ILogger<AutomationServer> automationLogger = _serviceProvider.GetRequiredService<ILogger<AutomationServer>>();
-			_automationServer = new AutomationServer(request => { _mainThreadSynchronizationContext.Post(obj => PostAutomationRequest(request), null); }, inUri, automationLogger);
+			_automationServer = new AutomationServer(request => { _mainThreadSynchronizationContext.Post(obj => PostAutomationRequest(request), null); }, uri, automationLogger);
 
 			// Allow creating controls from now on
 			TabPanel.ResumeLayout(false);
@@ -202,7 +202,7 @@ namespace UnrealGameSync
 
 			foreach (string defaultIssueApiUrl in DeploymentSettings.DefaultIssueApiUrls)
 			{
-				_defaultIssueMonitors.Add(CreateIssueMonitor(defaultIssueApiUrl, inDefaultPerforceSettings.UserName));
+				_defaultIssueMonitors.Add(CreateIssueMonitor(defaultIssueApiUrl, defaultPerforceSettings.UserName));
 			}
 
 			_allowCreatingHandle = true;
@@ -1222,7 +1222,7 @@ namespace UnrealGameSync
 
 			PerforceSettings perforceSettings = Utility.OverridePerforceSettings(_defaultPerforceSettings, project.ServerAndPort, project.UserName);
 
-			ModalTask<OpenProjectInfo>? settingsTask = PerforceModalTask.Execute(this, "Opening Project", "Opening project, please wait...", perforceSettings, (p, c) => OpenProjectWindow.DetectSettingsAsync(p, project, _settings, projectLogger, c), projectLogger, taskFlags);
+			ModalTask<OpenProjectInfo>? settingsTask = PerforceModalTask.Execute(this, "Opening Project", "Opening project, please wait...", perforceSettings, (p, c) => OpenProjectWindow.DetectSettingsAsync(p, project, _settings, _oidcTokenManager, projectLogger, c), projectLogger, taskFlags);
 			if (settingsTask == null || settingsTask.Failed)
 			{
 				if (settingsTask != null)
@@ -1289,7 +1289,7 @@ namespace UnrealGameSync
 			}
 
 			// Now that we have the project settings, we can construct the tab
-			WorkspaceControl newWorkspace = new WorkspaceControl(this, _dataFolder, _apiUrl, openProjectInfo, _serviceProvider, _settings, _oidcTokenManager);
+			WorkspaceControl newWorkspace = new WorkspaceControl(this, _dataFolder, _apiUrl, openProjectInfo, _serviceProvider, _settings);
 			
 			newWorkspace.Parent = TabPanel;
 			newWorkspace.Dock = DockStyle.Fill;
diff --git a/Engine/Source/Programs/UnrealGameSync/UnrealGameSync/Forms/OIDCLoginWindow.cs b/Engine/Source/Programs/UnrealGameSync/UnrealGameSync/Forms/OIDCLoginWindow.cs
deleted file mode 100644
diff --git a/Engine/Source/Programs/UnrealGameSync/UnrealGameSync/Forms/OpenProjectWindow.cs b/Engine/Source/Programs/UnrealGameSync/UnrealGameSync/Forms/OpenProjectWindow.cs
--- a/Engine/Source/Programs/UnrealGameSync/UnrealGameSync/Forms/OpenProjectWindow.cs
+++ b/Engine/Source/Programs/UnrealGameSync/UnrealGameSync/Forms/OpenProjectWindow.cs
@@ -1,5 +1,6 @@
 // Copyright Epic Games, Inc. All Rights Reserved.
 
+using EpicGames.OIDC;
 using EpicGames.Perforce;
 using Microsoft.Extensions.DependencyInjection;
 using Microsoft.Extensions.Logging;
@@ -277,10 +278,11 @@ namespace UnrealGameSync
 			if(TryGetSelectedProject(out selectedProject))
 			{
 				ILogger<OpenProjectInfo> logger = _serviceProvider.GetRequiredService<ILogger<OpenProjectInfo>>();
+				OidcTokenManager oidcTokenManager = _serviceProvider.GetRequiredService<OidcTokenManager>();
 
 				PerforceSettings newPerforceSettings = Utility.OverridePerforceSettings(Perforce, selectedProject.ServerAndPort, selectedProject.UserName);
 
-				ModalTask<OpenProjectInfo>? newOpenProjectInfo = PerforceModalTask.Execute(this, "Opening project", "Opening project, please wait...", newPerforceSettings, (x, y) => DetectSettingsAsync(x, selectedProject, _settings, logger, y), logger);
+				ModalTask<OpenProjectInfo>? newOpenProjectInfo = PerforceModalTask.Execute(this, "Opening project", "Opening project, please wait...", newPerforceSettings, (x, y) => DetectSettingsAsync(x, selectedProject, _settings, oidcTokenManager, logger, y), logger);
 				if (newOpenProjectInfo != null && newOpenProjectInfo.Succeeded)
 				{
 					_openProjectInfo = newOpenProjectInfo.Result;
@@ -290,9 +292,9 @@ namespace UnrealGameSync
 			}
 		}
 
-		public static async Task<OpenProjectInfo> DetectSettingsAsync(IPerforceConnection perforce, UserSelectedProjectSettings selectedProject, UserSettings userSettings, ILogger<OpenProjectInfo> logger, CancellationToken cancellationToken)
+		public static async Task<OpenProjectInfo> DetectSettingsAsync(IPerforceConnection perforce, UserSelectedProjectSettings selectedProject, UserSettings userSettings, OidcTokenManager oidcTokenManager, ILogger<OpenProjectInfo> logger, CancellationToken cancellationToken)
 		{
-			OpenProjectInfo settings = await OpenProjectInfo.CreateAsync(perforce, selectedProject, userSettings, logger, cancellationToken);
+			OpenProjectInfo settings = await OpenProjectInfo.CreateAsync(perforce, selectedProject, userSettings, oidcTokenManager, logger, cancellationToken);
 			if (DeploymentSettings.OnDetectProjectSettings != null)
 			{
 				string? message;
diff --git a/Engine/Source/Programs/UnrealGameSync/UnrealGameSync/Jupiter/JupiterMonitor.cs b/Engine/Source/Programs/UnrealGameSync/UnrealGameSync/Jupiter/JupiterMonitor.cs
deleted file mode 100644
diff --git a/Engine/Source/Programs/UnrealGameSync/UnrealGameSync/PerforceMonitor.cs b/Engine/Source/Programs/UnrealGameSync/UnrealGameSync/PerforceMonitor.cs
--- a/Engine/Source/Programs/UnrealGameSync/UnrealGameSync/PerforceMonitor.cs
+++ b/Engine/Source/Programs/UnrealGameSync/UnrealGameSync/PerforceMonitor.cs
@@ -1,6 +1,7 @@
 // Copyright Epic Games, Inc. All Rights Reserved.
 
 using EpicGames.Core;
+using EpicGames.OIDC;
 using EpicGames.Perforce;
 using Microsoft.Extensions.DependencyInjection;
 using Microsoft.Extensions.Logging;
@@ -9,6 +10,7 @@ using System.Collections.Generic;
 using System.Diagnostics;
 using System.Diagnostics.CodeAnalysis;
 using System.Linq;
+using System.Text;
 using System.Threading;
 using System.Threading.Tasks;
 
@@ -49,6 +51,8 @@ namespace UnrealGameSync
 		readonly DirectoryReference _cacheFolder;
 		readonly List<KeyValuePair<FileReference, DateTime>> _localConfigFiles;
 		readonly IAsyncDisposer _asyncDisposeTasks;
+		readonly OidcTokenManager _oidcTokenManager;
+
 		string[] prevCodeRules = Array.Empty<string>();
 
 		readonly SynchronizationContext _synchronizationContext;
@@ -63,25 +67,27 @@ namespace UnrealGameSync
 			private set;
 		}
 
-		public PerforceMonitor(IPerforceSettings inPerforceSettings, ProjectInfo projectInfo, ConfigFile inProjectConfigFile, DirectoryReference inCacheFolder, List<KeyValuePair<FileReference, DateTime>> inLocalConfigFiles, IServiceProvider inServiceProvider)
+		public PerforceMonitor(IPerforceSettings perforceSettings, ProjectInfo projectInfo, ConfigFile projectConfigFile, DirectoryReference cacheFolder, List<KeyValuePair<FileReference, DateTime>> localConfigFiles, OidcTokenClient? oidcTokenClient, IServiceProvider serviceProvider)
 		{
-			_perforceSettings = inPerforceSettings;
+			_perforceSettings = perforceSettings;
 			_branchClientPath = projectInfo.ClientRootPath;
 			_selectedClientFileName = projectInfo.ClientFileName;
 			_selectedProjectIdentifier = projectInfo.ProjectIdentifier;
 			_pendingMaxChangesValue = InitialMaxChangesValue;
 			LastChangeByCurrentUser = -1;
 			LastCodeChangeByCurrentUser = -1;
-			_logger = inServiceProvider.GetRequiredService<ILogger<PerforceMonitor>>();
+			_logger = serviceProvider.GetRequiredService<ILogger<PerforceMonitor>>();
 			_isEnterpriseProject = projectInfo.IsEnterpriseProject;
-			LatestProjectConfigFile = inProjectConfigFile;
-			_cacheFolder = inCacheFolder;
-			_localConfigFiles = inLocalConfigFiles;
-			_asyncDisposeTasks = inServiceProvider.GetRequiredService<IAsyncDisposer>();
+			LatestProjectConfigFile = projectConfigFile;
+			_cacheFolder = cacheFolder;
+			_localConfigFiles = localConfigFiles;
+			_asyncDisposeTasks = serviceProvider.GetRequiredService<IAsyncDisposer>();
 			_synchronizationContext = SynchronizationContext.Current!;
 			_cancellationSource = new CancellationTokenSource();
+			_oidcTokenManager = serviceProvider.GetRequiredService<OidcTokenManager>();
 
 			AvailableArchives = (new List<IArchiveInfo>()).AsReadOnly();
+			LatestOidcTokenClient = oidcTokenClient;
 		}
 
 		public void Start()
@@ -525,6 +531,7 @@ namespace UnrealGameSync
 			if(_localConfigFiles.Any(x => FileReference.GetLastWriteTimeUtc(x.Key) != x.Value))
 			{
 				await UpdateProjectConfigFileAsync(perforce, cancellationToken);
+				// TODO: Also check OIDC config
 				_synchronizationContext.Post(_ => OnUpdateMetadata?.Invoke(), null);
 			}
 
@@ -554,6 +561,7 @@ namespace UnrealGameSync
 		{
 			_localConfigFiles.Clear();
 			LatestProjectConfigFile = await ConfigUtils.ReadProjectConfigFileAsync(perforce, _branchClientPath, _selectedClientFileName, _cacheFolder, _localConfigFiles, _logger, cancellationToken);
+			LatestOidcTokenClient = await ConfigUtils.CreateOidcTokenClientAsync(_oidcTokenManager, LatestProjectConfigFile, _selectedProjectIdentifier, perforce, _branchClientPath, _selectedClientFileName, _localConfigFiles, _cacheFolder, _logger, cancellationToken);
 		}
 
 		public List<ChangesRecord> GetChanges()
@@ -603,6 +611,12 @@ namespace UnrealGameSync
 			private set;
 		}
 
+		public OidcTokenClient? LatestOidcTokenClient
+		{
+			get;
+			private set;
+		}
+
 		public IReadOnlyList<IArchiveInfo> AvailableArchives
 		{
 			get;
diff --git a/Engine/Source/Programs/UnrealGameSync/UnrealGameSync/Program.cs b/Engine/Source/Programs/UnrealGameSync/UnrealGameSync/Program.cs
--- a/Engine/Source/Programs/UnrealGameSync/UnrealGameSync/Program.cs
+++ b/Engine/Source/Programs/UnrealGameSync/UnrealGameSync/Program.cs
@@ -1,6 +1,7 @@
 // Copyright Epic Games, Inc. All Rights Reserved.
 
 using EpicGames.Core;
+using EpicGames.OIDC;
 using EpicGames.Perforce;
 using Microsoft.Extensions.DependencyInjection;
 using Microsoft.Extensions.Logging;
@@ -168,6 +169,8 @@ namespace UnrealGameSync
 				ServiceCollection services = new ServiceCollection();
 				services.AddLogging(builder => builder.AddProvider(loggerProvider));
 				services.AddSingleton<IAsyncDisposer, AsyncDisposer>();
+				services.AddSingleton(sp => TokenStoreFactory.CreateTokenStore());
+				services.AddSingleton<OidcTokenManager>();
 
 				await using (ServiceProvider serviceProvider = services.BuildServiceProvider())
 				{
diff --git a/Engine/Source/Programs/UnrealGameSync/UnrealGameSync/ProgramApplicationContext.cs b/Engine/Source/Programs/UnrealGameSync/UnrealGameSync/ProgramApplicationContext.cs
--- a/Engine/Source/Programs/UnrealGameSync/UnrealGameSync/ProgramApplicationContext.cs
+++ b/Engine/Source/Programs/UnrealGameSync/UnrealGameSync/ProgramApplicationContext.cs
@@ -16,6 +16,7 @@ using System.Threading.Tasks;
 using System.Windows.Forms;
 using EpicGames.OIDC;
 using UnrealGameSync.Forms;
+using Microsoft.Extensions.Logging.Abstractions;
 
 #nullable enable
 
@@ -56,9 +57,6 @@ namespace UnrealGameSync
 		ModalTaskWindow? _startupWindow;
 		MainWindow? _mainWindowInstance;
 
-		ITokenStore? _tokenStore;
-		OidcTokenManager? _oidcTokenManager;
-
 		public ProgramApplicationContext(IPerforceSettings defaultPerforceSettings, UpdateMonitor updateMonitor, string? apiUrl, DirectoryReference dataFolder, EventWaitHandle activateEvent, bool restoreState, string? updateSpawn, string? projectFileName, bool preview, IServiceProvider serviceProvider, string? uri)
 		{
 			_defaultPerforceSettings = defaultPerforceSettings;
@@ -189,7 +187,8 @@ namespace UnrealGameSync
 			foreach (UserSelectedProjectSettings projectSettings in _settings.OpenProjects)
 			{
 				ILogger<OpenProjectInfo> logger = serviceProvider.GetRequiredService<ILogger<OpenProjectInfo>>();
-				Task<OpenProjectInfo> startupTask = Task.Run(() => OpenProjectInfo.CreateAsync(defaultPerforceSettings, projectSettings, _settings, logger, _startupCancellationSource.Token), _startupCancellationSource.Token);
+				OidcTokenManager oidcTokenManager = serviceProvider.GetRequiredService<OidcTokenManager>();
+				Task<OpenProjectInfo> startupTask = Task.Run(() => OpenProjectInfo.CreateAsync(defaultPerforceSettings, projectSettings, _settings, oidcTokenManager, logger, _startupCancellationSource.Token), _startupCancellationSource.Token);
 				startupTasks.Add((projectSettings, new ModalTask<OpenProjectInfo>(startupTask)));
 			}
 			_startupTask = Task.Run(() => WaitForStartupTasks(startupTasks));
@@ -257,57 +256,6 @@ namespace UnrealGameSync
 			// Clear out the cache folder
 			Utility.ClearPrintCache(_cacheFolder);
 
-			_tokenStore = TokenStoreFactory.CreateTokenStore();
-			List<(DirectoryInfo, DirectoryInfo?)> configurationLocations = new List<(DirectoryInfo, DirectoryInfo?)>();
-			List<string> allowedProviders = new List<string>();
-			foreach ((UserSelectedProjectSettings, ModalTask<OpenProjectInfo>) startupTask in startupTasks)
-			{
-				try
-				{
-
-					ConfigFile configFile = startupTask.Item2.Result.LatestProjectConfigFile;
-					if (configFile == null)
-					{
-						continue;
-					}
-
-					ConfigSection? providerSection = configFile.FindSection("OIDCProvider");
-					if (providerSection == null)
-					{
-						continue;
-					}
-
-					string[] oidcAllowedProviders = providerSection.GetValues("OidcProviderAllowList", Array.Empty<string>());
-					if (oidcAllowedProviders.Length == 0)
-					{
-						continue;
-					}
-
-					allowedProviders.AddRange(oidcAllowedProviders);
-					
-					ProjectInfo projectInfo = startupTask.Item2.Result.ProjectInfo;
-					configurationLocations.Add((projectInfo.EngineDir.ToDirectoryInfo(), projectInfo.ProjectDir?.ToDirectoryInfo()));
-				}
-				catch (Exception)
-				{
-					// ignore any projects that failed to load
-				}
-			}
-
-			_oidcTokenManager = OidcTokenManager.CreateTokenManager(
-				ProviderConfigurationFactory.MergeConfiguration(configurationLocations), 
-				_tokenStore,
-				allowedProviders
-			);
-			// Verify that none of the projects we are opening needs a OIDC login, if they do prompt for the login
-			if (_oidcTokenManager?.HasUnfinishedLogin() ?? false)
-			{
-				OidcLoginWindow loginDialog = new OidcLoginWindow(_oidcTokenManager);
-				loginDialog.ShowDialog();
-
-				_tokenStore.Save();
-			}
-
 			// Get the application path
 			string originalExe = Assembly.GetExecutingAssembly().Location;
 			if (Path.GetExtension(originalExe).Equals(".dll", StringComparison.OrdinalIgnoreCase))
@@ -320,7 +268,7 @@ namespace UnrealGameSync
 			}
 
 			// Create the main window 
-			_mainWindowInstance = new MainWindow(_updateMonitor, _apiUrl, _dataFolder, _updateSpawn ?? originalExe, _preview, startupTasks, _defaultPerforceSettings, _serviceProvider, _settings, _uri, _oidcTokenManager);
+			_mainWindowInstance = new MainWindow(_updateMonitor, _apiUrl, _dataFolder, _updateSpawn ?? originalExe, _preview, startupTasks, _defaultPerforceSettings, _serviceProvider, _settings, _uri);
 			if(visible)
 			{
 				_mainWindowInstance.Show();
diff --git a/Engine/Source/Programs/UnrealGameSync/UnrealGameSync/UnrealGameSync.csproj b/Engine/Source/Programs/UnrealGameSync/UnrealGameSync/UnrealGameSync.csproj
--- a/Engine/Source/Programs/UnrealGameSync/UnrealGameSync/UnrealGameSync.csproj
+++ b/Engine/Source/Programs/UnrealGameSync/UnrealGameSync/UnrealGameSync.csproj
@@ -30,7 +30,6 @@
   </ItemGroup>
   <ItemGroup>
     <ProjectReference Include="..\Shared\EpicGames.Core\EpicGames.Core.csproj" />
-    <ProjectReference Include="..\Shared\EpicGames.Jupiter\EpicGames.Jupiter.csproj" />
     <ProjectReference Include="..\Shared\EpicGames.OIDC\EpicGames.OIDC.csproj" />
     <ProjectReference Include="..\Shared\EpicGames.Perforce\EpicGames.Perforce.csproj" />
     <ProjectReference Include="..\UnrealGameSyncShared\UnrealGameSyncShared.csproj" />
diff --git a/Engine/Source/Programs/UnrealGameSync/UnrealGameSyncShared/ConfigUtils.cs b/Engine/Source/Programs/UnrealGameSync/UnrealGameSyncShared/ConfigUtils.cs
--- a/Engine/Source/Programs/UnrealGameSync/UnrealGameSyncShared/ConfigUtils.cs
+++ b/Engine/Source/Programs/UnrealGameSync/UnrealGameSyncShared/ConfigUtils.cs
@@ -1,7 +1,9 @@
 // Copyright Epic Games, Inc. All Rights Reserved.
 
 using EpicGames.Core;
+using EpicGames.OIDC;
 using EpicGames.Perforce;
+using Microsoft.Extensions.Configuration;
 using Microsoft.Extensions.Logging;
 using System;
 using System.Collections.Generic;
@@ -9,6 +11,7 @@ using System.Diagnostics.CodeAnalysis;
 using System.IO;
 using System.Linq;
 using System.Runtime.InteropServices;
+using System.Text;
 using System.Threading;
 using System.Threading.Tasks;
 
@@ -142,6 +145,101 @@ namespace UnrealGameSync
 			return projectConfig;
 		}
 
+		public static async Task<List<string[]>> ReadConfigFiles(IPerforceConnection perforce, IEnumerable<string> depotPaths, List<KeyValuePair<FileReference, DateTime>> localFiles, DirectoryReference cacheFolder, ILogger logger, CancellationToken cancellationToken)
+		{
+			List<string[]> contents = new List<string[]>();
+
+			List<PerforceResponse<FStatRecord>> responses = await perforce.TryFStatAsync(FStatOptions.IncludeFileSizes, depotPaths.ToArray(), cancellationToken).ToListAsync(cancellationToken);
+			foreach (PerforceResponse<FStatRecord> response in responses)
+			{
+				if (response.Succeeded)
+				{
+					string[]? lines = null;
+
+					// Skip file records which are still in the workspace, but were synced from a different branch. For these files, the action seems to be empty, so filter against that.
+					FStatRecord fileRecord = response.Data;
+					if (fileRecord.HeadAction == FileAction.None)
+					{
+						continue;
+					}
+
+					// If this file is open for edit, read the local version
+					string? localFileName = fileRecord.ClientFile;
+					if (localFileName != null && File.Exists(localFileName) && (File.GetAttributes(localFileName) & FileAttributes.ReadOnly) == 0)
+					{
+						try
+						{
+							DateTime lastModifiedTime = File.GetLastWriteTimeUtc(localFileName);
+							localFiles.Add(new KeyValuePair<FileReference, DateTime>(new FileReference(localFileName), lastModifiedTime));
+							lines = await File.ReadAllLinesAsync(localFileName, cancellationToken);
+						}
+						catch (Exception ex)
+						{
+							logger.LogInformation(ex, "Failed to read local config file for {Path}", localFileName);
+						}
+					}
+
+					// Otherwise try to get it from perforce
+					if (lines == null && fileRecord.DepotFile != null)
+					{
+						lines = await Utility.TryPrintFileUsingCacheAsync(perforce, fileRecord.DepotFile, cacheFolder, fileRecord.Digest, logger, cancellationToken);
+					}
+
+					// Merge the text with the config file
+					if (lines != null)
+					{
+						contents.Add(lines);
+					}
+				}
+			}
+
+			return contents;
+		}
+
+		public static async Task<OidcTokenClient?> CreateOidcTokenClientAsync(OidcTokenManager oidcTokenManager, ConfigFile projectConfigFile, string projectIdentifier, IPerforceConnection perforce, string clientRoot, string? clientProjectFile, List<KeyValuePair<FileReference, DateTime>> localFiles, DirectoryReference cacheFolder, ILogger logger, CancellationToken cancellationToken)
+		{
+			OidcTokenClient? oidcTokenClient = null;
+
+			string? oidcProvider;
+			if (TryGetProjectSetting(projectConfigFile, projectIdentifier, "OidcProvider", out oidcProvider))
+			{
+				List<string> configFiles = new List<string>();
+				configFiles.AddRange(ProviderConfigurationFactory.ConfigPaths.Select(x => $"{clientRoot}/Engine/{x}"));
+
+				if (clientProjectFile != null && clientProjectFile.EndsWith(".uproject", StringComparison.OrdinalIgnoreCase))
+				{
+					string clientProjectPath = clientProjectFile.Substring(0, clientProjectFile.LastIndexOf('/'));
+					configFiles.AddRange(ProviderConfigurationFactory.ConfigPaths.Select(x => $"{clientProjectPath}/{x}"));
+				}
+
+				List<string[]> latestOidcConfigFiles = await ConfigUtils.ReadConfigFiles(perforce, configFiles, localFiles, cacheFolder, logger, cancellationToken);
+
+				ConfigurationBuilder configBuilder = new ConfigurationBuilder();
+				foreach (string[] configFile in latestOidcConfigFiles)
+				{
+					MemoryStream configStream = new MemoryStream(Encoding.UTF8.GetBytes(String.Join("\n", configFile)));
+					configBuilder.AddJsonStream(configStream);
+				}
+
+				IConfigurationRoot config = configBuilder.Build();
+				OidcTokenOptions oidcOptions = OidcTokenOptions.Bind(config);
+
+				if (oidcOptions.Providers.TryGetValue(oidcProvider, out ProviderInfo? providerInfo))
+				{
+					oidcTokenClient = oidcTokenManager.FindOrAddClient(oidcProvider, providerInfo);
+
+					// Trigger an update of the access token if necessary, just to make sure we have the most current state
+					try
+					{
+						await oidcTokenClient.GetAccessTokenAsync(cancellationToken);
+					}
+					catch { }
+				}
+			}
+
+			return oidcTokenClient;
+		}
+
 		public static FileReference GetEditorTargetFile(ProjectInfo projectInfo, ConfigFile projectConfig)
 		{
 			if (projectInfo.ProjectPath.EndsWith(".uproject", StringComparison.OrdinalIgnoreCase))
diff --git a/Engine/Source/Programs/UnrealGameSync/UnrealGameSyncShared/OpenProjectInfo.cs b/Engine/Source/Programs/UnrealGameSync/UnrealGameSyncShared/OpenProjectInfo.cs
--- a/Engine/Source/Programs/UnrealGameSync/UnrealGameSyncShared/OpenProjectInfo.cs
+++ b/Engine/Source/Programs/UnrealGameSync/UnrealGameSyncShared/OpenProjectInfo.cs
@@ -1,6 +1,7 @@
 // Copyright Epic Games, Inc. All Rights Reserved.
 
 using EpicGames.Core;
+using EpicGames.OIDC;
 using EpicGames.Perforce;
 using Microsoft.Extensions.Logging;
 using System;
@@ -24,8 +25,9 @@ namespace UnrealGameSync
 		public ConfigFile WorkspaceProjectConfigFile { get; }
 		public IReadOnlyList<string>? WorkspaceProjectStreamFilter { get; }
 		public List<KeyValuePair<FileReference, DateTime>> LocalConfigFiles { get; }
+		public OidcTokenClient? OidcTokenClient { get; }
 
-		public OpenProjectInfo(UserSelectedProjectSettings selectedProject, IPerforceSettings perforceSettings, ProjectInfo projectInfo, UserWorkspaceSettings workspaceSettings, WorkspaceStateWrapper workspaceStateWrapper, ConfigFile latestProjectConfigFile, ConfigFile workspaceProjectConfigFile, IReadOnlyList<string>? workspaceProjectStreamFilter, List<KeyValuePair<FileReference, DateTime>> localConfigFiles)
+		public OpenProjectInfo(UserSelectedProjectSettings selectedProject, IPerforceSettings perforceSettings, ProjectInfo projectInfo, UserWorkspaceSettings workspaceSettings, WorkspaceStateWrapper workspaceStateWrapper, ConfigFile latestProjectConfigFile, ConfigFile workspaceProjectConfigFile, IReadOnlyList<string>? workspaceProjectStreamFilter, List<KeyValuePair<FileReference, DateTime>> localConfigFiles, OidcTokenClient? oidcTokenClient)
 		{
 			SelectedProject = selectedProject;
 
@@ -37,9 +39,10 @@ namespace UnrealGameSync
 			WorkspaceProjectConfigFile = workspaceProjectConfigFile;
 			WorkspaceProjectStreamFilter = workspaceProjectStreamFilter;
 			LocalConfigFiles = localConfigFiles;
+			OidcTokenClient = oidcTokenClient;
 		}
 
-		public static async Task<OpenProjectInfo> CreateAsync(IPerforceSettings defaultPerforceSettings, UserSelectedProjectSettings selectedProject, UserSettings userSettings, ILogger<OpenProjectInfo> logger, CancellationToken cancellationToken)
+		public static async Task<OpenProjectInfo> CreateAsync(IPerforceSettings defaultPerforceSettings, UserSelectedProjectSettings selectedProject, UserSettings userSettings, OidcTokenManager oidcTokenManager, ILogger<OpenProjectInfo> logger, CancellationToken cancellationToken)
 		{
 			PerforceSettings perforceSettings = Utility.OverridePerforceSettings(defaultPerforceSettings, selectedProject.ServerAndPort, selectedProject.UserName);
 			using IPerforceConnection perforce = await PerforceConnection.CreateAsync(perforceSettings, logger);
@@ -52,10 +55,10 @@ namespace UnrealGameSync
 			}
 
 			// Execute like a regular task
-			return await CreateAsync(perforce, selectedProject, userSettings, logger, cancellationToken);
+			return await CreateAsync(perforce, selectedProject, userSettings, oidcTokenManager, logger, cancellationToken);
 		}
 
-		public static async Task<OpenProjectInfo> CreateAsync(IPerforceConnection defaultConnection, UserSelectedProjectSettings selectedProject, UserSettings userSettings, ILogger<OpenProjectInfo> logger, CancellationToken cancellationToken)
+		public static async Task<OpenProjectInfo> CreateAsync(IPerforceConnection defaultConnection, UserSelectedProjectSettings selectedProject, UserSettings userSettings, OidcTokenManager oidcTokenManager, ILogger<OpenProjectInfo> logger, CancellationToken cancellationToken)
 		{
 			using IDisposable loggerScope = logger.BeginScope("Project {SelectedProject}", selectedProject.ToString());
 			logger.LogInformation("Detecting settings for {Project}", selectedProject);
@@ -222,11 +225,14 @@ namespace UnrealGameSync
 				List<KeyValuePair<FileReference, DateTime>> localConfigFiles = new List<KeyValuePair<FileReference, DateTime>>();
 				ConfigFile latestProjectConfigFile = await ConfigUtils.ReadProjectConfigFileAsync(perforceClient, projectInfo, localConfigFiles, logger, cancellationToken);
 
+				// Read the initial OIDC config files
+				OidcTokenClient? oidcTokenClient = await ConfigUtils.CreateOidcTokenClientAsync(oidcTokenManager, latestProjectConfigFile, projectInfo.ProjectIdentifier, perforceClient, branchClientPath, newSelectedClientFileName, localConfigFiles, projectInfo.CacheFolder, logger, cancellationToken);
+
 				// Get the local config file and stream filter
 				ConfigFile workspaceProjectConfigFile = await WorkspaceUpdate.ReadProjectConfigFile(branchDirectoryName, newSelectedFileName, logger);
 				IReadOnlyList<string>? workspaceProjectStreamFilter = await WorkspaceUpdate.ReadProjectStreamFilter(perforceClient, workspaceProjectConfigFile, cancellationToken);
 
-				OpenProjectInfo workspaceSettings = new OpenProjectInfo(selectedProject, perforceSettings, projectInfo, userWorkspaceSettings, workspaceStateWrapper, latestProjectConfigFile, workspaceProjectConfigFile, workspaceProjectStreamFilter, localConfigFiles);
+				OpenProjectInfo workspaceSettings = new OpenProjectInfo(selectedProject, perforceSettings, projectInfo, userWorkspaceSettings, workspaceStateWrapper, latestProjectConfigFile, workspaceProjectConfigFile, workspaceProjectStreamFilter, localConfigFiles, oidcTokenClient);
 
 				return workspaceSettings;
 			}
