Revert change to construct OidcTokenClient instances directly; use DI instead.

diff --git a/Engine/Source/Programs/Shared/EpicGames.OIDC/OidcTokenManager.cs b/Engine/Source/Programs/Shared/EpicGames.OIDC/OidcTokenManager.cs
--- a/Engine/Source/Programs/Shared/EpicGames.OIDC/OidcTokenManager.cs
+++ b/Engine/Source/Programs/Shared/EpicGames.OIDC/OidcTokenManager.cs
@@ -45,7 +45,7 @@ namespace EpicGames.OIDC
 			return client;
 		}
 
-		public OidcTokenManager(IOptionsMonitor<OidcTokenOptions> settings, ITokenStore tokenStore, List<string>? allowedProviders = null)
+		public OidcTokenManager(IServiceProvider provider, IOptionsMonitor<OidcTokenOptions> settings, ITokenStore tokenStore, List<string>? allowedProviders = null)
 			: this(tokenStore)
 		{
 			Dictionary<string, string> refreshTokens = new Dictionary<string, string>();
@@ -65,7 +65,7 @@ namespace EpicGames.OIDC
 					continue;
 				}
 
-				OidcTokenClient tokenClient = new OidcTokenClient(key, providerInfo, _tokenStore);
+				OidcTokenClient tokenClient = ActivatorUtilities.CreateInstance<OidcTokenClient>(provider, key, providerInfo);
 
 				if (refreshTokens.TryGetValue(key, out string? refreshToken))
 				{
