EpicGames.Core: Backout ManagedProcess.cs, did not intend to submit it

#rnx

diff --git a/Engine/Source/Programs/Shared/EpicGames.Core/ManagedProcess.cs b/Engine/Source/Programs/Shared/EpicGames.Core/ManagedProcess.cs
--- a/Engine/Source/Programs/Shared/EpicGames.Core/ManagedProcess.cs
+++ b/Engine/Source/Programs/Shared/EpicGames.Core/ManagedProcess.cs
@@ -295,11 +295,6 @@ namespace EpicGames.Core
 		[DllImport("kernel32.dll")]
 		static extern bool IsProcessInJob(IntPtr hProcess, IntPtr hJob, out bool result);
 
-		[DllImport("kernel32.dll", SetLastError = true)]
-		static extern bool GenerateConsoleCtrlEvent(UInt32 dwCtrlEvent, UInt32 dwProcessGroupId);
-
-		public const UInt32 CTRL_C_EVENT = 0;
-
 		[DllImport("kernel32.dll", SetLastError = true)]
 		static extern int TerminateProcess(SafeHandleZeroOrMinusOneIsInvalid hProcess, uint uExitCode);
 
@@ -939,9 +934,7 @@ namespace EpicGames.Core
 		{
 			if(_processHandle != null)
 			{
-				_ = GenerateConsoleCtrlEvent(CTRL_C_EVENT, (uint)Id);
-				_ = WaitForSingleObject(_processHandle, 1000);
-				//_ = TerminateProcess(_processHandle, 0);
+				_ = TerminateProcess(_processHandle, 0);
 				_ = WaitForSingleObject(_processHandle, INFINITE);
 
 				_processHandle.Dispose();
