Improved checks for keywords placed in invalid conditional blocks.

diff --git a/Engine/Source/Programs/Shared/EpicGames.UHT/Parsers/UhtHeaderFileParser.cs b/Engine/Source/Programs/Shared/EpicGames.UHT/Parsers/UhtHeaderFileParser.cs
--- a/Engine/Source/Programs/Shared/EpicGames.UHT/Parsers/UhtHeaderFileParser.cs
+++ b/Engine/Source/Programs/Shared/EpicGames.UHT/Parsers/UhtHeaderFileParser.cs
@@ -3,8 +3,7 @@
 using System;
 using System.Collections.Generic;
 using System.Diagnostics.CodeAnalysis;
-using System.IO;
-using System.Runtime.CompilerServices;
+using System.Reflection;
 using EpicGames.Core;
 using EpicGames.UHT.Tables;
 using EpicGames.UHT.Tokenizer;
@@ -40,7 +39,7 @@ namespace EpicGames.UHT.Parsers
 	/// Compiler directives
 	/// </summary>
 	[Flags]
-	public enum UhtCompilerDirective
+	public enum UhtCompilerDirective : uint
 	{
 		/// <summary>
 		/// No compile directives
@@ -96,6 +95,21 @@ namespace EpicGames.UHT.Parsers
 		/// This directive is unrecognized and does not change the code generation at all
 		/// </summary>
 		Unrecognized = 1 << 9,
+
+		/// <summary>
+		/// The following flags are always ignored when keywords test for allowed conditional blocks
+		/// </summary>
+		AllowedCheckIgnoredFlags = CPPBlock | NotCPPBlock | ZeroBlock | OneBlock | WithHotReload,
+
+		/// <summary>
+		/// Default compiler directives to be allowed
+		/// </summary>
+		DefaultAllowedCheck = WithEditor | WithEditorOnlyData,
+
+		/// <summary>
+		/// All flags are allowed
+		/// </summary>
+		SilenceAllowedCheck = ~None,
 	}
 
 	/// <summary>
@@ -104,6 +118,85 @@ namespace EpicGames.UHT.Parsers
 	/// </summary>
 	public static class UhtCompilerDirectiveExtensions
 	{
+		private static readonly Lazy<List<string>> s_names = new (() =>
+		{
+			List<string> outList = new();
+			FieldInfo[] fields = typeof(UhtCompilerDirective).GetFields();
+			for (int bit = 0; bit < 32; ++bit)
+			{
+				bool found = false;
+				uint mask = (uint)1 << bit;
+				foreach (FieldInfo field in fields)
+				{
+					if (field.IsSpecialName)
+					{
+						continue;
+					}
+					object? value = field.GetValue(null);
+					if (value != null)
+					{
+						if (mask == (uint)value)
+						{
+							outList.Add(GetCompilerDirectiveText((UhtCompilerDirective)value));
+							found = true;
+							break;
+						}
+					}
+				}
+				if (!found)
+				{
+					outList.Add($"0x{mask:X8}");
+				}
+			}
+			return outList;
+		});
+
+		/// <summary>
+		/// Return the text associated with the given compiler directive
+		/// </summary>
+		/// <param name="compilerDirective">Directive in question</param>
+		/// <returns>String representation</returns>
+		public static string GetCompilerDirectiveText(this UhtCompilerDirective compilerDirective)
+		{
+			switch (compilerDirective)
+			{
+				case UhtCompilerDirective.CPPBlock: return "CPP";
+				case UhtCompilerDirective.NotCPPBlock: return "!CPP";
+				case UhtCompilerDirective.ZeroBlock: return "0";
+				case UhtCompilerDirective.OneBlock: return "1";
+				case UhtCompilerDirective.WithHotReload: return "WITH_HOT_RELOAD";
+				case UhtCompilerDirective.WithEditor: return "WITH_EDITOR";
+				case UhtCompilerDirective.WithEditorOnlyData: return "WITH_EDITORONLY_DATA";
+				case UhtCompilerDirective.WithEngine: return "WITH_ENGINE";
+				case UhtCompilerDirective.WithCoreUObject: return "WITH_COREUOBJECT";
+				default: return "<unrecognized>";
+			}
+		}
+
+		/// <summary>
+		/// Return a string list of the given compiler directives
+		/// </summary>
+		/// <param name="inFlags"></param>
+		/// <returns></returns>
+		public static List<string> ToStringList(this UhtCompilerDirective inFlags)
+		{
+			List<string> names = s_names.Value;
+			ulong intFlags = (ulong)inFlags;
+			List<string> outList = new();
+			for (int bit = 0; bit < 32; ++bit)
+			{
+				ulong mask = (ulong)1 << bit;
+				if (mask > intFlags)
+				{
+					break;
+				}
+				if ((mask & intFlags) != 0)
+				{
+					outList.Add(names[bit]);
+				}
+			}
+			return outList;
+		}
 
 		/// <summary>
 		/// Test to see if any of the specified flags are set
@@ -147,8 +240,8 @@ namespace EpicGames.UHT.Parsers
 	public static class UhtAccessSpecifierKeywords
 	{
 		#region Keywords
-		[UhtKeyword(Extends = UhtTableNames.ClassBase, Keyword = "public")]
-		[UhtKeyword(Extends = UhtTableNames.ScriptStruct, Keyword = "public")]
+		[UhtKeyword(Extends = UhtTableNames.ClassBase, Keyword = "public", AllowedCompilerDirectives = UhtCompilerDirective.SilenceAllowedCheck)]
+		[UhtKeyword(Extends = UhtTableNames.ScriptStruct, Keyword = "public", AllowedCompilerDirectives = UhtCompilerDirective.SilenceAllowedCheck)]
 		[SuppressMessage("CodeQuality", "IDE0051:Remove unused private members", Justification = "Attribute accessed method")]
 		[SuppressMessage("Style", "IDE0060:Remove unused parameter", Justification = "Attribute accessed method")]
 		private static UhtParseResult PublicKeyword(UhtParsingScope topScope, UhtParsingScope actionScope, ref UhtToken token)
@@ -156,8 +249,8 @@ namespace EpicGames.UHT.Parsers
 			return SetAccessSpecifier(topScope, UhtAccessSpecifier.Public);
 		}
 
-		[UhtKeyword(Extends = UhtTableNames.ClassBase, Keyword = "protected")]
-		[UhtKeyword(Extends = UhtTableNames.ScriptStruct, Keyword = "protected")]
+		[UhtKeyword(Extends = UhtTableNames.ClassBase, Keyword = "protected", AllowedCompilerDirectives = UhtCompilerDirective.SilenceAllowedCheck)]
+		[UhtKeyword(Extends = UhtTableNames.ScriptStruct, Keyword = "protected", AllowedCompilerDirectives = UhtCompilerDirective.SilenceAllowedCheck)]
 		[SuppressMessage("CodeQuality", "IDE0051:Remove unused private members", Justification = "Attribute accessed method")]
 		[SuppressMessage("Style", "IDE0060:Remove unused parameter", Justification = "Attribute accessed method")]
 		private static UhtParseResult ProtectedKeyword(UhtParsingScope topScope, UhtParsingScope actionScope, ref UhtToken token)
@@ -165,8 +258,8 @@ namespace EpicGames.UHT.Parsers
 			return SetAccessSpecifier(topScope, UhtAccessSpecifier.Protected);
 		}
 
-		[UhtKeyword(Extends = UhtTableNames.ClassBase, Keyword = "private")]
-		[UhtKeyword(Extends = UhtTableNames.ScriptStruct, Keyword = "private")]
+		[UhtKeyword(Extends = UhtTableNames.ClassBase, Keyword = "private", AllowedCompilerDirectives = UhtCompilerDirective.SilenceAllowedCheck)]
+		[UhtKeyword(Extends = UhtTableNames.ScriptStruct, Keyword = "private", AllowedCompilerDirectives = UhtCompilerDirective.SilenceAllowedCheck)]
 		[SuppressMessage("CodeQuality", "IDE0051:Remove unused private members", Justification = "Attribute accessed method")]
 		[SuppressMessage("Style", "IDE0060:Remove unused parameter", Justification = "Attribute accessed method")]
 		private static UhtParseResult PrivateKeyword(UhtParsingScope topScope, UhtParsingScope actionScope, ref UhtToken token)
@@ -540,6 +633,13 @@ namespace EpicGames.UHT.Parsers
 				{
 					if (keywordInfo.AllScopes || topScope == currentScope)
 					{
+						UhtCompilerDirective currentDirective = topScope.HeaderParser.GetCurrentCompositeCompilerDirective() & ~UhtCompilerDirective.AllowedCheckIgnoredFlags;
+						if (currentDirective.HasAnyFlags(~keywordInfo.AllowedCompilerDirectives))
+						{
+							List<string> strings = (keywordInfo.AllowedCompilerDirectives).ToStringList();
+							string directives = UhtUtilities.MergeTypeNames(strings, "or", false);
+							topScope.TokenReader.LogError($"'{token.Value.ToString()}' must not be inside preprocessor blocks, except for {directives}");
+						}
 						parseResult = keywordInfo.Delegate(topScope, currentScope, ref token);
 					}
 				}
@@ -758,7 +858,7 @@ namespace EpicGames.UHT.Parsers
 				if (SupportsElif(oldCompilerDirective) != SupportsElif(newCompilerDirective))
 				{
 					throw new UhtException(TokenReader, directive.InputLine,
-						$"Mixing {GetCompilerDirectiveText(oldCompilerDirective)} with {GetCompilerDirectiveText(newCompilerDirective)} in an #elif preprocessor block is not supported");
+						$"Mixing {oldCompilerDirective.GetCompilerDirectiveText()} with {newCompilerDirective.GetCompilerDirectiveText()} in an #elif preprocessor block is not supported");
 				}
 				PushCompilerDirective(newCompilerDirective);
 				if (IsRestrictedDirective(GetCurrentNonCompositeCompilerDirective()))
@@ -1007,23 +1107,6 @@ namespace EpicGames.UHT.Parsers
 				compilerDirective == UhtCompilerDirective.WithCoreUObject;
 		}
 
-		private static string GetCompilerDirectiveText(UhtCompilerDirective compilerDirective)
-		{
-			switch (compilerDirective)
-			{
-				case UhtCompilerDirective.CPPBlock: return "CPP";
-				case UhtCompilerDirective.NotCPPBlock: return "!CPP";
-				case UhtCompilerDirective.ZeroBlock: return "0";
-				case UhtCompilerDirective.OneBlock: return "1";
-				case UhtCompilerDirective.WithHotReload: return "WITH_HOT_RELOAD";
-				case UhtCompilerDirective.WithEditor: return "WITH_EDITOR";
-				case UhtCompilerDirective.WithEditorOnlyData: return "WITH_EDITORONLY_DATA";
-				case UhtCompilerDirective.WithEngine: return "WITH_ENGINE";
-				case UhtCompilerDirective.WithCoreUObject: return "WITH_COREUOBJECT";
-				default: return "<unrecognized>";
-			}
-		}
-
 		private static void SkipVirtualAndAPI(IUhtTokenReader replayReader)
 		{
 			while (true)
diff --git a/Engine/Source/Programs/Shared/EpicGames.UHT/Tables/UhtKeywordTable.cs b/Engine/Source/Programs/Shared/EpicGames.UHT/Tables/UhtKeywordTable.cs
--- a/Engine/Source/Programs/Shared/EpicGames.UHT/Tables/UhtKeywordTable.cs
+++ b/Engine/Source/Programs/Shared/EpicGames.UHT/Tables/UhtKeywordTable.cs
@@ -42,6 +42,11 @@ namespace EpicGames.UHT.Tables
 		/// If true, do not include in usage errors
 		/// </summary>
 		public bool DisableUsageError { get; set; } = false;
+
+		/// <summary>
+		/// List of the allowed compiler directives. 
+		/// </summary>
+		public UhtCompilerDirective AllowedCompilerDirectives { get; set; } = UhtCompilerDirective.DefaultAllowedCheck;
 	}
 
 	/// <summary>
@@ -105,6 +110,11 @@ namespace EpicGames.UHT.Tables
 		/// </summary>
 		public bool DisableUsageError { get; }
 
+		/// <summary>
+		/// List of the allowed compiler directives. 
+		/// </summary>
+		public UhtCompilerDirective AllowedCompilerDirectives { get; }
+
 		/// <summary>
 		/// Construct a new keyword
 		/// </summary>
@@ -120,12 +130,14 @@ namespace EpicGames.UHT.Tables
 				AllowText = attribute.AllowText;
 				AllScopes = attribute.AllScopes;
 				DisableUsageError = attribute.DisableUsageError;
+				AllowedCompilerDirectives = attribute.AllowedCompilerDirectives;
 			}
 			else
 			{
 				AllowText = null;
 				AllScopes = false;
 				DisableUsageError = false;
+				AllowedCompilerDirectives = UhtCompilerDirective.DefaultAllowedCheck;
 			}
 		}
 	}
