Fixed C#/C++ UHT to allow functions with parameters in WITH_EDITORONLY_DATA blocks.

diff --git a/Engine/Source/Programs/Shared/EpicGames.UHT/Parsers/UhtPropertyParser.cs b/Engine/Source/Programs/Shared/EpicGames.UHT/Parsers/UhtPropertyParser.cs
--- a/Engine/Source/Programs/Shared/EpicGames.UHT/Parsers/UhtPropertyParser.cs
+++ b/Engine/Source/Programs/Shared/EpicGames.UHT/Parsers/UhtPropertyParser.cs
@@ -791,26 +791,29 @@ namespace EpicGames.UHT.Parsers
 
 			propertySettings.LineNumber = tokenReader.InputLine;
 
-			UhtCompilerDirective compilerDirective = topScope.HeaderParser.GetCurrentCompositeCompilerDirective();
-			if (compilerDirective.HasAnyFlags(UhtCompilerDirective.WithEditorOnlyData))
-			{
-				propertySettings.PropertyFlags |= EPropertyFlags.EditorOnly;
-			}
-			else if (compilerDirective.HasAnyFlags(UhtCompilerDirective.WithEditor))
-			{
-				// Checking for this error is a bit tricky given legacy code.  
-				// 1) If already wrapped in WITH_EDITORONLY_DATA (see above), then we ignore the error via the else 
-				// 2) Ignore any module that is an editor module
-				UhtPackage package = topScope.ScopeType.HeaderFile.Package;
-				UHTManifest.Module module = package.Module;
-				bool isEditorModule =
-					module.ModuleType == UHTModuleType.EngineEditor ||
-					module.ModuleType == UHTModuleType.GameEditor ||
-					module.ModuleType == UHTModuleType.EngineUncooked ||
-					module.ModuleType == UHTModuleType.GameUncooked;
-				if (propertySettings.PropertyCategory == UhtPropertyCategory.Member && !isEditorModule)
+			if (propertySettings.PropertyCategory == UhtPropertyCategory.Member)
+			{
+				UhtCompilerDirective compilerDirective = topScope.HeaderParser.GetCurrentCompositeCompilerDirective();
+				if (compilerDirective.HasAnyFlags(UhtCompilerDirective.WithEditorOnlyData))
+				{
+					propertySettings.PropertyFlags |= EPropertyFlags.EditorOnly;
+				}
+				else if (compilerDirective.HasAnyFlags(UhtCompilerDirective.WithEditor))
 				{
-					tokenReader.LogError("UProperties should not be wrapped by WITH_EDITOR, use WITH_EDITORONLY_DATA instead.");
+					// Checking for this error is a bit tricky given legacy code.  
+					// 1) If already wrapped in WITH_EDITORONLY_DATA (see above), then we ignore the error via the else 
+					// 2) Ignore any module that is an editor module
+					UhtPackage package = topScope.ScopeType.HeaderFile.Package;
+					UHTManifest.Module module = package.Module;
+					bool isEditorModule =
+						module.ModuleType == UHTModuleType.EngineEditor ||
+						module.ModuleType == UHTModuleType.GameEditor ||
+						module.ModuleType == UHTModuleType.EngineUncooked ||
+						module.ModuleType == UHTModuleType.GameUncooked;
+					if (!isEditorModule)
+					{
+						tokenReader.LogError("UProperties should not be wrapped by WITH_EDITOR, use WITH_EDITORONLY_DATA instead.");
+					}
 				}
 			}
 
@@ -877,7 +880,8 @@ namespace EpicGames.UHT.Parsers
 			// Check for any disallowed flags
 			if (propertySettings.PropertyFlags.HasAnyFlags(propertySettings.DisallowPropertyFlags))
 			{
-				tokenReader.LogError("Specified type modifiers not allowed here");
+				EPropertyFlags extraFlags = propertySettings.PropertyFlags & propertySettings.DisallowPropertyFlags;
+				tokenReader.LogError($"Specified type modifiers not allowed here '{String.Join(" | ", extraFlags.ToStringList())}'");
 			}
 
 			if (_options.HasAnyFlags(UhtPropertyParseOptions.AddModuleRelativePath))
diff --git a/Engine/Source/Programs/UnrealHeaderTool/Private/HeaderParser.cpp b/Engine/Source/Programs/UnrealHeaderTool/Private/HeaderParser.cpp
--- a/Engine/Source/Programs/UnrealHeaderTool/Private/HeaderParser.cpp
+++ b/Engine/Source/Programs/UnrealHeaderTool/Private/HeaderParser.cpp
@@ -3160,25 +3160,28 @@ void FHeaderParser::GetVarType(
 		}
 	}
 
-	uint32 CurrentCompilerDirective = GetCurrentCompilerDirective();
-	if ((CurrentCompilerDirective & ECompilerDirective::WithEditorOnlyData) != 0)
-	{
-		Flags |= CPF_EditorOnly;
-	}
-	else if ((CurrentCompilerDirective & ECompilerDirective::WithEditor) != 0)
+	if (VariableCategory == EVariableCategory::Member)
 	{
-		// Checking for this error is a bit tricky given legacy code.  
-		// 1) If already wrapped in WITH_EDITORONLY_DATA (see above), then we ignore the error via the else 
-		// 2) Ignore any module that is an editor module
-		const FManifestModule& Module = Scope->GetFileScope()->GetSourceFile()->GetPackageDef().GetModule();
-		const bool bIsEditorModule =
-			Module.ModuleType == EBuildModuleType::EngineEditor ||
-			Module.ModuleType == EBuildModuleType::GameEditor || 
-			Module.ModuleType == EBuildModuleType::EngineUncooked || 
-			Module.ModuleType == EBuildModuleType::GameUncooked;
-		if (VariableCategory == EVariableCategory::Member && !bIsEditorModule)
+		uint32 CurrentCompilerDirective = GetCurrentCompilerDirective();
+		if ((CurrentCompilerDirective & ECompilerDirective::WithEditorOnlyData) != 0)
 		{
-			LogError(TEXT("UProperties should not be wrapped by WITH_EDITOR, use WITH_EDITORONLY_DATA instead."));
+			Flags |= CPF_EditorOnly;
+		}
+		else if ((CurrentCompilerDirective & ECompilerDirective::WithEditor) != 0)
+		{
+			// Checking for this error is a bit tricky given legacy code.  
+			// 1) If already wrapped in WITH_EDITORONLY_DATA (see above), then we ignore the error via the else 
+			// 2) Ignore any module that is an editor module
+			const FManifestModule& Module = Scope->GetFileScope()->GetSourceFile()->GetPackageDef().GetModule();
+			const bool bIsEditorModule =
+				Module.ModuleType == EBuildModuleType::EngineEditor ||
+				Module.ModuleType == EBuildModuleType::GameEditor ||
+				Module.ModuleType == EBuildModuleType::EngineUncooked ||
+				Module.ModuleType == EBuildModuleType::GameUncooked;
+			if (!bIsEditorModule)
+			{
+				LogError(TEXT("UProperties should not be wrapped by WITH_EDITOR, use WITH_EDITORONLY_DATA instead."));
+			}
 		}
 	}
 
