Fixed C# UHT strict documentation tests for function.  Would throw an exception or get lost in an infinite loop.

diff --git a/Engine/Source/Programs/Shared/EpicGames.UHT/Types/UhtFunction.cs b/Engine/Source/Programs/Shared/EpicGames.UHT/Types/UhtFunction.cs
--- a/Engine/Source/Programs/Shared/EpicGames.UHT/Types/UhtFunction.cs
+++ b/Engine/Source/Programs/Shared/EpicGames.UHT/Types/UhtFunction.cs
@@ -1024,7 +1024,7 @@ namespace EpicGames.UHT.Types
 		/// <param name="paramMap">Parameter map to add the tooltip</param>
 		private static void AddParameterToolTip(StringView name, StringView text, Dictionary<StringView, StringView> paramMap)
 		{
-			paramMap.Add(name, new StringView(UhtFCString.TabsToSpaces(text.Memory.Trim(), 4, false)));
+			paramMap.TryAdd(name, new StringView(UhtFCString.TabsToSpaces(text.Memory.Trim(), 4, false)));
 		}
 
 		/// <summary>
@@ -1066,6 +1066,10 @@ namespace EpicGames.UHT.Types
 					line = new StringView(input, 0, lineEnd);
 					input = new StringView(input, lineEnd + 1);
 				}
+				else
+				{
+					input = new StringView();
+				}
 
 				// If this is a parameter, invoke the method to extract the name first
 				int paramStart = line.Span.IndexOf(ParamTag, StringComparison.OrdinalIgnoreCase);
