UnrealBuild: Escape quotes for multi-target builds

#rnx

diff --git a/Engine/Source/Programs/Shared/EpicGames.Build/Automation/FormatCommandLine.cs b/Engine/Source/Programs/Shared/EpicGames.Build/Automation/FormatCommandLine.cs
--- a/Engine/Source/Programs/Shared/EpicGames.Build/Automation/FormatCommandLine.cs
+++ b/Engine/Source/Programs/Shared/EpicGames.Build/Automation/FormatCommandLine.cs
@@ -41,6 +41,9 @@ namespace UnrealBuildBase
 				return Argument;
 			}
 
+			// If the argument has quotes in it, the nested quotes must be escaped
+			Argument = Argument.Replace("\"", "\\\"");
+
 			// If it does have a space, and it's formatted as an option (ie. -Something=), try to insert quotes after the equals character
 			int EqualsIdx = Argument.IndexOf('=');
 			if(Argument.StartsWith("-") && EqualsIdx != -1 && EqualsIdx < SpaceIdx)
