Horde: Fix debugger display of CbAttachment objects.

diff --git a/Engine/Source/Programs/Shared/EpicGames.Serialization/CbField.cs b/Engine/Source/Programs/Shared/EpicGames.Serialization/CbField.cs
--- a/Engine/Source/Programs/Shared/EpicGames.Serialization/CbField.cs
+++ b/Engine/Source/Programs/Shared/EpicGames.Serialization/CbField.cs
@@ -191,7 +191,7 @@ namespace EpicGames.Serialization
 	/// <summary>
 	/// A binary attachment, referenced by <see cref="IoHash"/>
 	/// </summary>
-	[DebuggerDisplay("{_hash}")]
+	[DebuggerDisplay("{Hash}")]
 	[JsonConverter(typeof(CbBinaryAttachmentJsonConverter))]
 	[TypeConverter(typeof(CbBinaryAttachmentTypeConverter))]
 	public struct CbBinaryAttachment : IEquatable<CbBinaryAttachment>
@@ -273,7 +273,7 @@ namespace EpicGames.Serialization
 	/// <summary>
 	/// An object attachment, referenced by <see cref="IoHash"/>
 	/// </summary>
-	[DebuggerDisplay("{_hash}")]
+	[DebuggerDisplay("{Hash}")]
 	[JsonConverter(typeof(CbObjectAttachmentJsonConverter))]
 	[TypeConverter(typeof(CbObjectAttachmentTypeConverter))]
 	public struct CbObjectAttachment : IEquatable<CbObjectAttachment>
