Horde: Add some attributes to make debugging RE directory objects easier.

diff --git a/Engine/Source/Programs/Shared/EpicGames.Horde/Compute/DirectoryTree.cs b/Engine/Source/Programs/Shared/EpicGames.Horde/Compute/DirectoryTree.cs
--- a/Engine/Source/Programs/Shared/EpicGames.Horde/Compute/DirectoryTree.cs
+++ b/Engine/Source/Programs/Shared/EpicGames.Horde/Compute/DirectoryTree.cs
@@ -3,6 +3,7 @@
 using EpicGames.Core;
 using EpicGames.Serialization;
 using System.Collections.Generic;
+using System.Diagnostics;
 
 namespace EpicGames.Horde.Compute
 {
@@ -27,6 +28,7 @@ namespace EpicGames.Horde.Compute
 	/// <summary>
 	/// Reference to a named file within a <see cref="DirectoryTree"/>
 	/// </summary>
+	[DebuggerDisplay("{Name}")]
 	public class FileNode
 	{
 		/// <summary>
@@ -82,6 +84,7 @@ namespace EpicGames.Horde.Compute
 	/// <summary>
 	/// Reference to a named directory within a <see cref="DirectoryTree"/>
 	/// </summary>
+	[DebuggerDisplay("{Name}")]
 	public class DirectoryNode
 	{
 		/// <summary>
