EpicGames.Serialization: Refactor CbWriter to be an interface with extension methods, to make it easier to iterate on less allocation-heavy methods for piecing together CB objects.

diff --git a/Engine/Source/Programs/Shared/EpicGames.Serialization/CbField.cs b/Engine/Source/Programs/Shared/EpicGames.Serialization/CbField.cs
--- a/Engine/Source/Programs/Shared/EpicGames.Serialization/CbField.cs
+++ b/Engine/Source/Programs/Shared/EpicGames.Serialization/CbField.cs
@@ -2077,28 +2077,6 @@ namespace EpicGames.Serialization
 			return Current;
 		}
 
-		/// <summary>
-		/// Copy the field range into a buffer of exactly GetRangeSize() bytes.
-		/// </summary>
-		/// <param name="buffer"></param>
-		public void CopyRangeTo(Span<byte> buffer)
-		{
-			ReadOnlyMemory<byte> source;
-			if (TryGetRangeView(out source))
-			{
-				source.Span.CopyTo(buffer);
-			}
-			else
-			{
-				for (CbFieldIterator it = new CbFieldIterator(this); it; it.MoveNext())
-				{
-					int size = it.Current.GetSize();
-					it.Current.CopyTo(buffer);
-					buffer = buffer.Slice(size);
-				}
-			}
-		}
-
 		/// <summary>
 		/// Invoke the visitor for every attachment in the field range.
 		/// </summary>
@@ -2129,25 +2107,6 @@ namespace EpicGames.Serialization
 			}
 		}
 
-		/// <summary>
-		/// Try to get a view of every field in the range as they would be serialized.
-		///
-		/// A view is available if each field contains its type. Access the equivalent for other field
-		/// ranges through FCbFieldIterator::CloneRange or CopyRangeTo.
-		/// </summary>
-		/// <returns></returns>
-		bool TryGetRangeView(out ReadOnlyMemory<byte> outView)
-		{
-			throw new NotImplementedException();
-/*			FMemoryView View;
-			if (FieldType::TryGetView(View))
-			{
-				OutView = MakeMemoryView(View.GetData(), FieldsEnd);
-				return true;
-			}
-			return false;*/
-		}
-
 		/// <summary>
 		/// Move to the next element
 		/// </summary>
diff --git a/Engine/Source/Programs/Shared/EpicGames.Serialization/CbWriter.cs b/Engine/Source/Programs/Shared/EpicGames.Serialization/CbWriter.cs
--- a/Engine/Source/Programs/Shared/EpicGames.Serialization/CbWriter.cs
+++ b/Engine/Source/Programs/Shared/EpicGames.Serialization/CbWriter.cs
@@ -2,6 +2,7 @@
 
 using EpicGames.Core;
 using System;
+using System.Buffers;
 using System.Buffers.Binary;
 using System.Collections.Generic;
 using System.IO;
@@ -34,9 +35,46 @@
 	}
 
 	/// <summary>
+	/// Interface for compact binary writers
+	/// </summary>
+	public interface ICbWriter
+	{
+		/// <summary>
+		/// Begin writing an object field
+		/// </summary>
+		/// <param name="name">Name of the field. May be empty for fields that are not part of another object.</param>
+		void BeginObject(Utf8String name);
+
+		/// <summary>
+		/// End the current object
+		/// </summary>
+		void EndObject();
+
+		/// <summary>
+		/// Begin writing a named array field
+		/// </summary>
+		/// <param name="name">Name of the field, or an empty string.</param>
+		/// <param name="elementType">Type of the field. May be <see cref="CbFieldType.None"/> for non-uniform arrays.</param>
+		void BeginArray(Utf8String name, CbFieldType elementType);
+
+		/// <summary>
+		/// End the current array
+		/// </summary>
+		void EndArray();
+
+		/// <summary>
+		/// Writes the header for a named field
+		/// </summary>
+		/// <param name="type">Type of the field</param>
+		/// <param name="name">Name of the field. May be empty for fields that are not part of another object.</param>
+		/// <param name="length">Length of data for the field</param>
+		Span<byte> WriteField(CbFieldType type, Utf8String name, int length);
+	}
+
+	/// <summary>
 	/// Forward-only writer for compact binary objects
 	/// </summary>
-	public class CbWriter
+	public class CbWriter : ICbWriter
 	{
 		/// <summary>
 		/// Stores information about an object or array scope within the written buffer which requires a header to be inserted containing
@@ -148,7 +186,7 @@
 		/// <returns>New chunk object</returns>
 		Chunk AllocChunk(int offset, int maxLength)
 		{
-			for(int idx = _freeChunks.Count - 1; idx >= 0; idx--)
+			for (int idx = _freeChunks.Count - 1; idx >= 0; idx--)
 			{
 				Chunk chunk = _freeChunks[idx];
 				if (chunk._data.Length >= maxLength)
@@ -206,7 +244,7 @@
 		/// </summary>
 		/// <param name="length"></param>
 		/// <returns>The allocated memory</returns>
-		Memory<byte> Allocate(int length)
+		public Span<byte> Allocate(int length)
 		{
 			Chunk lastChunk = CurrentChunk;
 			if (lastChunk._length + length > lastChunk._data.Length)
@@ -216,7 +254,7 @@
 				_chunks.Add(lastChunk);
 			}
 
-			Memory<byte> buffer = lastChunk._data.AsMemory(lastChunk._length, length);
+			Span<byte> buffer = lastChunk._data.AsSpan(lastChunk._length, length);
 			lastChunk._length += length;
 			_currentOffset += length;
 			return buffer;
@@ -248,95 +286,71 @@
 		}
 
 		/// <summary>
-		/// Writes the header for an unnamed field
-		/// </summary>
-		/// <param name="type"></param>
-		void WriteFieldHeader(CbFieldType type)
-		{
-			Scope scope = CurrentScope;
-			if (!CbFieldUtils.IsArray(scope._fieldType))
-			{
-				throw new CbWriterException($"Anonymous fields are not allowed within fields of type {scope._fieldType}");
-			}
-			
-			if (scope._uniformFieldType == CbFieldType.None)
-			{
-				Allocate(1).Span[0] = (byte)type;
-			}
-			else if (scope._uniformFieldType != type)
-			{
-				throw new CbWriterException($"Mismatched type for uniform array - expected {scope._uniformFieldType}, not {type}");
-			}
-			scope._count++;
-		}
-
-		/// <summary>
 		/// Writes the header for a named field
 		/// </summary>
 		/// <param name="type"></param>
 		/// <param name="name"></param>
-		void WriteFieldHeader(CbFieldType type, Utf8String name)
+		public Span<byte> WriteField(CbFieldType type, Utf8String name, int length)
 		{
-			Scope scope = CurrentScope;
-			if (!CbFieldUtils.IsObject(scope._fieldType))
+			if (name.IsEmpty)
 			{
-				throw new CbWriterException($"Named fields are not allowed within fields of type {scope._fieldType}");
+				Scope scope = CurrentScope;
+				if (!CbFieldUtils.IsArray(scope._fieldType))
+				{
+					throw new CbWriterException($"Anonymous fields are not allowed within fields of type {scope._fieldType}");
+				}
+
+				if (scope._uniformFieldType == CbFieldType.None)
+				{
+					Allocate(1)[0] = (byte)type;
+				}
+				else if (scope._uniformFieldType != type)
+				{
+					throw new CbWriterException($"Mismatched type for uniform array - expected {scope._uniformFieldType}, not {type}");
+				}
+				scope._count++;
 			}
-
-			int nameVarIntLength = VarInt.MeasureUnsigned(name.Length);
-			if (scope._uniformFieldType == CbFieldType.None)
-			{
-				Span<byte> buffer = Allocate(1 + nameVarIntLength + name.Length).Span;
-				buffer[0] = (byte)(type | CbFieldType.HasFieldName);
-				WriteBinaryPayload(buffer[1..], name.Span);
-			}
 			else
 			{
-				if (scope._uniformFieldType != type)
+				Scope scope = CurrentScope;
+				if (!CbFieldUtils.IsObject(scope._fieldType))
 				{
-					throw new CbWriterException($"Mismatched type for uniform object - expected {scope._uniformFieldType}, not {type}");
+					throw new CbWriterException($"Named fields are not allowed within fields of type {scope._fieldType}");
 				}
-				WriteBinaryPayload(name.Span);
+
+				int nameVarIntLength = VarInt.MeasureUnsigned(name.Length);
+				if (scope._uniformFieldType == CbFieldType.None)
+				{
+					Span<byte> buffer = Allocate(1 + nameVarIntLength + name.Length);
+					buffer[0] = (byte)(type | CbFieldType.HasFieldName);
+					WriteBinaryPayload(buffer[1..], name.Span);
+				}
+				else
+				{
+					if (scope._uniformFieldType != type)
+					{
+						throw new CbWriterException($"Mismatched type for uniform object - expected {scope._uniformFieldType}, not {type}");
+					}
+					Span<byte> buffer = Allocate(name.Length);
+					WriteBinaryPayload(buffer, name.Span);
+				}
+				scope._count++;
 			}
-			scope._count++;
+			return Allocate(length);
 		}
 
 		/// <summary>
-		/// Copies an entire field value to the output
+		/// Writes the payload for a binary value
 		/// </summary>
-		/// <param name="field"></param>
-		public void WriteFieldValue(CbField field)
+		/// <param name="output">Output buffer</param>
+		/// <param name="value">Value to be written</param>
+		static void WriteBinaryPayload(Span<byte> output, ReadOnlySpan<byte> value)
 		{
-			WriteFieldHeader(field.GetType());
-			int size = (int)field.GetPayloadSize();
-			field.GetPayloadView().CopyTo(Allocate(size));
+			int varIntLength = VarInt.WriteUnsigned(output, value.Length);
+			output = output[varIntLength..];
+
+			value.CopyTo(output);
+			CheckSize(output, value.Length);
 		}
 
 		/// <summary>
-		/// Copies an entire field value to the output, using the name from the field
-		/// </summary>
-		/// <param name="field"></param>
-		public void WriteField(CbField field) => WriteField(field.GetName(), field);
-
-		/// <summary>
-		/// Copies an entire field value to the output
-		/// </summary>
-		/// <param name="name">Name of the field</param>
-		/// <param name="field"></param>
-		public void WriteField(Utf8String name, CbField field)
-		{
-			WriteFieldHeader(field.GetType(), name);
-			int size = (int)field.GetPayloadSize();
-			field.GetPayloadView().CopyTo(Allocate(size));
-		}
-
-		/// <summary>
-		/// Begin writing an object field
-		/// </summary>
-		public void BeginObject()
-		{
-			WriteFieldHeader(CbFieldType.Object);
-			PushScope(CbFieldType.Object, CbFieldType.None);
-		}
-
-		/// <summary>
@@ -345,7 +359,7 @@
 		/// <param name="name">Name of the field</param>
 		public void BeginObject(Utf8String name)
 		{
-			WriteFieldHeader(CbFieldType.Object, name);
+			WriteField(CbFieldType.Object, name, 0);
 			PushScope(CbFieldType.Object, CbFieldType.None);
 		}
 
@@ -358,22 +372,22 @@
 		}
 
 		/// <summary>
-		/// Begin writing an array field
-		/// </summary>
-		public void BeginArray()
-		{
-			WriteFieldHeader(CbFieldType.Array);
-			PushScope(CbFieldType.Array, CbFieldType.None);
-		}
-
-		/// <summary>
 		/// Begin writing a named array field
 		/// </summary>
 		/// <param name="name"></param>
-		public void BeginArray(Utf8String name)
+		public void BeginArray(Utf8String name, CbFieldType elementType)
 		{
-			WriteFieldHeader(CbFieldType.Array, name);
-			PushScope(CbFieldType.Array, CbFieldType.None);
+			if (elementType == CbFieldType.None)
+			{
+				WriteField(CbFieldType.Array, name, 0);
+				PushScope(CbFieldType.Array, CbFieldType.None);
+			}
+			else
+			{
+				WriteField(CbFieldType.UniformArray, name, 0);
+				PushScope(CbFieldType.UniformArray, elementType);
+				Allocate(1)[0] = (byte)elementType;
+			}
 		}
 
 		/// <summary>
@@ -385,458 +399,6 @@
 		}
 
 		/// <summary>
-		/// Begin writing a uniform array field
-		/// </summary>
-		/// <param name="fieldType">The field type for elements in the array</param>
-		public void BeginUniformArray(CbFieldType fieldType)
-		{
-			WriteFieldHeader(CbFieldType.UniformArray);
-			PushScope(CbFieldType.UniformArray, fieldType);
-			Allocate(1).Span[0] = (byte)fieldType;
-		}
-
-		/// <summary>
-		/// Begin writing a named uniform array field
-		/// </summary>
-		/// <param name="name">Name of the field</param>
-		/// <param name="fieldType">The field type for elements in the array</param>
-		public void BeginUniformArray(Utf8String name, CbFieldType fieldType)
-		{
-			WriteFieldHeader(CbFieldType.UniformArray, name);
-			PushScope(CbFieldType.UniformArray, fieldType);
-			Allocate(1).Span[0] = (byte)fieldType;
-		}
-
-		/// <summary>
-		/// End the current array
-		/// </summary>
-		public void EndUniformArray()
-		{
-			PopScope();
-		}
-
-		/// <summary>
-		/// Write a null field
-		/// </summary>
-		public void WriteNullValue()
-		{
-			WriteFieldHeader(CbFieldType.Null);
-		}
-
-		/// <summary>
-		/// Write a named null field
-		/// </summary>
-		/// <param name="name">Name of the field</param>
-		public void WriteNull(Utf8String name)
-		{
-			WriteFieldHeader(CbFieldType.Null, name);
-		}
-
-		/// <summary>
-		/// Writes a boolean value
-		/// </summary>
-		/// <param name="value"></param>
-		public void WriteBoolValue(bool value)
-		{
-			WriteFieldHeader(value? CbFieldType.BoolTrue : CbFieldType.BoolFalse);
-		}
-
-		/// <summary>
-		/// Writes a boolean value
-		/// </summary>
-		/// <param name="name">Name of the field</param>
-		/// <param name="value"></param>
-		public void WriteBool(Utf8String name, bool value)
-		{
-			WriteFieldHeader(value ? CbFieldType.BoolTrue : CbFieldType.BoolFalse, name);
-		}
-
-		/// <summary>
-		/// Writes the payload for an integer
-		/// </summary>
-		/// <param name="value">Value to write</param>
-		void WriteIntegerPayload(ulong value)
-		{
-			int length = VarInt.MeasureUnsigned(value);
-			Span<byte> buffer = Allocate(length).Span;
-			VarInt.WriteUnsigned(buffer, value);
-		}
-
-		/// <summary>
-		/// Writes an unnamed integer field
-		/// </summary>
-		/// <param name="value">Value to be written</param>
-		public void WriteIntegerValue(int value)
-		{
-			WriteIntegerValue((long)value);
-		}
-
-		/// <summary>
-		/// Writes an unnamed integer field
-		/// </summary>
-		/// <param name="value">Value to be written</param>
-		public void WriteIntegerValue(long value)
-		{
-			if (value >= 0)
-			{
-				WriteFieldHeader(CbFieldType.IntegerPositive);
-				WriteIntegerPayload((ulong)value);
-			}
-			else
-			{
-				WriteFieldHeader(CbFieldType.IntegerNegative);
-				WriteIntegerPayload((ulong)-value);
-			}
-		}
-
-		/// <summary>
-		/// Writes an named integer field
-		/// </summary>
-		/// <param name="name">Name of the field</param>
-		/// <param name="value">Value to be written</param>
-		public void WriteInteger(Utf8String name, int value)
-		{
-			WriteInteger(name, (long)value);
-		}
-
-		/// <summary>
-		/// Writes an named integer field
-		/// </summary>
-		/// <param name="name">Name of the field</param>
-		/// <param name="value">Value to be written</param>
-		public void WriteInteger(Utf8String name, long value)
-		{
-			if (value >= 0)
-			{
-				WriteFieldHeader(CbFieldType.IntegerPositive, name);
-				WriteIntegerPayload((ulong)value);
-			}
-			else
-			{
-				WriteFieldHeader(CbFieldType.IntegerNegative, name);
-				WriteIntegerPayload((ulong)-value);
-			}
-		}
-
-		/// <summary>
-		/// Writes an unnamed integer field
-		/// </summary>
-		/// <param name="value">Value to be written</param>
-		public void WriteIntegerValue(ulong value)
-		{
-			WriteFieldHeader(CbFieldType.IntegerPositive);
-			WriteIntegerPayload(value);
-		}
-
-		/// <summary>
-		/// Writes a named integer field
-		/// </summary>
-		/// <param name="name">Name of the field</param>
-		/// <param name="value">Value to be written</param>
-		public void WriteInteger(Utf8String name, ulong value)
-		{
-			WriteFieldHeader(CbFieldType.IntegerPositive, name);
-			WriteIntegerPayload(value);
-		}
-
-		/// <summary>
-		/// Writes an unnamed double field
-		/// </summary>
-		/// <param name="value">Value to be written</param>
-		public void WriteDoubleValue(double value)
-		{
-			WriteFieldHeader(CbFieldType.Float64);
-			BinaryPrimitives.WriteInt64BigEndian(Allocate(sizeof(double)).Span, BitConverter.DoubleToInt64Bits(value));
-		}
-
-		/// <summary>
-		/// Writes a named double field
-		/// </summary>
-		/// <param name="name">Name of the field</param>
-		/// <param name="value">Value to be written</param>
-		public void WriteDouble(Utf8String name, double value)
-		{
-			WriteFieldHeader(CbFieldType.Float64, name);
-			BinaryPrimitives.WriteInt64BigEndian(Allocate(sizeof(double)).Span, BitConverter.DoubleToInt64Bits(value));
-		}
-
-		/// <summary>
-		/// Writes the payload for a <see cref="DateTime"/> value
-		/// </summary>
-		/// <param name="dateTime">The value to write</param>
-		void WriteDateTimePayload(DateTime dateTime)
-		{
-			Span<byte> buffer = Allocate(sizeof(long)).Span;
-			BinaryPrimitives.WriteInt64BigEndian(buffer, dateTime.Ticks);
-		}
-
-		/// <summary>
-		/// Writes an unnamed <see cref="DateTime"/> field
-		/// </summary>
-		/// <param name="value">Value to be written</param>
-		public void WriteDateTimeValue(DateTime value)
-		{
-			WriteFieldHeader(CbFieldType.DateTime);
-			WriteDateTimePayload(value);
-		}
-
-		/// <summary>
-		/// Writes a named <see cref="DateTime"/> field
-		/// </summary>
-		/// <param name="name">Name of the field</param>
-		/// <param name="value">Value to be written</param>
-		public void WriteDateTime(Utf8String name, DateTime value)
-		{
-			WriteFieldHeader(CbFieldType.DateTime, name);
-			WriteDateTimePayload(value);
-		}
-
-		/// <summary>
-		/// Writes the payload for a hash
-		/// </summary>
-		/// <param name="hash"></param>
-		void WriteHashPayload(IoHash hash)
-		{
-			Span<byte> buffer = Allocate(IoHash.NumBytes).Span;
-			hash.CopyTo(buffer);
-		}
-
-		/// <summary>
-		/// Writes an unnamed <see cref="IoHash"/> field
-		/// </summary>
-		/// <param name="hash"></param>
-		public void WriteHashValue(IoHash hash)
-		{
-			WriteFieldHeader(CbFieldType.Hash);
-			WriteHashPayload(hash);
-		}
-
-		/// <summary>
-		/// Writes a named <see cref="IoHash"/> field
-		/// </summary>
-		/// <param name="name">Name of the field</param>
-		/// <param name="value">Value to be written</param>
-		public void WriteHash(Utf8String name, IoHash value)
-		{
-			WriteFieldHeader(CbFieldType.Hash, name);
-			WriteHashPayload(value);
-		}
-
-		/// <summary>
-		/// Writes an unnamed reference to a binary attachment
-		/// </summary>
-		/// <param name="hash">Hash of the attachment</param>
-		public void WriteBinaryAttachmentValue(IoHash hash)
-		{
-			WriteFieldHeader(CbFieldType.BinaryAttachment);
-			WriteHashPayload(hash);
-		}
-
-		/// <summary>
-		/// Writes a named reference to a binary attachment
-		/// </summary>
-		/// <param name="name">Name of the field</param>
-		/// <param name="hash">Hash of the attachment</param>
-		public void WriteBinaryAttachment(Utf8String name, IoHash hash)
-		{
-			WriteFieldHeader(CbFieldType.BinaryAttachment, name);
-			WriteHashPayload(hash);
-		}
-
-		/// <summary>
-		/// Writes the payload for an object to the buffer
-		/// </summary>
-		/// <param name="obj"></param>
-		void WriteObjectPayload(CbObject obj)
-		{
-			CbField field = obj.AsField();
-			Memory<byte> buffer = Allocate(field.Payload.Length);
-			field.Payload.CopyTo(buffer);
-		}
-
-		/// <summary>
-		/// Writes an object directly into the writer
-		/// </summary>
-		/// <param name="obj">Object to write</param>
-		public void WriteObject(CbObject obj)
-		{
-			WriteFieldHeader(CbFieldType.Object);
-			WriteObjectPayload(obj);
-		}
-
-		/// <summary>
-		/// Writes an object directly into the writer
-		/// </summary>
-		/// <param name="name">Name of the object</param>
-		/// <param name="obj">Object to write</param>
-		public void WriteObject(Utf8String name, CbObject obj)
-		{
-			WriteFieldHeader(CbFieldType.Object, name);
-			WriteObjectPayload(obj);
-		}
-
-		/// <summary>
-		/// Writes an unnamed reference to an object attachment
-		/// </summary>
-		/// <param name="hash">Hash of the attachment</param>
-		public void WriteObjectAttachmentValue(IoHash hash)
-		{
-			WriteFieldHeader(CbFieldType.ObjectAttachment);
-			WriteHashPayload(hash);
-		}
-
-		/// <summary>
-		/// Writes a named reference to an object attachment
-		/// </summary>
-		/// <param name="name">Name of the field</param>
-		/// <param name="hash">Hash of the attachment</param>
-		public void WriteObjectAttachment(Utf8String name, IoHash hash)
-		{
-			WriteFieldHeader(CbFieldType.ObjectAttachment, name);
-			WriteHashPayload(hash);
-		}
-
-		/// <summary>
-		/// Writes the payload for a binary value
-		/// </summary>
-		/// <param name="output">Output buffer</param>
-		/// <param name="value">Value to be written</param>
-		static void WriteBinaryPayload(Span<byte> output, ReadOnlySpan<byte> value)
-		{
-			int varIntLength = VarInt.WriteUnsigned(output, value.Length);
-			output = output[varIntLength..];
-
-			value.CopyTo(output);
-			CheckSize(output, value.Length);
-		}
-
-		/// <summary>
-		/// Writes the payload for a binary value
-		/// </summary>
-		/// <param name="value">Value to be written</param>
-		void WriteBinaryPayload(ReadOnlySpan<byte> value)
-		{
-			int valueVarIntLength = VarInt.MeasureUnsigned(value.Length);
-			Span<byte> buffer = Allocate(valueVarIntLength + value.Length).Span;
-			WriteBinaryPayload(buffer, value);
-		}
-
-		/// <summary>
-		/// Writes an unnamed string value
-		/// </summary>
-		/// <param name="value">Value to be written</param>
-		public void WriteStringValue(string value) => WriteUtf8StringValue(value);
-
-		/// <summary>
-		/// Writes a named string value
-		/// </summary>
-		/// <param name="name">Name of the field</param>
-		/// <param name="value">Value to be written</param>
-		public void WriteString(Utf8String name, string? value)
-		{
-			if(value != null)
-			{
-				WriteUtf8String(name, value);
-			}
-		}
-
-		/// <summary>
-		/// Writes an unnamed string value
-		/// </summary>
-		/// <param name="value">Value to be written</param>
-		public void WriteUtf8StringValue(Utf8String value)
-		{
-			WriteFieldHeader(CbFieldType.String);
-			WriteBinaryPayload(value.Span);
-		}
-
-		/// <summary>
-		/// Writes a named string value
-		/// </summary>
-		/// <param name="name">Name of the field</param>
-		/// <param name="value">Value to be written</param>
-		public void WriteUtf8String(Utf8String name, Utf8String value)
-		{
-			if (value.Length > 0)
-			{
-				WriteFieldHeader(CbFieldType.String, name);
-				WriteBinaryPayload(value.Span);
-			}
-		}
-
-		/// <summary>
-		/// Writes an unnamed binary value
-		/// </summary>
-		/// <param name="value">Value to be written</param>
-		public void WriteBinarySpanValue(ReadOnlySpan<byte> value)
-		{
-			WriteFieldHeader(CbFieldType.Binary);
-			WriteBinaryPayload(value);
-		}
-
-		/// <summary>
-		/// Writes a named binary value
-		/// </summary>
-		/// <param name="name">Name of the field</param>
-		/// <param name="value">Value to be written</param>
-		public void WriteBinarySpan(Utf8String name, ReadOnlySpan<byte> value)
-		{
-			WriteFieldHeader(CbFieldType.Binary, name);
-			WriteBinaryPayload(value);
-		}
-
-		/// <summary>
-		/// Writes an unnamed binary value
-		/// </summary>
-		/// <param name="value">Value to be written</param>
-		public void WriteBinaryValue(ReadOnlyMemory<byte> value)
-		{
-			WriteBinarySpanValue(value.Span);
-		}
-
-		/// <summary>
-		/// Writes a named binary value
-		/// </summary>
-		/// <param name="name">Name of the field</param>
-		/// <param name="value">Value to be written</param>
-		public void WriteBinary(Utf8String name, ReadOnlyMemory<byte> value)
-		{
-			WriteBinarySpan(name, value.Span);
-		}
-
-		/// <summary>
-		/// Writes an unnamed binary value
-		/// </summary>
-		/// <param name="value">Value to be written</param>
-		public void WriteBinaryArrayValue(byte[] value)
-		{
-			WriteBinarySpanValue(value.AsSpan());
-		}
-
-		/// <summary>
-		/// Writes a named binary value
-		/// </summary>
-		/// <param name="name">Name of the field</param>
-		/// <param name="value">Value to be written</param>
-		public void WriteBinaryArray(Utf8String name, byte[] value)
-		{
-			WriteBinarySpan(name, value.AsSpan());
-		}
-
-		/// <summary>
-		/// Check that the given span is the required size
-		/// </summary>
-		/// <param name="span"></param>
-		/// <param name="expectedSize"></param>
-		static void CheckSize(Span<byte> span, int expectedSize)
-		{
-			if (span.Length != expectedSize)
-			{
-				throw new Exception("Size of buffer is not correct");
-			}
-		}
-
-		/// <summary>
 		/// Computes the hash for this object
 		/// </summary>
 		/// <returns>Hash for the object</returns>
@@ -1068,5 +630,358 @@
 					throw new InvalidOperationException();
 			}
 		}
+
+		/// <summary>
+		/// Check that the given span is the required size
+		/// </summary>
+		/// <param name="span"></param>
+		/// <param name="expectedSize"></param>
+		static void CheckSize(Span<byte> span, int expectedSize)
+		{
+			if (span.Length != expectedSize)
+			{
+				throw new Exception("Size of buffer is not correct");
+			}
+		}
 	}
+
+	/// <summary>
+	/// Extension methods for <see cref="CbWriter"/>
+	/// </summary>
+	public static class CbWriterExtensions
+	{
+		static int MeasureFieldWithLength(int length) => length + VarInt.MeasureUnsigned(length);
+
+		static Span<byte> WriteFieldWithLength(this ICbWriter writer, CbFieldType type, Utf8String name, int length)
+		{
+			int fullLength = MeasureFieldWithLength(length);
+			Span<byte> buffer = writer.WriteField(type, name, fullLength);
+
+			int lengthLength = VarInt.WriteUnsigned(buffer, length);
+			return buffer.Slice(lengthLength);
+		}
+
+		/// <summary>
+		/// Begin writing an object field
+		/// </summary>
+		public static void BeginObject(this ICbWriter writer) => writer.BeginObject(default);
+
+		/// <summary>
+		/// Begin writing an array field
+		/// </summary>
+		public static void BeginArray(this ICbWriter writer) => writer.BeginArray(default, CbFieldType.None);
+
+		/// <summary>
+		/// Begin writing a named array field
+		/// </summary>
+		/// <param name="name">Name of the field</param>
+		public static void BeginArray(this ICbWriter writer, Utf8String name) => writer.BeginArray(name, CbFieldType.None);
+
+		/// <summary>
+		/// Begin writing a uniform array field
+		/// </summary>
+		/// <param name="fieldType">The field type for elements in the array</param>
+		public static void BeginUniformArray(this ICbWriter writer, CbFieldType fieldType) => BeginUniformArray(writer, default, fieldType);
+
+		/// <summary>
+		/// Begin writing a named uniform array field
+		/// </summary>
+		/// <param name="name">Name of the field</param>
+		/// <param name="fieldType">The field type for elements in the array</param>
+		public static void BeginUniformArray(this ICbWriter writer, Utf8String name, CbFieldType fieldType) => writer.BeginArray(name, fieldType);
+
+		/// <summary>
+		/// End writing a uniform array field
+		/// </summary>
+		/// <param name="writer"></param>
+		public static void EndUniformArray(this ICbWriter writer) => writer.EndArray();
+
+		/// <summary>
+		/// Copies an entire field value to the output
+		/// </summary>
+		/// <param name="field"></param>
+		public static void WriteFieldValue(this ICbWriter writer, CbField field) => WriteField(writer, default, field);
+
+		/// <summary>
+		/// Copies an entire field value to the output, using the name from the field
+		/// </summary>
+		/// <param name="field"></param>
+		public static void WriteField(this ICbWriter writer, CbField field) => WriteField(writer, field.GetName(), field);
+
+		/// <summary>
+		/// Copies an entire field value to the output
+		/// </summary>
+		/// <param name="name">Name of the field</param>
+		/// <param name="field"></param>
+		public static void WriteField(this ICbWriter writer, Utf8String name, CbField field)
+		{
+			ReadOnlySpan<byte> source = field.GetPayloadView().Span;
+			Span<byte> target = writer.WriteField(field.GetType(), name, source.Length);
+			source.CopyTo(target);
+		}
+
+		/// <summary>
+		/// Write a null field
+		/// </summary>
+		public static void WriteNullValue(this ICbWriter writer) => WriteNull(writer, default);
+
+		/// <summary>
+		/// Write a named null field
+		/// </summary>
+		/// <param name="name">Name of the field</param>
+		public static void WriteNull(this ICbWriter writer, Utf8String name) => writer.WriteField(CbFieldType.Null, name, 0);
+
+		/// <summary>
+		/// Writes a boolean value
+		/// </summary>
+		/// <param name="value"></param>
+		public static void WriteBoolValue(this ICbWriter writer, bool value) => WriteBool(writer, default, value);
+
+		/// <summary>
+		/// Writes a boolean value
+		/// </summary>
+		/// <param name="name">Name of the field</param>
+		/// <param name="value"></param>
+		public static void WriteBool(this ICbWriter writer, Utf8String name, bool value) => writer.WriteField(value ? CbFieldType.BoolTrue : CbFieldType.BoolFalse, name, 0);
+
+		/// <summary>
+		/// Writes an unnamed integer field
+		/// </summary>
+		/// <param name="value">Value to be written</param>
+		public static void WriteIntegerValue(this ICbWriter writer, int value) => WriteInteger(writer, default, value);
+
+		/// <summary>
+		/// Writes an unnamed integer field
+		/// </summary>
+		/// <param name="value">Value to be written</param>
+		public static void WriteIntegerValue(this ICbWriter writer, long value) => WriteInteger(writer, default, value);
+
+		/// <summary>
+		/// Writes an named integer field
+		/// </summary>
+		/// <param name="name">Name of the field</param>
+		/// <param name="value">Value to be written</param>
+		public static void WriteInteger(this ICbWriter writer, Utf8String name, int value) => WriteInteger(writer, name, (long)value);
+
+		/// <summary>
+		/// Writes an named integer field
+		/// </summary>
+		/// <param name="name">Name of the field</param>
+		/// <param name="value">Value to be written</param>
+		public static void WriteInteger(this ICbWriter writer, Utf8String name, long value)
+		{
+			if (value >= 0)
+			{
+				int length = VarInt.MeasureUnsigned((ulong)value);
+				Span<byte> data = writer.WriteField(CbFieldType.IntegerPositive, name, length);
+				VarInt.WriteUnsigned(data, (ulong)value);
+			}
+			else
+			{
+				int length = VarInt.MeasureUnsigned((ulong)-value);
+				Span<byte> data = writer.WriteField(CbFieldType.IntegerNegative, name, length);
+				VarInt.WriteUnsigned(data, (ulong)-value);
+			}
+		}
+
+		/// <summary>
+		/// Writes an unnamed integer field
+		/// </summary>
+		/// <param name="value">Value to be written</param>
+		public static void WriteIntegerValue(this ICbWriter writer, ulong value) => WriteInteger(writer, default, value);
+
+		/// <summary>
+		/// Writes a named integer field
+		/// </summary>
+		/// <param name="name">Name of the field</param>
+		/// <param name="value">Value to be written</param>
+		public static void WriteInteger(this ICbWriter writer, Utf8String name, ulong value)
+		{
+			int length = VarInt.MeasureUnsigned((ulong)value);
+			Span<byte> data = writer.WriteField(CbFieldType.IntegerPositive, name, length);
+			VarInt.WriteUnsigned(data, (ulong)value);
+		}
+
+		/// <summary>
+		/// Writes an unnamed double field
+		/// </summary>
+		/// <param name="value">Value to be written</param>
+		public static void WriteDoubleValue(this ICbWriter writer, double value) => WriteDouble(writer, default, value);
+
+		/// <summary>
+		/// Writes a named double field
+		/// </summary>
+		/// <param name="name">Name of the field</param>
+		/// <param name="value">Value to be written</param>
+		public static void WriteDouble(this ICbWriter writer, Utf8String name, double value)
+		{
+			Span<byte> buffer = writer.WriteField(CbFieldType.Float64, name, sizeof(double));
+			BinaryPrimitives.WriteDoubleBigEndian(buffer, value);
+		}
+
+		/// <summary>
+		/// Writes an unnamed <see cref="DateTime"/> field
+		/// </summary>
+		/// <param name="value">Value to be written</param>
+		public static void WriteDateTimeValue(this ICbWriter writer, DateTime value) => WriteDateTime(writer, default, value);
+
+		/// <summary>
+		/// Writes a named <see cref="DateTime"/> field
+		/// </summary>
+		/// <param name="name">Name of the field</param>
+		/// <param name="value">Value to be written</param>
+		public static void WriteDateTime(this ICbWriter writer, Utf8String name, DateTime value)
+		{
+			Span<byte> buffer = writer.WriteField(CbFieldType.DateTime, name, sizeof(long));
+			BinaryPrimitives.WriteInt64BigEndian(buffer, value.Ticks);
+		}
+
+		/// <summary>
+		/// Writes an unnamed <see cref="IoHash"/> field
+		/// </summary>
+		/// <param name="value">Value to be written</param>
+		public static void WriteHashValue(this ICbWriter writer, IoHash value) => WriteHash(writer, default, value);
+
+		/// <summary>
+		/// Writes a named <see cref="IoHash"/> field
+		/// </summary>
+		/// <param name="name">Name of the field</param>
+		/// <param name="value">Value to be written</param>
+		public static void WriteHash(this ICbWriter writer, Utf8String name, IoHash value)
+		{
+			Span<byte> buffer = writer.WriteField(CbFieldType.Hash, name, IoHash.NumBytes);
+			value.CopyTo(buffer);
+		}
+
+		/// <summary>
+		/// Writes an unnamed reference to a binary attachment
+		/// </summary>
+		/// <param name="hash">Hash of the attachment</param>
+		public static void WriteBinaryAttachmentValue(this ICbWriter writer, IoHash hash) => WriteBinaryAttachment(writer, default, hash);
+
+		/// <summary>
+		/// Writes a named reference to a binary attachment
+		/// </summary>
+		/// <param name="name">Name of the field</param>
+		/// <param name="hash">Hash of the attachment</param>
+		public static void WriteBinaryAttachment(this ICbWriter writer, Utf8String name, IoHash hash)
+		{
+			Span<byte> buffer = writer.WriteField(CbFieldType.BinaryAttachment, name, IoHash.NumBytes);
+			hash.CopyTo(buffer);
+		}
+
+		/// <summary>
+		/// Writes an object directly into the writer
+		/// </summary>
+		/// <param name="obj">Object to write</param>
+		public static void WriteObject(this ICbWriter writer, CbObject obj) => WriteObject(writer, default, obj);
+
+		/// <summary>
+		/// Writes an object directly into the writer
+		/// </summary>
+		/// <param name="name">Name of the object</param>
+		/// <param name="obj">Object to write</param>
+		public static void WriteObject(this ICbWriter writer, Utf8String name, CbObject obj)
+		{
+			ReadOnlyMemory<byte> view = obj.AsField().Payload;
+			Span<byte> buffer = writer.WriteField(CbFieldType.Object, name, view.Length);
+			view.Span.CopyTo(buffer);
+		}
+
+		/// <summary>
+		/// Writes an unnamed reference to an object attachment
+		/// </summary>
+		/// <param name="hash">Hash of the attachment</param>
+		public static void WriteObjectAttachmentValue(this ICbWriter writer, IoHash hash) => WriteObjectAttachment(writer, default, hash);
+
+		/// <summary>
+		/// Writes a named reference to an object attachment
+		/// </summary>
+		/// <param name="name">Name of the field</param>
+		/// <param name="hash">Hash of the attachment</param>
+		public static void WriteObjectAttachment(this ICbWriter writer, Utf8String name, IoHash hash)
+		{
+			Span<byte> buffer = writer.WriteField(CbFieldType.ObjectAttachment, name, IoHash.NumBytes);
+			hash.CopyTo(buffer);
+		}
+
+		/// <summary>
+		/// Writes an unnamed string value
+		/// </summary>
+		/// <param name="value">Value to be written</param>
+		public static void WriteStringValue(this ICbWriter writer, string value) => WriteUtf8StringValue(writer, value);
+
+		/// <summary>
+		/// Writes a named string value
+		/// </summary>
+		/// <param name="name">Name of the field</param>
+		/// <param name="value">Value to be written</param>
+		public static void WriteString(this ICbWriter writer, Utf8String name, string? value)
+		{
+			if(value != null)
+			{
+				writer.WriteUtf8String(name, value);
+			}
+		}
+
+		/// <summary>
+		/// Writes an unnamed string value
+		/// </summary>
+		/// <param name="value">Value to be written</param>
+		public static void WriteUtf8StringValue(this ICbWriter writer, Utf8String value) => WriteUtf8String(writer, default, value);
+
+		/// <summary>
+		/// Writes a named string value
+		/// </summary>
+		/// <param name="name">Name of the field</param>
+		/// <param name="value">Value to be written</param>
+		public static void WriteUtf8String(this ICbWriter writer, Utf8String name, Utf8String value)
+		{
+			Span<byte> buffer = WriteFieldWithLength(writer, CbFieldType.String, name, value.Length);
+			value.Span.CopyTo(buffer);
+		}
+
+		/// <summary>
+		/// Writes an unnamed binary value
+		/// </summary>
+		/// <param name="value">Value to be written</param>
+		public static void WriteBinarySpanValue(this ICbWriter writer, ReadOnlySpan<byte> value) => WriteBinarySpan(writer, default, value);
+
+		/// <summary>
+		/// Writes a named binary value
+		/// </summary>
+		/// <param name="name">Name of the field</param>
+		/// <param name="value">Value to be written</param>
+		public static void WriteBinarySpan(this ICbWriter writer, Utf8String name, ReadOnlySpan<byte> value)
+		{
+			Span<byte> buffer = WriteFieldWithLength(writer, CbFieldType.Binary, name, value.Length);
+			value.CopyTo(buffer);
+		}
+
+		/// <summary>
+		/// Writes an unnamed binary value
+		/// </summary>
+		/// <param name="value">Value to be written</param>
+		public static void WriteBinaryValue(this ICbWriter writer, ReadOnlyMemory<byte> value) => writer.WriteBinarySpanValue(value.Span);
+
+		/// <summary>
+		/// Writes a named binary value
+		/// </summary>
+		/// <param name="name">Name of the field</param>
+		/// <param name="value">Value to be written</param>
+		public static void WriteBinary(this ICbWriter writer, Utf8String name, ReadOnlyMemory<byte> value) => writer.WriteBinarySpan(name, value.Span);
+
+		/// <summary>
+		/// Writes an unnamed binary value
+		/// </summary>
+		/// <param name="value">Value to be written</param>
+		public static void WriteBinaryArrayValue(this ICbWriter writer, byte[] value) => writer.WriteBinarySpanValue(value.AsSpan());
+
+		/// <summary>
+		/// Writes a named binary value
+		/// </summary>
+		/// <param name="name">Name of the field</param>
+		/// <param name="value">Value to be written</param>
+		public static void WriteBinaryArray(this ICbWriter writer, Utf8String name, byte[] value) => writer.WriteBinarySpan(name, value.AsSpan());
+	}
 }
