Adding List<> extension method "AddAll(...)" that can replace "AddRange(new string[] { ... })"

diff --git a/Engine/Source/Programs/Shared/EpicGames.Core/ListExtensions.cs b/Engine/Source/Programs/Shared/EpicGames.Core/ListExtensions.cs
--- a/Engine/Source/Programs/Shared/EpicGames.Core/ListExtensions.cs
+++ b/Engine/Source/Programs/Shared/EpicGames.Core/ListExtensions.cs
@@ -3,6 +3,7 @@
 using System;
 using System.Collections;
 using System.Collections.Generic;
+using System.Linq;
 using System.Threading.Tasks;
 
 namespace EpicGames.Core
@@ -65,6 +66,17 @@ namespace EpicGames.Core
 			list.Sort((x, y) => comparer.Compare(selector(x), selector(y)));
 		}
 
+		/// <summary>
+		/// Add all arguments to the list
+		/// </summary>
+		/// <typeparam name="TElement">List element</typeparam>
+		/// <param name="list">List to add to</param>
+		/// <param name="arguments">Arguments to add to list</param>
+		public static void AddAll<TElement>(this List<TElement> list, params TElement[] arguments)
+		{
+			list.AddRange(arguments.AsEnumerable());
+		}
+
 		/// <summary>
 		/// Convert all elements of a list to a base type
 		/// </summary>
