Horde: Missing file.

diff --git a/Engine/Source/Programs/Shared/EpicGames.Horde/Compute/Cpp/CppComputeOutputNode.cs b/Engine/Source/Programs/Shared/EpicGames.Horde/Compute/Cpp/CppComputeOutputNode.cs
new file mode 100644
--- /dev/null
+++ b/Engine/Source/Programs/Shared/EpicGames.Horde/Compute/Cpp/CppComputeOutputNode.cs
@@ -0,0 +1,71 @@
+// Copyright Epic Games, Inc. All Rights Reserved.
+
+using EpicGames.Core;
+using EpicGames.Horde.Logs;
+using EpicGames.Horde.Storage;
+using EpicGames.Horde.Storage.Nodes;
+using System.Collections.Generic;
+
+namespace EpicGames.Horde.Compute.Cpp
+{
+	/// <summary>
+	/// Output from executing a C++ remote execution job
+	/// </summary>
+	[TreeNode("{B6D273EB-FA61-4F37-9FC2-7532EB274339}", 1)]
+	public class CppComputeOutputNode : TreeNode
+	{
+		/// <summary>
+		/// Exit code from the process
+		/// </summary>
+		public int ExitCode { get; }
+
+		/// <summary>
+		/// Log from the process execution
+		/// </summary>
+		public TreeNodeRef<LogNode> Log { get; }
+
+		/// <summary>
+		/// Reference to the tree of output files
+		/// </summary>
+		public TreeNodeRef<DirectoryNode> Tree { get; }
+
+		/// <summary>
+		/// Constructor
+		/// </summary>
+		/// <param name="exitCode">Exit code for the process</param>
+		/// <param name="log">Log object</param>
+		/// <param name="tree">Tree of output files</param>
+		public CppComputeOutputNode(int exitCode, TreeNodeRef<LogNode> log, TreeNodeRef<DirectoryNode> tree)
+		{
+			ExitCode = exitCode;
+			Log = log;
+			Tree = tree;
+		}
+
+		/// <summary>
+		/// Deserialization constructor
+		/// </summary>
+		/// <param name="reader">Reader to deserialize from</param>
+		public CppComputeOutputNode(ITreeNodeReader reader)
+		{
+			ExitCode = (int)reader.ReadUnsignedVarInt();
+			Log = reader.ReadRef<LogNode>();
+			Tree = reader.ReadRef<DirectoryNode>();
+		}
+
+		/// <inheritdoc/>
+		public override void Serialize(ITreeNodeWriter writer)
+		{
+			writer.WriteUnsignedVarInt(ExitCode);
+			writer.WriteRef(Log);
+			writer.WriteRef(Tree);
+		}
+
+		/// <inheritdoc/>
+		public override IEnumerable<TreeNodeRef> EnumerateRefs()
+		{
+			yield return Log;
+			yield return Tree;
+		}
+	}
+}
