Horde: CI warning fixes.

diff --git a/Engine/Source/Programs/Horde/Horde.Build/Agents/Fleet/ComputeQueueAwsMetricStrategy.cs b/Engine/Source/Programs/Horde/Horde.Build/Agents/Fleet/ComputeQueueAwsMetricStrategy.cs
--- a/Engine/Source/Programs/Horde/Horde.Build/Agents/Fleet/ComputeQueueAwsMetricStrategy.cs
+++ b/Engine/Source/Programs/Horde/Horde.Build/Agents/Fleet/ComputeQueueAwsMetricStrategy.cs
@@ -43,6 +43,16 @@ namespace Horde.Build.Agents.Fleet
 		}
 	}
 
+	static class ComputeServiceExtensions
+	{
+		public static Task<int> GetNumQueuedTasksForPoolAsync(this ComputeService computeService, ClusterId clusterId, IPool pool)
+		{
+			// Will need to reimplement this functionality in ComputeService if we want to use this strategy, but conditions on tasks
+			// may reference properties that are specific to an agent rather than a pool...
+			throw new NotImplementedException();
+		}
+	}
+
 	/// <summary>
 	/// A no-op strategy that reports size of compute task queue for a given pool as AWS CloudWatch metrics
 	/// A metric that later can be used as a source for AWS-controlled auto-scaling policies. 
@@ -50,7 +60,7 @@ namespace Horde.Build.Agents.Fleet
 	public class ComputeQueueAwsMetricStrategy : IPoolSizeStrategy
 	{
 		private readonly IAmazonCloudWatch _cloudWatch;
-		private readonly IComputeService _computeService;
+		private readonly ComputeService _computeService;
 		private readonly ComputeQueueAwsMetricSettings _settings;
 		private readonly ILogger<ComputeQueueAwsMetricStrategy> _logger;
 
@@ -61,7 +71,7 @@ namespace Horde.Build.Agents.Fleet
 		/// <param name="computeService"></param>
 		/// <param name="settings"></param>
 		/// <param name="logger"></param>
-		public ComputeQueueAwsMetricStrategy(IAmazonCloudWatch cloudWatch, IComputeService computeService,
+		public ComputeQueueAwsMetricStrategy(IAmazonCloudWatch cloudWatch, ComputeService computeService,
 			ComputeQueueAwsMetricSettings settings, ILogger<ComputeQueueAwsMetricStrategy> logger)
 		{
 			_cloudWatch = cloudWatch;
diff --git a/Engine/Source/Programs/Shared/EpicGames.Serialization/CbWriter.cs b/Engine/Source/Programs/Shared/EpicGames.Serialization/CbWriter.cs
--- a/Engine/Source/Programs/Shared/EpicGames.Serialization/CbWriter.cs
+++ b/Engine/Source/Programs/Shared/EpicGames.Serialization/CbWriter.cs
@@ -131,6 +131,9 @@ namespace EpicGames.Serialization
 			}
 		}
 
+		/// <summary>
+		/// Size of data to preallocate by default
+		/// </summary>
 		public const int DefaultChunkSize = 1024;
 
 		readonly List<Chunk> _chunks = new List<Chunk>();
@@ -288,9 +291,10 @@ namespace EpicGames.Serialization
 		/// <summary>
 		/// Writes the header for a named field
 		/// </summary>
-		/// <param name="type"></param>
-		/// <param name="name"></param>
-		public Span<byte> WriteField(CbFieldType type, Utf8String name, int length)
+		/// <param name="type">Type of the field</param>
+		/// <param name="name">Name of the field</param>
+		/// <param name="size">Size of the field</param>
+		public Span<byte> WriteField(CbFieldType type, Utf8String name, int size)
 		{
 			if (name.IsEmpty)
 			{
@@ -336,7 +340,7 @@ namespace EpicGames.Serialization
 				}
 				scope._count++;
 			}
-			return Allocate(length);
+			return Allocate(size);
 		}
 
 		/// <summary>
@@ -375,6 +379,7 @@ namespace EpicGames.Serialization
 		/// Begin writing a named array field
 		/// </summary>
 		/// <param name="name"></param>
+		/// <param name="elementType">Type of elements in the array</param>
 		public void BeginArray(Utf8String name, CbFieldType elementType)
 		{
 			if (elementType == CbFieldType.None)
@@ -664,28 +669,33 @@ namespace EpicGames.Serialization
 		/// <summary>
 		/// Begin writing an object field
 		/// </summary>
+		/// <param name="writer">Writer for output data</param>
 		public static void BeginObject(this ICbWriter writer) => writer.BeginObject(default);
 
 		/// <summary>
 		/// Begin writing an array field
 		/// </summary>
+		/// <param name="writer">Writer for output data</param>
 		public static void BeginArray(this ICbWriter writer) => writer.BeginArray(default, CbFieldType.None);
 
 		/// <summary>
 		/// Begin writing a named array field
 		/// </summary>
+		/// <param name="writer">Writer for output data</param>
 		/// <param name="name">Name of the field</param>
 		public static void BeginArray(this ICbWriter writer, Utf8String name) => writer.BeginArray(name, CbFieldType.None);
 
 		/// <summary>
 		/// Begin writing a uniform array field
 		/// </summary>
+		/// <param name="writer">Writer for output data</param>
 		/// <param name="fieldType">The field type for elements in the array</param>
 		public static void BeginUniformArray(this ICbWriter writer, CbFieldType fieldType) => BeginUniformArray(writer, default, fieldType);
 
 		/// <summary>
 		/// Begin writing a named uniform array field
 		/// </summary>
+		/// <param name="writer">Writer for output data</param>
 		/// <param name="name">Name of the field</param>
 		/// <param name="fieldType">The field type for elements in the array</param>
 		public static void BeginUniformArray(this ICbWriter writer, Utf8String name, CbFieldType fieldType) => writer.BeginArray(name, fieldType);
@@ -693,24 +703,27 @@ namespace EpicGames.Serialization
 		/// <summary>
 		/// End writing a uniform array field
 		/// </summary>
-		/// <param name="writer"></param>
+		/// <param name="writer">Writer for output data</param>
 		public static void EndUniformArray(this ICbWriter writer) => writer.EndArray();
 
 		/// <summary>
 		/// Copies an entire field value to the output
 		/// </summary>
+		/// <param name="writer">Writer for output data</param>
 		/// <param name="field"></param>
 		public static void WriteFieldValue(this ICbWriter writer, CbField field) => WriteField(writer, default, field);
 
 		/// <summary>
 		/// Copies an entire field value to the output, using the name from the field
 		/// </summary>
+		/// <param name="writer">Writer for output data</param>
 		/// <param name="field"></param>
 		public static void WriteField(this ICbWriter writer, CbField field) => WriteField(writer, field.GetName(), field);
 
 		/// <summary>
 		/// Copies an entire field value to the output
 		/// </summary>
+		/// <param name="writer">Writer for output data</param>
 		/// <param name="name">Name of the field</param>
 		/// <param name="field"></param>
 		public static void WriteField(this ICbWriter writer, Utf8String name, CbField field)
@@ -723,42 +736,49 @@ namespace EpicGames.Serialization
 		/// <summary>
 		/// Write a null field
 		/// </summary>
+		/// <param name="writer">Writer for output data</param>
 		public static void WriteNullValue(this ICbWriter writer) => WriteNull(writer, default);
 
 		/// <summary>
 		/// Write a named null field
 		/// </summary>
+		/// <param name="writer">Writer for output data</param>
 		/// <param name="name">Name of the field</param>
 		public static void WriteNull(this ICbWriter writer, Utf8String name) => writer.WriteField(CbFieldType.Null, name, 0);
 
 		/// <summary>
 		/// Writes a boolean value
 		/// </summary>
-		/// <param name="value"></param>
+		/// <param name="writer">Writer for output data</param>
+		/// <param name="value">Value to be written</param>
 		public static void WriteBoolValue(this ICbWriter writer, bool value) => WriteBool(writer, default, value);
 
 		/// <summary>
 		/// Writes a boolean value
 		/// </summary>
+		/// <param name="writer">Writer for output data</param>
 		/// <param name="name">Name of the field</param>
-		/// <param name="value"></param>
+		/// <param name="value">Value to be written</param>
 		public static void WriteBool(this ICbWriter writer, Utf8String name, bool value) => writer.WriteField(value ? CbFieldType.BoolTrue : CbFieldType.BoolFalse, name, 0);
 
 		/// <summary>
 		/// Writes an unnamed integer field
 		/// </summary>
+		/// <param name="writer">Writer for output data</param>
 		/// <param name="value">Value to be written</param>
 		public static void WriteIntegerValue(this ICbWriter writer, int value) => WriteInteger(writer, default, value);
 
 		/// <summary>
 		/// Writes an unnamed integer field
 		/// </summary>
+		/// <param name="writer">Writer for output data</param>
 		/// <param name="value">Value to be written</param>
 		public static void WriteIntegerValue(this ICbWriter writer, long value) => WriteInteger(writer, default, value);
 
 		/// <summary>
 		/// Writes an named integer field
 		/// </summary>
+		/// <param name="writer">Writer for output data</param>
 		/// <param name="name">Name of the field</param>
 		/// <param name="value">Value to be written</param>
 		public static void WriteInteger(this ICbWriter writer, Utf8String name, int value) => WriteInteger(writer, name, (long)value);
@@ -766,6 +786,7 @@ namespace EpicGames.Serialization
 		/// <summary>
 		/// Writes an named integer field
 		/// </summary>
+		/// <param name="writer">Writer for output data</param>
 		/// <param name="name">Name of the field</param>
 		/// <param name="value">Value to be written</param>
 		public static void WriteInteger(this ICbWriter writer, Utf8String name, long value)
@@ -787,12 +808,14 @@ namespace EpicGames.Serialization
 		/// <summary>
 		/// Writes an unnamed integer field
 		/// </summary>
+		/// <param name="writer">Writer for output data</param>
 		/// <param name="value">Value to be written</param>
 		public static void WriteIntegerValue(this ICbWriter writer, ulong value) => WriteInteger(writer, default, value);
 
 		/// <summary>
 		/// Writes a named integer field
 		/// </summary>
+		/// <param name="writer">Writer for output data</param>
 		/// <param name="name">Name of the field</param>
 		/// <param name="value">Value to be written</param>
 		public static void WriteInteger(this ICbWriter writer, Utf8String name, ulong value)
@@ -805,12 +828,14 @@ namespace EpicGames.Serialization
 		/// <summary>
 		/// Writes an unnamed double field
 		/// </summary>
+		/// <param name="writer">Writer for output data</param>
 		/// <param name="value">Value to be written</param>
 		public static void WriteDoubleValue(this ICbWriter writer, double value) => WriteDouble(writer, default, value);
 
 		/// <summary>
 		/// Writes a named double field
 		/// </summary>
+		/// <param name="writer">Writer for output data</param>
 		/// <param name="name">Name of the field</param>
 		/// <param name="value">Value to be written</param>
 		public static void WriteDouble(this ICbWriter writer, Utf8String name, double value)
@@ -822,12 +847,14 @@ namespace EpicGames.Serialization
 		/// <summary>
 		/// Writes an unnamed <see cref="DateTime"/> field
 		/// </summary>
+		/// <param name="writer">Writer for output data</param>
 		/// <param name="value">Value to be written</param>
 		public static void WriteDateTimeValue(this ICbWriter writer, DateTime value) => WriteDateTime(writer, default, value);
 
 		/// <summary>
 		/// Writes a named <see cref="DateTime"/> field
 		/// </summary>
+		/// <param name="writer">Writer for output data</param>
 		/// <param name="name">Name of the field</param>
 		/// <param name="value">Value to be written</param>
 		public static void WriteDateTime(this ICbWriter writer, Utf8String name, DateTime value)
@@ -839,12 +866,14 @@ namespace EpicGames.Serialization
 		/// <summary>
 		/// Writes an unnamed <see cref="IoHash"/> field
 		/// </summary>
+		/// <param name="writer">Writer for output data</param>
 		/// <param name="value">Value to be written</param>
 		public static void WriteHashValue(this ICbWriter writer, IoHash value) => WriteHash(writer, default, value);
 
 		/// <summary>
 		/// Writes a named <see cref="IoHash"/> field
 		/// </summary>
+		/// <param name="writer">Writer for output data</param>
 		/// <param name="name">Name of the field</param>
 		/// <param name="value">Value to be written</param>
 		public static void WriteHash(this ICbWriter writer, Utf8String name, IoHash value)
@@ -856,12 +885,14 @@ namespace EpicGames.Serialization
 		/// <summary>
 		/// Writes an unnamed reference to a binary attachment
 		/// </summary>
+		/// <param name="writer">Writer for output data</param>
 		/// <param name="hash">Hash of the attachment</param>
 		public static void WriteBinaryAttachmentValue(this ICbWriter writer, IoHash hash) => WriteBinaryAttachment(writer, default, hash);
 
 		/// <summary>
 		/// Writes a named reference to a binary attachment
 		/// </summary>
+		/// <param name="writer">Writer for output data</param>
 		/// <param name="name">Name of the field</param>
 		/// <param name="hash">Hash of the attachment</param>
 		public static void WriteBinaryAttachment(this ICbWriter writer, Utf8String name, IoHash hash)
@@ -873,12 +904,14 @@ namespace EpicGames.Serialization
 		/// <summary>
 		/// Writes an object directly into the writer
 		/// </summary>
+		/// <param name="writer">Writer for output data</param>
 		/// <param name="obj">Object to write</param>
 		public static void WriteObject(this ICbWriter writer, CbObject obj) => WriteObject(writer, default, obj);
 
 		/// <summary>
 		/// Writes an object directly into the writer
 		/// </summary>
+		/// <param name="writer">Writer for output data</param>
 		/// <param name="name">Name of the object</param>
 		/// <param name="obj">Object to write</param>
 		public static void WriteObject(this ICbWriter writer, Utf8String name, CbObject obj)
@@ -891,12 +924,14 @@ namespace EpicGames.Serialization
 		/// <summary>
 		/// Writes an unnamed reference to an object attachment
 		/// </summary>
+		/// <param name="writer">Writer for output data</param>
 		/// <param name="hash">Hash of the attachment</param>
 		public static void WriteObjectAttachmentValue(this ICbWriter writer, IoHash hash) => WriteObjectAttachment(writer, default, hash);
 
 		/// <summary>
 		/// Writes a named reference to an object attachment
 		/// </summary>
+		/// <param name="writer">Writer for output data</param>
 		/// <param name="name">Name of the field</param>
 		/// <param name="hash">Hash of the attachment</param>
 		public static void WriteObjectAttachment(this ICbWriter writer, Utf8String name, IoHash hash)
@@ -908,12 +943,14 @@ namespace EpicGames.Serialization
 		/// <summary>
 		/// Writes an unnamed string value
 		/// </summary>
+		/// <param name="writer">Writer for output data</param>
 		/// <param name="value">Value to be written</param>
 		public static void WriteStringValue(this ICbWriter writer, string value) => WriteUtf8StringValue(writer, value);
 
 		/// <summary>
 		/// Writes a named string value
 		/// </summary>
+		/// <param name="writer">Writer for output data</param>
 		/// <param name="name">Name of the field</param>
 		/// <param name="value">Value to be written</param>
 		public static void WriteString(this ICbWriter writer, Utf8String name, string? value)
@@ -927,12 +964,14 @@ namespace EpicGames.Serialization
 		/// <summary>
 		/// Writes an unnamed string value
 		/// </summary>
+		/// <param name="writer">Writer for output data</param>
 		/// <param name="value">Value to be written</param>
 		public static void WriteUtf8StringValue(this ICbWriter writer, Utf8String value) => WriteUtf8String(writer, default, value);
 
 		/// <summary>
 		/// Writes a named string value
 		/// </summary>
+		/// <param name="writer">Writer for output data</param>
 		/// <param name="name">Name of the field</param>
 		/// <param name="value">Value to be written</param>
 		public static void WriteUtf8String(this ICbWriter writer, Utf8String name, Utf8String value)
@@ -944,12 +983,14 @@ namespace EpicGames.Serialization
 		/// <summary>
 		/// Writes an unnamed binary value
 		/// </summary>
+		/// <param name="writer">Writer for output data</param>
 		/// <param name="value">Value to be written</param>
 		public static void WriteBinarySpanValue(this ICbWriter writer, ReadOnlySpan<byte> value) => WriteBinarySpan(writer, default, value);
 
 		/// <summary>
 		/// Writes a named binary value
 		/// </summary>
+		/// <param name="writer">Writer for output data</param>
 		/// <param name="name">Name of the field</param>
 		/// <param name="value">Value to be written</param>
 		public static void WriteBinarySpan(this ICbWriter writer, Utf8String name, ReadOnlySpan<byte> value)
@@ -961,12 +1002,14 @@ namespace EpicGames.Serialization
 		/// <summary>
 		/// Writes an unnamed binary value
 		/// </summary>
+		/// <param name="writer">Writer for output data</param>
 		/// <param name="value">Value to be written</param>
 		public static void WriteBinaryValue(this ICbWriter writer, ReadOnlyMemory<byte> value) => writer.WriteBinarySpanValue(value.Span);
 
 		/// <summary>
 		/// Writes a named binary value
 		/// </summary>
+		/// <param name="writer">Writer for output data</param>
 		/// <param name="name">Name of the field</param>
 		/// <param name="value">Value to be written</param>
 		public static void WriteBinary(this ICbWriter writer, Utf8String name, ReadOnlyMemory<byte> value) => writer.WriteBinarySpan(name, value.Span);
@@ -974,12 +1017,14 @@ namespace EpicGames.Serialization
 		/// <summary>
 		/// Writes an unnamed binary value
 		/// </summary>
+		/// <param name="writer">Writer for output data</param>
 		/// <param name="value">Value to be written</param>
 		public static void WriteBinaryArrayValue(this ICbWriter writer, byte[] value) => writer.WriteBinarySpanValue(value.AsSpan());
 
 		/// <summary>
 		/// Writes a named binary value
 		/// </summary>
+		/// <param name="writer">Writer for output data</param>
 		/// <param name="name">Name of the field</param>
 		/// <param name="value">Value to be written</param>
 		public static void WriteBinaryArray(this ICbWriter writer, Utf8String name, byte[] value) => writer.WriteBinarySpan(name, value.AsSpan());
