Horde: Fix build warnings.

diff --git a/Engine/Source/Programs/Shared/EpicGames.Horde/Compute/ComputeMessageTypes.cs b/Engine/Source/Programs/Shared/EpicGames.Horde/Compute/ComputeMessageTypes.cs
--- a/Engine/Source/Programs/Shared/EpicGames.Horde/Compute/ComputeMessageTypes.cs
+++ b/Engine/Source/Programs/Shared/EpicGames.Horde/Compute/ComputeMessageTypes.cs
@@ -149,7 +149,7 @@ namespace EpicGames.Horde.Compute
 			{
 				using (Stream stream = await storage.ReadBlobAsync(locator, cancellationToken))
 				{
-					MemoryStream target = new MemoryStream();
+					using MemoryStream target = new MemoryStream();
 					await stream.CopyToAsync(target, cancellationToken);
 					response.Data = target.ToArray();
 				}
@@ -158,7 +158,7 @@ namespace EpicGames.Horde.Compute
 			{
 				using (Stream stream = await storage.ReadBlobRangeAsync(locator, offset, length, cancellationToken))
 				{
-					MemoryStream target = new MemoryStream();
+					using MemoryStream target = new MemoryStream();
 					await stream.CopyToAsync(target, cancellationToken);
 					response.Data = target.ToArray();
 				}
