Horde: C++ compute improvements.

* Add an in-memory cache for data retrieved from the initiator.
* Allow appending partial line chunks to LogChunkNode instances.
* Allow constructing a server compute handler instance with the URL of a server rather than a HttpClient factory.

diff --git a/Engine/Source/Programs/Horde/Horde.Agent/Commands/Compute/CppComputeCommand.cs b/Engine/Source/Programs/Horde/Horde.Agent/Commands/Compute/CppComputeCommand.cs
--- a/Engine/Source/Programs/Horde/Horde.Agent/Commands/Compute/CppComputeCommand.cs
+++ b/Engine/Source/Programs/Horde/Horde.Agent/Commands/Compute/CppComputeCommand.cs
@@ -60,9 +60,9 @@ namespace Horde.Agent.Commands
 		/// <inheritdoc/>
 		public override async Task<int> ExecuteAsync(ILogger logger)
 		{
-			await using IComputeClient channel = CreateClient();
+			await using IComputeClient client = CreateClient();
 
-			IComputeRequest<bool> request = await channel.AddRequestAsync<bool>(default, null, HandleRequestAsync);
+			IComputeRequest<bool> request = await client.AddRequestAsync<bool>(default, null, HandleRequestAsync);
 			bool result = await request.Result;
 
 			return result? 0 : 1;
diff --git a/Engine/Source/Programs/Horde/Horde.Agent/Leases/Handlers/ComputeHandler.cs b/Engine/Source/Programs/Horde/Horde.Agent/Leases/Handlers/ComputeHandler.cs
--- a/Engine/Source/Programs/Horde/Horde.Agent/Leases/Handlers/ComputeHandler.cs
+++ b/Engine/Source/Programs/Horde/Horde.Agent/Leases/Handlers/ComputeHandler.cs
@@ -7,7 +7,6 @@ using System.Linq;
 using System.Net.Sockets;
 using System.Threading;
 using System.Threading.Tasks;
-using Amazon.EC2;
 using EpicGames.Core;
 using EpicGames.Horde.Compute;
 using EpicGames.Horde.Compute.Cpp;
@@ -17,7 +16,7 @@ using EpicGames.Horde.Storage.Backends;
 using EpicGames.Horde.Storage.Nodes;
 using Horde.Agent.Services;
 using HordeCommon.Rpc.Tasks;
-using Microsoft.CodeAnalysis.CSharp.Syntax;
+using Microsoft.Extensions.Caching.Memory;
 using Microsoft.Extensions.Logging;
 
 namespace Horde.Agent.Leases.Handlers
@@ -27,13 +26,15 @@ namespace Horde.Agent.Leases.Handlers
 	/// </summary>
 	class ComputeHandler : LeaseHandler<ComputeTask>
 	{
+		readonly IMemoryCache _memoryCache;
 		readonly ILogger _logger;
 
 		/// <summary>
 		/// Constructor
 		/// </summary>
-		public ComputeHandler(ILogger<ComputeHandler> logger)
+		public ComputeHandler(IMemoryCache memoryCache, ILogger<ComputeHandler> logger)
 		{
+			_memoryCache = memoryCache;
 			_logger = logger;
 		}
 
@@ -99,7 +100,7 @@ namespace Horde.Agent.Leases.Handlers
 			DirectoryReference sandboxDir = DirectoryReference.Combine(Program.DataDir, "Sandbox");
 
 			using ComputeStorageClient store = new ComputeStorageClient(channel);
-			TreeReader reader = new TreeReader(store, null, _logger);
+			TreeReader reader = new TreeReader(store, _memoryCache, _logger);
 
 			CppComputeNode node = await reader.ReadNodeAsync<CppComputeNode>(locator, cancellationToken);
 			DirectoryNode directoryNode = await node.Sandbox.ExpandAsync(reader, cancellationToken);
diff --git a/Engine/Source/Programs/Horde/Horde.Agent/Program.cs b/Engine/Source/Programs/Horde/Horde.Agent/Program.cs
--- a/Engine/Source/Programs/Horde/Horde.Agent/Program.cs
+++ b/Engine/Source/Programs/Horde/Horde.Agent/Program.cs
@@ -203,6 +203,8 @@ namespace Horde.Agent
 
 			services.AddSingleton<IStorageClientFactory, StorageClientFactory>();
 
+			services.AddMemoryCache();
+
 			// Allow commands to augment the service collection for their own DI service providers
 			services.AddSingleton<DefaultServices>(x => new DefaultServices(config, services));
 
diff --git a/Engine/Source/Programs/Shared/EpicGames.Horde/Compute/ServerComputeClient.cs b/Engine/Source/Programs/Shared/EpicGames.Horde/Compute/ServerComputeClient.cs
--- a/Engine/Source/Programs/Shared/EpicGames.Horde/Compute/ServerComputeClient.cs
+++ b/Engine/Source/Programs/Shared/EpicGames.Horde/Compute/ServerComputeClient.cs
@@ -115,6 +115,7 @@ namespace EpicGames.Horde.Compute
 		}
 
 		readonly object _lockObject = new object();
+		readonly HttpClient? _defaultHttpClient;
 		readonly Func<HttpClient> _createHttpClient;
 		readonly CancellationTokenSource _cancellationSource = new CancellationTokenSource();
 		readonly Dictionary<ReadOnlyMemory<byte>, RequestInfo> _requests = new Dictionary<ReadOnlyMemory<byte>, RequestInfo>(MemoryComparer.Instance);
@@ -122,6 +123,26 @@ namespace EpicGames.Horde.Compute
 		Task _listenTask = Task.CompletedTask;
 		readonly ILogger _logger;
 
+		/// <summary>
+		/// Constructor
+		/// </summary>
+		/// <param name="serverUri">Uri of the server to connect to</param>
+		/// <param name="logger">Logger for diagnostic messages</param>
+		public ServerComputeClient(Uri serverUri, ILogger logger)
+		{
+#pragma warning disable CA2000 // Dispose objects before losing scope
+			// This is disposed via HttpClient
+			SocketsHttpHandler handler = new SocketsHttpHandler();
+			handler.PooledConnectionLifetime = TimeSpan.FromMinutes(2.0);
+
+			_defaultHttpClient = new HttpClient(handler, true);
+			_defaultHttpClient.BaseAddress = serverUri;
+#pragma warning restore CA2000 // Dispose objects before losing scope
+
+			_createHttpClient = GetDefaultHttpClient;
+			_logger = logger;
+		}
+
 		/// <summary>
 		/// Constructor
 		/// </summary>
@@ -133,11 +154,18 @@ namespace EpicGames.Horde.Compute
 			_logger = logger;
 		}
 
+		/// <summary>
+		/// Gets the default http client
+		/// </summary>
+		/// <returns></returns>
+		HttpClient GetDefaultHttpClient() => _defaultHttpClient!;
+
 		/// <inheritdoc/>
 		public async ValueTask DisposeAsync()
 		{
 			await StopAsync();
 			_cancellationSource.Dispose();
+			_defaultHttpClient?.Dispose();
 		}
 
 		/// <inheritdoc/>
diff --git a/Engine/Source/Programs/Shared/EpicGames.Horde/Logs/LogChunkNode.cs b/Engine/Source/Programs/Shared/EpicGames.Horde/Logs/LogChunkNode.cs
--- a/Engine/Source/Programs/Shared/EpicGames.Horde/Logs/LogChunkNode.cs
+++ b/Engine/Source/Programs/Shared/EpicGames.Horde/Logs/LogChunkNode.cs
@@ -170,21 +170,11 @@ namespace EpicGames.Horde.Logs
 		/// <param name="lineOffsets">Offsets of each line within the text</param>
 		public static void UpdateLineOffsets(ReadOnlySpan<byte> data, int start, List<int> lineOffsets)
 		{
-			if (start < data.Length)
+			for (int idx = start; idx < data.Length; idx++)
 			{
-				// Make sure the data ends with a newline
-				if (data[data.Length - 1] != '\n')
+				if (data[idx] == '\n')
 				{
-					throw new InvalidDataException("Chunk data must end with a newline");
-				}
-
-				// Calculate the new number of newlines
-				for (int idx = start; idx < data.Length; idx++)
-				{
-					if (data[idx] == '\n')
-					{
-						lineOffsets.Add(idx + 1);
-					}
+					lineOffsets.Add(idx + 1);
 				}
 			}
 		}
