UGS: Fix threading issue when opening multiple startup projects with OIDC configuration.

diff --git a/Engine/Source/Programs/Shared/EpicGames.OIDC/OidcTokenManager.cs b/Engine/Source/Programs/Shared/EpicGames.OIDC/OidcTokenManager.cs
--- a/Engine/Source/Programs/Shared/EpicGames.OIDC/OidcTokenManager.cs
+++ b/Engine/Source/Programs/Shared/EpicGames.OIDC/OidcTokenManager.cs
@@ -18,7 +18,6 @@ using IdentityModel.OidcClient;
 using IdentityModel.OidcClient.Results;
 using Microsoft.Extensions.Configuration;
 using Microsoft.Extensions.DependencyInjection;
-using Microsoft.Extensions.Logging;
 using Microsoft.Extensions.Options;
 using Microsoft.Win32;
 
@@ -26,6 +25,7 @@ namespace EpicGames.OIDC
 {
 	public class OidcTokenManager
 	{
+		private readonly object _lockObject = new object();
 		private readonly Dictionary<string, OidcTokenClient> _tokenClients = new Dictionary<string, OidcTokenClient>(StringComparer.OrdinalIgnoreCase);
 		private readonly ITokenStore _tokenStore;
 
@@ -36,6 +36,8 @@ namespace EpicGames.OIDC
 
 		public OidcTokenClient FindOrAddClient(string name, ProviderInfo providerInfo)
 		{
+			lock (_lockObject)
+			{
 			OidcTokenClient? client;
 			if (!_tokenClients.TryGetValue(name, out client))
 			{
@@ -44,6 +46,7 @@ namespace EpicGames.OIDC
 			}
 			return client;
 		}
+		}
 
 		public OidcTokenManager(IServiceProvider provider, IOptionsMonitor<OidcTokenOptions> settings, ITokenStore tokenStore, List<string>? allowedProviders = null)
 			: this(tokenStore)
