EpicGames.Core: Fix NET 3.1 compilation of ByteString struct.

diff --git a/Engine/Source/Programs/Shared/EpicGames.Core/ByteString.cs b/Engine/Source/Programs/Shared/EpicGames.Core/ByteString.cs
--- a/Engine/Source/Programs/Shared/EpicGames.Core/ByteString.cs
+++ b/Engine/Source/Programs/Shared/EpicGames.Core/ByteString.cs
@@ -117,9 +117,18 @@ namespace EpicGames.Core
 		/// <inheritdoc/>
 		public override int GetHashCode()
 		{
+#if NETCOREAPP3_1
+			HashCode hashCode = new HashCode();
+			foreach (byte value in Data.Span)
+			{
+				hashCode.Add(value);
+			}
+			return hashCode.ToHashCode();
+#else
 			HashCode hashCode = new HashCode();
 			hashCode.AddBytes(Data.Span);
 			return hashCode.ToHashCode();
+#endif
 		}
 
 		/// <inheritdoc/>
