EpicGames.Core: Add missing equality and comparison operators causing static analysis warnings.

diff --git a/Engine/Source/Programs/Shared/EpicGames.Core/ByteString.cs b/Engine/Source/Programs/Shared/EpicGames.Core/ByteString.cs
--- a/Engine/Source/Programs/Shared/EpicGames.Core/ByteString.cs
+++ b/Engine/Source/Programs/Shared/EpicGames.Core/ByteString.cs
@@ -138,6 +138,15 @@ namespace EpicGames.Core
 		/// Creates a Utf8 string represending the bytes in this string
 		/// </summary>
 		public Utf8String ToUtf8String() => StringUtils.FormatUtf8HexString(Data.Span);
+
+#pragma warning disable CS1591 // Missing XML comment for publicly visible type or member
+		public static bool operator ==(ByteString left, ByteString right) => left.Equals(right);
+		public static bool operator !=(ByteString left, ByteString right) => !(left == right);
+		public static bool operator <(ByteString left, ByteString right) => left.CompareTo(right) < 0;
+		public static bool operator <=(ByteString left, ByteString right) => left.CompareTo(right) <= 0;
+		public static bool operator >(ByteString left, ByteString right) => left.CompareTo(right) > 0;
+		public static bool operator >=(ByteString left, ByteString right) => left.CompareTo(right) >= 0;
+#pragma warning restore CS1591 // Missing XML comment for publicly visible type or member
 	}
 
 	/// <summary>
