Reduced the lock contention while updating the queue empty event.

#rnx

diff --git a/Engine/Source/Programs/Shared/EpicGames.Core/ThreadPoolWorkQueue.cs b/Engine/Source/Programs/Shared/EpicGames.Core/ThreadPoolWorkQueue.cs
--- a/Engine/Source/Programs/Shared/EpicGames.Core/ThreadPoolWorkQueue.cs
+++ b/Engine/Source/Programs/Shared/EpicGames.Core/ThreadPoolWorkQueue.cs
@@ -63,13 +63,9 @@ namespace EpicGames.Core
 		/// <param name="actionToExecute">The action to add</param>
 		public void Enqueue(Action actionToExecute)
 		{
-			lock(_lockObject)
+			if (Interlocked.Increment(ref _numOutstandingJobs) == 1)
 			{
-				if(_numOutstandingJobs == 0)
-				{
-					_emptyEvent.Reset();
-				}
-				_numOutstandingJobs++;
+				SetEventState();
 			}
 
 #if SINGLE_THREAD
@@ -98,13 +94,9 @@ namespace EpicGames.Core
 			}
 			finally
 			{
-				lock(_lockObject)
+				if (Interlocked.Decrement(ref _numOutstandingJobs) == 0)
 				{
-					_numOutstandingJobs--;
-					if(_numOutstandingJobs == 0)
-					{
-						_emptyEvent.Set();
-					}
+					SetEventState();
 				}
 			}
 		}
@@ -156,5 +148,25 @@ namespace EpicGames.Core
 				}
 			}
 		}
+
+		/// <summary>
+		/// When the number of outstanding jobs transitions between 0 and 1 and visa-versa, we need to set the event 
+		/// based on the actual counter.  While the adjustment of the counter happens with interlocked instructions,
+		/// the setting of the event must be done under a lock to serialize the set/reset
+		/// </summary>
+		private void SetEventState()
+		{
+			lock (_lockObject)
+			{
+				if (_numOutstandingJobs > 0)
+				{
+					_emptyEvent.Reset();
+				}
+				else
+				{
+					_emptyEvent.Set();
+				}
+			}
+		}
 	}
 }
