Horde: Add Json converter for StringId.

diff --git a/Engine/Source/Programs/Shared/EpicGames.Horde/StringId.cs b/Engine/Source/Programs/Shared/EpicGames.Horde/StringId.cs
--- a/Engine/Source/Programs/Shared/EpicGames.Horde/StringId.cs
+++ b/Engine/Source/Programs/Shared/EpicGames.Horde/StringId.cs
@@ -1,8 +1,12 @@
 // Copyright Epic Games, Inc. All Rights Reserved.
 
 using System;
+using System.ComponentModel;
 using System.Diagnostics.CodeAnalysis;
+using System.Globalization;
 using System.Text;
+using System.Text.Json;
+using System.Text.Json.Serialization;
 using EpicGames.Core;
 
 namespace EpicGames.Horde
@@ -10,6 +14,9 @@ namespace EpicGames.Horde
 	/// <summary>
 	/// Normalized string identifier for a resource
 	/// </summary>
+	[JsonSchemaString]
+	[JsonConverter(typeof(StringIdJsonConverter))]
+	[TypeConverter(typeof(StringIdTypeConverter))]
 	public struct StringId : IEquatable<StringId>, IEquatable<string>, IEquatable<ReadOnlyMemory<char>>
 	{
 		/// <summary>
@@ -218,4 +225,34 @@ namespace EpicGames.Horde
 		/// <returns>True if the two string ids are not equal</returns>
 		public static bool operator !=(StringId left, StringId right) => !left.Equals(right);
 	}
+
+	/// <summary>
+	/// Class which serializes <see cref="StringId"/> types
+	/// </summary>
+	sealed class StringIdJsonConverter : JsonConverter<StringId>
+	{
+		/// <inheritdoc/>
+		public override StringId Read(ref Utf8JsonReader reader, Type typeToConvert, JsonSerializerOptions options) => new StringId(new Utf8String(reader.GetUtf8String().ToArray()));
+
+		/// <inheritdoc/>
+		public override void Write(Utf8JsonWriter writer, StringId value, JsonSerializerOptions options) => writer.WriteStringValue(value.Span);
+	}
+
+	/// <summary>
+	/// Class which serializes <see cref="StringId"/> types
+	/// </summary>
+	sealed class StringIdTypeConverter : TypeConverter
+	{
+		/// <inheritdoc/>
+		public override bool CanConvertFrom(ITypeDescriptorContext? context, Type sourceType) => sourceType == typeof(string);
+
+		/// <inheritdoc/>
+		public override object? ConvertFrom(ITypeDescriptorContext? context, CultureInfo? culture, object value) => new StringId((string)value);
+
+		/// <inheritdoc/>
+		public override bool CanConvertTo(ITypeDescriptorContext? context, Type? destinationType) => destinationType == typeof(string);
+
+		/// <inheritdoc/>
+		public override object? ConvertTo(ITypeDescriptorContext? context, CultureInfo? culture, object? value, Type destinationType) => ((StringId)value!).Text.ToString();
+	}
 }
