EpicGames.Core: Fix warnings generating documentation.

diff --git a/Engine/Source/Programs/Shared/EpicGames.Core/SocketExtensions.cs b/Engine/Source/Programs/Shared/EpicGames.Core/SocketExtensions.cs
--- a/Engine/Source/Programs/Shared/EpicGames.Core/SocketExtensions.cs
+++ b/Engine/Source/Programs/Shared/EpicGames.Core/SocketExtensions.cs
@@ -18,7 +18,7 @@ namespace EpicGames.Core
 		/// </summary>
 		/// <param name="socket">Socket to read from</param>
 		/// <param name="buffer">Buffer to store the data</param>
-		/// <param name="socketFlags">Flags for the <see cref="Socket.ReceiveAsync"/> call</param>
+		/// <param name="socketFlags">Flags for the socket receive call</param>
 		/// <param name="cancellationToken">Cancellation token for the operation</param>
 		public static async Task ReceiveFullAsync(this Socket socket, Memory<byte> buffer, SocketFlags socketFlags, CancellationToken cancellationToken)
 		{
@@ -38,7 +38,8 @@ namespace EpicGames.Core
 		/// </summary>
 		/// <param name="socket">Socket to write to</param>
 		/// <param name="buffer">Buffer to write</param>
-		/// <param name="cancellationToken">Flags for the <see cref="Socket.SendAsync"/> call</param>
+		/// <param name="socketFlags">Flags for the socket sent call</param>
+		/// <param name="cancellationToken">Cancellation token for the operation</param>
 		/// <returns></returns>
 		public static async Task SendFullAsync(this Socket socket, ReadOnlyMemory<byte> buffer, SocketFlags socketFlags, CancellationToken cancellationToken)
 		{
