EpicGames.Build: Don't throw an exception if a path in AdditionalSearchPath doesn't exist

#rnx

diff --git a/Engine/Source/Programs/Shared/EpicGames.Build/System/Rules.cs b/Engine/Source/Programs/Shared/EpicGames.Build/System/Rules.cs
--- a/Engine/Source/Programs/Shared/EpicGames.Build/System/Rules.cs
+++ b/Engine/Source/Programs/Shared/EpicGames.Build/System/Rules.cs
@@ -118,23 +118,10 @@ namespace UnrealBuildBase
 				}
 			}
 
-			// Process the additional search path, if sent in
+			// Process the additional search path, if sent in, only for paths that exist
 			if (AdditionalSearchPaths != null)
 			{
-				foreach (DirectoryReference AdditionalSearchPath in AdditionalSearchPaths)
-				{
-					if (AdditionalSearchPath != null)
-					{
-						if (DirectoryReference.Exists(AdditionalSearchPath))
-						{
-							Folders.Add(AdditionalSearchPath);
-						}
-						else
-						{
-							throw new Exception($"Couldn't find AdditionalSearchPath for rules source files '{AdditionalSearchPath}'");
-						}
-					}
-				}
+				Folders.AddRange(AdditionalSearchPaths.Where(x => x != null && DirectoryReference.Exists(x)));
 			}
 
 			return FindAllRulesFiles(Folders, RulesFileType);
