Remove the comment and tooltip metadata from shipping builds.

diff --git a/Engine/Source/Programs/Shared/EpicGames.UHT/Utils/UhtStringBuilderExtensions.cs b/Engine/Source/Programs/Shared/EpicGames.UHT/Utils/UhtStringBuilderExtensions.cs
--- a/Engine/Source/Programs/Shared/EpicGames.UHT/Utils/UhtStringBuilderExtensions.cs
+++ b/Engine/Source/Programs/Shared/EpicGames.UHT/Utils/UhtStringBuilderExtensions.cs
@@ -21,6 +21,11 @@ namespace EpicGames.UHT.Utils
 		/// </summary>
 		public static StringView SpacesString = new(new string(' ', 128));
 
+		/// <summary>
+		/// Names of meta data entries that will not appear in shipping builds for game code
+		/// </summary>
+		private static readonly HashSet<string> s_hiddenMetaDataNames = new(new string[]{ UhtNames.Comment, UhtNames.ToolTip }, StringComparer.OrdinalIgnoreCase);
+
 		/// <summary>
 		/// Append tabs to the builder
 		/// </summary>
@@ -237,6 +242,7 @@ namespace EpicGames.UHT.Utils
 		/// <returns>Destination builder</returns>
 		private static StringBuilder AppendMetaDataDef(this StringBuilder builder, UhtType type, string? staticsName, string? namePrefix, string name, string? nameSuffix, string? metaNameSuffix, int tabs)
 		{
+			bool isPartOfEngine = type.Package.IsPartOfEngine;
 			if (!type.MetaData.IsEmpty())
 			{
 				List<KeyValuePair<string, string>> sortedMetaData = type.MetaData.GetSorted();
@@ -246,7 +252,16 @@ namespace EpicGames.UHT.Utils
 
 				foreach (KeyValuePair<string, string> kvp in sortedMetaData)
 				{
+					bool restricted = !isPartOfEngine && s_hiddenMetaDataNames.Contains(kvp.Key);
+					if (restricted)
+					{
+						builder.Append("#if !UE_BUILD_SHIPPING\r\n");
+					}
 					builder.AppendTabs(tabs + 1).Append("{ ").AppendUTF8LiteralString(kvp.Key).Append(", ").AppendUTF8LiteralString(kvp.Value).Append(" },\r\n");
+					if (restricted)
+					{
+						builder.Append("#endif\r\n");
+					}
 				}
 
 				builder.AppendTabs(tabs).Append("};\r\n");
