Horde: Change compute client to connect directly to remotes rather than opening a listen socket.

diff --git a/Engine/Source/Programs/Horde/Horde.Agent/Commands/Compute/CppComputeCommand.cs b/Engine/Source/Programs/Horde/Horde.Agent/Commands/Compute/CppComputeCommand.cs
--- a/Engine/Source/Programs/Horde/Horde.Agent/Commands/Compute/CppComputeCommand.cs
+++ b/Engine/Source/Programs/Horde/Horde.Agent/Commands/Compute/CppComputeCommand.cs
@@ -3,6 +3,7 @@
 using System;
 using System.Collections.Generic;
 using System.Net.Http;
+using System.Net.Http.Headers;
 using System.Text.Json;
 using System.Threading;
 using System.Threading.Tasks;
@@ -35,8 +36,11 @@ namespace Horde.Agent.Commands
 			public List<string> OutputPaths { get; set; } = new List<string>();
 		}
 
+		[CommandLine("-Cluster")]
+		public string ClusterId { get; set; } = "default";
+
 		[CommandLine("-Loopback")]
-		bool Loopback { get; set; }
+		public bool Loopback { get; set; }
 
 		[CommandLine("-Task=", Required = true)]
 		FileReference TaskFile { get; set; } = null!;
@@ -62,8 +66,7 @@ namespace Horde.Agent.Commands
 		{
 			await using IComputeClient client = CreateClient();
 
-			IComputeRequest<bool> request = await client.AddRequestAsync<bool>(default, null, HandleRequestAsync);
-			bool result = await request.Result;
+			bool result = await client.ExecuteAsync(new ClusterId(ClusterId), null, HandleRequestAsync, CancellationToken.None);
 
 			return result? 0 : 1;
 		}
@@ -152,8 +155,11 @@ namespace Horde.Agent.Commands
 
 		HttpClient CreateHttpClient()
 		{
+			ServerProfile profile = _settings.GetCurrentServerProfile();
+
 			HttpClient client = _httpClientFactory.CreateClient();
-			client.BaseAddress = _settings.GetCurrentServerProfile().Url;
+			client.BaseAddress = profile.Url;
+			client.DefaultRequestHeaders.Authorization = new AuthenticationHeaderValue("Bearer", profile.Token);
 			return client;
 		}
 
diff --git a/Engine/Source/Programs/Horde/Horde.Agent/Commands/Compute/DdcComputeCommand.cs b/Engine/Source/Programs/Horde/Horde.Agent/Commands/Compute/DdcComputeCommand.cs
--- a/Engine/Source/Programs/Horde/Horde.Agent/Commands/Compute/DdcComputeCommand.cs
+++ b/Engine/Source/Programs/Horde/Horde.Agent/Commands/Compute/DdcComputeCommand.cs
@@ -39,10 +39,7 @@ namespace Horde.Agent.Commands
 			logger.LogInformation("Connecting to server: {Server}", profile.Name);
 
 			await using ServerComputeClient client = new ServerComputeClient(CreateHttpClient, logger);
-			client.Start();
-
-			IComputeRequest<object?> request = await client.AddRequestAsync(ClusterId, null, TestCommandAsync);
-			await request.Result;
+			await client.ExecuteAsync(ClusterId, null, TestCommandAsync, CancellationToken.None);
 
 			return 0;
 		}
diff --git a/Engine/Source/Programs/Shared/EpicGames.Horde/Compute/IComputeClient.cs b/Engine/Source/Programs/Shared/EpicGames.Horde/Compute/IComputeClient.cs
--- a/Engine/Source/Programs/Shared/EpicGames.Horde/Compute/IComputeClient.cs
+++ b/Engine/Source/Programs/Shared/EpicGames.Horde/Compute/IComputeClient.cs
@@ -17,6 +17,7 @@ namespace EpicGames.Horde.Compute
 		/// <param name="clusterId">Cluster to execute the request</param>
 		/// <param name="requirements">Requirements for the agent</param>
 		/// <param name="handler">Handler for the connection</param>
-		public Task<IComputeRequest<TResult>> AddRequestAsync<TResult>(ClusterId clusterId, Requirements? requirements, Func<IComputeChannel, CancellationToken, Task<TResult>> handler);
+		/// <param name="cancellationToken">Cancellation token for the operation</param>
+		public Task<TResult> ExecuteAsync<TResult>(ClusterId clusterId, Requirements? requirements, Func<IComputeChannel, CancellationToken, Task<TResult>> handler, CancellationToken cancellationToken);
 	}
 }
diff --git a/Engine/Source/Programs/Shared/EpicGames.Horde/Compute/IComputeRequest.cs b/Engine/Source/Programs/Shared/EpicGames.Horde/Compute/IComputeRequest.cs
deleted file mode 100644
diff --git a/Engine/Source/Programs/Shared/EpicGames.Horde/Compute/LoopbackComputeClient.cs b/Engine/Source/Programs/Shared/EpicGames.Horde/Compute/LoopbackComputeClient.cs
--- a/Engine/Source/Programs/Shared/EpicGames.Horde/Compute/LoopbackComputeClient.cs
+++ b/Engine/Source/Programs/Shared/EpicGames.Horde/Compute/LoopbackComputeClient.cs
@@ -14,14 +14,6 @@ namespace EpicGames.Horde.Compute
 	/// </summary>
 	public sealed class LoopbackComputeClient : IComputeClient
 	{
-		class ComputeRequest<TResult> : IComputeRequest<TResult>
-		{
-			public Task<TResult> Result { get; }
-
-			public ComputeRequest(Task<TResult> result) => Result = result;
-			public void Cancel() { }
-		}
-
 		readonly Task _listenerTask;
 		readonly SocketComputeChannel _inner;
 		readonly Socket _socket;
@@ -66,9 +58,9 @@ namespace EpicGames.Horde.Compute
 		}
 
 		/// <inheritdoc/>
-		public Task<IComputeRequest<TResult>> AddRequestAsync<TResult>(ClusterId clusterId, Requirements? requirements, Func<IComputeChannel, CancellationToken, Task<TResult>> handler)
+		public async Task<TResult> ExecuteAsync<TResult>(ClusterId clusterId, Requirements? requirements, Func<IComputeChannel, CancellationToken, Task<TResult>> handler, CancellationToken cancellationToken)
 		{
-			return Task.FromResult<IComputeRequest<TResult>>(new ComputeRequest<TResult>(handler(_inner, CancellationToken.None)));
+			return await handler(_inner, cancellationToken);
 		}
 	}
 }
diff --git a/Engine/Source/Programs/Shared/EpicGames.Horde/Compute/ServerComputeClient.cs b/Engine/Source/Programs/Shared/EpicGames.Horde/Compute/ServerComputeClient.cs
--- a/Engine/Source/Programs/Shared/EpicGames.Horde/Compute/ServerComputeClient.cs
+++ b/Engine/Source/Programs/Shared/EpicGames.Horde/Compute/ServerComputeClient.cs
@@ -3,11 +3,10 @@
 using EpicGames.Core;
 using Microsoft.Extensions.Logging;
 using System;
-using System.Collections.Generic;
 using System.Net;
 using System.Net.Http;
+using System.Net.Http.Json;
 using System.Net.Sockets;
-using System.Security.Cryptography;
 using System.Threading;
 using System.Threading.Tasks;
 
@@ -23,107 +22,23 @@ namespace EpicGames.Horde.Compute
 		/// </summary>
 		public const int NonceLength = 64;
 
-		class AddComputeTaskRequest
+		class AssignComputeRequest
 		{
 			public Requirements? Requirements { get; set; }
-			public int? RemotePort { get; set; }
-			public string Nonce { get; set; } = String.Empty;
-			public string AesKey { get; set; } = String.Empty;
-			public string AesIv { get; set; } = String.Empty;
-		}
-
-		class AddComputeTasksRequest
-		{
-			public Requirements? Requirements { get; set; }
-			public int RemotePort { get; set; }
-			public List<AddComputeTaskRequest> Tasks { get; set; } = new List<AddComputeTaskRequest>();
-		}
-
-		abstract class RequestInfo : IDisposable
-		{
-			public ReadOnlyMemory<byte> Nonce { get; }
-			public ReadOnlyMemory<byte> AesKey { get; }
-			public ReadOnlyMemory<byte> AesIv { get; }
-
-			protected readonly TaskCompletionSource<IComputeChannel> _computeChannelSource;
-			protected readonly CancellationTokenSource _cancellationSource;
-
-			public RequestInfo(CancellationToken cancellationToken)
-			{
-				Nonce = RandomNumberGenerator.GetBytes(NonceLength);
-
-				using (Aes aes = Aes.Create())
-				{
-					AesKey = aes.Key;
-					AesIv = aes.IV;
-				}
-
-				_computeChannelSource = new TaskCompletionSource<IComputeChannel>();
-				_cancellationSource = CancellationTokenSource.CreateLinkedTokenSource(cancellationToken);
-			}
-
-			public abstract Task RunAsync(IComputeChannel channel);
-
-			public void Cancel()
-			{
-				_computeChannelSource.TrySetCanceled();
-				_cancellationSource.Cancel();
-			}
-
-			public void Dispose()
-			{
-				_cancellationSource.Dispose();
-			}
 		}
 
-		class RequestInfo<TResult> : RequestInfo, IComputeRequest<TResult>
+		class AssignComputeResponse
 		{
-			readonly Func<IComputeChannel, CancellationToken, Task<TResult>> _handler;
-
-			public Task<TResult> Result { get; }
-
-			public RequestInfo(Func<IComputeChannel, CancellationToken, Task<TResult>> handler, CancellationToken cancellationToken)
-				: base(cancellationToken)
-			{
-				_handler = handler;
-
-				Result = RunWrapperAsync();
-			}
-
-			public override Task RunAsync(IComputeChannel channel)
-			{
-				_computeChannelSource.SetResult(channel);
-				return Result;
-			}
-
-			async Task<TResult> RunWrapperAsync()
-			{
-				IComputeChannel channel = await _computeChannelSource.Task;
-				return await _handler(channel, _cancellationSource.Token);
-			}
-		}
-
-		class MemoryComparer : IEqualityComparer<ReadOnlyMemory<byte>>
-		{
-			public static MemoryComparer Instance { get; } = new MemoryComparer();
-
-			public bool Equals(ReadOnlyMemory<byte> x, ReadOnlyMemory<byte> y) => x.Span.SequenceEqual(y.Span);
-
-			public int GetHashCode(ReadOnlyMemory<byte> memory)
-			{
-				HashCode hasher = new HashCode();
-				hasher.AddBytes(memory.Span);
-				return hasher.ToHashCode();
-			}
+			public string Ip { get; set; } = String.Empty;
+			public int Port { get; set; }
+			public string Nonce { get; set; } = String.Empty;
+			public string AesKey { get; set; } = String.Empty;
+			public string AesIv { get; set; } = String.Empty;
 		}
 
-		readonly object _lockObject = new object();
 		readonly HttpClient? _defaultHttpClient;
 		readonly Func<HttpClient> _createHttpClient;
 		readonly CancellationTokenSource _cancellationSource = new CancellationTokenSource();
-		readonly Dictionary<ReadOnlyMemory<byte>, RequestInfo> _requests = new Dictionary<ReadOnlyMemory<byte>, RequestInfo>(MemoryComparer.Instance);
-		TcpListener? _listener;
-		Task _listenTask = Task.CompletedTask;
 		readonly ILogger _logger;
 
 		/// <summary>
@@ -164,157 +79,54 @@ namespace EpicGames.Horde.Compute
 		HttpClient GetDefaultHttpClient() => _defaultHttpClient!;
 
 		/// <inheritdoc/>
-		public async ValueTask DisposeAsync()
+		public ValueTask DisposeAsync()
 		{
-			await StopAsync();
-			_cancellationSource.Dispose();
-			_defaultHttpClient?.Dispose();
+			Dispose();
+			return new ValueTask();
 		}
 
 		/// <inheritdoc/>
 		public void Dispose()
 		{
-			DisposeAsync().AsTask().Wait();
-		}
-
-		/// <summary>
-		/// Starts the client
-		/// </summary>
-		public void Start()
-		{
-			_listener = new TcpListener(IPAddress.IPv6Any, 0);
-			_listener.Start();
-			_listenTask = Task.Run(() => ListenAsync(_listener, _cancellationSource.Token));
+			_cancellationSource.Dispose();
+			_defaultHttpClient?.Dispose();
 		}
 
-		/// <summary>
-		/// Cancels all running tasks and waits for them to complete
-		/// </summary>
-		public async Task StopAsync()
+		/// <inheritdoc/>
+		public async Task<TResult> ExecuteAsync<TResult>(ClusterId clusterId, Requirements? requirements, Func<IComputeChannel, CancellationToken, Task<TResult>> handler, CancellationToken cancellationToken)
 		{
-			_cancellationSource.Cancel();
-			try
-			{
-				await _listenTask;
-			}
-			catch (OperationCanceledException)
-			{
-			}
-		}
+			// Assign a compute worker
+			HttpClient client = _createHttpClient();
 
-		/// <summary>
-		/// Adds a new remote request
-		/// </summary>
-		/// <param name="clusterId">Cluster to execute the request</param>
-		/// <param name="requirements">Requirements for the agent</param>
-		/// <param name="handler">Handler for the connection</param>
-		public async Task<IComputeRequest<TResult>> AddRequestAsync<TResult>(ClusterId clusterId, Requirements? requirements, Func<IComputeChannel, CancellationToken, Task<TResult>> handler)
-		{
-			RequestInfo<TResult> requestInfo = new RequestInfo<TResult>(handler, _cancellationSource.Token);
-			await AddRequestAsync(clusterId, requirements, requestInfo);
-			return requestInfo;
-		}
+			AssignComputeRequest request = new AssignComputeRequest();
+			request.Requirements = requirements;
 
-		async Task AddRequestAsync(ClusterId clusterId, Requirements? requirements, RequestInfo requestInfo)
-		{
-			try
+			AssignComputeResponse? responseMessage;
+			using (HttpResponseMessage response = await client.PostAsync($"api/v2/compute/{clusterId}", request, _cancellationSource.Token))
 			{
-				HttpClient client = _createHttpClient();
-
-				AddComputeTaskRequest task = new AddComputeTaskRequest();
-				task.Nonce = StringUtils.FormatHexString(requestInfo.Nonce.Span);
-				task.AesKey = StringUtils.FormatHexString(requestInfo.AesKey.Span);
-				task.AesIv = StringUtils.FormatHexString(requestInfo.AesIv.Span);
-
-				AddComputeTasksRequest request = new AddComputeTasksRequest();
-				request.Requirements = requirements;
-				request.RemotePort = ((IPEndPoint)_listener!.LocalEndpoint).Port;
-				request.Tasks.Add(task);
-
-				using (HttpResponseMessage response = await client.PostAsync($"api/v2/compute/{clusterId}", request, _cancellationSource.Token))
+				response.EnsureSuccessStatusCode();
+				responseMessage = await response.Content.ReadFromJsonAsync<AssignComputeResponse>(cancellationToken: cancellationToken);
+				if (responseMessage == null)
 				{
-					response.EnsureSuccessStatusCode();
-					_requests.Add(requestInfo.Nonce, requestInfo);
+					throw new InvalidOperationException();
 				}
 			}
-			catch
-			{
-				requestInfo.Dispose();
-			}
-		}
 
-		async Task ListenAsync(TcpListener listener, CancellationToken cancellationToken)
-		{
-			List<Task> tasks = new List<Task>();
-			try
-			{
-				for (; ; )
-				{
-					TcpClient client = await listener.AcceptTcpClientAsync(cancellationToken);
-					_logger.LogInformation("Received connection from {Remote}", client.Client.RemoteEndPoint);
+			// Connect to the remote machine
+			using Socket socket = new Socket(SocketType.Stream, ProtocolType.Tcp);
+			await socket.ConnectAsync(IPAddress.Parse(responseMessage.Ip), responseMessage.Port, cancellationToken);
 
-					Task task = HandleConnectionGuardedAsync(client, cancellationToken);
-					tasks.Add(task);
-				}
-			}
-			finally
-			{
-				await Task.WhenAll(tasks);
-				listener.Stop();
-			}
-		}
+			// Send the nonce
+			byte[] nonce = StringUtils.ParseHexString(responseMessage.Nonce);
+			await socket.SendFullAsync(nonce, SocketFlags.None, cancellationToken);
+			_logger.LogInformation("Connected to {Ip} with nonce {Nonce}", responseMessage.Ip, responseMessage.Nonce);
 
-		async Task HandleConnectionGuardedAsync(TcpClient client, CancellationToken cancellationToken)
-		{
-			try
-			{
-				await HandleConnectionAsync(client, cancellationToken);
-			}
-			catch (Exception ex)
-			{
-				_logger.LogInformation(ex, "Exception while handling request: {Message}", ex.Message);
-			}
-		}
-
-		async Task HandleConnectionAsync(TcpClient client, CancellationToken cancellationToken)
-		{
-			using TcpClient disposeClient = client;
-			Socket socket = client.Client;
+			// Pass the rest of the call over to the handler
+			byte[] aesKey = StringUtils.ParseHexString(responseMessage.AesKey);
+			byte[] aesIv = StringUtils.ParseHexString(responseMessage.AesIv);
 
-			// Read the nonce for this connection
-			byte[] nonceBuffer = new byte[NonceLength];
-			for (int nonceLength = 0; nonceLength < nonceBuffer.Length;)
-			{
-				int read = await socket.ReceiveAsync(nonceBuffer.AsMemory(nonceLength), SocketFlags.Partial, cancellationToken);
-				if (read == 0)
-				{
-					return;
-				}
-				nonceLength += read;
-			}
-
-			// Dispatch the connection to the appropriate handler
-			RequestInfo? requestInfo = null;
-			try
-			{
-				// Find the matching request
-				lock (_lockObject)
-				{
-					if (!_requests.Remove(nonceBuffer, out requestInfo))
-					{
-						_logger.LogInformation("Invalid nonce for connection; ignoring.");
-						return;
-					}
-				}
-
-				// Create a crypto stream for communication, using the predetermined key for this request
-				using SocketComputeChannel channel = new SocketComputeChannel(socket, requestInfo.AesKey, requestInfo.AesIv);
-				await requestInfo.RunAsync(channel);
-			}
-			finally
-			{
-				requestInfo?.Dispose();
-			}
+			using SocketComputeChannel channel = new SocketComputeChannel(socket, aesKey, aesIv);
+			return await handler(channel, cancellationToken);
 		}
 	}
 }
