EpicGames.Serialization: Refactoring CbWriter.

* New base class (CbWriterBase) has a virtual method to allocate chunks of memory for output, while still sharing common functionality for tracking scope state.
* Data for fixing up size/offset fields is written to blocks of memory returned by the CbWriterBase allocator rather than just to the final output buffer.
* Added support for embedded externally allocated blocks of data into the output stream, correctly updating size fields of containing fields.

diff --git a/Engine/Source/Programs/Shared/EpicGames.Serialization.Tests/CbWriterTests.cs b/Engine/Source/Programs/Shared/EpicGames.Serialization.Tests/CbWriterTests.cs
new file mode 100644
--- /dev/null
+++ b/Engine/Source/Programs/Shared/EpicGames.Serialization.Tests/CbWriterTests.cs
@@ -0,0 +1,230 @@
+// Copyright Epic Games, Inc. All Rights Reserved.
+
+using System;
+using System.Collections.Generic;
+using System.Linq;
+using System.Text;
+using System.Threading.Tasks;
+using EpicGames.Core;
+using Microsoft.VisualStudio.TestTools.UnitTesting;
+
+namespace EpicGames.Serialization.Tests
+{
+	[TestClass]
+	public class CbWriterTests
+	{
+		[TestMethod]
+		public void EmptyObject()
+		{
+			CbWriter writer = new CbWriter();
+			writer.BeginObject();
+			writer.EndObject();
+
+			byte[] data = writer.ToByteArray();
+			Assert.AreEqual(2, data.Length);
+			Assert.AreEqual(data[0], (byte)CbFieldType.Object);
+			Assert.AreEqual(data[1], 0);
+		}
+
+		[TestMethod]
+		public void EmptyArray()
+		{
+			CbWriter writer = new CbWriter();
+			writer.BeginArray();
+			writer.EndArray();
+
+			byte[] data = writer.ToByteArray();
+			Assert.AreEqual(3, data.Length);
+			Assert.AreEqual((byte)CbFieldType.Array, data[0]);
+			Assert.AreEqual(1, data[1]);
+			Assert.AreEqual(0, data[2]);
+		}
+
+		[TestMethod]
+		public void Object()
+		{
+			CbWriter writer = new CbWriter();
+			writer.BeginObject();
+			writer.WriteInteger("a", 1);
+			writer.WriteInteger("b", 2);
+			writer.WriteInteger("c", 3);
+			writer.EndArray();
+
+			byte[] data = writer.ToByteArray();
+			Assert.AreEqual(14, data.Length);
+			Assert.AreEqual((byte)CbFieldType.Object, data[0]);
+			Assert.AreEqual(12, data[1]); // Length
+
+			Assert.AreEqual((byte)(CbFieldType.IntegerPositive | CbFieldType.HasFieldName), data[2]);
+			Assert.AreEqual(1, data[3]);
+			Assert.AreEqual((byte)'a', data[4]);
+			Assert.AreEqual(1, data[5]);
+
+			Assert.AreEqual((byte)(CbFieldType.IntegerPositive | CbFieldType.HasFieldName), data[6]);
+			Assert.AreEqual(1, data[7]);
+			Assert.AreEqual((byte)'b', data[8]);
+			Assert.AreEqual(2, data[9]);
+
+			Assert.AreEqual((byte)(CbFieldType.IntegerPositive | CbFieldType.HasFieldName), data[10]);
+			Assert.AreEqual(1, data[11]);
+			Assert.AreEqual((byte)'c', data[12]);
+			Assert.AreEqual(3, data[13]);
+		}
+
+		[TestMethod]
+		public void Array()
+		{
+			CbWriter writer = new CbWriter();
+			writer.BeginArray();
+			writer.WriteIntegerValue(1);
+			writer.WriteIntegerValue(2);
+			writer.WriteIntegerValue(3);
+			writer.EndArray();
+
+			byte[] data = writer.ToByteArray();
+			Assert.AreEqual(9, data.Length);
+			Assert.AreEqual((byte)CbFieldType.Array, data[0]);
+			Assert.AreEqual(7, data[1]); // Length
+			Assert.AreEqual(3, data[2]); // Item count
+			Assert.AreEqual((byte)CbFieldType.IntegerPositive, data[3]);
+			Assert.AreEqual(1, data[4]);
+			Assert.AreEqual((byte)CbFieldType.IntegerPositive, data[5]);
+			Assert.AreEqual(2, data[6]);
+			Assert.AreEqual((byte)CbFieldType.IntegerPositive, data[7]);
+			Assert.AreEqual(3, data[8]);
+		}
+
+		[TestMethod]
+		public void UniformArray()
+		{
+			CbWriter writer = new CbWriter();
+			writer.BeginUniformArray(CbFieldType.IntegerPositive);
+			writer.WriteIntegerValue(1);
+			writer.WriteIntegerValue(2);
+			writer.WriteIntegerValue(3);
+			writer.EndArray();
+
+			byte[] data = writer.ToByteArray();
+			Assert.AreEqual(7, data.Length);
+			Assert.AreEqual((byte)CbFieldType.UniformArray, data[0]);
+			Assert.AreEqual(5, data[1]); // Length
+			Assert.AreEqual(3, data[2]); // Item count
+			Assert.AreEqual((byte)CbFieldType.IntegerPositive, data[3]);
+			Assert.AreEqual(1, data[4]);
+			Assert.AreEqual(2, data[5]);
+			Assert.AreEqual(3, data[6]);
+		}
+
+		[TestMethod]
+		public void NestedArray()
+		{
+			CbWriter writer = new CbWriter();
+			writer.BeginObject();
+			writer.BeginArray("a");
+			writer.WriteIntegerValue(1);
+			writer.EndArray();
+			writer.EndObject();
+
+			byte[] data = writer.ToByteArray();
+			Assert.AreEqual(9, data.Length);
+			Assert.AreEqual((byte)CbFieldType.Object, data[0]);
+			Assert.AreEqual(7, data[1]); // Length
+
+			Assert.AreEqual((byte)(CbFieldType.Array | CbFieldType.HasFieldName), data[2]);
+			Assert.AreEqual(1, data[3]); // Name length
+			Assert.AreEqual((byte)'a', data[4]); // Name
+			Assert.AreEqual(3, data[5]); // Length
+			Assert.AreEqual(1, data[6]); // Item count
+
+			Assert.AreEqual((byte)CbFieldType.IntegerPositive, data[7]);
+			Assert.AreEqual(1, data[8]);
+		}
+
+		[TestMethod]
+		public void RawData()
+		{
+			byte[] test = { 1, 2, 3, 4 };
+
+			CbWriter writer = new CbWriter();
+			writer.BeginObject();
+			writer.WriteBinaryReference("a", test);
+			writer.EndObject();
+
+			byte[] data = writer.ToByteArray();
+			Assert.AreEqual(10, data.Length);
+			Assert.AreEqual((byte)CbFieldType.Object, data[0]);
+			Assert.AreEqual(8, data[1]); // Length
+
+			Assert.AreEqual((byte)(CbFieldType.Binary | CbFieldType.HasFieldName), data[2]);
+			Assert.AreEqual(1, data[3]); // Name length
+			Assert.AreEqual((byte)'a', data[4]); // Name
+			Assert.AreEqual(4, data[5]); // Length
+
+			Assert.AreEqual(1, data[6]);
+			Assert.AreEqual(2, data[7]);
+			Assert.AreEqual(3, data[8]);
+			Assert.AreEqual(4, data[9]);
+		}
+
+		[TestMethod]
+		public void UniformArrayOfArrays()
+		{
+			CbWriter writer = new CbWriter();
+			writer.BeginUniformArray(CbFieldType.Array);
+			writer.BeginArray();
+			writer.WriteNullValue();
+			writer.EndArray();
+			writer.EndUniformArray();
+
+			byte[] data = writer.ToByteArray();
+			Assert.AreEqual(7, data.Length);
+			Assert.AreEqual((byte)CbFieldType.UniformArray, data[0]);
+			Assert.AreEqual(5, data[1]); // Length
+			Assert.AreEqual(1, data[2]); // Item Count
+
+			Assert.AreEqual((byte)CbFieldType.Array, data[3]);
+			Assert.AreEqual(2, data[4]); // Length
+			Assert.AreEqual(1, data[5]); // Item Count
+
+			Assert.AreEqual((byte)CbFieldType.Null, data[6]);
+		}
+
+		class TestObj
+		{
+			[CbField("obj")]
+			public List<CbObjectAttachment> ObjectAttachments { get; set; } = new List<CbObjectAttachment>();
+
+			[CbField("bin")]
+			public List<CbBinaryAttachment> BinaryAttachments { get; set; } = new List<CbBinaryAttachment>();
+		}
+
+		[TestMethod]
+		public void Attachments()
+		{
+			List<IoHash> objHashes = Enumerable.Range(0, 5).Select(x => IoHash.Compute(Encoding.UTF8.GetBytes($"obj{x}"))).ToList();
+			List<IoHash> binHashes = Enumerable.Range(0, 5).Select(x => IoHash.Compute(Encoding.UTF8.GetBytes($"bin{x}"))).ToList();
+
+			CbWriter writer = new CbWriter();
+			writer.BeginObject();
+			writer.BeginUniformArray("obj", CbFieldType.ObjectAttachment);
+			for (int idx = 0; idx < 5; idx++)
+			{
+				writer.WriteObjectAttachmentValue(objHashes[idx]);
+			}
+			writer.EndUniformArray();
+			writer.BeginArray("bin");
+			for (int idx = 0; idx < 5; idx++)
+			{
+				writer.WriteBinaryAttachmentValue(binHashes[idx]);
+			}
+			writer.EndArray();
+			writer.EndObject();
+
+			byte[] data = writer.ToByteArray();
+
+			TestObj obj = CbSerializer.Deserialize<TestObj>(data);
+			Assert.IsTrue(obj.ObjectAttachments.Select(x => x.Hash).SequenceEqual(objHashes));
+			Assert.IsTrue(obj.BinaryAttachments.Select(x => x.Hash).SequenceEqual(binHashes));
+		}
+	}
+}
diff --git a/Engine/Source/Programs/Shared/EpicGames.Serialization/CbWriter.cs b/Engine/Source/Programs/Shared/EpicGames.Serialization/CbWriter.cs
--- a/Engine/Source/Programs/Shared/EpicGames.Serialization/CbWriter.cs
+++ b/Engine/Source/Programs/Shared/EpicGames.Serialization/CbWriter.cs
@@ -2,9 +2,9 @@
 
 using EpicGames.Core;
 using System;
-using System.Buffers;
 using System.Buffers.Binary;
 using System.Collections.Generic;
+using System.Diagnostics;
 using System.IO;
 
 namespace EpicGames.Serialization
@@ -69,12 +69,18 @@ namespace EpicGames.Serialization
 		/// <param name="name">Name of the field. May be empty for fields that are not part of another object.</param>
 		/// <param name="length">Length of data for the field</param>
 		Span<byte> WriteField(CbFieldType type, Utf8String name, int length);
+
+		/// <summary>
+		/// Writes a reference to an external binary field into the output stream
+		/// </summary>
+		/// <param name="data">Data to reference</param>
+		void WriteReference(ReadOnlyMemory<byte> data);
 	}
 
 	/// <summary>
-	/// Forward-only writer for compact binary objects
+	/// Base class for <see cref="ICbWriter"/> implementations. Tracks structural data for fixing up lengths and offsets, without managing any buffers for field data.
 	/// </summary>
-	public class CbWriter : ICbWriter
+	public abstract class CbWriterBase : ICbWriter
 	{
 		/// <summary>
 		/// Stores information about an object or array scope within the written buffer which requires a header to be inserted containing
@@ -83,210 +89,290 @@ namespace EpicGames.Serialization
 		class Scope
 		{
 			public CbFieldType _fieldType;
+			public bool _writeFieldType;
 			public CbFieldType _uniformFieldType;
-			public int _offset; // Offset to insert the length/count
-			public int _length; // Excludes the size of this field's headers, and child fields' headers.
-			public int _count;
-			public List<Scope> _children = new List<Scope>();
-			public int _sizeOfChildHeaders; // Sum of additional headers for child items, recursively.
+			public Utf8String _name;
+			public int _itemCount;
+			public ReadOnlyMemory<byte> _data;
+			public int _length;
+			public Scope? _firstChild;
+			public Scope? _lastChild;
+			public Scope? _nextSibling;
 
-			public Scope(CbFieldType fieldType, CbFieldType uniformFieldType, int offset)
+			public void Reset()
 			{
-				Reset(fieldType, uniformFieldType, offset);
+				_fieldType = CbFieldType.None;
+				_writeFieldType = true;
+				_uniformFieldType = CbFieldType.None;
+				_name = default;
+				_itemCount = 0;
+				_data = default;
+				_length = 0;
+				_firstChild = null;
+				_lastChild = null;
+				_nextSibling = null;
 			}
 
-			public void Reset(CbFieldType fieldType, CbFieldType uniformFieldType, int offset)
+			public void AddChild(Scope child)
 			{
-				_fieldType = fieldType;
-				_uniformFieldType = uniformFieldType;
-				_offset = offset;
-				_length = 0;
-				_count = 0;
-				_children.Clear();
-				_sizeOfChildHeaders = 0;
+				if (_lastChild == null)
+				{
+					_firstChild = child;
+				}
+				else
+				{
+					_lastChild._nextSibling = child;
+				}
+				_lastChild = child;
 			}
 		}
 
+		readonly Scope _rootScope = new Scope();
+		readonly Stack<Scope> _openScopes = new Stack<Scope>();
+		readonly Stack<Scope> _freeScopes = new Stack<Scope>();
+
+		Memory<byte> _buffer = Memory<byte>.Empty;
+		int _bufferPos = 0; // Offset of the first field in the current buffer
+		int _bufferEnd = 0; // Current end of the buffer
+
 		/// <summary>
-		/// Chunk of written data. Chunks are allocated as needed and chained together with scope annotations to produce the output data.
+		/// Constructor
 		/// </summary>
-		class Chunk
+		protected CbWriterBase()
 		{
-			public int _offset;
-			public int _length;
-			public byte[] _data;
-			public List<Scope> _scopes = new List<Scope>();
-
-			public Chunk(int offset, int maxLength)
-			{
-				_data = new byte[maxLength];
-				Reset(offset);
-			}
-
-			public void Reset(int offset)
-			{
-				_offset = offset;
-				_length = 0;
-				_scopes.Clear();
-			}
+			_openScopes.Push(_rootScope);
 		}
 
 		/// <summary>
-		/// Size of data to preallocate by default
+		/// Resets the current contents of the writer
 		/// </summary>
-		public const int DefaultChunkSize = 1024;
+		protected void Reset()
+		{
+			AddChildrenToFreeList(_rootScope);
+			_rootScope.Reset();
+		}
 
-		readonly List<Chunk> _chunks = new List<Chunk>();
-		readonly Stack<Scope> _openScopes = new Stack<Scope>();
-		Chunk CurrentChunk => _chunks[^1];
-		Scope CurrentScope => _openScopes.Peek();
-		int _currentOffset;
-		readonly List<Chunk> _freeChunks = new List<Chunk>();
-		readonly List<Scope> _freeScopes = new List<Scope>();
+		void AddChildrenToFreeList(Scope root)
+		{
+			for (Scope? child = root._firstChild; child != null; child = child._nextSibling)
+			{
+				AddChildrenToFreeList(child);
+				child.Reset();
+				_freeScopes.Push(child);
+			}
+		}
 
 		/// <summary>
-		/// Constructor
+		/// Allocate a scope object
 		/// </summary>
-		public CbWriter()
-			: this(DefaultChunkSize)
+		/// <returns>New scope object</returns>
+		Scope AllocScope()
 		{
+			Scope? scope;
+			if (!_freeScopes.TryPop(out scope))
+			{
+				scope = new Scope();
+			}
+			return scope;
 		}
 
 		/// <summary>
-		/// Constructor
+		/// Creates a scope containing leaf data
 		/// </summary>
-		/// <param name="reserve">Amount of data to reserve for output</param>
-		public CbWriter(int reserve)
+		/// <param name="data"></param>
+		void AddLeafData(ReadOnlyMemory<byte> data)
 		{
-			_chunks.Add(new Chunk(0, reserve));
-			_openScopes.Push(new Scope(CbFieldType.Array, CbFieldType.None, 0));
+			Scope scope = AllocScope();
+			scope._data = data;
+			scope._length = data.Length;
+
+			Scope currentScope = _openScopes.Peek();
+			currentScope.AddChild(scope);
 		}
 
 		/// <summary>
-		/// 
+		/// Insert a new scope
 		/// </summary>
-		public void Clear()
+		Scope EnterScope(CbFieldType fieldType, Utf8String name)
 		{
-			foreach (Chunk chunk in _chunks)
-			{
-				FreeChunk(chunk);
-			}
+			Scope currentScope = _openScopes.Peek();
 
-			_currentOffset = 0;
+			Scope scope = AllocScope();
+			scope._fieldType = fieldType;
+			scope._writeFieldType = currentScope._uniformFieldType == CbFieldType.None;
+			scope._name = name;
 
-			_chunks.Clear();
-			_chunks.Add(AllocChunk(0, DefaultChunkSize));
+			currentScope.AddChild(scope);
 
-			_openScopes.Clear();
-			_openScopes.Push(AllocScope(CbFieldType.Array, CbFieldType.None, 0));
+			_openScopes.Push(scope);
+			return scope;
 		}
 
 		/// <summary>
-		/// Allocate a new chunk object
+		/// Pop a scope from the current open list
 		/// </summary>
-		/// <param name="offset">Offset of the chunk</param>
-		/// <param name="maxLength">Maximum length of the chunk</param>
-		/// <returns>New chunk object</returns>
-		Chunk AllocChunk(int offset, int maxLength)
+		void LeaveScope()
 		{
-			for (int idx = _freeChunks.Count - 1; idx >= 0; idx--)
+			WriteFields();
+			Scope scope = _openScopes.Peek();
+
+			// Measure the length of all children
+			int childrenLength = 0;
+			for (Scope? child = scope._firstChild; child != null; child = child._nextSibling)
+			{
+				childrenLength += child._length;
+			}
+
+			// Measure the total length of this scope
+			if (scope._fieldType != CbFieldType.None)
 			{
-				Chunk chunk = _freeChunks[idx];
-				if (chunk._data.Length >= maxLength)
+				// Measure the size of the field header
+				int headerLength = 0;
+				if (scope._writeFieldType)
 				{
-					_freeChunks.RemoveAt(idx);
-					chunk.Reset(offset);
-					return chunk;
+					headerLength++;
+					if (!scope._name.IsEmpty)
+					{
+						headerLength += VarInt.MeasureUnsigned(scope._name.Length) + scope._name.Length;
+					}
 				}
+
+				// Measure the size of the payload
+				int payloadLength = 0;
+				if (CbFieldUtils.IsArray(scope._fieldType))
+				{
+					payloadLength += VarInt.MeasureUnsigned(scope._itemCount);
+				}
+				if (scope._fieldType == CbFieldType.UniformObject || scope._fieldType == CbFieldType.UniformArray)
+				{
+					payloadLength++;
+				}
+
+				// Measure the size of writing the payload
+				int payloadLengthBytes = VarInt.MeasureUnsigned(payloadLength + childrenLength);
+
+				// Allocate the header
+				Memory<byte> header = Allocate(headerLength + payloadLengthBytes + payloadLength);
+				scope._data = header;
+				_bufferPos = _bufferEnd;
+
+				// Write all the fields to the header buffer
+				Span<byte> span = header.Span;
+				if (scope._writeFieldType)
+				{
+					if (scope._name.IsEmpty)
+					{
+						span[0] = (byte)scope._fieldType;
+						span = span.Slice(1);
+					}
+					else
+					{
+						span[0] = (byte)(scope._fieldType | CbFieldType.HasFieldName);
+						span = span.Slice(1);
+
+						int bytesWritten = VarInt.WriteUnsigned(span, scope._name.Length);
+						span = span.Slice(bytesWritten);
+
+						scope._name.Span.CopyTo(span);
+						span = span.Slice(scope._name.Length);
+					}
+				}
+
+				VarInt.WriteUnsigned(span, payloadLength + childrenLength);
+				span = span.Slice(payloadLengthBytes);
+
+				if (CbFieldUtils.IsArray(scope._fieldType))
+				{
+					int itemCountBytes = VarInt.WriteUnsigned(span, scope._itemCount);
+					span = span.Slice(itemCountBytes);
+				}
+
+				// Write the type for uniform arrays
+				if (scope._fieldType == CbFieldType.UniformObject || scope._fieldType == CbFieldType.UniformArray)
+				{
+					span[0] = (byte)scope._uniformFieldType;
+					span = span.Slice(1);
+				}
+
+				Debug.Assert(span.Length == 0);
 			}
-			return new Chunk(offset, maxLength);
+
+			// Set the final size of this scope, and pop it from the current open stack
+			scope._length = childrenLength + scope._data.Length;
+			_openScopes.Pop();
 		}
 
 		/// <summary>
-		/// Adds a chunk to the free list
+		/// Begin writing an object field
 		/// </summary>
-		/// <param name="chunk"></param>
-		void FreeChunk(Chunk chunk)
+		/// <param name="name">Name of the field</param>
+		public void BeginObject(Utf8String name)
 		{
-			// Add the scopes to the free list
-			_freeScopes.AddRange(chunk._scopes);
-			chunk._scopes.Clear();
+			WriteFields();
 
-			// Insert it into the free list, sorted by descending size
-			for (int idx = 0; ; idx++)
-			{
-				if (idx == _freeChunks.Count || chunk._data.Length >= _freeChunks[idx]._data.Length)
-				{
-					_freeChunks.Insert(idx, chunk);
-					break;
-				}
-			}
+			Scope parentScope = _openScopes.Peek();
+			parentScope._itemCount++;
+
+			EnterScope(CbFieldType.Object, name);
 		}
 
 		/// <summary>
-		/// Allocate a scope object
+		/// End the current object
 		/// </summary>
-		/// <param name="fieldType"></param>
-		/// <param name="uniformFieldType"></param>
-		/// <param name="offset"></param>
-		/// <returns></returns>
-		Scope AllocScope(CbFieldType fieldType, CbFieldType uniformFieldType, int offset)
+		public void EndObject() => LeaveScope();
+
+		/// <summary>
+		/// Begin writing a named array field
+		/// </summary>
+		/// <param name="name"></param>
+		/// <param name="elementType">Type of elements in the array</param>
+		public void BeginArray(Utf8String name, CbFieldType elementType)
 		{
-			if (_freeScopes.Count > 0)
-			{
-				Scope scope = _freeScopes[^1];
-				scope.Reset(fieldType, uniformFieldType, offset);
-				_freeScopes.RemoveAt(_freeScopes.Count - 1);
-				return scope;
-			}
-			return new Scope(fieldType, uniformFieldType, offset);
+			WriteFields();
+
+			Scope parentScope = _openScopes.Peek();
+			parentScope._itemCount++;
+
+			Scope scope = EnterScope((elementType == CbFieldType.None) ? CbFieldType.Array : CbFieldType.UniformArray, name);
+			scope._uniformFieldType = elementType;
 		}
 
 		/// <summary>
-		/// Ensure that a block of contiguous memory of the given length is available in the output buffer
+		/// End the current array
 		/// </summary>
-		/// <param name="length"></param>
-		/// <returns>The allocated memory</returns>
-		public Span<byte> Allocate(int length)
+		public void EndArray() => LeaveScope();
+
+		void WriteFields()
 		{
-			Chunk lastChunk = CurrentChunk;
-			if (lastChunk._length + length > lastChunk._data.Length)
+			if (_bufferPos < _bufferEnd)
 			{
-				int chunkSize = Math.Max(length, DefaultChunkSize);
-				lastChunk = AllocChunk(_currentOffset, chunkSize);
-				_chunks.Add(lastChunk);
+				AddLeafData(_buffer.Slice(_bufferPos, _bufferEnd - _bufferPos));
+				_bufferPos = _bufferEnd;
 			}
-
-			Span<byte> buffer = lastChunk._data.AsSpan(lastChunk._length, length);
-			lastChunk._length += length;
-			_currentOffset += length;
-			return buffer;
 		}
 
-		/// <summary>
-		/// Insert a new scope
-		/// </summary>
-		/// <param name="fieldType"></param>
-		/// <param name="uniformFieldType"></param>
-		void PushScope(CbFieldType fieldType, CbFieldType uniformFieldType)
+		private Memory<byte> Allocate(int length)
 		{
-			Scope newScope = AllocScope(fieldType, uniformFieldType, _currentOffset);
-			CurrentScope._children.Add(newScope);
-			_openScopes.Push(newScope);
+			if (_bufferEnd + length > _buffer.Length)
+			{
+				WriteFields();
+				_buffer = AllocateChunk(length);
+				_bufferPos = 0;
+				_bufferEnd = 0;
+			}
+			
+			Memory<byte> data = _buffer.Slice(_bufferEnd, length);
+			_bufferEnd += length;
 
-			CurrentChunk._scopes.Add(newScope);
+			return data;
 		}
 
 		/// <summary>
-		/// Pop a scope from the current open list
+		/// Allocates a chunk of data for storing CB fragments. This should be relatively coarse; the returned chunk will be reused for subsequent writes until full.
 		/// </summary>
-		void PopScope()
-		{
-			Scope scope = CurrentScope;
-			scope._length = _currentOffset - scope._offset;
-			scope._sizeOfChildHeaders = ComputeSizeOfChildHeaders(scope);
-			_openScopes.Pop();
-		}
+		/// <param name="minSize">Minimum size of the chunk</param>
+		/// <returns>New chunk of memory</returns>
+		protected abstract Memory<byte> AllocateChunk(int minSize);
 
 		/// <summary>
 		/// Writes the header for a named field
@@ -296,51 +382,69 @@ namespace EpicGames.Serialization
 		/// <param name="size">Size of the field</param>
 		public Span<byte> WriteField(CbFieldType type, Utf8String name, int size)
 		{
+			WriteFieldHeader(type, name);
+			return Allocate(size).Span;
+		}
+
+		void WriteFieldHeader(CbFieldType type, Utf8String name)
+		{
+			Scope scope = _openScopes.Peek();
 			if (name.IsEmpty)
 			{
-				Scope scope = CurrentScope;
-				if (!CbFieldUtils.IsArray(scope._fieldType))
+				CbFieldType scopeType = scope._fieldType;
+				if (!CbFieldUtils.IsArray(scopeType))
 				{
-					throw new CbWriterException($"Anonymous fields are not allowed within fields of type {scope._fieldType}");
+					throw new CbWriterException($"Anonymous fields are not allowed within fields of type {scopeType}");
 				}
 
-				if (scope._uniformFieldType == CbFieldType.None)
+				CbFieldType elementType = scope._uniformFieldType;
+				if (elementType == CbFieldType.None)
 				{
-					Allocate(1)[0] = (byte)type;
+					Allocate(1).Span[0] = (byte)type;
 				}
-				else if (scope._uniformFieldType != type)
+				else if (elementType != type)
 				{
-					throw new CbWriterException($"Mismatched type for uniform array - expected {scope._uniformFieldType}, not {type}");
+					throw new CbWriterException($"Mismatched type for uniform array - expected {elementType}, not {type}");
 				}
-				scope._count++;
+
+				scope._itemCount++;
 			}
 			else
 			{
-				Scope scope = CurrentScope;
-				if (!CbFieldUtils.IsObject(scope._fieldType))
+				CbFieldType scopeType = scope._fieldType;
+				if (!CbFieldUtils.IsObject(scopeType))
 				{
-					throw new CbWriterException($"Named fields are not allowed within fields of type {scope._fieldType}");
+					throw new CbWriterException($"Named fields are not allowed within fields of type {scopeType}");
 				}
 
+				CbFieldType elementType = scope._uniformFieldType;
+
 				int nameVarIntLength = VarInt.MeasureUnsigned(name.Length);
-				if (scope._uniformFieldType == CbFieldType.None)
+				if (elementType == CbFieldType.None)
 				{
-					Span<byte> buffer = Allocate(1 + nameVarIntLength + name.Length);
+					Span<byte> buffer = Allocate(1 + nameVarIntLength + name.Length).Span;
 					buffer[0] = (byte)(type | CbFieldType.HasFieldName);
 					WriteBinaryPayload(buffer[1..], name.Span);
 				}
 				else
 				{
-					if (scope._uniformFieldType != type)
+					if (elementType != type)
 					{
-						throw new CbWriterException($"Mismatched type for uniform object - expected {scope._uniformFieldType}, not {type}");
+						throw new CbWriterException($"Mismatched type for uniform object - expected {elementType}, not {type}");
 					}
-					Span<byte> buffer = Allocate(name.Length);
-					WriteBinaryPayload(buffer, name.Span);
+					Memory<byte> buffer = Allocate(name.Length);
+					WriteBinaryPayload(buffer.Span, name.Span);
 				}
-				scope._count++;
+
+				scope._itemCount++;
 			}
-			return Allocate(size);
+		}
+
+		/// <inheritdoc/>
+		public void WriteReference(ReadOnlyMemory<byte> data)
+		{
+			WriteFields();
+			AddLeafData(data);
 		}
 
 		/// <summary>
@@ -354,53 +458,49 @@ namespace EpicGames.Serialization
 			output = output[varIntLength..];
 
 			value.CopyTo(output);
-			CheckSize(output, value.Length);
 		}
 
 		/// <summary>
-		/// Begin writing an object field
+		/// Gets the size of the serialized data
 		/// </summary>
-		/// <param name="name">Name of the field</param>
-		public void BeginObject(Utf8String name)
+		/// <returns></returns>
+		public int GetSize()
 		{
-			WriteField(CbFieldType.Object, name, 0);
-			PushScope(CbFieldType.Object, CbFieldType.None);
+			if (_openScopes.Count > 1)
+			{
+				throw new CbWriterException("Unfinished scope in writer");
+			}
+
+			int length = 0;
+			for (Scope? child = _rootScope._firstChild; child != null; child = child._nextSibling)
+			{
+				length += child._length;
+			}
+
+			return length;
 		}
 
 		/// <summary>
-		/// End the current object
+		/// Copy the data from this writer to a buffer
 		/// </summary>
-		public void EndObject()
+		/// <param name="buffer"></param>
+		public void CopyTo(Span<byte> buffer)
 		{
-			PopScope();
+			Copy(_rootScope, buffer);
 		}
 
-		/// <summary>
-		/// Begin writing a named array field
-		/// </summary>
-		/// <param name="name"></param>
-		/// <param name="elementType">Type of elements in the array</param>
-		public void BeginArray(Utf8String name, CbFieldType elementType)
+		static Span<byte> Copy(Scope scope, Span<byte> span)
 		{
-			if (elementType == CbFieldType.None)
+			if (scope._data.Length > 0)
 			{
-				WriteField(CbFieldType.Array, name, 0);
-				PushScope(CbFieldType.Array, CbFieldType.None);
+				scope._data.Span.CopyTo(span);
+				span = span.Slice(scope._data.Length);
 			}
-			else
+			for (Scope? child = scope._firstChild; child != null; child = child._nextSibling)
 			{
-				WriteField(CbFieldType.UniformArray, name, 0);
-				PushScope(CbFieldType.UniformArray, elementType);
-				Allocate(1)[0] = (byte)elementType;
+				span = Copy(child, span);
 			}
-		}
-
-		/// <summary>
-		/// End the current array
-		/// </summary>
-		public void EndArray()
-		{
-			PopScope();
+			return span;
 		}
 
 		/// <summary>
@@ -411,7 +511,7 @@ namespace EpicGames.Serialization
 		{
 			using (Blake3.Hasher hasher = Blake3.Hasher.New())
 			{
-				foreach (ReadOnlyMemory<byte> segment in EnumerateSegments())
+				foreach (ReadOnlyMemory<byte> segment in GetSegments())
 				{
 					hasher.Update(segment.Span);
 				}
@@ -420,88 +520,63 @@ namespace EpicGames.Serialization
 		}
 
 		/// <summary>
-		/// Gets the size of the serialized data
+		/// Convert the data into a compact binary object
 		/// </summary>
 		/// <returns></returns>
-		public int GetSize()
+		public CbObject ToObject()
 		{
-			if (_openScopes.Count != 1)
-			{
-				throw new CbWriterException("Unfinished scope in writer");
-			}
-
-			return _currentOffset + ComputeSizeOfChildHeaders(CurrentScope);
+			return new CbObject(ToByteArray());
 		}
 
 		/// <summary>
-		/// Gets the contents of this writer as a stream
+		/// Convert the data into a flat array
 		/// </summary>
-		/// <returns>New stream for the contents of this object</returns>
-		public Stream AsStream() => new ReadStream(EnumerateSegments().GetEnumerator(), GetSize());
-
-		private IEnumerable<ReadOnlyMemory<byte>> EnumerateSegments()
+		/// <returns></returns>
+		public byte[] ToByteArray()
 		{
-			byte[] scopeHeader = new byte[64];
-
-			int sourceOffset = 0;
-			foreach (Chunk chunk in _chunks)
-			{
-				foreach (Scope scope in chunk._scopes)
-				{
-					ReadOnlyMemory<byte> sourceData = chunk._data.AsMemory(sourceOffset - chunk._offset, scope._offset - sourceOffset);
-					yield return sourceData;
-
-					sourceOffset += sourceData.Length;
-
-					int headerLength = WriteScopeHeader(scopeHeader, scope);
-					yield return scopeHeader.AsMemory(0, headerLength);
-				}
-
-				ReadOnlyMemory<byte> lastSourceData = chunk._data.AsMemory(sourceOffset - chunk._offset, (chunk._offset + chunk._length) - sourceOffset);
-				yield return lastSourceData;
-
-				sourceOffset += lastSourceData.Length;
-			}
+			byte[] buffer = new byte[GetSize()];
+			CopyTo(buffer);
+			return buffer;
 		}
 
 		/// <summary>
-		/// Copy the data from this writer to a buffer
+		/// Enumerate all the segments in the data that has been written
 		/// </summary>
-		/// <param name="buffer"></param>
-		public void CopyTo(Span<byte> buffer)
+		/// <returns>Sequence of segments</returns>
+		public List<ReadOnlyMemory<byte>> GetSegments()
 		{
-			int bufferOffset = 0;
+			List<ReadOnlyMemory<byte>> segments = new List<ReadOnlyMemory<byte>>();
+			GetSegments(_rootScope, segments);
+			return segments;
+		}
 
-			int sourceOffset = 0;
-			foreach (Chunk chunk in _chunks)
+		static void GetSegments(Scope scope, List<ReadOnlyMemory<byte>> segments)
+		{
+			if (scope._data.Length > 0)
 			{
-				foreach (Scope scope in chunk._scopes)
-				{
-					ReadOnlySpan<byte> sourceData = chunk._data.AsSpan(sourceOffset - chunk._offset, scope._offset - sourceOffset);
-					sourceData.CopyTo(buffer.Slice(bufferOffset));
-
-					bufferOffset += sourceData.Length;
-					sourceOffset += sourceData.Length;
-
-					bufferOffset += WriteScopeHeader(buffer.Slice(bufferOffset), scope);
-				}
-
-				ReadOnlySpan<byte> lastSourceData = chunk._data.AsSpan(sourceOffset - chunk._offset, (chunk._offset + chunk._length) - sourceOffset);
-				lastSourceData.CopyTo(buffer.Slice(bufferOffset));
-				bufferOffset += lastSourceData.Length;
-				sourceOffset += lastSourceData.Length;
+				segments.Add(scope._data);
+			}
+			for (Scope? child = scope._firstChild; child != null; child = child._nextSibling)
+			{
+				GetSegments(child, segments);
 			}
 		}
 
+		/// <summary>
+		/// Gets the contents of this writer as a stream
+		/// </summary>
+		/// <returns>New stream for the contents of this object</returns>
+		public Stream AsStream() => new ReadStream(GetSegments().GetEnumerator(), GetSize());
+
 		class ReadStream : Stream
 		{
-			readonly IEnumerator<ReadOnlyMemory<byte>> _enumerator;
+			readonly IEnumerator<ReadOnlyMemory<byte>> _segments;
 			ReadOnlyMemory<byte> _segment;
 			long _positionInternal;
 
-			public ReadStream(IEnumerator<ReadOnlyMemory<byte>> enumerator, long length)
+			public ReadStream(IEnumerator<ReadOnlyMemory<byte>> segments, long length)
 			{
-				_enumerator = enumerator;
+				_segments = segments;
 				Length = length;
 			}
 
@@ -535,11 +610,11 @@ namespace EpicGames.Serialization
 				{
 					while (_segment.Length == 0)
 					{
-						if (!_enumerator.MoveNext())
+						if (!_segments.MoveNext())
 						{
 							return readLength;
 						}
-						_segment = _enumerator.Current;
+						_segment = _segments.Current;
 					}
 
 					int copyLength = Math.Min(_segment.Length, buffer.Length);
@@ -565,88 +640,61 @@ namespace EpicGames.Serialization
 			/// <inheritdoc/>
 			public override void Write(byte[] buffer, int offset, int count) => throw new NotSupportedException();
 		}
+	}
 
+	/// <summary>
+	/// Forward-only writer for compact binary objects
+	/// </summary>
+	public class CbWriter : CbWriterBase
+	{
 		/// <summary>
-		/// Convert the data into a compact binary object
+		/// Size of data to preallocate by default
 		/// </summary>
-		/// <returns></returns>
-		public CbObject ToObject()
-		{
-			return new CbObject(ToByteArray());
-		}
+		public const int DefaultChunkSize = 1024;
+
+		readonly List<byte[]> _chunks = new List<byte[]>();
+		readonly List<byte[]> _freeChunks = new List<byte[]>();
 
 		/// <summary>
-		/// Convert the data into a flat array
+		/// Constructor
 		/// </summary>
-		/// <returns></returns>
-		public byte[] ToByteArray()
+		public CbWriter()
 		{
-			byte[] buffer = new byte[GetSize()];
-			CopyTo(buffer);
-			return buffer;
 		}
 
 		/// <summary>
-		/// Comptues the size of any child headers
+		/// Constructor
 		/// </summary>
-		/// <param name="scope"></param>
-		static int ComputeSizeOfChildHeaders(Scope scope)
+		/// <param name="reserve">Amount of data to reserve for output</param>
+		public CbWriter(int reserve)
 		{
-			int sizeOfChildHeaders = 0;
-			foreach (Scope childScope in scope._children)
-			{
-				switch (childScope._fieldType)
-				{
-					case CbFieldType.Object:
-					case CbFieldType.UniformObject:
-						sizeOfChildHeaders += childScope._sizeOfChildHeaders + VarInt.MeasureUnsigned(childScope._length + childScope._sizeOfChildHeaders);
-						break;
-					case CbFieldType.Array:
-					case CbFieldType.UniformArray:
-						int arrayCountLength = VarInt.MeasureUnsigned(childScope._count);
-						sizeOfChildHeaders += childScope._sizeOfChildHeaders + VarInt.MeasureUnsigned(childScope._length + childScope._sizeOfChildHeaders + arrayCountLength) + arrayCountLength;
-						break;
-					default:
-						throw new InvalidOperationException();
-				}
-			}
-			return sizeOfChildHeaders;
+			_freeChunks.Add(new byte[reserve]);
 		}
 
 		/// <summary>
-		/// Writes the header for a particular scope
+		/// Clear the current contents of the writer
 		/// </summary>
-		/// <param name="span"></param>
-		/// <param name="scope"></param>
-		/// <returns></returns>
-		static int WriteScopeHeader(Span<byte> span, Scope scope)
+		public void Clear()
 		{
-			switch (scope._fieldType)
-			{
-				case CbFieldType.Object:
-				case CbFieldType.UniformObject:
-					return VarInt.WriteUnsigned(span, scope._length + scope._sizeOfChildHeaders);
-				case CbFieldType.Array:
-				case CbFieldType.UniformArray:
-					int numItemsLength = VarInt.MeasureUnsigned(scope._count);
-					int offset = VarInt.WriteUnsigned(span, scope._length + scope._sizeOfChildHeaders + numItemsLength);
-					return offset + VarInt.WriteUnsigned(span.Slice(offset), scope._count);
-				default:
-					throw new InvalidOperationException();
-			}
+			base.Reset();
+
+			_freeChunks.AddRange(_chunks);
+			_chunks.Clear();
 		}
 
-		/// <summary>
-		/// Check that the given span is the required size
-		/// </summary>
-		/// <param name="span"></param>
-		/// <param name="expectedSize"></param>
-		static void CheckSize(Span<byte> span, int expectedSize)
+		/// <inheritdoc/>
+		protected override Memory<byte> AllocateChunk(int minSize)
 		{
-			if (span.Length != expectedSize)
+			for (int idx = 0; idx < _freeChunks.Count; idx++)
 			{
-				throw new Exception("Size of buffer is not correct");
+				byte[] data = _freeChunks[idx];
+				if (data.Length >= minSize)
+				{
+					_freeChunks.RemoveAt(idx);
+					return data;
+				}
 			}
+			return new byte[Math.Max(minSize, DefaultChunkSize)];
 		}
 	}
 
@@ -980,6 +1028,20 @@ namespace EpicGames.Serialization
 			value.Span.CopyTo(buffer);
 		}
 
+		/// <summary>
+		/// Writes an external binary value into the output stream
+		/// </summary>
+		/// <param name="writer">Writer for output data</param>
+		/// <param name="name">Name of the field</param>
+		/// <param name="data">Data to reference</param>
+		public static void WriteBinaryReference(this ICbWriter writer, Utf8String name, ReadOnlyMemory<byte> data)
+		{
+			int lengthBytes = VarInt.MeasureUnsigned(data.Length);
+			Span<byte> span = writer.WriteField(CbFieldType.Binary, name, lengthBytes);
+			VarInt.WriteUnsigned(span, data.Length);
+			writer.WriteReference(data);
+		}
+
 		/// <summary>
 		/// Writes an unnamed binary value
 		/// </summary>
