OidcToken - Ignore stored tokens if a "Key not valid for use in specified state" exception is raised. Fixes rare issue issue on Windows machiens were the stored token can not be read.

diff --git a/Engine/Source/Programs/Shared/EpicGames.OIDC/WindowsTokenStore.cs b/Engine/Source/Programs/Shared/EpicGames.OIDC/WindowsTokenStore.cs
--- a/Engine/Source/Programs/Shared/EpicGames.OIDC/WindowsTokenStore.cs
+++ b/Engine/Source/Programs/Shared/EpicGames.OIDC/WindowsTokenStore.cs
@@ -120,6 +120,12 @@ namespace EpicGames.OIDC
 					refreshToken = "";
 					return false;
 				}
+				if (e.NativeErrorCode == unchecked((int)0x8009000B)) // key not valid for use in specified state
+				{
+					// unable to decrypt the data, ignore it
+					refreshToken = "";
+					return false;
+				}
 				throw;
 			}
 		}
