EpicGames.Serialization: Fix writing of CB fields at the root of the output object.

* CbWriter was expecting BeginObject/BeginArray call before any fields were written, so did not handle writing fields directly to the root scope.
* Fields written at the root scope were not being flushed immediately after writing. Rest of the class assumes that any complete fields at root scope had already been accounted for in object size.

diff --git a/Engine/Source/Programs/Shared/EpicGames.Serialization.Tests/CbWriterTests.cs b/Engine/Source/Programs/Shared/EpicGames.Serialization.Tests/CbWriterTests.cs
--- a/Engine/Source/Programs/Shared/EpicGames.Serialization.Tests/CbWriterTests.cs
+++ b/Engine/Source/Programs/Shared/EpicGames.Serialization.Tests/CbWriterTests.cs
@@ -226,5 +226,29 @@ namespace EpicGames.Serialization.Tests
 			Assert.IsTrue(obj.ObjectAttachments.Select(x => x.Hash).SequenceEqual(objHashes));
 			Assert.IsTrue(obj.BinaryAttachments.Select(x => x.Hash).SequenceEqual(binHashes));
 		}
+
+		[TestMethod]
+		public void EmbeddedObject()
+		{
+			// Create an empty CbObject
+			CbObject obj;
+			{
+				CbWriter writer = new CbWriter();
+				writer.BeginObject();
+				writer.EndObject();
+				obj = writer.ToObject();
+			}
+
+			// Embed it in another writer at the root
+			{
+				CbWriter writer = new CbWriter();
+				writer.WriteObject(obj);
+
+				byte[] data = writer.ToByteArray();
+				Assert.AreEqual(2, data.Length);
+				Assert.AreEqual((byte)CbFieldType.Object, data[0]);
+				Assert.AreEqual(0, data[1]);
+			}
+		}
 	}
 }
diff --git a/Engine/Source/Programs/Shared/EpicGames.Serialization/CbWriter.cs b/Engine/Source/Programs/Shared/EpicGames.Serialization/CbWriter.cs
--- a/Engine/Source/Programs/Shared/EpicGames.Serialization/CbWriter.cs
+++ b/Engine/Source/Programs/Shared/EpicGames.Serialization/CbWriter.cs
@@ -127,7 +127,7 @@ namespace EpicGames.Serialization
 			}
 		}
 
-		readonly Scope _rootScope = new Scope();
+		readonly Scope _rootScope = new Scope { _fieldType = CbFieldType.Array };
 		readonly Stack<Scope> _openScopes = new Stack<Scope>();
 		readonly Stack<Scope> _freeScopes = new Stack<Scope>();
 
@@ -383,7 +383,14 @@ namespace EpicGames.Serialization
 		public Span<byte> WriteField(CbFieldType type, Utf8String name, int size)
 		{
 			WriteFieldHeader(type, name);
-			return Allocate(size).Span;
+
+			Span<byte> span = Allocate(size).Span;
+			if (_openScopes.Count == 1)
+			{
+				// If this field is at the root, flush it immediately
+				WriteFields();
+			}
+			return span;
 		}
 
 		void WriteFieldHeader(CbFieldType type, Utf8String name)
