Horde: Compute service refactor.

- Message framing changed back to use simple
- Support for multiplexing multiple requests over a single socket. Requests are tagged with an arbitrary channel id and routed to a queue for the appropriate IComputeChannel on the client.
- Clients can now request a message writer through the IComputeChannel interface, allowing pooling of write buffers for generating new messages and reducing data copies.
- Added accounting for machine resources in compute leases. Compute requests specify min/max desired resources, receive allocated resources in respoonse.
- Switched to AES-GCM for encryption. Message length is unencrypted but validated as part of the message packet.

diff --git a/Engine/Source/Programs/Horde/Horde.Agent/Commands/Compute/CppComputeCommand.cs b/Engine/Source/Programs/Horde/Horde.Agent/Commands/Compute/CppComputeCommand.cs
--- a/Engine/Source/Programs/Horde/Horde.Agent/Commands/Compute/CppComputeCommand.cs
+++ b/Engine/Source/Programs/Horde/Horde.Agent/Commands/Compute/CppComputeCommand.cs
@@ -6,6 +6,7 @@ using System.Net.Http;
 using System.Net.Http.Headers;
 using System.Text.Json;
 using System.Threading;
+using System.Threading.Channels;
 using System.Threading.Tasks;
 using EpicGames.Core;
 using EpicGames.Horde.Compute;
@@ -71,30 +72,36 @@ namespace Horde.Agent.Commands
 			return result? 0 : 1;
 		}
 
-		async Task<bool> HandleRequestAsync(IComputeChannel channel, CancellationToken cancellationToken)
+		async Task<bool> HandleRequestAsync(IComputeLease lease, CancellationToken cancellationToken)
 		{
+			IComputeChannel channel = lease.OpenChannel(1);
+
 			MemoryStorageClient storage = new MemoryStorageClient();
-			NodeLocator node = await CreateComputeNodeAsync(TaskFile, storage, cancellationToken);
 
-			CppComputeMessage message = new CppComputeMessage { Locator = node };
-			await channel.WriteCbMessageAsync(message, cancellationToken);
+			NodeLocator node = await CreateComputeNodeAsync(TaskFile, storage, cancellationToken);
+			await lease.DefaultChannel.CppStartAsync(node, channel.Id, cancellationToken);
 
 			for (; ; )
 			{
-				object? obj = await channel.ReadCbMessageAsync(cancellationToken);
-				switch (obj)
+				IComputeMessage message = await channel.ReadAsync(cancellationToken);
+				switch (message.Type)
 				{
-					case null:
+					case ComputeMessageType.None:
 						return false;
-					case BlobReadMessage readBlob:
-						await channel.WriteBlobDataAsync(readBlob, storage, cancellationToken);
+					case ComputeMessageType.CppBlobRead:
+						{
+							CppBlobReadMessage cppBlobRead = message.AsCppBlobRead();
+							await channel.CppBlobDataAsync(cppBlobRead, storage, cancellationToken);
+						}
 						break;
-					case CppComputeOutputMessage cppOutput:
+					case ComputeMessageType.CppResult:
 						{
+							NodeLocator cppOutputLocator = message.AsCppResult();
+
 							using ComputeStorageClient remoteStorage = new ComputeStorageClient(channel);
 							TreeReader reader = new TreeReader(remoteStorage, null, _logger);
 
-							CppComputeOutputNode computeOutput = await reader.ReadNodeAsync<CppComputeOutputNode>(cppOutput.Locator, cancellationToken);
+							CppComputeOutputNode computeOutput = await reader.ReadNodeAsync<CppComputeOutputNode>(cppOutputLocator, cancellationToken);
 							_logger.LogInformation("Exit code: {ExitCode}", computeOutput.ExitCode);
 
 							LogNode logNode = await computeOutput.Log.ExpandAsync(reader, cancellationToken);
@@ -107,10 +114,10 @@ namespace Horde.Agent.Commands
 								}
 							}
 
-							await channel.WriteCbMessageAsync(new CppComputeFinishMessage(), cancellationToken);
-							await channel.WriteCbMessageAsync(new CloseMessage(), cancellationToken);
+							await channel.CppFinishAsync(cancellationToken);
+							await lease.DefaultChannel.CloseAsync(cancellationToken);
 						}
-						return true;
+						return false;
 					default:
 						throw new NotImplementedException();
 				}
@@ -163,10 +170,10 @@ namespace Horde.Agent.Commands
 			return client;
 		}
 
-		async Task RunListenerAsync(IComputeChannel channel, CancellationToken cancellationToken)
+		async Task RunListenerAsync(IComputeLease lease, CancellationToken cancellationToken)
 		{
 			ComputeHandler handler = ActivatorUtilities.CreateInstance<ComputeHandler>(_serviceProvider);
-			await handler.RunAsync(channel, cancellationToken);
+			await handler.RunAsync(lease, cancellationToken);
 		}
 	}
 }
diff --git a/Engine/Source/Programs/Horde/Horde.Agent/Commands/Compute/DdcComputeCommand.cs b/Engine/Source/Programs/Horde/Horde.Agent/Commands/Compute/DdcComputeCommand.cs
--- a/Engine/Source/Programs/Horde/Horde.Agent/Commands/Compute/DdcComputeCommand.cs
+++ b/Engine/Source/Programs/Horde/Horde.Agent/Commands/Compute/DdcComputeCommand.cs
@@ -1,5 +1,7 @@
 // Copyright Epic Games, Inc. All Rights Reserved.
 
+using System;
+using System.Diagnostics;
 using System.Net.Http;
 using System.Threading;
 using System.Threading.Tasks;
@@ -51,14 +53,21 @@ namespace Horde.Agent.Commands
 			return client;
 		}
 
-		async Task<object?> TestCommandAsync(IComputeChannel channel, CancellationToken cancellationToken)
+		async Task<object?> TestCommandAsync(IComputeLease lease, CancellationToken cancellationToken)
 		{
-			ComputeMessageWriter writer = new ComputeMessageWriter(channel);
-			await writer.WriteCbMessageAsync(new XorRequestMessage { Value = 123, Payload = new byte[] { 1, 2, 3, 4, 5 } }, cancellationToken);
+			IComputeChannel channel = lease.DefaultChannel;
 
-			XorResponseMessage response = await channel.ReadCbMessageAsync<XorResponseMessage>(cancellationToken);
+			await channel.XorRequestAsync(new byte[] { 1, 2, 3, 4, 5 }, (byte)123, cancellationToken);
+			IComputeMessage response = await channel.ReadAsync(cancellationToken);
 
-			await writer.WriteCbMessageAsync(new CloseMessage(), cancellationToken);
+			byte[] result = response.Data.ToArray();
+			Debug.Assert(result[0] == (byte)(1 ^ 123));
+			Debug.Assert(result[1] == (byte)(2 ^ 123));
+			Debug.Assert(result[2] == (byte)(3 ^ 123));
+			Debug.Assert(result[3] == (byte)(4 ^ 123));
+			Debug.Assert(result[4] == (byte)(5 ^ 123));
+
+			await channel.CloseAsync(cancellationToken);
 			return null;
 		}
 	}
diff --git a/Engine/Source/Programs/Horde/Horde.Agent/Leases/Handlers/ComputeHandler.cs b/Engine/Source/Programs/Horde/Horde.Agent/Leases/Handlers/ComputeHandler.cs
--- a/Engine/Source/Programs/Horde/Horde.Agent/Leases/Handlers/ComputeHandler.cs
+++ b/Engine/Source/Programs/Horde/Horde.Agent/Leases/Handlers/ComputeHandler.cs
@@ -57,9 +57,9 @@ namespace Horde.Agent.Leases.Handlers
 					return LeaseResult.Success;
 				}
 
-				using (SocketComputeChannel channel = new SocketComputeChannel(tcpClient.Client, computeTask.AesKey.Memory, computeTask.AesIv.Memory))
+				await using (ComputeLease lease = new ComputeLease(tcpClient.Client, computeTask.Key.Span, computeTask.Nonce.Span.Slice(0, ComputeLease.NonceLength), computeTask.Resources))
 				{
-					await RunAsync(channel, cancellationToken);
+					await RunAsync(lease, cancellationToken);
 					return LeaseResult.Success;
 				}
 			}
@@ -74,37 +74,55 @@ namespace Horde.Agent.Leases.Handlers
 			}
 		}
 
-		public async Task RunAsync(IComputeChannel channel, CancellationToken cancellationToken)
+		public async Task RunAsync(IComputeLease lease, CancellationToken cancellationToken)
 		{
+			IComputeChannel channel = lease.DefaultChannel;
 			for (; ; )
 			{
-				object? request = await channel.ReadCbMessageAsync(cancellationToken);
-				switch (request)
+				using IComputeMessage message = await channel.ReadAsync(cancellationToken);
+				switch (message.Type)
 				{
-					case null:
-					case CloseMessage _:
+					case ComputeMessageType.None:
+					case ComputeMessageType.Close:
 						return;
-					case XorRequestMessage xorRequest:
+					case ComputeMessageType.XorRequest:
 						{
-							XorResponseMessage response = new XorResponseMessage();
-							response.Payload = new byte[xorRequest.Payload.Length];
-							for (int idx = 0; idx < xorRequest.Payload.Length; idx++)
-							{
-								response.Payload[idx] = (byte)(xorRequest.Payload[idx] ^ xorRequest.Value);
-							}
-							await channel.WriteCbMessageAsync(response, cancellationToken);
+							XorRequestMessage xorRequest = message.AsXorRequest();
+							await RunXorAsync(channel, xorRequest.Data, xorRequest.Value, cancellationToken);
 						}
 						break;
-					case CppComputeMessage cppCompute:
-						await RunCppAsync(channel, cppCompute.Locator, cancellationToken);
+					case ComputeMessageType.CppBegin:
+						{
+							CppBeginMessage cppBegin = message.AsCppBegin();
+							await RunCppAsync(lease.OpenChannel(cppBegin.ReplyChannelId), cppBegin.Locator, cancellationToken);
+						}
 						break;
+					default:
+						throw new NotImplementedException();
 				}
 			}
 		}
 
+		static async Task RunXorAsync(IComputeChannel channel, ReadOnlyMemory<byte> source, byte value, CancellationToken cancellationToken)
+		{
+			using (IComputeMessageWriter writer = channel.CreateMessage(ComputeMessageType.XorResponse, source.Length))
+			{
+				XorData(source.Span, writer.GetSpanAndAdvance(source.Length), value);
+				await writer.SendAsync(cancellationToken);
+			}
+		}
+
+		static void XorData(ReadOnlySpan<byte> source, Span<byte> target, byte value)
+		{
+			for (int idx = 0; idx < source.Length; idx++)
+			{
+				target[idx] = (byte)(source[idx] ^ value);
+			}
+		}
+
 		public async Task RunCppAsync(IComputeChannel channel, NodeLocator locator, CancellationToken cancellationToken)
 		{
-			DirectoryReference sandboxDir = DirectoryReference.Combine(Program.DataDir, "Sandbox");
+			DirectoryReference sandboxDir = DirectoryReference.Combine(Program.DataDir, "Sandbox", channel.Id.ToString());
 
 			using ComputeStorageClient store = new ComputeStorageClient(channel);
 			TreeReader reader = new TreeReader(store, _memoryCache, _logger);
@@ -164,21 +182,25 @@ namespace Horde.Agent.Leases.Handlers
 				CppComputeOutputNode outputNode = new CppComputeOutputNode(exitCode, logNodeRef, new TreeNodeRef<DirectoryNode>(outputTree));
 				NodeHandle outputHandle = await writer.FlushAsync(outputNode, cancellationToken);
 
-				CppComputeOutputMessage outputMessage = new CppComputeOutputMessage { Locator = outputHandle.Locator };
-				await channel.WriteCbMessageAsync(outputMessage, cancellationToken);
+				await channel.CppResultAsync(outputHandle.Locator, cancellationToken);
 			}
 
 			for (; ; )
 			{
-				object? request = await channel.ReadCbMessageAsync(cancellationToken);
-				switch (request)
+				using IComputeMessage message = await channel.ReadAsync(cancellationToken);
+				switch (message.Type)
 				{
-					case null:
-					case CppComputeFinishMessage _:
+					case ComputeMessageType.CppEnd:
+						FileUtils.ForceDeleteDirectory(sandboxDir);
 						return;
-					case BlobReadMessage blobRead:
-						await channel.WriteBlobDataAsync(blobRead, storage, cancellationToken);
+					case ComputeMessageType.CppBlobRead:
+						{
+							CppBlobReadMessage cppBlobRead = message.AsCppBlobRead();
+							await channel.CppBlobDataAsync(cppBlobRead, storage, cancellationToken);
+						}
 						break;
+					default:
+						throw new NotImplementedException();
 				}
 			}
 		}
diff --git a/Engine/Source/Programs/Horde/Horde.Agent/Services/ComputeListenerService.cs b/Engine/Source/Programs/Horde/Horde.Agent/Services/ComputeListenerService.cs
--- a/Engine/Source/Programs/Horde/Horde.Agent/Services/ComputeListenerService.cs
+++ b/Engine/Source/Programs/Horde/Horde.Agent/Services/ComputeListenerService.cs
@@ -174,7 +174,7 @@ namespace Horde.Agent.Services
 
 				// Read the nonce for the connection
 				byte[] nonceBuffer = new byte[ServerComputeClient.NonceLength];
-				await socket.ReceiveFullAsync(nonceBuffer, SocketFlags.None, cancellationToken);
+				await socket.ReceiveMessageAsync(nonceBuffer, SocketFlags.None, cancellationToken);
 
 				// Register the socket
 				ByteString nonce = new ByteString(nonceBuffer);
diff --git a/Engine/Source/Programs/Horde/Horde.Build/Compute/ComputeController.cs b/Engine/Source/Programs/Horde/Horde.Build/Compute/ComputeController.cs
--- a/Engine/Source/Programs/Horde/Horde.Build/Compute/ComputeController.cs
+++ b/Engine/Source/Programs/Horde/Horde.Build/Compute/ComputeController.cs
@@ -48,17 +48,12 @@ namespace Horde.Build.Compute
 		/// <summary>
 		/// AES key for the channel, as a hex string
 		/// </summary>
-		public string AesKey { get; set; } = String.Empty;
-
-		/// <summary>
-		/// AES IV for the channel, as a hex string
-		/// </summary>
-		public string AesIv { get; set; } = String.Empty;
+		public string Key { get; set; } = String.Empty;
 
 		/// <summary>
 		/// Resources assigned to this machine
 		/// </summary>
-		public Dictionary<string, int> Resources { get; set; } = new Dictionary<string, int>();
+		public Dictionary<string, int> AssignedResources { get; set; } = new Dictionary<string, int>();
 	}
 
 	/// <summary>
@@ -114,8 +109,12 @@ namespace Horde.Build.Compute
 			response.Ip = computeResource.Ip.ToString();
 			response.Port = computeResource.Port;
 			response.Nonce = StringUtils.FormatHexString(computeResource.Task.Nonce.Span);
-			response.AesKey = StringUtils.FormatHexString(computeResource.Task.AesKey.Span);
-			response.AesIv = StringUtils.FormatHexString(computeResource.Task.AesIv.Span);
+			response.Key = StringUtils.FormatHexString(computeResource.Task.Key.Span);
+
+			foreach (KeyValuePair<string, int> pair in computeResource.Task.Resources)
+			{
+				response.AssignedResources.Add(pair.Key, pair.Value);
+			}
 
 			return response;
 		}
diff --git a/Engine/Source/Programs/Horde/Horde.Build/Compute/ComputeTaskSource.cs b/Engine/Source/Programs/Horde/Horde.Build/Compute/ComputeTaskSource.cs
--- a/Engine/Source/Programs/Horde/Horde.Build/Compute/ComputeTaskSource.cs
+++ b/Engine/Source/Programs/Horde/Horde.Build/Compute/ComputeTaskSource.cs
@@ -42,20 +42,14 @@ namespace Horde.Build.Compute
 		/// </summary>
 		public ComputeTask Task { get; }
 
-		/// <summary>
-		/// Resources assigned to this lease
-		/// </summary>
-		public Dictionary<string, int> AssignedResources { get; }
-
 		/// <summary>
 		/// Constructor
 		/// </summary>
-		public ComputeResource(IPAddress ip, int port, ComputeTask task, Dictionary<string, int> assignedResources)
+		public ComputeResource(IPAddress ip, int port, ComputeTask task)
 		{
 			Ip = ip;
 			Port = port;
 			Task = task;
-			AssignedResources = assignedResources;
 		}
 	}
 
@@ -117,7 +111,7 @@ namespace Horde.Build.Compute
 						Dictionary<string, int> assignedResources = new Dictionary<string, int>(); 
 						if (node.Value.Agent.MeetsRequirements(requirements, assignedResources))
 						{
-							computeTask ??= CreateComputeTask();
+							computeTask ??= CreateComputeTask(assignedResources);
 
 							byte[] payload = Any.Pack(computeTask).ToByteArray();
 							AgentLease lease = new AgentLease(LeaseId.GenerateNewId(), "Compute task", null, null, null, LeaseState.Pending, assignedResources, requirements.Exclusive, payload);
@@ -125,7 +119,7 @@ namespace Horde.Build.Compute
 							if (node.Value.Lease.TrySetResult(lease))
 							{
 								Waiter waiter = node.Value;
-								return new ComputeResource(waiter.Ip, waiter.Port, computeTask, assignedResources);
+								return new ComputeResource(waiter.Ip, waiter.Port, computeTask);
 							}
 						}
 					}
@@ -134,17 +128,12 @@ namespace Horde.Build.Compute
 			return null;
 		}
 
-		static ComputeTask CreateComputeTask()
+		static ComputeTask CreateComputeTask(Dictionary<string, int> assignedResources)
 		{
 			ComputeTask computeTask = new ComputeTask();
 			computeTask.Nonce = UnsafeByteOperations.UnsafeWrap(RandomNumberGenerator.GetBytes(ServerComputeClient.NonceLength));
-
-			using (Aes aes = Aes.Create())
-			{
-				computeTask.AesKey = UnsafeByteOperations.UnsafeWrap(aes.Key);
-				computeTask.AesIv = UnsafeByteOperations.UnsafeWrap(aes.IV);
-			}
-
+			computeTask.Key = UnsafeByteOperations.UnsafeWrap(ComputeLease.CreateKey());
+			computeTask.Resources.Add(assignedResources);
 			return computeTask;
 		}
 
diff --git a/Engine/Source/Programs/Horde/Horde.Build/Utilities/BackgroundTask.cs b/Engine/Source/Programs/Horde/Horde.Build/Utilities/BackgroundTask.cs
deleted file mode 100644
diff --git a/Engine/Source/Programs/Shared/EpicGames.Core/BackgroundTask.cs b/Engine/Source/Programs/Shared/EpicGames.Core/BackgroundTask.cs
new file mode 100644
--- /dev/null
+++ b/Engine/Source/Programs/Shared/EpicGames.Core/BackgroundTask.cs
@@ -0,0 +1,93 @@
+// Copyright Epic Games, Inc. All Rights Reserved.
+
+using System;
+using System.Threading;
+using System.Threading.Tasks;
+
+namespace EpicGames.Core
+{
+	/// <summary>
+	/// Runs a task in the background and allows stopping it on demand
+	/// </summary>
+	public sealed class BackgroundTask : IDisposable, IAsyncDisposable
+	{
+		readonly Func<CancellationToken, Task> _runTask;
+		Task _task = Task.CompletedTask;
+		readonly CancellationTokenSource _cancellationTokenSource = new CancellationTokenSource();
+
+		/// <summary>
+		/// Accessor for the inner task
+		/// </summary>
+		public Task Task => _task;
+
+		/// <summary>
+		/// Constructor. Note that the task does not start until <see cref="Start"/> is called.
+		/// </summary>
+		/// <param name="runTask"></param>
+		public BackgroundTask(Func<CancellationToken, Task> runTask)
+		{
+			_runTask = runTask;
+		}
+
+		/// <inheritdoc/>
+		public void Dispose()
+		{
+			if (!_task.IsCompleted)
+			{
+				_cancellationTokenSource.Cancel();
+				StopAsync().Wait();
+			}
+			_cancellationTokenSource.Dispose();
+		}
+
+		/// <inheritdoc/>
+		public async ValueTask DisposeAsync()
+		{
+			await StopAsync();
+			_cancellationTokenSource.Dispose();
+		}
+
+		/// <summary>
+		/// Creates and starts a new background task instance
+		/// </summary>
+		/// <param name="runTask"></param>
+		/// <returns></returns>
+		public static BackgroundTask StartNew(Func<CancellationToken, Task> runTask)
+		{
+			BackgroundTask task = new BackgroundTask(runTask);
+			task.Start();
+			return task;
+		}
+
+		/// <summary>
+		/// Starts the task
+		/// </summary>
+		public void Start()
+		{
+			if (!_task.IsCompleted)
+			{
+				throw new InvalidOperationException("Background task is already running");
+			}
+			_task = Task.Run(() => _runTask(_cancellationTokenSource.Token), _cancellationTokenSource.Token);
+		}
+
+		/// <summary>
+		/// Signals the cancellation token and waits for the task to finish
+		/// </summary>
+		/// <returns></returns>
+		public async Task StopAsync()
+		{
+			if (!_task.IsCompleted)
+			{
+				try
+				{
+					_cancellationTokenSource.Cancel();
+					await _task;
+				}
+				catch (OperationCanceledException)
+				{
+				}
+			}
+		}
+	}
+}
diff --git a/Engine/Source/Programs/Shared/EpicGames.Core/MemoryWriter.cs b/Engine/Source/Programs/Shared/EpicGames.Core/MemoryWriter.cs
--- a/Engine/Source/Programs/Shared/EpicGames.Core/MemoryWriter.cs
+++ b/Engine/Source/Programs/Shared/EpicGames.Core/MemoryWriter.cs
@@ -144,6 +144,111 @@ namespace EpicGames.Core
 		}
 	}
 
+	/// <summary>
+	/// Writes into an expandable memory block
+	/// </summary>
+	public class PooledMemoryWriter : IMemoryWriter, IDisposable
+	{
+		MemoryPool<byte> _pool;
+		IMemoryOwner<byte> _owner;
+		Memory<byte> _memory;
+		int _length;
+
+		/// <summary>
+		/// Returns the span that has been written to
+		/// </summary>
+		public Span<byte> WrittenSpan => _memory.Span.Slice(0, _length);
+
+		/// <summary>
+		/// Returns the memory that has been written to
+		/// </summary>
+		public Memory<byte> WrittenMemory => _memory.Slice(0, _length);
+
+		/// <summary>
+		/// Constructor
+		/// </summary>
+		/// <param name="initialSize">Initial size of the allocated data</param>
+		public PooledMemoryWriter(int initialSize = 4096)
+			: this(MemoryPool<byte>.Shared, initialSize)
+		{
+		}
+
+		/// <summary>
+		/// Constructor
+		/// </summary>
+		/// <param name="pool">Pool to allocate data from</param>
+		/// <param name="initialSize">Initial size of the allocated data</param>
+		public PooledMemoryWriter(MemoryPool<byte> pool, int initialSize = 4096)
+		{
+			_pool = pool;
+			_owner = pool.Rent(initialSize);
+		}
+
+		/// <inheritdoc/>
+		public void Dispose()
+		{
+			Dispose(true);
+			GC.SuppressFinalize(this);
+		}
+
+		/// <summary>
+		/// Overridable dispose method
+		/// </summary>
+		/// <param name="disposing">Whether the </param>
+		protected virtual void Dispose(bool disposing)
+		{
+			_owner.Dispose();
+		}
+
+		/// <inheritdoc/>
+		public int Length => _length;
+
+		/// <summary>
+		/// Resets the current contents of the buffer
+		/// </summary>
+		public void Clear()
+		{
+			_length = 0;
+		}
+
+		/// <summary>
+		/// Resets the current contents of the buffer
+		/// </summary>
+		public void Clear(int minSize)
+		{
+			Clear();
+			if (minSize > _owner.Memory.Length)
+			{
+				_owner.Dispose();
+				_owner = _pool.Rent(minSize);
+			}
+		}
+
+		/// <inheritdoc/>
+		public Span<byte> GetSpan(int sizeHint = 0) => GetMemory(sizeHint).Span;
+
+		/// <inheritdoc/>
+		public Memory<byte> GetMemory(int minSize = 0)
+		{
+			int requiredLength = _length + Math.Max(minSize, 1);
+			if (_memory.Length < requiredLength)
+			{
+				IMemoryOwner<byte> newOwner = _pool.Rent(requiredLength + Math.Max(requiredLength / 2, 4096));
+				_memory.CopyTo(newOwner.Memory);
+				_memory = newOwner.Memory;
+				_owner.Dispose();
+				_owner = newOwner;
+			}
+			return _memory.Slice(_length);
+		}
+
+		/// <inheritdoc/>
+		public void Advance(int length)
+		{
+			_length += length;
+		}
+	}
+
 	/// <summary>
 	/// Class for building byte sequences, similar to StringBuilder. Allocates memory in chunks to avoid copying data.
 	/// </summary>
diff --git a/Engine/Source/Programs/Shared/EpicGames.Core/SocketExtensions.cs b/Engine/Source/Programs/Shared/EpicGames.Core/SocketExtensions.cs
--- a/Engine/Source/Programs/Shared/EpicGames.Core/SocketExtensions.cs
+++ b/Engine/Source/Programs/Shared/EpicGames.Core/SocketExtensions.cs
@@ -20,17 +20,33 @@ namespace EpicGames.Core
 		/// <param name="buffer">Buffer to store the data</param>
 		/// <param name="socketFlags">Flags for the socket receive call</param>
 		/// <param name="cancellationToken">Cancellation token for the operation</param>
-		public static async Task ReceiveFullAsync(this Socket socket, Memory<byte> buffer, SocketFlags socketFlags, CancellationToken cancellationToken)
+		public static async Task ReceiveMessageAsync(this Socket socket, Memory<byte> buffer, SocketFlags socketFlags, CancellationToken cancellationToken)
+		{
+			if (!await TryReceiveMessageAsync(socket, buffer, socketFlags, cancellationToken))
+			{
+				throw new EndOfStreamException();
+			}
+		}
+
+		/// <summary>
+		/// Reads a complete buffer from the given socket, retrying reads until the buffer is full.
+		/// </summary>
+		/// <param name="socket">Socket to read from</param>
+		/// <param name="buffer">Buffer to store the data</param>
+		/// <param name="socketFlags">Flags for the socket receive call</param>
+		/// <param name="cancellationToken">Cancellation token for the operation</param>
+		public static async Task<bool> TryReceiveMessageAsync(this Socket socket, Memory<byte> buffer, SocketFlags socketFlags, CancellationToken cancellationToken)
 		{
 			for (int offset = 0; offset < buffer.Length;)
 			{
 				int read = await socket.ReceiveAsync(buffer.Slice(offset), socketFlags, cancellationToken);
 				if (read == 0)
 				{
-					throw new EndOfStreamException();
+					return false;
 				}
 				offset += read;
 			}
+			return true;
 		}
 
 		/// <summary>
@@ -41,7 +57,7 @@ namespace EpicGames.Core
 		/// <param name="socketFlags">Flags for the socket sent call</param>
 		/// <param name="cancellationToken">Cancellation token for the operation</param>
 		/// <returns></returns>
-		public static async Task SendFullAsync(this Socket socket, ReadOnlyMemory<byte> buffer, SocketFlags socketFlags, CancellationToken cancellationToken)
+		public static async Task SendMessageAsync(this Socket socket, ReadOnlyMemory<byte> buffer, SocketFlags socketFlags, CancellationToken cancellationToken)
 		{
 			while (buffer.Length > 0)
 			{
diff --git a/Engine/Source/Programs/Shared/EpicGames.Horde/Compute/ComputeLease.cs b/Engine/Source/Programs/Shared/EpicGames.Horde/Compute/ComputeLease.cs
new file mode 100644
--- /dev/null
+++ b/Engine/Source/Programs/Shared/EpicGames.Horde/Compute/ComputeLease.cs
@@ -0,0 +1,375 @@
+// Copyright Epic Games, Inc. All Rights Reserved.
+
+using System;
+using System.Buffers;
+using System.Buffers.Binary;
+using System.Collections.Generic;
+using System.Net.Sockets;
+using System.Security.Cryptography;
+using System.Threading;
+using System.Threading.Tasks;
+using EpicGames.Core;
+using EpicGames.Serialization;
+
+namespace EpicGames.Horde.Compute
+{
+	/// <summary>
+	/// Handles a compute lease with a particular remote, multiplexing messages across a single socket.
+	/// </summary>
+	public sealed class ComputeLease : IComputeLease
+	{
+		/// <summary>
+		/// Length of the required encrption key. 
+		/// </summary>
+		public const int KeyLength = 32;
+
+		/// <summary>
+		/// Length of the nonce. This should be a cryptographically random number, and does not have to be secret.
+		/// </summary>
+		public const int NonceLength = 12;
+
+		/// <summary>
+		/// Maximum length of a single message
+		/// </summary>
+		public const int MaxMessageLength = 1024 * 1024;
+
+		sealed class Message : IComputeMessage
+		{
+			internal IMemoryOwner<byte>? _owner;
+			internal readonly TaskCompletionSource<Message> _next = new TaskCompletionSource<Message>();
+
+			/// <summary>
+			/// Type of the message
+			/// </summary>
+			public ComputeMessageType Type { get; }
+
+			/// <summary>
+			/// Message data
+			/// </summary>
+			public ReadOnlyMemory<byte> Data { get; private set; }
+
+			public Message(IMemoryOwner<byte>? owner, ComputeMessageType type, ReadOnlyMemory<byte> data)
+			{
+				_owner = owner;
+
+				Type = type;
+				Data = data;
+			}
+
+			/// <inheritdoc/>
+			public void Dispose()
+			{
+				if (_owner != null)
+				{
+					_owner.Dispose();
+					_owner = null;
+				}
+			}
+		}
+
+		/// <summary>
+		/// Allows creating new messages in rented memory
+		/// </summary>
+		class Writer : IComputeMessageWriter
+		{
+			readonly ComputeLease _owner;
+			PooledMemoryWriter _writer;
+
+			public int Length => _writer.Length;
+
+			public Writer(ComputeLease owner, int channelId, ComputeMessageType type, PooledMemoryWriter writer)
+			{
+				_owner = owner;
+				_writer = writer;
+
+				_writer.WriteUInt32(0); // Placeholder for size
+				_writer.WriteUnsignedVarInt(channelId);
+				_writer.WriteUInt8((byte)type);
+			}
+
+			public void Dispose()
+			{
+				if (_writer != null)
+				{
+					_owner.ReleaseWriter(_writer);
+					_writer = null!;
+				}
+			}
+
+			public ValueTask SendAsync(CancellationToken cancellationToken)
+			{
+				int bodySize = _writer.Length - HeaderLength;
+				_writer.Advance(FooterLength);
+
+				Span<byte> message = _writer.WrittenMemory.Span;
+
+				Span<byte> header = message.Slice(0, HeaderLength);
+				BinaryPrimitives.WriteInt32LittleEndian(header, bodySize);
+
+				Span<byte> body = message.Slice(HeaderLength, bodySize);
+				Span<byte> footer = message.Slice(header.Length + body.Length);
+				_owner._aesGcm.Encrypt(_owner._writeNonce, body, body, footer, header);
+				IncrementNonce(_owner._writeNonce);
+
+				return _owner.SendMessageAsync(_writer.WrittenMemory, cancellationToken);
+			}
+
+			/// <inheritdoc/>
+			public void Advance(int count) => _writer.Advance(count);
+
+			/// <inheritdoc/>
+			public Memory<byte> GetMemory(int sizeHint = 0) => _writer.GetMemory(sizeHint);
+
+			/// <inheritdoc/>
+			public Span<byte> GetSpan(int sizeHint = 0) => _writer.GetSpan(sizeHint);
+		}
+
+		/// <summary>
+		/// Multiplexes data onto the owner's socket
+		/// </summary>
+		class Channel : IComputeChannel
+		{
+			readonly ComputeLease _owner;
+			readonly int _channelId;
+
+			TaskCompletionSource<Message> _readTail;
+			Task<Message> _readNext;
+
+			public int Id => _channelId;
+
+			public Channel(ComputeLease owner, int channelId)
+			{
+				_owner = owner;
+				_channelId = channelId;
+				_readTail = new TaskCompletionSource<Message>();
+				_readNext = _readTail.Task;
+			}
+
+			public IComputeMessageWriter CreateMessage(ComputeMessageType type, int reserve) => _owner.CreateWriter(type, _channelId, reserve);
+
+			public void Dispose()
+			{
+				_owner.RemoveChannel(_channelId);
+			}
+
+			public void MarkComplete()
+			{
+#pragma warning disable CA2000 // Dispose objects before losing scope
+				_readTail.SetResult(new Message(null, ComputeMessageType.None, ReadOnlyMemory<byte>.Empty));
+#pragma warning restore CA2000 // Dispose objects before losing scope
+			}
+
+			public void QueueMessage(Message message)
+			{
+				_readTail.SetResult(message);
+				_readTail = message._next;
+			}
+
+			/// <inheritdoc/>
+			public async ValueTask<IComputeMessage> ReadAsync(CancellationToken cancellationToken)
+			{
+				Message message = await _readNext;
+				_readNext = message._next.Task;
+				return message;
+			}
+		}
+
+		/// <inheritdoc/>
+		public IComputeChannel DefaultChannel => _defaultChannel;
+
+		/// <inheritdoc/>
+		public IReadOnlyDictionary<string, int> AssignedResources { get; }
+
+		const int HeaderLength = sizeof(int); // Unencrypted size of the message
+		const int FooterLength = 16; // 16-byte auth tag for encryption
+
+		readonly Socket _socket;
+		readonly MemoryPool<byte> _pool;
+		readonly AesGcm _aesGcm;
+		readonly BackgroundTask _backgroundTask;
+		readonly Channel _defaultChannel;
+
+		readonly object _lockObject = new object();
+		readonly Dictionary<int, Channel> _channels = new Dictionary<int, Channel>();
+
+		readonly byte[] _readNonce;
+		readonly byte[] _writeNonce;
+
+		readonly Stack<PooledMemoryWriter> _freeWriters = new Stack<PooledMemoryWriter>();
+
+		readonly SemaphoreSlim _writeSemaphore;
+
+		/// <summary>
+		/// Constructor
+		/// </summary>
+		/// <param name="socket">Socket for communication</param>
+		/// <param name="key">AES encryption key (256 bits / 32 bytes)</param>
+		/// <param name="nonce">Cryptographic nonce to identify the connection. Must be longer than <see cref="NonceLength"/>.</param>
+		/// <param name="assignedResources">Resources assigned to this lease</param>
+		public ComputeLease(Socket socket, ReadOnlySpan<byte> key, ReadOnlySpan<byte> nonce, IReadOnlyDictionary<string, int> assignedResources)
+			: this(socket, key, nonce, assignedResources, MemoryPool<byte>.Shared)
+		{
+		}
+
+		/// <summary>
+		/// Constructor
+		/// </summary>
+		/// <param name="socket">Socket for communication</param>
+		/// <param name="key">AES encryption key (256 bits / 32 bytes)</param>
+		/// <param name="nonce">Cryptographic nonce to identify the connection. Must be longer than <see cref="NonceLength"/>.</param>
+		/// <param name="assignedResources">Resources assigned to this lease</param>
+		/// <param name="pool">Memory pool from which to allocate message buffers</param>
+		public ComputeLease(Socket socket, ReadOnlySpan<byte> key, ReadOnlySpan<byte> nonce, IReadOnlyDictionary<string, int> assignedResources, MemoryPool<byte> pool)
+		{
+			if (key.Length != KeyLength)
+			{
+				throw new ArgumentException($"Key must be {KeyLength} bytes", nameof(key));
+			}
+
+			_socket = socket;
+			_pool = pool;
+			_aesGcm = new AesGcm(key);
+
+			_readNonce = nonce.Slice(0, NonceLength).ToArray();
+			_writeNonce = nonce.Slice(0, NonceLength).ToArray();
+
+			_writeSemaphore = new SemaphoreSlim(1);
+
+			_defaultChannel = new Channel(this, 0);
+			_channels.Add(0, _defaultChannel);
+
+			_backgroundTask = BackgroundTask.StartNew(BackgroundReadAsync);
+
+			AssignedResources = assignedResources;
+		}
+
+		/// <inheritdoc/>
+		public async ValueTask DisposeAsync()
+		{
+			await _backgroundTask.DisposeAsync();
+
+			foreach (Channel channel in _channels.Values)
+			{
+				channel.MarkComplete();
+			}
+
+			_writeSemaphore.Dispose();
+			_aesGcm.Dispose();
+		}
+
+		/// <summary>
+		/// Creates a new encryption key for the compute lead
+		/// </summary>
+		/// <returns>New key bytes</returns>
+		public static byte[] CreateKey() => RandomNumberGenerator.GetBytes(KeyLength);
+
+		/// <inheritdoc/>
+		public IComputeChannel OpenChannel(int channelId)
+		{
+			Channel channel = new Channel(this, channelId);
+			lock (_lockObject)
+			{
+				_channels.Add(channelId, channel);
+			}
+			return channel;
+		}
+
+		async Task BackgroundReadAsync(CancellationToken cancellationToken)
+		{
+			byte[] sizeBuffer = new byte[4];
+			while (!cancellationToken.IsCancellationRequested)
+			{
+				if (!await _socket.TryReceiveMessageAsync(sizeBuffer, SocketFlags.None, cancellationToken))
+				{
+					break;
+				}
+
+				int length = BinaryPrimitives.ReadInt32LittleEndian(sizeBuffer);
+				if (length < 0 || length > MaxMessageLength)
+				{
+					throw new InvalidOperationException();
+				}
+
+				IMemoryOwner<byte> owner = _pool.Rent(length + FooterLength);
+				try
+				{
+					// Decrypt the message
+					await _socket.ReceiveMessageAsync(owner.Memory.Slice(0, length + FooterLength), SocketFlags.None, cancellationToken);
+					_aesGcm.Decrypt(_readNonce, owner.Memory.Span.Slice(0, length), owner.Memory.Span.Slice(length, FooterLength), owner.Memory.Span.Slice(0, length), sizeBuffer);
+					IncrementNonce(_readNonce);
+
+					// Decode the message data
+					ReadOnlyMemory<byte> data = owner.Memory.Slice(0, length);
+
+					int channelId = (int)VarInt.ReadUnsigned(data.Span, out int channelIdByteCount);
+					data = data.Slice(channelIdByteCount);
+
+					ComputeMessageType type = (ComputeMessageType)data.Span[0];
+					data = data.Slice(1);
+
+#pragma warning disable CA2000 // Dispose objects before losing scope
+					Message message = new Message(owner, type, data);
+#pragma warning restore CA2000 // Dispose objects before losing scope
+
+					Channel? channel;
+					lock (_lockObject)
+					{
+						channel = _channels[channelId];
+					}
+
+					channel.QueueMessage(message);
+				}
+				catch
+				{
+					owner.Dispose();
+					throw;
+				}
+			}
+		}
+
+		void RemoveChannel(int channelId)
+		{
+			lock (_lockObject)
+			{
+				_channels.Remove(channelId);
+			}
+		}
+
+		Writer CreateWriter(ComputeMessageType type, int channelId, int reserve)
+		{
+			PooledMemoryWriter? writer;
+			if (_freeWriters.TryPop(out writer))
+			{
+				writer.Clear(reserve);
+			}
+			else
+			{
+				writer = new PooledMemoryWriter(reserve);
+			}
+			return new Writer(this, channelId, type, writer);
+		}
+
+		void ReleaseWriter(PooledMemoryWriter writer)
+		{
+			_freeWriters.Push(writer);
+		}
+
+		async ValueTask SendMessageAsync(ReadOnlyMemory<byte> message, CancellationToken cancellationToken)
+		{
+			await _writeSemaphore.WaitAsync(cancellationToken);
+			try
+			{
+				await _socket.SendMessageAsync(message, SocketFlags.None, cancellationToken);
+			}
+			finally
+			{
+				_writeSemaphore.Release();
+			}
+		}
+
+		static void IncrementNonce(byte[] nonce)
+		{
+			BinaryPrimitives.WriteInt64LittleEndian(nonce, BinaryPrimitives.ReadInt64LittleEndian(nonce) + 1);
+		}
+	}
+}
diff --git a/Engine/Source/Programs/Shared/EpicGames.Horde/Compute/ComputeMessage.cs b/Engine/Source/Programs/Shared/EpicGames.Horde/Compute/ComputeMessage.cs
deleted file mode 100644
diff --git a/Engine/Source/Programs/Shared/EpicGames.Horde/Compute/ComputeMessageTypes.cs b/Engine/Source/Programs/Shared/EpicGames.Horde/Compute/ComputeMessageTypes.cs
deleted file mode 100644
diff --git a/Engine/Source/Programs/Shared/EpicGames.Horde/Compute/ComputeStorageClient.cs b/Engine/Source/Programs/Shared/EpicGames.Horde/Compute/ComputeStorageClient.cs
--- a/Engine/Source/Programs/Shared/EpicGames.Horde/Compute/ComputeStorageClient.cs
+++ b/Engine/Source/Programs/Shared/EpicGames.Horde/Compute/ComputeStorageClient.cs
@@ -36,6 +36,27 @@ namespace EpicGames.Horde.Compute
 
 		#region Blobs
 
+		sealed class BlobStream : ReadOnlyMemoryStream
+		{
+			readonly IComputeMessage _message;
+
+			public BlobStream(IComputeMessage message)
+				: base(message.Data)
+			{
+				_message = message;
+			}
+
+			protected override void Dispose(bool disposing)
+			{
+				base.Dispose(disposing);
+
+				if (disposing)
+				{
+					_message.Dispose();
+				}
+			}
+		}
+
 		/// <inheritdoc/>
 		public override async Task<Stream> ReadBlobAsync(BlobLocator locator, CancellationToken cancellationToken = default)
 		{
@@ -48,11 +69,16 @@ namespace EpicGames.Horde.Compute
 			await _semaphore.WaitAsync(cancellationToken);
 			try
 			{
-				BlobReadMessage request = new BlobReadMessage { Locator = locator, Offset = offset, Length = length };
-				await _channel.WriteCbMessageAsync(request, cancellationToken);
+				await _channel.CppBlobReadAsync(locator, offset, length, cancellationToken);
+
+				IComputeMessage response = await _channel.ReadAsync(cancellationToken);
+				if (response.Type != ComputeMessageType.CppBlobData)
+				{
+					response.Dispose();
+					throw new NotImplementedException();
+				}
 
-				BlobDataMessage response = await _channel.ReadCbMessageAsync<BlobDataMessage>(cancellationToken);
-				return new ReadOnlyMemoryStream(response.Data);
+				return new BlobStream(response);
 			}
 			finally
 			{
diff --git a/Engine/Source/Programs/Shared/EpicGames.Horde/Compute/IComputeChannel.cs b/Engine/Source/Programs/Shared/EpicGames.Horde/Compute/IComputeChannel.cs
--- a/Engine/Source/Programs/Shared/EpicGames.Horde/Compute/IComputeChannel.cs
+++ b/Engine/Source/Programs/Shared/EpicGames.Horde/Compute/IComputeChannel.cs
@@ -3,26 +3,45 @@
 using System;
 using System.Threading;
 using System.Threading.Tasks;
+using EpicGames.Core;
 
 namespace EpicGames.Horde.Compute
 {
 	/// <summary>
-	/// Full-duplex channel for sending and reciving messages
+	/// Full-duplex channel for sending and receiving messages
 	/// </summary>
 	public interface IComputeChannel
 	{
+		/// <summary>
+		/// Identifier for the current channel
+		/// </summary>
+		int Id { get; }
+
 		/// <summary>
 		/// Reads a message from the channel
 		/// </summary>
 		/// <param name="cancellationToken">Cancellation token for the operation</param>
-		/// <returns>Data for the message that was read. This buffer is ephemeral, and may be overwritten on the next read call.</returns>
-		Task<ReadOnlyMemory<byte>> ReadMessageAsync(CancellationToken cancellationToken);
+		/// <returns>Data for a message that was read. Must be disposed.</returns>
+		ValueTask<IComputeMessage> ReadAsync(CancellationToken cancellationToken);
+
+		/// <summary>
+		/// Creates a new builder for a message
+		/// </summary>
+		/// <param name="type">Type of the message</param>
+		/// <param name="sizeHint">Hint for the expected message size</param>
+		/// <returns>New builder for messages</returns>
+		IComputeMessageWriter CreateMessage(ComputeMessageType type, int sizeHint = 0);
+	}
 
+	/// <summary>
+	/// Writer for compute messages
+	/// </summary>
+	public interface IComputeMessageWriter : IMemoryWriter, IDisposable
+	{
 		/// <summary>
-		/// Writes a message to the channel
+		/// Sends the current message
 		/// </summary>
-		/// <param name="buffer">Buffer containing the data to be written</param>
 		/// <param name="cancellationToken">Cancellation token for the operation</param>
-		Task WriteMessageAsync(ReadOnlyMemory<byte> buffer, CancellationToken cancellationToken);
+		ValueTask SendAsync(CancellationToken cancellationToken);
 	}
 }
diff --git a/Engine/Source/Programs/Shared/EpicGames.Horde/Compute/IComputeClient.cs b/Engine/Source/Programs/Shared/EpicGames.Horde/Compute/IComputeClient.cs
--- a/Engine/Source/Programs/Shared/EpicGames.Horde/Compute/IComputeClient.cs
+++ b/Engine/Source/Programs/Shared/EpicGames.Horde/Compute/IComputeClient.cs
@@ -18,6 +18,6 @@ namespace EpicGames.Horde.Compute
 		/// <param name="requirements">Requirements for the agent</param>
 		/// <param name="handler">Handler for the connection</param>
 		/// <param name="cancellationToken">Cancellation token for the operation</param>
-		public Task<TResult> ExecuteAsync<TResult>(ClusterId clusterId, Requirements? requirements, Func<IComputeChannel, CancellationToken, Task<TResult>> handler, CancellationToken cancellationToken);
+		public Task<TResult> ExecuteAsync<TResult>(ClusterId clusterId, Requirements? requirements, Func<IComputeLease, CancellationToken, Task<TResult>> handler, CancellationToken cancellationToken);
 	}
 }
diff --git a/Engine/Source/Programs/Shared/EpicGames.Horde/Compute/IComputeLease.cs b/Engine/Source/Programs/Shared/EpicGames.Horde/Compute/IComputeLease.cs
new file mode 100644
--- /dev/null
+++ b/Engine/Source/Programs/Shared/EpicGames.Horde/Compute/IComputeLease.cs
@@ -0,0 +1,30 @@
+// Copyright Epic Games, Inc. All Rights Reserved.
+
+using System;
+using System.Collections.Generic;
+
+namespace EpicGames.Horde.Compute
+{
+	/// <summary>
+	/// Full-duplex channel for sending and reciving messages
+	/// </summary>
+	public interface IComputeLease : IAsyncDisposable
+	{
+		/// <summary>
+		/// The control channel, used to transmit messages before additional channels have been created
+		/// </summary>
+		IComputeChannel DefaultChannel { get; }
+
+		/// <summary>
+		/// Resources assigned to this lease
+		/// </summary>
+		IReadOnlyDictionary<string, int> AssignedResources { get; }
+
+		/// <summary>
+		/// Opens a new channel with the given identifier.
+		/// </summary>
+		/// <param name="id">The pipe identifier. These values are arbitrary, as long as both ends agree on a channel to communicate on. Channel 0 is reserved for the control channel.</param>
+		/// <returns>New channel instance</returns>
+		IComputeChannel OpenChannel(int id);
+	}
+}
diff --git a/Engine/Source/Programs/Shared/EpicGames.Horde/Compute/IComputeMessage.cs b/Engine/Source/Programs/Shared/EpicGames.Horde/Compute/IComputeMessage.cs
new file mode 100644
--- /dev/null
+++ b/Engine/Source/Programs/Shared/EpicGames.Horde/Compute/IComputeMessage.cs
@@ -0,0 +1,311 @@
+// Copyright Epic Games, Inc. All Rights Reserved.
+
+using System;
+using System.IO;
+using System.Threading;
+using System.Threading.Tasks;
+using EpicGames.Core;
+using EpicGames.Horde.Storage;
+using Microsoft.CodeAnalysis.CSharp.Syntax;
+
+namespace EpicGames.Horde.Compute
+{
+	/// <summary>
+	/// Type of a compute message
+	/// </summary>
+	public enum ComputeMessageType
+	{
+		/// <summary>
+		/// No message was received (end of stream)
+		/// </summary>
+		None = 0x00,
+
+		/// <summary>
+		/// Gracefully close the current connection
+		/// </summary>
+		Close = 0x01,
+
+		#region Test Requests
+
+		/// <summary>
+		/// Xor a block of data with a value
+		/// </summary>
+		XorRequest = 0x10,
+
+		/// <summary>
+		/// Result from an <see cref="XorRequest"/> request.
+		/// </summary>
+		XorResponse = 0x11,
+
+		#endregion
+
+		#region C++ Builds
+
+		/// <summary>
+		/// Initiate a C++ build
+		/// </summary>
+		CppBegin = 0x20,
+
+		/// <summary>
+		/// Gives output from a C++ build
+		/// </summary>
+		CppResult = 0x21,
+
+		/// <summary>
+		/// Finish a C++ build and close the current channel
+		/// </summary>
+		CppEnd = 0x22,
+
+		/// <summary>
+		/// Reads a blob pertaining to a C++ build
+		/// </summary>
+		CppBlobRead = 0x23,
+
+		/// <summary>
+		/// Response to a <see cref="CppBlobRead"/> request.
+		/// </summary>
+		CppBlobData = 0x24,
+
+		#endregion
+	}
+
+	/// <summary>
+	/// Information about a received compute message
+	/// </summary>
+	public interface IComputeMessage : IDisposable
+	{
+		/// <summary>
+		/// Type of the message
+		/// </summary>
+		public ComputeMessageType Type { get; }
+
+		/// <summary>
+		/// Data that was read
+		/// </summary>
+		public ReadOnlyMemory<byte> Data { get; }
+	}
+
+	/// <summary>
+	/// Message for running an XOR command
+	/// </summary>
+	/// <param name="Data">Data to xor</param>
+	/// <param name="Value">Value to XOR with</param>
+	public record struct XorRequestMessage(ReadOnlyMemory<byte> Data, byte Value);
+
+	/// <summary>
+	/// Message for starting a C++ build
+	/// </summary>
+	/// <param name="Locator">Locator for the task description</param>
+	/// <param name="ReplyChannelId">Channel to reply on</param>
+	public record struct CppBeginMessage(NodeLocator Locator, int ReplyChannelId);
+
+	/// <summary>
+	/// Creates a blob read request
+	/// </summary>
+	public record struct CppBlobReadMessage(BlobLocator Locator, int Offset, int Length);
+
+	/// <summary>
+	/// Output from a C++ build request
+	/// </summary>
+	public record struct CppOutputMessage(NodeLocator Locator);
+
+	/// <summary>
+	/// Wraps various requests across compute channels
+	/// </summary>
+	public static class ComputeChannelExtensions
+	{
+		/// <summary>
+		/// Send a close message to the remote
+		/// </summary>
+		/// <param name="channel">Channel to write to</param>
+		/// <param name="cancellationToken">Cancellation token for the operation</param>
+		public static async Task CloseAsync(this IComputeChannel channel, CancellationToken cancellationToken)
+		{
+			using (IComputeMessageWriter writer = channel.CreateMessage(ComputeMessageType.Close))
+			{
+				await writer.SendAsync(cancellationToken);
+			}
+		}
+
+		/// <summary>
+		/// 
+		/// </summary>
+		/// <param name="channel"></param>
+		/// <param name="data"></param>
+		/// <param name="value"></param>
+		/// <param name="cancellationToken"></param>
+		/// <returns></returns>
+		public static async Task XorRequestAsync(this IComputeChannel channel, ReadOnlyMemory<byte> data, byte value, CancellationToken cancellationToken)
+		{
+			using (IComputeMessageWriter writer = channel.CreateMessage(ComputeMessageType.XorRequest))
+			{
+				writer.WriteFixedLengthBytes(data.Span);
+				writer.WriteUInt8(value);
+				await writer.SendAsync(cancellationToken);
+			}
+		}
+
+		/// <summary>
+		/// 
+		/// </summary>
+		/// <param name="message"></param>
+		/// <returns></returns>
+		public static XorRequestMessage AsXorRequest(this IComputeMessage message)
+		{
+			ReadOnlyMemory<byte> data = message.Data;
+			return new XorRequestMessage(data[0..^1], data.Span[^1]);
+		}
+
+		/// <summary>
+		/// Starts a C++ build task
+		/// </summary>
+		/// <param name="channel">Channel to write to</param>
+		/// <param name="locator">Locator for the node containing the task definition</param>
+		/// <param name="replyChannelId">Channel to reply to with future messages</param>
+		/// <param name="cancellationToken">Cancellation token for the operation</param>
+		public static async Task CppStartAsync(this IComputeChannel channel, NodeLocator locator, int replyChannelId, CancellationToken cancellationToken)
+		{
+			using (IComputeMessageWriter writer = channel.CreateMessage(ComputeMessageType.CppBegin))
+			{
+				writer.WriteNodeLocator(locator);
+				writer.WriteUnsignedVarInt(replyChannelId);
+				await writer.SendAsync(cancellationToken);
+			}
+		}
+
+		/// <summary>
+		/// 
+		/// </summary>
+		/// <param name="message"></param>
+		/// <returns></returns>
+		public static CppBeginMessage AsCppBegin(this IComputeMessage message)
+		{
+			MemoryReader reader = new MemoryReader(message.Data);
+			NodeLocator locator = reader.ReadNodeLocator();
+			int replyChannelId = (int)reader.ReadUnsignedVarInt();
+			return new CppBeginMessage(locator, replyChannelId);
+		}
+
+		/// <summary>
+		/// Sends the result for a C++ compute task
+		/// </summary>
+		/// <param name="message"></param>
+		public static NodeLocator AsCppResult(this IComputeMessage message)
+		{
+			MemoryReader reader = new MemoryReader(message.Data);
+			return reader.ReadNodeLocator();
+		}
+
+		/// <summary>
+		/// Sends the result for a C++ compute task
+		/// </summary>
+		/// <param name="channel">Channel to write to</param>
+		/// <param name="locator">Locator for the result</param>
+		/// <param name="cancellationToken">Cancellation token for the operation</param>
+		public static async Task CppResultAsync(this IComputeChannel channel, NodeLocator locator, CancellationToken cancellationToken)
+		{
+			using (IComputeMessageWriter writer = channel.CreateMessage(ComputeMessageType.CppResult))
+			{
+				writer.WriteNodeLocator(locator);
+				await writer.SendAsync(cancellationToken);
+			}
+		}
+
+		/// <summary>
+		/// Completes a C++ build task
+		/// </summary>
+		/// <param name="channel">Channel to write to</param>
+		/// <param name="cancellationToken">Cancellation token for the operation</param>
+		public static async Task CppFinishAsync(this IComputeChannel channel, CancellationToken cancellationToken)
+		{
+			using (IComputeMessageWriter writer = channel.CreateMessage(ComputeMessageType.CppEnd))
+			{
+				await writer.SendAsync(cancellationToken);
+			}
+		}
+
+		/// <summary>
+		/// 
+		/// </summary>
+		/// <param name="message"></param>
+		/// <returns></returns>
+		public static CppBlobReadMessage AsCppBlobRead(this IComputeMessage message)
+		{
+			MemoryReader reader = new MemoryReader(message.Data);
+			BlobLocator locator = reader.ReadBlobLocator();
+			int offset = (int)reader.ReadUnsignedVarInt();
+			int length = (int)reader.ReadUnsignedVarInt();
+			return new CppBlobReadMessage(locator, offset, length);
+		}
+
+		/// <summary>
+		/// 
+		/// </summary>
+		/// <param name="channel"></param>
+		/// <param name="locator"></param>
+		/// <param name="offset"></param>
+		/// <param name="length"></param>
+		/// <param name="cancellationToken"></param>
+		/// <returns></returns>
+		public static async Task CppBlobReadAsync(this IComputeChannel channel, BlobLocator locator, int offset, int length, CancellationToken cancellationToken)
+		{
+			using (IComputeMessageWriter writer = channel.CreateMessage(ComputeMessageType.CppBlobRead))
+			{
+				writer.WriteBlobLocator(locator);
+				writer.WriteUnsignedVarInt(offset);
+				writer.WriteUnsignedVarInt(length);
+				await writer.SendAsync(cancellationToken);
+			}
+		}
+
+		/// <summary>
+		/// Writes blob data to a compute channel
+		/// </summary>
+		/// <param name="channel">Channel to write to</param>
+		/// <param name="message">The read request</param>
+		/// <param name="storage">Storage client to retrieve the blob from</param>
+		/// <param name="cancellationToken">Cancellation token for the operation</param>
+		public static Task CppBlobDataAsync(this IComputeChannel channel, CppBlobReadMessage message, IStorageClient storage, CancellationToken cancellationToken)
+		{
+			return CppBlobDataAsync(channel, message.Locator, message.Offset, message.Length, storage, cancellationToken);
+		}
+
+		/// <summary>
+		/// Writes blob data to a compute channel
+		/// </summary>
+		/// <param name="channel">Channel to write to</param>
+		/// <param name="locator">Locator for the blob to send</param>
+		/// <param name="offset">Starting offset of the data</param>
+		/// <param name="length">Length of the data</param>
+		/// <param name="storage">Storage client to retrieve the blob from</param>
+		/// <param name="cancellationToken">Cancellation token for the operation</param>
+		public static async Task CppBlobDataAsync(this IComputeChannel channel, BlobLocator locator, int offset, int length, IStorageClient storage, CancellationToken cancellationToken)
+		{
+			byte[] data;
+			if (offset == 0 && length == 0)
+			{
+				using (Stream stream = await storage.ReadBlobAsync(locator, cancellationToken))
+				{
+					using MemoryStream target = new MemoryStream();
+					await stream.CopyToAsync(target, cancellationToken);
+					data = target.ToArray();
+				}
+			}
+			else
+			{
+				using (Stream stream = await storage.ReadBlobRangeAsync(locator, offset, length, cancellationToken))
+				{
+					using MemoryStream target = new MemoryStream();
+					await stream.CopyToAsync(target, cancellationToken);
+					data = target.ToArray();
+				}
+			}
+
+			using (IComputeMessageWriter writer = channel.CreateMessage(ComputeMessageType.CppBlobData, length + 128))
+			{
+				writer.WriteFixedLengthBytes(data);
+				await writer.SendAsync(cancellationToken);
+			}
+		}
+	}
+}
diff --git a/Engine/Source/Programs/Shared/EpicGames.Horde/Compute/LoopbackComputeClient.cs b/Engine/Source/Programs/Shared/EpicGames.Horde/Compute/LoopbackComputeClient.cs
--- a/Engine/Source/Programs/Shared/EpicGames.Horde/Compute/LoopbackComputeClient.cs
+++ b/Engine/Source/Programs/Shared/EpicGames.Horde/Compute/LoopbackComputeClient.cs
@@ -1,11 +1,13 @@
 // Copyright Epic Games, Inc. All Rights Reserved.
 
 using System;
+using System.Collections.Generic;
 using System.Net;
 using System.Net.Sockets;
 using System.Security.Cryptography;
 using System.Threading;
 using System.Threading.Tasks;
+using EpicGames.Core;
 
 namespace EpicGames.Horde.Compute
 {
@@ -14,8 +16,10 @@ namespace EpicGames.Horde.Compute
 	/// </summary>
 	public sealed class LoopbackComputeClient : IComputeClient
 	{
-		readonly Task _listenerTask;
-		readonly SocketComputeChannel _inner;
+		readonly byte[] _key;
+		readonly byte[] _nonce;
+		readonly BackgroundTask _listenerTask;
+		readonly Socket _listener;
 		readonly Socket _socket;
 
 		/// <summary>
@@ -23,44 +27,45 @@ namespace EpicGames.Horde.Compute
 		/// </summary>
 		/// <param name="serverFunc">Callback to process the server side of the connection</param>
 		/// <param name="port">Port to connect on</param>
-		public LoopbackComputeClient(Func<IComputeChannel, CancellationToken, Task> serverFunc, int port = 2000)
+		public LoopbackComputeClient(Func<IComputeLease, CancellationToken, Task> serverFunc, int port = 2000)
 		{
-			using Aes aes = Aes.Create();
-			_listenerTask = RunListenerAsync(aes.Key, aes.IV, serverFunc, port, CancellationToken.None);
+			_key = ComputeLease.CreateKey();
+			_nonce = RandomNumberGenerator.GetBytes(ComputeLease.NonceLength);
+
+			_listener = new Socket(SocketType.Stream, ProtocolType.IP);
+			_listener.Bind(new IPEndPoint(IPAddress.Loopback, port));
+			_listener.Listen();
+
+			_listenerTask = BackgroundTask.StartNew(ctx => RunListenerAsync(_listener, _key, _nonce, serverFunc, ctx));
 
 			_socket = new Socket(SocketType.Stream, ProtocolType.IP);
 			_socket.Connect(IPAddress.Loopback, port);
-
-			_inner = new SocketComputeChannel(_socket, aes.Key, aes.IV);
 		}
 
 		/// <inheritdoc/>
 		public async ValueTask DisposeAsync()
 		{
-			_inner.Dispose();
 			_socket.Dispose();
-			await _listenerTask;
+			await _listenerTask.DisposeAsync();
+			_listener.Dispose();
 		}
 
 		/// <summary>
 		/// Sets up the loopback listener and calls the server method
 		/// </summary>
-		static async Task RunListenerAsync(ReadOnlyMemory<byte> aesKey, ReadOnlyMemory<byte> aesIv, Func<IComputeChannel, CancellationToken, Task> serverFunc, int port, CancellationToken cancellationToken)
+		static async Task RunListenerAsync(Socket listener, ReadOnlyMemory<byte> key, ReadOnlyMemory<byte> nonce, Func<IComputeLease, CancellationToken, Task> serverFunc, CancellationToken cancellationToken)
 		{
-			using Socket listener = new Socket(SocketType.Stream, ProtocolType.IP);
-			listener.Bind(new IPEndPoint(IPAddress.Loopback, port));
-			listener.Listen();
-
 			using Socket socket = await listener.AcceptAsync(cancellationToken);
-			using SocketComputeChannel channel = new SocketComputeChannel(socket, aesKey, aesIv);
 
-			await serverFunc(channel, cancellationToken);
+			await using ComputeLease lease = new ComputeLease(socket, key.Span, nonce.Span, new Dictionary<string, int>());
+			await serverFunc(lease, cancellationToken);
 		}
 
 		/// <inheritdoc/>
-		public async Task<TResult> ExecuteAsync<TResult>(ClusterId clusterId, Requirements? requirements, Func<IComputeChannel, CancellationToken, Task<TResult>> handler, CancellationToken cancellationToken)
+		public async Task<TResult> ExecuteAsync<TResult>(ClusterId clusterId, Requirements? requirements, Func<IComputeLease, CancellationToken, Task<TResult>> handler, CancellationToken cancellationToken)
 		{
-			return await handler(_inner, cancellationToken);
+			await using ComputeLease lease = new ComputeLease(_socket, _key, _nonce, new Dictionary<string, int>());
+			return await handler(lease, cancellationToken);
 		}
 	}
 }
diff --git a/Engine/Source/Programs/Shared/EpicGames.Horde/Compute/ServerComputeClient.cs b/Engine/Source/Programs/Shared/EpicGames.Horde/Compute/ServerComputeClient.cs
--- a/Engine/Source/Programs/Shared/EpicGames.Horde/Compute/ServerComputeClient.cs
+++ b/Engine/Source/Programs/Shared/EpicGames.Horde/Compute/ServerComputeClient.cs
@@ -3,6 +3,7 @@
 using EpicGames.Core;
 using Microsoft.Extensions.Logging;
 using System;
+using System.Collections.Generic;
 using System.Net;
 using System.Net.Http;
 using System.Net.Http.Json;
@@ -32,8 +33,8 @@ namespace EpicGames.Horde.Compute
 			public string Ip { get; set; } = String.Empty;
 			public int Port { get; set; }
 			public string Nonce { get; set; } = String.Empty;
-			public string AesKey { get; set; } = String.Empty;
-			public string AesIv { get; set; } = String.Empty;
+			public string Key { get; set; } = String.Empty;
+			public Dictionary<string, int> AssignedResources { get; set; } = new Dictionary<string, int>();
 		}
 
 		readonly HttpClient? _defaultHttpClient;
@@ -93,7 +94,7 @@ namespace EpicGames.Horde.Compute
 		}
 
 		/// <inheritdoc/>
-		public async Task<TResult> ExecuteAsync<TResult>(ClusterId clusterId, Requirements? requirements, Func<IComputeChannel, CancellationToken, Task<TResult>> handler, CancellationToken cancellationToken)
+		public async Task<TResult> ExecuteAsync<TResult>(ClusterId clusterId, Requirements? requirements, Func<IComputeLease, CancellationToken, Task<TResult>> handler, CancellationToken cancellationToken)
 		{
 			// Assign a compute worker
 			HttpClient client = _createHttpClient();
@@ -118,14 +119,13 @@ namespace EpicGames.Horde.Compute
 
 			// Send the nonce
 			byte[] nonce = StringUtils.ParseHexString(responseMessage.Nonce);
-			await socket.SendFullAsync(nonce, SocketFlags.None, cancellationToken);
+			await socket.SendMessageAsync(nonce, SocketFlags.None, cancellationToken);
 			_logger.LogInformation("Connected to {Ip} with nonce {Nonce}", responseMessage.Ip, responseMessage.Nonce);
 
 			// Pass the rest of the call over to the handler
-			byte[] aesKey = StringUtils.ParseHexString(responseMessage.AesKey);
-			byte[] aesIv = StringUtils.ParseHexString(responseMessage.AesIv);
+			byte[] key = StringUtils.ParseHexString(responseMessage.Key);
 
-			using SocketComputeChannel channel = new SocketComputeChannel(socket, aesKey, aesIv);
+			await using ComputeLease channel = new ComputeLease(socket, key, nonce, responseMessage.AssignedResources);
 			return await handler(channel, cancellationToken);
 		}
 	}
diff --git a/Engine/Source/Programs/Shared/EpicGames.Horde/Compute/SocketComputeChannel.cs b/Engine/Source/Programs/Shared/EpicGames.Horde/Compute/SocketComputeChannel.cs
deleted file mode 100644
diff --git a/Engine/Source/Programs/Shared/EpicGames.Serialization/CbWriter.cs b/Engine/Source/Programs/Shared/EpicGames.Serialization/CbWriter.cs
--- a/Engine/Source/Programs/Shared/EpicGames.Serialization/CbWriter.cs
+++ b/Engine/Source/Programs/Shared/EpicGames.Serialization/CbWriter.cs
@@ -6,6 +6,7 @@ using System.Buffers.Binary;
 using System.Collections.Generic;
 using System.Diagnostics;
 using System.IO;
+using System.Linq;
 
 namespace EpicGames.Serialization
 {
@@ -689,6 +690,19 @@ namespace EpicGames.Serialization
 			_chunks.Clear();
 		}
 
+		/// <summary>
+		/// Ensures that the required space is available in a contiguous chunk
+		/// </summary>
+		/// <param name="reserve">Minimum required space</param>
+		public void Reserve(int reserve)
+		{
+			int allocate = Math.Max(DefaultChunkSize, reserve);
+			if (!_freeChunks.Any(x => x.Length >= reserve))
+			{
+				_freeChunks.Add(new byte[allocate + 4096]);
+			}
+		}
+
 		/// <inheritdoc/>
 		protected override Memory<byte> AllocateChunk(int minSize)
 		{
