Horde: Fix rare data corruption when streaming data into bundles.

If data is being appended to an existing file node in chunks, and the bundle size exceeds the desired max while data is still in the buffer, a new bundle (and write buffer) will be created without preserving the partially written data. Resolution is to only flush bundles when complete nodes have been added.

diff --git a/Engine/Source/Programs/Shared/EpicGames.Horde.Tests/BundleStoreTests.cs b/Engine/Source/Programs/Shared/EpicGames.Horde.Tests/BundleStoreTests.cs
--- a/Engine/Source/Programs/Shared/EpicGames.Horde.Tests/BundleStoreTests.cs
+++ b/Engine/Source/Programs/Shared/EpicGames.Horde.Tests/BundleStoreTests.cs
@@ -231,6 +231,51 @@ namespace EpicGames.Horde.Tests
 			Assert.IsTrue(worldData.SequenceEqual(Encoding.UTF8.GetBytes("world")));
 		}
 
+		[TestMethod]
+		public async Task StreamTestAsync()
+		{
+			using IMemoryCache cache = new MemoryCache(new MemoryCacheOptions());
+			MemoryStorageClient store = new MemoryStorageClient();
+
+			const int length = 4096;
+
+			byte[] chunk = new byte[length];
+			new Random(0).NextBytes(chunk);
+
+			// Generate a tree
+			NodeLocator locator;
+			{
+				using TreeWriter writer = new TreeWriter(store, new TreeOptions { MaxBlobSize = 1024 });
+
+				ChunkingOptions options = new ChunkingOptions();
+				options.LeafOptions.MinSize = 128;
+				options.LeafOptions.TargetSize = 256;
+				options.LeafOptions.MaxSize = 64 * 1024;
+
+				FileNodeWriter fileWriter = new FileNodeWriter(writer, options);
+				for (int idx = 0; idx < chunk.Length / 16; idx++)
+				{
+					await fileWriter.AppendAsync(chunk.AsMemory(idx * 16, 16), CancellationToken.None);
+				}
+
+				NodeHandle handle = await fileWriter.FlushAsync(CancellationToken.None);
+				locator = handle.Locator;
+			}
+
+			// Check we can read it back in
+			{
+				TreeReader reader = new TreeReader(store, cache, NullLogger.Instance);
+
+				FileNode newRoot = await reader.ReadNodeAsync<FileNode>(locator);
+
+				using MemoryStream stream = new MemoryStream();
+				await newRoot.CopyToStreamAsync(reader, stream, CancellationToken.None);
+
+				byte[] output = stream.ToArray();
+				Assert.IsTrue(chunk.SequenceEqual(output));
+			}
+		}
+
 		[TestMethod]
 		public async Task LargeFileTestAsync()
 		{
diff --git a/Engine/Source/Programs/Shared/EpicGames.Horde/Storage/TreeWriter.cs b/Engine/Source/Programs/Shared/EpicGames.Horde/Storage/TreeWriter.cs
--- a/Engine/Source/Programs/Shared/EpicGames.Horde/Storage/TreeWriter.cs
+++ b/Engine/Source/Programs/Shared/EpicGames.Horde/Storage/TreeWriter.cs
@@ -701,13 +701,6 @@ namespace EpicGames.Horde.Storage
 		/// <returns>Buffer to be written into.</returns>
 		public Memory<byte> GetOutputBuffer(int usedSize, int desiredSize)
 		{
-			// If the bundle is full, start the process of writing it to disk
-			if (_currentBundle != null && _currentBundle.UncompressedLength > 0 && _currentBundle.UncompressedLength + desiredSize > _options.MaxBlobSize)
-			{
-				Complete();
-			}
-
-			// If we don't yet have a bundle in this writer, create one
 			PendingBundle currentBundle = GetCurrentBundle();
 			return currentBundle.GetBuffer(usedSize, desiredSize);
 		}
