Fix coding standard issues in some RigVM code.

#rnx

diff --git a/Engine/Source/Programs/Shared/EpicGames.UHT/Exporters/CodeGen/UhtHeaderCodeGeneratorCppFile.cs b/Engine/Source/Programs/Shared/EpicGames.UHT/Exporters/CodeGen/UhtHeaderCodeGeneratorCppFile.cs
--- a/Engine/Source/Programs/Shared/EpicGames.UHT/Exporters/CodeGen/UhtHeaderCodeGeneratorCppFile.cs
+++ b/Engine/Source/Programs/Shared/EpicGames.UHT/Exporters/CodeGen/UhtHeaderCodeGeneratorCppFile.cs
@@ -707,7 +707,7 @@ namespace EpicGames.UHT.Exporters.CodeGen
 			if (scriptStruct.RigVMStructInfo != null)
 			{
 				string constPrefix = "";
-				if (scriptStruct.RigVMStructInfo.ExecuteContextMember == String.Empty)
+				if (String.IsNullOrEmpty(scriptStruct.RigVMStructInfo.ExecuteContextMember))
 				{
 					constPrefix = "const ";
 				}
@@ -724,7 +724,7 @@ namespace EpicGames.UHT.Exporters.CodeGen
 						.Append(methodInfo.Name)
 						.Append("()\r\n");
 					builder.Append("{\r\n");
-					if (scriptStruct.RigVMStructInfo.ExecuteContextMember == String.Empty)
+					if (String.IsNullOrEmpty(scriptStruct.RigVMStructInfo.ExecuteContextMember))
 					{
 						builder.Append("\t").Append(scriptStruct.RigVMStructInfo.ExecuteContextType).Append(" TemporaryExecuteContext;\r\n");
 					}
diff --git a/Engine/Source/Programs/Shared/EpicGames.UHT/Exporters/CodeGen/UhtHeaderCodeGeneratorHFile.cs b/Engine/Source/Programs/Shared/EpicGames.UHT/Exporters/CodeGen/UhtHeaderCodeGeneratorHFile.cs
--- a/Engine/Source/Programs/Shared/EpicGames.UHT/Exporters/CodeGen/UhtHeaderCodeGeneratorHFile.cs
+++ b/Engine/Source/Programs/Shared/EpicGames.UHT/Exporters/CodeGen/UhtHeaderCodeGeneratorHFile.cs
@@ -146,7 +146,7 @@ namespace EpicGames.UHT.Exporters.CodeGen
 					builder.Append("#define ").Append(scriptStruct.SourceName).Append('_').Append(methodInfo.Name).Append("() \\\r\n");
 					builder.Append('\t').Append(methodInfo.ReturnType).Append(' ').Append(scriptStruct.SourceName).Append("::Static").Append(methodInfo.Name).Append("( \\\r\n");
 					builder.Append("\t\t");
-					if (scriptStruct.RigVMStructInfo.ExecuteContextMember == String.Empty)
+					if (String.IsNullOrEmpty(scriptStruct.RigVMStructInfo.ExecuteContextMember))
 					{
 						builder.Append("const ");
 					}
@@ -175,7 +175,7 @@ namespace EpicGames.UHT.Exporters.CodeGen
 					if (scriptStruct.RigVMStructInfo != null)
 					{
 						string constPrefix = "";
-						if (scriptStruct.RigVMStructInfo.ExecuteContextMember == String.Empty)
+						if (String.IsNullOrEmpty(scriptStruct.RigVMStructInfo.ExecuteContextMember))
 						{
 							constPrefix = "const ";
 						}
@@ -185,7 +185,7 @@ namespace EpicGames.UHT.Exporters.CodeGen
 							builder.Append('\t').Append(methodInfo.ReturnType).Append(' ').Append(methodInfo.Name).Append('(').Append(constPrefix).Append(scriptStruct.RigVMStructInfo.ExecuteContextType).Append("& InExecuteContext); \\\r\n");
 							builder.Append("\tstatic ").Append(methodInfo.ReturnType).Append(" Static").Append(methodInfo.Name).Append("( \\\r\n");
 							builder.Append("\t\t");
-							if (scriptStruct.RigVMStructInfo.ExecuteContextMember == String.Empty)
+							if (String.IsNullOrEmpty(scriptStruct.RigVMStructInfo.ExecuteContextMember))
 							{
 								builder.Append("const ");
 							}
diff --git a/Engine/Source/Programs/Shared/EpicGames.UHT/Types/UhtScriptStruct.cs b/Engine/Source/Programs/Shared/EpicGames.UHT/Types/UhtScriptStruct.cs
--- a/Engine/Source/Programs/Shared/EpicGames.UHT/Types/UhtScriptStruct.cs
+++ b/Engine/Source/Programs/Shared/EpicGames.UHT/Types/UhtScriptStruct.cs
@@ -923,7 +923,7 @@ namespace EpicGames.UHT.Types
 		{
 			if (RigVMStructInfo != null)
 			{
-				if (this.MetaData.TryGetValueHierarchical("ExecuteContext", out string? ExecuteContextMetadata))
+				if (MetaData.TryGetValueHierarchical("ExecuteContext", out string? ExecuteContextMetadata))
 				{
 					RigVMStructInfo.ExecuteContextType = ExecuteContextMetadata;
 				}
@@ -935,7 +935,7 @@ namespace EpicGames.UHT.Types
 						UhtRigVMParameter parameter = new UhtRigVMParameter(property, RigVMStructInfo.Members.Count);
 						if (parameter.IsExecuteContext())
 						{
-							if (RigVMStructInfo.ExecuteContextMember == String.Empty)
+							if (String.IsNullOrEmpty(RigVMStructInfo.ExecuteContextMember))
 							{
 								RigVMStructInfo.ExecuteContextMember = parameter.Name;
 							}
