Horde: Add an error return path for C++ compute messages.

diff --git a/Engine/Source/Programs/Horde/Horde.Agent/Commands/Compute/CppComputeCommand.cs b/Engine/Source/Programs/Horde/Horde.Agent/Commands/Compute/CppComputeCommand.cs
--- a/Engine/Source/Programs/Horde/Horde.Agent/Commands/Compute/CppComputeCommand.cs
+++ b/Engine/Source/Programs/Horde/Horde.Agent/Commands/Compute/CppComputeCommand.cs
@@ -94,9 +94,9 @@ namespace Horde.Agent.Commands
 							await channel.CppBlobDataAsync(cppBlobRead, storage, cancellationToken);
 						}
 						break;
-					case ComputeMessageType.CppResult:
+					case ComputeMessageType.CppSuccess:
 						{
-							NodeLocator cppOutputLocator = message.AsCppResult();
+							NodeLocator cppOutputLocator = message.AsCppSuccess();
 
 							using ComputeStorageClient remoteStorage = new ComputeStorageClient(channel);
 							TreeReader reader = new TreeReader(remoteStorage, null, _logger);
@@ -118,6 +118,12 @@ namespace Horde.Agent.Commands
 							await lease.DefaultChannel.CloseAsync(cancellationToken);
 						}
 						return false;
+					case ComputeMessageType.CppFailure:
+						{
+							string error = message.AsCppFailure();
+							_logger.LogError("Failed: {Error}", error);
+						}
+						return false;
 					default:
 						throw new NotImplementedException();
 				}
diff --git a/Engine/Source/Programs/Horde/Horde.Agent/Leases/Handlers/ComputeHandler.cs b/Engine/Source/Programs/Horde/Horde.Agent/Leases/Handlers/ComputeHandler.cs
--- a/Engine/Source/Programs/Horde/Horde.Agent/Leases/Handlers/ComputeHandler.cs
+++ b/Engine/Source/Programs/Horde/Horde.Agent/Leases/Handlers/ComputeHandler.cs
@@ -121,6 +121,18 @@ namespace Horde.Agent.Leases.Handlers
 		}
 
 		public async Task RunCppAsync(IComputeChannel channel, NodeLocator locator, CancellationToken cancellationToken)
+		{
+			try
+			{
+				await RunCppInternalAsync(channel, locator, cancellationToken);
+			}
+			catch (Exception ex)
+			{
+				await channel.CppFailureAsync(ex.ToString(), cancellationToken);
+			}
+		}
+
+		async Task RunCppInternalAsync(IComputeChannel channel, NodeLocator locator, CancellationToken cancellationToken)
 		{
 			DirectoryReference sandboxDir = DirectoryReference.Combine(Program.DataDir, "Sandbox", channel.Id.ToString());
 
@@ -182,7 +194,7 @@ namespace Horde.Agent.Leases.Handlers
 				CppComputeOutputNode outputNode = new CppComputeOutputNode(exitCode, logNodeRef, new TreeNodeRef<DirectoryNode>(outputTree));
 				NodeHandle outputHandle = await writer.FlushAsync(outputNode, cancellationToken);
 
-				await channel.CppResultAsync(outputHandle.Locator, cancellationToken);
+				await channel.CppSuccessAsync(outputHandle.Locator, cancellationToken);
 			}
 
 			for (; ; )
diff --git a/Engine/Source/Programs/Shared/EpicGames.Horde/Compute/IComputeMessage.cs b/Engine/Source/Programs/Shared/EpicGames.Horde/Compute/IComputeMessage.cs
--- a/Engine/Source/Programs/Shared/EpicGames.Horde/Compute/IComputeMessage.cs
+++ b/Engine/Source/Programs/Shared/EpicGames.Horde/Compute/IComputeMessage.cs
@@ -49,22 +49,27 @@ namespace EpicGames.Horde.Compute
 		/// <summary>
 		/// Gives output from a C++ build
 		/// </summary>
-		CppResult = 0x21,
+		CppSuccess = 0x21,
+
+		/// <summary>
+		/// Information about a failed C++ task
+		/// </summary>
+		CppFailure = 0x22,
 
 		/// <summary>
 		/// Finish a C++ build and close the current channel
 		/// </summary>
-		CppEnd = 0x22,
+		CppEnd = 0x23,
 
 		/// <summary>
 		/// Reads a blob pertaining to a C++ build
 		/// </summary>
-		CppBlobRead = 0x23,
+		CppBlobRead = 0x24,
 
 		/// <summary>
 		/// Response to a <see cref="CppBlobRead"/> request.
 		/// </summary>
-		CppBlobData = 0x24,
+		CppBlobData = 0x25,
 
 		#endregion
 	}
@@ -190,7 +195,7 @@ namespace EpicGames.Horde.Compute
 		/// Sends the result for a C++ compute task
 		/// </summary>
 		/// <param name="message"></param>
-		public static NodeLocator AsCppResult(this IComputeMessage message)
+		public static NodeLocator AsCppSuccess(this IComputeMessage message)
 		{
 			MemoryReader reader = new MemoryReader(message.Data);
 			return reader.ReadNodeLocator();
@@ -202,15 +207,39 @@ namespace EpicGames.Horde.Compute
 		/// <param name="channel">Channel to write to</param>
 		/// <param name="locator">Locator for the result</param>
 		/// <param name="cancellationToken">Cancellation token for the operation</param>
-		public static async Task CppResultAsync(this IComputeChannel channel, NodeLocator locator, CancellationToken cancellationToken)
+		public static async Task CppSuccessAsync(this IComputeChannel channel, NodeLocator locator, CancellationToken cancellationToken)
 		{
-			using (IComputeMessageWriter writer = channel.CreateMessage(ComputeMessageType.CppResult))
+			using (IComputeMessageWriter writer = channel.CreateMessage(ComputeMessageType.CppSuccess))
 			{
 				writer.WriteNodeLocator(locator);
 				await writer.SendAsync(cancellationToken);
 			}
 		}
 
+		/// <summary>
+		/// Sends the result for a C++ compute task
+		/// </summary>
+		/// <param name="message"></param>
+		public static string AsCppFailure(this IComputeMessage message)
+		{
+			return new Utf8String(message.Data).ToString();
+		}
+
+		/// <summary>
+		/// Sends the failure error for a C++ compute task
+		/// </summary>
+		/// <param name="channel">Channel to write to</param>
+		/// <param name="error">Error string to send to the caller</param>
+		/// <param name="cancellationToken">Cancellation token for the operation</param>
+		public static async Task CppFailureAsync(this IComputeChannel channel, string error, CancellationToken cancellationToken)
+		{
+			using (IComputeMessageWriter writer = channel.CreateMessage(ComputeMessageType.CppFailure))
+			{
+				writer.WriteString(error);
+				await writer.SendAsync(cancellationToken);
+			}
+		}
+
 		/// <summary>
 		/// Completes a C++ build task
 		/// </summary>
