Horde: Fix decoding of compute errors.

diff --git a/Engine/Source/Programs/Shared/EpicGames.Horde/Compute/IComputeMessage.cs b/Engine/Source/Programs/Shared/EpicGames.Horde/Compute/IComputeMessage.cs
--- a/Engine/Source/Programs/Shared/EpicGames.Horde/Compute/IComputeMessage.cs
+++ b/Engine/Source/Programs/Shared/EpicGames.Horde/Compute/IComputeMessage.cs
@@ -222,7 +222,8 @@ namespace EpicGames.Horde.Compute
 		/// <param name="message"></param>
 		public static string AsCppFailure(this IComputeMessage message)
 		{
-			return new Utf8String(message.Data).ToString();
+			MemoryReader reader = new MemoryReader(message.Data);
+			return reader.ReadString();
 		}
 
 		/// <summary>
