Horde: Additional logging for running of compute commands.

diff --git a/Engine/Source/Programs/Shared/EpicGames.Horde/Compute/ServerComputeClient.cs b/Engine/Source/Programs/Shared/EpicGames.Horde/Compute/ServerComputeClient.cs
--- a/Engine/Source/Programs/Shared/EpicGames.Horde/Compute/ServerComputeClient.cs
+++ b/Engine/Source/Programs/Shared/EpicGames.Horde/Compute/ServerComputeClient.cs
@@ -96,6 +96,8 @@ namespace EpicGames.Horde.Compute
 		/// <inheritdoc/>
 		public async Task<TResult> ExecuteAsync<TResult>(ClusterId clusterId, Requirements? requirements, Func<IComputeLease, CancellationToken, Task<TResult>> handler, CancellationToken cancellationToken)
 		{
+			_logger.LogInformation("Requesting compute resource");
+
 			// Assign a compute worker
 			HttpClient client = _createHttpClient();
 
@@ -113,6 +115,8 @@ namespace EpicGames.Horde.Compute
 				}
 			}
 
+			_logger.LogInformation("Connecting to {Ip} with nonce {Nonce}...", responseMessage.Ip, responseMessage.Nonce);
+
 			// Connect to the remote machine
 			using Socket socket = new Socket(SocketType.Stream, ProtocolType.Tcp);
 			await socket.ConnectAsync(IPAddress.Parse(responseMessage.Ip), responseMessage.Port, cancellationToken);
@@ -120,7 +124,7 @@ namespace EpicGames.Horde.Compute
 			// Send the nonce
 			byte[] nonce = StringUtils.ParseHexString(responseMessage.Nonce);
 			await socket.SendMessageAsync(nonce, SocketFlags.None, cancellationToken);
-			_logger.LogInformation("Connected to {Ip} with nonce {Nonce}", responseMessage.Ip, responseMessage.Nonce);
+			_logger.LogInformation("Connection established.");
 
 			// Pass the rest of the call over to the handler
 			byte[] key = StringUtils.ParseHexString(responseMessage.Key);
