Horde: Allow overwriting read-only files when extracting files from temp storage.

diff --git a/Engine/Source/Programs/Shared/EpicGames.Horde/Storage/Nodes/FileNode.cs b/Engine/Source/Programs/Shared/EpicGames.Horde/Storage/Nodes/FileNode.cs
--- a/Engine/Source/Programs/Shared/EpicGames.Horde/Storage/Nodes/FileNode.cs
+++ b/Engine/Source/Programs/Shared/EpicGames.Horde/Storage/Nodes/FileNode.cs
@@ -108,6 +108,10 @@ namespace EpicGames.Horde.Storage.Nodes
 		/// <returns></returns>
 		public async Task CopyToFileAsync(TreeReader reader, FileInfo file, CancellationToken cancellationToken)
 		{
+			if(file.Exists && (file.Attributes & FileAttributes.ReadOnly) != 0)
+			{
+				file.Attributes &= ~FileAttributes.ReadOnly;
+			}
 			using (FileStream stream = file.OpenWrite())
 			{
 				await CopyToStreamAsync(reader, stream, cancellationToken);
