Horde: Fix error when reading 0-byte nodes from bundles with no packets.

diff --git a/Engine/Source/Programs/Shared/EpicGames.Horde/Storage/TreeReader.cs b/Engine/Source/Programs/Shared/EpicGames.Horde/Storage/TreeReader.cs
--- a/Engine/Source/Programs/Shared/EpicGames.Horde/Storage/TreeReader.cs
+++ b/Engine/Source/Programs/Shared/EpicGames.Horde/Storage/TreeReader.cs
@@ -633,6 +633,11 @@ namespace EpicGames.Horde.Storage
 		/// <returns>The decoded data</returns>
 		async ValueTask<ReadOnlyMemory<byte>> ReadBundlePacketAsync(BundleInfo bundleInfo, int packetIdx, CancellationToken cancellationToken)
 		{
+			if (packetIdx < 0 || packetIdx >= bundleInfo.Header.Packets.Count)
+			{
+				throw new ArgumentException("Packet index is out of range", nameof(packetIdx));
+			}
+
 			string decodedCacheKey = GetDecodedPacketCacheKey(bundleInfo.Locator.BlobId, packetIdx);
 			if (_cache != null && _cache.TryGetValue(decodedCacheKey, out ReadOnlyMemory<byte> decodedPacket))
 			{
@@ -692,6 +697,11 @@ namespace EpicGames.Horde.Storage
 		/// <returns>The encoded packet data</returns>
 		async ValueTask<ReadOnlyMemory<byte>> ReadEncodedPacketAsync(BundleInfo bundleInfo, int packetIdx)
 		{
+			if (packetIdx < 0 || packetIdx >= bundleInfo.Header.Packets.Count)
+			{
+				throw new ArgumentException("Packet index is out of range", nameof(packetIdx));
+			}
+
 			string encodedCacheKey = GetEncodedPacketCacheKey(bundleInfo.Locator.BlobId, packetIdx);
 			if (_cache != null && _cache.TryGetValue(encodedCacheKey, out ReadOnlyMemory<byte> encodedPacket))
 			{
@@ -733,7 +743,6 @@ namespace EpicGames.Horde.Storage
 			BundleExport export = bundleInfo.Header.Exports[locator.ExportIdx];
 
 			ExportInfo exportInfo = bundleInfo.Exports[locator.ExportIdx];
-			ReadOnlyMemory<byte> packetData = await ReadBundlePacketAsync(bundleInfo, exportInfo.PacketIdx, cancellationToken);
 
 			List<NodeLocator> refs = new List<NodeLocator>(export.References.Count);
 			for (int idx = 0; idx < export.References.Count; idx++)
@@ -742,7 +751,12 @@ namespace EpicGames.Horde.Storage
 				refs.Add(reference);
 			}
 
-			ReadOnlyMemory<byte> nodeData = packetData.Slice(exportInfo.Offset, export.Length);
+			ReadOnlyMemory<byte> nodeData = ReadOnlyMemory<byte>.Empty;
+			if (export.Length > 0)
+			{
+				ReadOnlyMemory<byte> packetData = await ReadBundlePacketAsync(bundleInfo, exportInfo.PacketIdx, cancellationToken);
+				nodeData = packetData.Slice(exportInfo.Offset, export.Length);
+			}
 
 			TypeInfo? typeInfo = bundleInfo.Types[export.TypeIdx];
 			if (typeInfo == null)
