Horde: Add support for creating a zip archive from a DirectoryNode tree on-the-fly.

diff --git a/Engine/Source/Programs/Shared/EpicGames.Horde/Storage/Nodes/DirectoryNode.cs b/Engine/Source/Programs/Shared/EpicGames.Horde/Storage/Nodes/DirectoryNode.cs
--- a/Engine/Source/Programs/Shared/EpicGames.Horde/Storage/Nodes/DirectoryNode.cs
+++ b/Engine/Source/Programs/Shared/EpicGames.Horde/Storage/Nodes/DirectoryNode.cs
@@ -3,10 +3,13 @@
 using EpicGames.Core;
 using Microsoft.Extensions.Logging;
 using System;
+using System.Buffers;
 using System.Collections.Generic;
 using System.Diagnostics;
 using System.Diagnostics.CodeAnalysis;
 using System.IO;
+using System.IO.Compression;
+using System.IO.Pipelines;
 using System.Linq;
 using System.Threading;
 using System.Threading.Tasks;
@@ -923,5 +926,131 @@ namespace EpicGames.Horde.Storage.Nodes
 
 			await Task.WhenAll(tasks);
 		}
+
+		/// <summary>
+		/// Returns a stream containing the zipped contents of this directory
+		/// </summary>
+		/// <param name="reader">Reader for other nodes</param>
+		/// <returns>Stream containing zipped archive data</returns>
+		public Stream GetZipStream(TreeReader reader) => new ZippedDirectoryStream(reader, this);
+	}
+
+	/// <summary>
+	/// Stream which zips a directory node tree dynamically
+	/// </summary>
+	class ZippedDirectoryStream : Stream
+	{
+		/// <inheritdoc/>
+		public override bool CanRead => true;
+
+		/// <inheritdoc/>
+		public override bool CanSeek => false;
+
+		/// <inheritdoc/>
+		public override bool CanWrite => false;
+
+		/// <inheritdoc/>
+		public override long Length => throw new NotImplementedException();
+
+		/// <inheritdoc/>
+		public override long Position { get => _position; set => throw new NotImplementedException(); }
+
+		readonly Pipe _pipe;
+		readonly BackgroundTask _backgroundTask;
+
+		long _position;
+		ReadOnlySequence<byte> _current = ReadOnlySequence<byte>.Empty;
+
+		/// <summary>
+		/// Constructor
+		/// </summary>
+		/// <param name="reader">Reader to read nodes through</param>
+		/// <param name="node">Root node to copy from</param>
+		public ZippedDirectoryStream(TreeReader reader, DirectoryNode node)
+		{
+			_pipe = new Pipe();
+			_backgroundTask = BackgroundTask.StartNew(ctx => CopyToPipeAsync(reader, node, _pipe.Writer, ctx));
+		}
+
+		/// <inheritdoc/>
+		public override async ValueTask DisposeAsync()
+		{
+			await base.DisposeAsync();
+
+			await _backgroundTask.DisposeAsync();
+		}
+
+		/// <inheritdoc/>
+		public override async ValueTask<int> ReadAsync(Memory<byte> buffer, CancellationToken cancellationToken = default)
+		{
+			while (_current.Length == 0)
+			{
+				ReadResult result = await _pipe.Reader.ReadAsync(cancellationToken);
+				if (result.IsCompleted)
+				{
+					return 0;
+				}
+				_current = result.Buffer;
+			}
+
+			int initialSize = buffer.Length;
+			while (buffer.Length > 0 && _current.Length > 0)
+			{
+				int copy = Math.Min(buffer.Length, _current.First.Length);
+				_current.First.Slice(0, copy).CopyTo(buffer);
+				_current = _current.Slice(copy);
+				buffer = buffer.Slice(copy);
+			}
+
+			if (_current.Length == 0)
+			{
+				_pipe.Reader.AdvanceTo(_current.End);
+			}
+
+			int length = initialSize - buffer.Length;
+			_position += length;
+			return length;
+		}
+
+		static async Task CopyToPipeAsync(TreeReader reader, DirectoryNode node, PipeWriter writer, CancellationToken cancellationToken)
+		{
+			using Stream outputStream = writer.AsStream();
+			using ZipArchive archive = new ZipArchive(outputStream, ZipArchiveMode.Create);
+			await CopyFilesAsync(reader, node, "", archive, cancellationToken);
+		}
+
+		static async Task CopyFilesAsync(TreeReader reader, DirectoryNode directory, string prefix, ZipArchive archive, CancellationToken cancellationToken)
+		{
+			foreach (DirectoryEntry directoryEntry in directory.Directories)
+			{
+				DirectoryNode node = await directoryEntry.ExpandAsync(reader, cancellationToken);
+				await CopyFilesAsync(reader, node, $"{prefix}{directoryEntry.Name}/", archive, cancellationToken);
+			}
+
+			foreach (FileEntry fileEntry in directory.Files)
+			{
+				ZipArchiveEntry entry = archive.CreateEntry($"{prefix}{fileEntry}");
+
+				using Stream entryStream = entry.Open();
+				await fileEntry.CopyToStreamAsync(reader, entryStream, cancellationToken);
+			}
+		}
+
+		/// <inheritdoc/>
+		public override void Flush()
+		{
+		}
+
+		/// <inheritdoc/>
+		public override int Read(byte[] buffer, int offset, int count) => ReadAsync(buffer.AsMemory(offset, count)).AsTask().Result;
+
+		/// <inheritdoc/>
+		public override long Seek(long offset, SeekOrigin origin) => throw new NotSupportedException();
+
+		/// <inheritdoc/>
+		public override void SetLength(long value) => throw new NotSupportedException();
+
+		/// <inheritdoc/>
+		public override void Write(byte[] buffer, int offset, int count) => throw new NotSupportedException();
 	}
 }
