Horde: Fix bundles being flushed unnecessarily often.

diff --git a/Engine/Source/Programs/Shared/EpicGames.Horde/Storage/Nodes/DirectoryNode.cs b/Engine/Source/Programs/Shared/EpicGames.Horde/Storage/Nodes/DirectoryNode.cs
--- a/Engine/Source/Programs/Shared/EpicGames.Horde/Storage/Nodes/DirectoryNode.cs
+++ b/Engine/Source/Programs/Shared/EpicGames.Horde/Storage/Nodes/DirectoryNode.cs
@@ -885,17 +885,35 @@ namespace EpicGames.Horde.Storage.Nodes
 
 		static async Task CopyFilesAsync(List<(DirectoryNode DirectoryNode, FileInfo FileInfo)> files, int minIdx, int maxIdx, FileEntry[] entries, ChunkingOptions options, TreeWriter baseWriter, CopyStats? copyStats, CancellationToken cancellationToken)
 		{
-			using TreeWriter writer = new TreeWriter(baseWriter);
+			TreeWriter writer = baseWriter;
+			try
+			{
+				if (minIdx != 0)
+				{
+					writer = new TreeWriter(baseWriter);
+				}
+
+				FileNodeWriter fileNodeWriter = new FileNodeWriter(writer, options);
+				for (int idx = minIdx; idx < maxIdx; idx++)
+				{
+					FileInfo fileInfo = files[idx].FileInfo;
+					NodeHandle handle = await fileNodeWriter.CreateAsync(fileInfo, cancellationToken);
+					entries[idx] = new FileEntry(fileInfo.Name, FileEntryFlags.None, fileNodeWriter.Length, handle);
+					copyStats?.Update(1, fileNodeWriter.Length);
+				}
 
-			FileNodeWriter fileNodeWriter = new FileNodeWriter(writer, options);
-			for(int idx = minIdx; idx < maxIdx; idx++)
+				if (minIdx != 0)
+				{
+					await writer.FlushAsync(cancellationToken);
+				}
+			}
+			finally
 			{
-				FileInfo fileInfo = files[idx].FileInfo;
-				NodeHandle handle = await fileNodeWriter.CreateAsync(fileInfo, cancellationToken);
-				entries[idx] = new FileEntry(fileInfo.Name, FileEntryFlags.None, fileNodeWriter.Length, handle);
-				copyStats?.Update(1, fileNodeWriter.Length);
+				if (minIdx != 0)
+				{
+					writer.Dispose();
+				}
 			}
-			await writer.FlushAsync(cancellationToken);
 		}
 
 		/// <summary>
diff --git a/Engine/Source/Programs/Shared/EpicGames.Horde/Storage/Nodes/FileNode.cs b/Engine/Source/Programs/Shared/EpicGames.Horde/Storage/Nodes/FileNode.cs
--- a/Engine/Source/Programs/Shared/EpicGames.Horde/Storage/Nodes/FileNode.cs
+++ b/Engine/Source/Programs/Shared/EpicGames.Horde/Storage/Nodes/FileNode.cs
@@ -490,7 +490,7 @@ namespace EpicGames.Horde.Storage.Nodes
 		{
 			Reset();
 			await AppendAsync(stream, cancellationToken);
-			return await FlushAsync(cancellationToken);
+			return await CompleteAsync(cancellationToken);
 		}
 
 		/// <summary>
