Horde: Add support for copying directory nodes from zip streams.

diff --git a/Engine/Source/Programs/Shared/EpicGames.Horde/Storage/Nodes/DirectoryNode.cs b/Engine/Source/Programs/Shared/EpicGames.Horde/Storage/Nodes/DirectoryNode.cs
--- a/Engine/Source/Programs/Shared/EpicGames.Horde/Storage/Nodes/DirectoryNode.cs
+++ b/Engine/Source/Programs/Shared/EpicGames.Horde/Storage/Nodes/DirectoryNode.cs
@@ -792,6 +792,47 @@ namespace EpicGames.Horde.Storage.Nodes
 			return node;
 		}
 
+		/// <summary>
+		/// Copies entries from a zip file
+		/// </summary>
+		/// <param name="stream">Input stream</param>
+		/// <param name="writer">Writer for new nodes</param>
+		/// <param name="cancellationToken">Cancellation token for the operation</param>
+		public async Task CopyFromZipStreamAsync(Stream stream, TreeWriter writer, CancellationToken cancellationToken)
+		{
+			FileNodeWriter fileWriter = new FileNodeWriter(writer, new ChunkingOptions());
+			using (ZipArchive archive = new ZipArchive(stream, ZipArchiveMode.Read, true))
+			{
+				foreach (ZipArchiveEntry entry in archive.Entries)
+				{
+					if (entry.Name.Length > 0)
+					{
+						NodeHandle node;
+						using (Stream entryStream = entry.Open())
+						{
+							node = await fileWriter.CreateAsync(entryStream, cancellationToken);
+						}
+
+						string fullName = entry.FullName;
+						int slashIdx = fullName.LastIndexOf('/');
+
+						DirectoryNode directory;
+						if (slashIdx == -1)
+						{
+							directory = this;
+						}
+						else
+						{
+							directory = await FindOrAddDirectoryAsync(null!, fullName.Substring(0, slashIdx), cancellationToken);
+						}
+
+						FileEntry fileEntry = new FileEntry(entry.Name, FileEntryFlags.None, fileWriter.Length, node);
+						directory.AddFile(fileEntry);
+					}
+				}
+			}
+		}
+
 		/// <summary>
 		/// Adds files from a directory on disk
 		/// </summary>
@@ -800,7 +841,6 @@ namespace EpicGames.Horde.Storage.Nodes
 		/// <param name="writer">Writer for new node data</param>
 		/// <param name="progress">Feedback interface for progress updates</param>
 		/// <param name="cancellationToken">Cancellation token for the operation</param>
-		/// <returns></returns>
 		public async Task CopyFromDirectoryAsync(DirectoryInfo directoryInfo, ChunkingOptions options, TreeWriter writer, IProgress<ICopyStats>? progress, CancellationToken cancellationToken)
 		{
 			// Enumerate all the files below this directory
@@ -950,7 +990,7 @@ namespace EpicGames.Horde.Storage.Nodes
 		/// </summary>
 		/// <param name="reader">Reader for other nodes</param>
 		/// <returns>Stream containing zipped archive data</returns>
-		public Stream GetZipStream(TreeReader reader) => new ZippedDirectoryStream(reader, this);
+		public Stream ToZipStream(TreeReader reader) => new ZippedDirectoryStream(reader, this);
 	}
 
 	/// <summary>
@@ -1004,11 +1044,12 @@ namespace EpicGames.Horde.Storage.Nodes
 			while (_current.Length == 0)
 			{
 				ReadResult result = await _pipe.Reader.ReadAsync(cancellationToken);
-				if (result.IsCompleted)
+				_current = result.Buffer;
+
+				if (result.IsCompleted && _current.Length == 0)
 				{
 					return 0;
 				}
-				_current = result.Buffer;
 			}
 
 			int initialSize = buffer.Length;
