Fix static analysis warning.

diff --git a/Engine/Source/Programs/Shared/EpicGames.Core/Log.cs b/Engine/Source/Programs/Shared/EpicGames.Core/Log.cs
--- a/Engine/Source/Programs/Shared/EpicGames.Core/Log.cs
+++ b/Engine/Source/Programs/Shared/EpicGames.Core/Log.cs
@@ -214,7 +214,7 @@ namespace EpicGames.Core
 		/// <returns>The created trace listener</returns>
 		public static void AddFileWriter(string name, FileReference outputFile)
 		{
-			Logger.LogInformation($"Log file: {OutputFile}", outputFile);
+			Logger.LogInformation("Log file: {OutputFile}", outputFile);
 
 			if (Log.BackupLogFiles && FileReference.Exists(outputFile))
 			{
