Horde: Fix incorrect type on NodeLocator <-> Json converter.

diff --git a/Engine/Source/Programs/Shared/EpicGames.Horde/Storage/NodeLocator.cs b/Engine/Source/Programs/Shared/EpicGames.Horde/Storage/NodeLocator.cs
--- a/Engine/Source/Programs/Shared/EpicGames.Horde/Storage/NodeLocator.cs
+++ b/Engine/Source/Programs/Shared/EpicGames.Horde/Storage/NodeLocator.cs
@@ -78,13 +78,13 @@ namespace EpicGames.Horde.Storage
 	/// <summary>
 	/// Type converter for BlobId to and from JSON
 	/// </summary>
-	sealed class NodeLocatorJsonConverter : JsonConverter<BlobLocator>
+	sealed class NodeLocatorJsonConverter : JsonConverter<NodeLocator>
 	{
 		/// <inheritdoc/>
-		public override BlobLocator Read(ref Utf8JsonReader reader, Type typeToConvert, JsonSerializerOptions options) => new BlobLocator(new Utf8String(reader.GetUtf8String().ToArray()));
+		public override NodeLocator Read(ref Utf8JsonReader reader, Type typeToConvert, JsonSerializerOptions options) => NodeLocator.Parse(reader.GetString() ?? String.Empty);
 
 		/// <inheritdoc/>
-		public override void Write(Utf8JsonWriter writer, BlobLocator value, JsonSerializerOptions options) => writer.WriteStringValue(value.Inner.Span);
+		public override void Write(Utf8JsonWriter writer, NodeLocator value, JsonSerializerOptions options) => writer.WriteStringValue(value.ToString());
 	}
 
 	/// <summary>
