Fixed inconsistencies between CbPackage serialization in C# and from Zen, these now write the same format. Resolves issues when using CbPackage during cook against Unreal Cloud DDC.

diff --git a/Engine/Source/Programs/Shared/EpicGames.Serialization.Tests/CbPackageTests.cs b/Engine/Source/Programs/Shared/EpicGames.Serialization.Tests/CbPackageTests.cs
--- a/Engine/Source/Programs/Shared/EpicGames.Serialization.Tests/CbPackageTests.cs
+++ b/Engine/Source/Programs/Shared/EpicGames.Serialization.Tests/CbPackageTests.cs
@@ -40,8 +40,8 @@ namespace EpicGames.Serialization.Tests
 			IoHash rootHash = IoHash.Compute(rootObject.GetView().Span);
 			
 			using CbPackageBuilder builder = new();
-			await builder.AddAttachment(rootHash, CbPackageAttachmentFlags.IsObject, rootObject.GetView().ToArray());
-			byte[] bytes = builder.ToByteArray();
+			builder.AddAttachment(rootHash, CbPackageAttachmentFlags.IsObject, rootObject.GetView().ToArray());
+			byte[] bytes = await builder.ToByteArray();
 
 			await using MemoryStream ms = new MemoryStream(bytes);
 			CbPackageReader reader = await CbPackageReader.Create(ms);
@@ -74,11 +74,11 @@ namespace EpicGames.Serialization.Tests
 			IoHash rootHash = IoHash.Compute(rootObject.GetView().Span);
 
 			using CbPackageBuilder builder = new();
-			await builder.AddAttachment(rootHash, CbPackageAttachmentFlags.IsObject, rootObject.GetView().ToArray());
-			await builder.AddAttachment(simpleHash, CbPackageAttachmentFlags.IsObject, simpleObject.GetView().ToArray());
-			await builder.AddAttachment(blobHash, 0, blob);
+			builder.AddAttachment(rootHash, CbPackageAttachmentFlags.IsObject, rootObject.GetView().ToArray());
+			builder.AddAttachment(simpleHash, CbPackageAttachmentFlags.IsObject, simpleObject.GetView().ToArray());
+			builder.AddAttachment(blobHash, 0, blob);
 
-			byte[] bytes = builder.ToByteArray();
+			byte[] bytes = await builder.ToByteArray();
 
 			await using MemoryStream ms = new MemoryStream(bytes);
 			CbPackageReader reader = await CbPackageReader.Create(ms);
@@ -115,7 +115,12 @@ namespace EpicGames.Serialization.Tests
 		public async Task ReadNoAttachments()
 		{
 			using CbPackageBuilder packageBuilder = new CbPackageBuilder();
-			byte[] buf = packageBuilder.ToByteArray();
+			SimpleObject simple = new SimpleObject();
+			CbObject simpleObject = CbSerializer.Serialize(simple);
+			IoHash simpleHash = IoHash.Compute(simpleObject.GetView().Span);
+			packageBuilder.AddAttachment(simpleHash, CbPackageAttachmentFlags.IsObject, simpleObject.GetView().ToArray());
+
+			byte[] buf = await packageBuilder.ToByteArray();
 			await using MemoryStream ms = new MemoryStream(buf);
 
 			try
diff --git a/Engine/Source/Programs/Shared/EpicGames.Serialization/CbPackage.cs b/Engine/Source/Programs/Shared/EpicGames.Serialization/CbPackage.cs
--- a/Engine/Source/Programs/Shared/EpicGames.Serialization/CbPackage.cs
+++ b/Engine/Source/Programs/Shared/EpicGames.Serialization/CbPackage.cs
@@ -3,6 +3,7 @@
 using System;
 using System.Collections.Generic;
 using System.IO;
+using System.Linq;
 using System.Text;
 using System.Threading.Tasks;
 using EpicGames.Core;
@@ -168,10 +169,12 @@ namespace EpicGames.Serialization
 	{
 		private readonly CbPackageHeader _header;
 		private readonly Stream _attachmentsStream;
+		private readonly List<CbPackageAttachmentEntry> _attachmentEntries = new List<CbPackageAttachmentEntry>();
 
-		private CbPackageReader(CbPackageHeader header, CbObject rootObject, IoHash rootHash, Stream attachmentsStream)
+		private CbPackageReader(CbPackageHeader header, CbObject rootObject, IoHash rootHash, List<CbPackageAttachmentEntry> attachmentEntries, Stream attachmentsStream)
 		{
 			_header = header;
+			_attachmentEntries = attachmentEntries;
 			_attachmentsStream = attachmentsStream;
 			RootObject = rootObject;
 			RootHash = rootHash;
@@ -205,34 +208,45 @@ namespace EpicGames.Serialization
 		{
 			CbPackageHeader header = await CbPackageHeader.Read(s);
 
-			if (header.AttachmentCount < 1)
+			List<CbPackageAttachmentEntry> entries = new List<CbPackageAttachmentEntry>();
+			// we expect a extra entry for the root object
+			uint attachmentsCount = header.AttachmentCount + 1;
+
+			for (int i = 0; i < attachmentsCount; i++)
+			{
+				CbPackageAttachmentEntry entry = await CbPackageAttachmentEntry.Read(s);
+				entries.Add(entry);
+			}
+
+			if (entries.Count < 1)
 			{
-				throw new Exception("At least 1 attachment was expected in a CbPackage");
+				throw new Exception("Did not find any root object entry when reading CbPackage");
 			}
 
 			// the first object is assumed to be a CbObject, the root object which references the other attachments
-			CbPackageAttachmentEntry entry = await CbPackageAttachmentEntry.Read(s);
-			if (!entry.Flags.HasFlag(CbPackageAttachmentFlags.IsObject))
+			CbPackageAttachmentEntry rootObject = entries.First();
+			entries.RemoveAt(0);
+			if (!rootObject.Flags.HasFlag(CbPackageAttachmentFlags.IsObject))
 			{
 				throw new Exception("First attachment must be a CbObject for a package");
 			}
-			if (entry.Flags.HasFlag(CbPackageAttachmentFlags.IsError))
+			if (rootObject.Flags.HasFlag(CbPackageAttachmentFlags.IsError))
 			{
 				throw new Exception("First attachment was a error object");
 			}
-			if (entry.Flags.HasFlag(CbPackageAttachmentFlags.IsCompressed))
+			if (rootObject.Flags.HasFlag(CbPackageAttachmentFlags.IsCompressed))
 			{
 				// TODO: We could support this being a compressed buffer and just remove it
 				throw new Exception("First attachment must not be compressed");
 			}
 
-			if (entry.PayloadSize > Int32.MaxValue)
+			if (rootObject.PayloadSize > Int32.MaxValue)
 			{
 				throw new Exception($"Package attachments larger then {Int32.MaxValue} not supported");
 			}
-			byte[] rootObjectBytes = await ReadStreamAsync(s, (int)entry.PayloadSize);
+			byte[] rootObjectBytes = await ReadStreamAsync(s, (int)rootObject.PayloadSize);
 
-			return new CbPackageReader(header, new CbObject(rootObjectBytes), entry.AttachmentHash, s);
+			return new CbPackageReader(header, new CbObject(rootObjectBytes), rootObject.AttachmentHash, entries, s);
 		}
 
 		/// <summary>
@@ -243,11 +257,9 @@ namespace EpicGames.Serialization
 		{
 			// close the stream after we have iterated the attachments as there should be nothing left in it
 			await using Stream s = _attachmentsStream;
-			// we have already read the first attachment as that is the root object
-			for (int i = 1; i < _header.AttachmentCount; i++)
+			
+			foreach (CbPackageAttachmentEntry entry in _attachmentEntries)
 			{
-				CbPackageAttachmentEntry entry = await CbPackageAttachmentEntry.Read(_attachmentsStream);
-
 				if (entry.PayloadSize > Int32.MaxValue)
 				{
 					throw new Exception($"Package attachments larger then {Int32.MaxValue} not supported");
@@ -273,33 +285,22 @@ namespace EpicGames.Serialization
 	/// </summary>
 	public sealed class CbPackageBuilder : IDisposable
 	{
-#pragma warning disable CA2213
-		private readonly MemoryStream _packageBuffer;
-#pragma warning restore CA2213
-		private uint _countOfWrittenAttachments;
+		private readonly List<(CbPackageAttachmentEntry, Stream)> _streamAttachments = new List<(CbPackageAttachmentEntry, Stream)>();
 
 		/// <summary>
 		/// Constructor for the package builder
 		/// </summary>
 		public CbPackageBuilder()
 		{
-			_packageBuffer = new MemoryStream();
-
-			// we will overwrite this header again after building the package builder
-			CbPackageHeader header = new CbPackageHeader
-			{
-				HeaderMagic = CbPackageConstants.CbPackageHeaderMagic,
-				AttachmentCount = 0,
-				Reserved1 = 0,
-				Reserved2 = 0,
-			};
-			header.Write(_packageBuffer);
 		}
 
 		/// <inheritdoc/>
 		public void Dispose()
 		{
-			_packageBuffer.Dispose();
+			foreach ((CbPackageAttachmentEntry _, Stream stream) in _streamAttachments)
+			{
+				stream.Dispose();
+			}
 		}
 
 		/// <summary>
@@ -309,7 +310,7 @@ namespace EpicGames.Serialization
 		/// <param name="flags">The flags that apply to the attachment</param>
 		/// <param name="blobMemory">In-memory buffer of the attachment</param>
 		/// <returns></returns>
-		public async Task AddAttachment(IoHash attachmentHash, CbPackageAttachmentFlags flags, byte[] blobMemory)
+		public void AddAttachment(IoHash attachmentHash, CbPackageAttachmentFlags flags, byte[] blobMemory)
 		{
 			CbPackageAttachmentEntry entry = new CbPackageAttachmentEntry()
 			{
@@ -317,9 +318,8 @@ namespace EpicGames.Serialization
 				Flags = flags,
 				PayloadSize = (ulong)blobMemory.LongLength
 			};
-			entry.Write(_packageBuffer);
-			await _packageBuffer.WriteAsync(blobMemory, 0, blobMemory.Length);
-			_countOfWrittenAttachments++;
+
+			_streamAttachments.Add((entry, new MemoryStream(blobMemory)));
 		}
 
 		/// <summary>
@@ -330,7 +330,7 @@ namespace EpicGames.Serialization
 		/// <param name="stream">The stream to read the attachment from</param>
 		/// <param name="length">The count of bytes to read from the stream</param>
 		/// <returns></returns>
-		public async Task AddAttachment(IoHash attachmentHash, CbPackageAttachmentFlags flags, Stream stream, ulong length)
+		public void AddAttachment(IoHash attachmentHash, CbPackageAttachmentFlags flags, Stream stream, ulong length)
 		{
 			CbPackageAttachmentEntry entry = new CbPackageAttachmentEntry()
 			{
@@ -338,29 +338,43 @@ namespace EpicGames.Serialization
 				Flags = flags,
 				PayloadSize = length
 			};
-			entry.Write(_packageBuffer);
-			await stream.CopyToAsync(_packageBuffer);
-			_countOfWrittenAttachments++;
+
+			_streamAttachments.Add((entry, stream));
 		}
 
 		/// <summary>
 		/// Generate a contiguous buffer of the cb package
 		/// </summary>
 		/// <returns></returns>
-		public byte[] ToByteArray()
+		public async Task<byte[]> ToByteArray()
 		{
-			// update the package with the correct header
-			_packageBuffer.Position = 0;
+			MemoryStream packageBuffer = new MemoryStream();
+
+			if (_streamAttachments.Count == 0)
+			{
+				throw new Exception("Expected the first attachment to be a CbObject but found no attachment");
+			}
+
+			// we will overwrite this header again after building the package builder
 			CbPackageHeader header = new CbPackageHeader
 			{
 				HeaderMagic = CbPackageConstants.CbPackageHeaderMagic,
-				AttachmentCount = _countOfWrittenAttachments,
+				AttachmentCount = (uint)_streamAttachments.Count - 1, // the root object does not count as a attachment
 				Reserved1 = 0,
 				Reserved2 = 0,
 			};
-			header.Write(_packageBuffer);
 
-			return _packageBuffer.ToArray();
+			header.Write(packageBuffer);
+			foreach ((CbPackageAttachmentEntry entry, Stream _) in _streamAttachments)
+			{
+				entry.Write(packageBuffer);
+			}
+
+			foreach ((CbPackageAttachmentEntry _, Stream stream) in _streamAttachments)
+			{
+				await stream.CopyToAsync(packageBuffer);
+			}
+			return packageBuffer.ToArray();
 		}
 	}
 }
diff --git a/Engine/Source/Programs/UnrealCloudDDC/Jupiter/Controllers/ReferencesController.cs b/Engine/Source/Programs/UnrealCloudDDC/Jupiter/Controllers/ReferencesController.cs
--- a/Engine/Source/Programs/UnrealCloudDDC/Jupiter/Controllers/ReferencesController.cs
+++ b/Engine/Source/Programs/UnrealCloudDDC/Jupiter/Controllers/ReferencesController.cs
@@ -240,7 +240,7 @@ namespace Jupiter.Controllers
                         IAsyncEnumerable<Attachment> attachments = _referenceResolver.GetAttachments(ns, cb);
 
                         using CbPackageBuilder writer = new CbPackageBuilder();
-                        await writer.AddAttachment(objectRecord.BlobIdentifier.AsIoHash(), CbPackageAttachmentFlags.IsObject, blobMemory);
+                        writer.AddAttachment(objectRecord.BlobIdentifier.AsIoHash(), CbPackageAttachmentFlags.IsObject, blobMemory);
 
                         await foreach (Attachment attachment in attachments)
                         {
@@ -282,15 +282,15 @@ namespace Jupiter.Controllers
                                     throw new NotSupportedException($"Unknown attachment type {attachment.GetType()}");
                                 }
 
-                                await writer.AddAttachment(attachmentHash, flags, attachmentContents.Stream, (ulong)attachmentContents.Length);
+                                writer.AddAttachment(attachmentHash, flags, attachmentContents.Stream, (ulong)attachmentContents.Length);
                             }
                             catch (Exception e)
                             {
                                 (CbObject errorObject, HttpStatusCode _) = ToErrorResult(e);
-                                await writer.AddAttachment(attachmentHash, CbPackageAttachmentFlags.IsError | CbPackageAttachmentFlags.IsObject, errorObject.GetView().ToArray());
+                                writer.AddAttachment(attachmentHash, CbPackageAttachmentFlags.IsError | CbPackageAttachmentFlags.IsObject, errorObject.GetView().ToArray());
                             }
                         }
-                        await using BlobContents contents = new BlobContents(writer.ToByteArray());
+                        await using BlobContents contents = new BlobContents(await writer.ToByteArray());
                         await WriteBody(contents, CustomMediaTypeNames.UnrealCompactBinaryPackage);
                         break;
                     }
