ManagedProcess: Don't create accounting job object with JOB_OBJECT_LIMIT_KILL_ON_JOB_CLOSE as the main job object is already managing those processes and it is causing mspdbsrv to be terminated while in use by cl as it is a shared process

diff --git a/Engine/Source/Programs/Shared/EpicGames.Core/ManagedProcess.cs b/Engine/Source/Programs/Shared/EpicGames.Core/ManagedProcess.cs
--- a/Engine/Source/Programs/Shared/EpicGames.Core/ManagedProcess.cs
+++ b/Engine/Source/Programs/Shared/EpicGames.Core/ManagedProcess.cs
@@ -124,7 +124,7 @@ namespace EpicGames.Core
 		/// <summary>
 		/// Constructor
 		/// </summary>
-		public ManagedProcessGroup()
+		public ManagedProcessGroup(bool bKillOnJobClose = true)
 		{
 			if(SupportsJobObjects)
 			{
@@ -137,7 +137,11 @@ namespace EpicGames.Core
 
 				// Configure the job object to terminate the processes added to it when the handle is closed
 				JOBOBJECT_EXTENDED_LIMIT_INFORMATION limitInformation = new JOBOBJECT_EXTENDED_LIMIT_INFORMATION();
-				limitInformation.BasicLimitInformation.LimitFlags = JOB_OBJECT_LIMIT_KILL_ON_JOB_CLOSE | JOB_OBJECT_LIMIT_BREAKAWAY_OK;
+				limitInformation.BasicLimitInformation.LimitFlags = JOB_OBJECT_LIMIT_BREAKAWAY_OK;
+				if (bKillOnJobClose)
+				{
+					limitInformation.BasicLimitInformation.LimitFlags |= JOB_OBJECT_LIMIT_KILL_ON_JOB_CLOSE;
+				}
 
 				int length = Marshal.SizeOf(typeof(JOBOBJECT_EXTENDED_LIMIT_INFORMATION));
 				IntPtr limitInformationPtr = Marshal.AllocHGlobal(length);
@@ -738,7 +742,7 @@ namespace EpicGames.Core
 					}
 
 					// Create a JobObject for each spawned process to do CPU usage accounting of spawned process and all its children
-					_accountingProcessGroup = new ManagedProcessGroup();
+					_accountingProcessGroup = new ManagedProcessGroup(bKillOnJobClose: false);
 					if (AssignProcessToJobObject(_accountingProcessGroup.JobHandle, processInfo.hProcess) == 0) 
 					{
 						throw new Win32Exception();
