Horde: Set the executable file permissions on zip entries created from directory nodes.

diff --git a/Engine/Source/Programs/Shared/EpicGames.Horde/Storage/Nodes/DirectoryNode.cs b/Engine/Source/Programs/Shared/EpicGames.Horde/Storage/Nodes/DirectoryNode.cs
--- a/Engine/Source/Programs/Shared/EpicGames.Horde/Storage/Nodes/DirectoryNode.cs
+++ b/Engine/Source/Programs/Shared/EpicGames.Horde/Storage/Nodes/DirectoryNode.cs
@@ -1096,6 +1096,11 @@ namespace EpicGames.Horde.Storage.Nodes
 			{
 				ZipArchiveEntry entry = archive.CreateEntry($"{prefix}{fileEntry}");
 
+				if((fileEntry.Flags & FileEntryFlags.Executable) != 0)
+				{
+					entry.ExternalAttributes |= 0b_111_111_101 << 16; // rwx rwx r-x
+				}
+
 				using Stream entryStream = entry.Open();
 				await fileEntry.CopyToStreamAsync(reader, entryStream, cancellationToken);
 			}
