Horde: Set read/write file modes on zipped files by default.

diff --git a/Engine/Source/Programs/Shared/EpicGames.Horde/Storage/Nodes/DirectoryNode.cs b/Engine/Source/Programs/Shared/EpicGames.Horde/Storage/Nodes/DirectoryNode.cs
--- a/Engine/Source/Programs/Shared/EpicGames.Horde/Storage/Nodes/DirectoryNode.cs
+++ b/Engine/Source/Programs/Shared/EpicGames.Horde/Storage/Nodes/DirectoryNode.cs
@@ -1096,10 +1096,14 @@ namespace EpicGames.Horde.Storage.Nodes
 			{
 				ZipArchiveEntry entry = archive.CreateEntry($"{prefix}{fileEntry}");
 
-				if((fileEntry.Flags & FileEntryFlags.Executable) != 0)
+				if ((fileEntry.Flags & FileEntryFlags.Executable) != 0)
 				{
 					entry.ExternalAttributes |= 0b_111_111_101 << 16; // rwx rwx r-x
 				}
+				else
+				{
+					entry.ExternalAttributes |= 0b_110_110_100 << 16; // rw- rw- r--
+				}
 
 				using Stream entryStream = entry.Open();
 				await fileEntry.CopyToStreamAsync(reader, entryStream, cancellationToken);
