Set the file mode on any zip entries extracted that have anything set in the top 16 bits, regardless of the platform that claims to have produced the zip file.

diff --git a/Engine/Source/Programs/Shared/EpicGames.Core/ZipArchiveExtensions.cs b/Engine/Source/Programs/Shared/EpicGames.Core/ZipArchiveExtensions.cs
--- a/Engine/Source/Programs/Shared/EpicGames.Core/ZipArchiveExtensions.cs
+++ b/Engine/Source/Programs/Shared/EpicGames.Core/ZipArchiveExtensions.cs
@@ -45,11 +45,6 @@ namespace EpicGames.Core
 			return entry;
 		}
 
-		/// <summary>
-		/// Internal field storing information about the platform that created a ZipArchiveEntry. Cannot interpret how to treat the attribute bits without reading this.
-		/// </summary>
-		static readonly FieldInfo s_versionMadeByPlatformField = typeof(ZipArchiveEntry).GetField("_versionMadeByPlatform", BindingFlags.NonPublic | BindingFlags.Instance)!;
-
 		/// <summary>
 		/// Extract a zip archive entry, preserving platform mode bits
 		/// </summary>
@@ -70,18 +65,18 @@ namespace EpicGames.Core
 
 			if (RuntimeInformation.IsOSPlatform(OSPlatform.Linux))
 			{
-				int madeByPlatform = Convert.ToInt32(s_versionMadeByPlatformField.GetValue(entry)!);
-				if (madeByPlatform == 3 || madeByPlatform == 19) // Unix or OSX
+				int mode = entry.ExternalAttributes >> 16;
+				if(mode != 0)
 				{
-					FileUtils.SetFileMode_Linux(targetFileName, (ushort)(entry.ExternalAttributes >> 16));
+					FileUtils.SetFileMode_Linux(targetFileName, (ushort)mode);
 				}
 			}
 			else if (RuntimeInformation.IsOSPlatform(OSPlatform.OSX))
 			{
-				int madeByPlatform = Convert.ToInt32(s_versionMadeByPlatformField.GetValue(entry)!);
-				if (madeByPlatform == 3 || madeByPlatform == 19) // Unix or OSX
+				int mode = entry.ExternalAttributes >> 16;
+				if (mode != 0)
 				{
-					FileUtils.SetFileMode_Mac(targetFileName, (ushort)(entry.ExternalAttributes >> 16));
+					FileUtils.SetFileMode_Mac(targetFileName, (ushort)mode);
 				}
 			}
 		}
