Horde: Move C++ compute commands into a separate namespace.

diff --git a/Engine/Source/Programs/Shared/EpicGames.Horde/Compute/ComputeStorageClient.cs b/Engine/Source/Programs/Shared/EpicGames.Horde/Compute/ComputeStorageClient.cs
--- a/Engine/Source/Programs/Shared/EpicGames.Horde/Compute/ComputeStorageClient.cs
+++ b/Engine/Source/Programs/Shared/EpicGames.Horde/Compute/ComputeStorageClient.cs
@@ -6,6 +6,7 @@ using System.IO;
 using System.Threading;
 using System.Threading.Tasks;
 using EpicGames.Core;
+using EpicGames.Horde.Compute.Cpp;
 using EpicGames.Horde.Storage;
 
 namespace EpicGames.Horde.Compute
diff --git a/Engine/Source/Programs/Shared/EpicGames.Horde/Compute/Cpp/CppMessages.cs b/Engine/Source/Programs/Shared/EpicGames.Horde/Compute/Cpp/CppMessages.cs
new file mode 100644
--- /dev/null
+++ b/Engine/Source/Programs/Shared/EpicGames.Horde/Compute/Cpp/CppMessages.cs
@@ -0,0 +1,210 @@
+// Copyright Epic Games, Inc. All Rights Reserved.
+
+using EpicGames.Core;
+using EpicGames.Horde.Storage;
+using System.IO;
+using System.Threading;
+using System.Threading.Tasks;
+
+namespace EpicGames.Horde.Compute.Cpp
+{
+	/// <summary>
+	/// Message for starting a C++ build
+	/// </summary>
+	/// <param name="Locator">Locator for the task description</param>
+	/// <param name="ReplyChannelId">Channel to reply on</param>
+	public record struct CppBeginMessage(NodeLocator Locator, int ReplyChannelId);
+
+	/// <summary>
+	/// Creates a blob read request
+	/// </summary>
+	public record struct CppBlobReadMessage(BlobLocator Locator, int Offset, int Length);
+
+	/// <summary>
+	/// Output from a C++ build request
+	/// </summary>
+	public record struct CppOutputMessage(NodeLocator Locator);
+
+	/// <summary>
+	/// Extension methods for C++ builds
+	/// </summary>
+	public static class ComputeChannelCppExtensions
+	{
+		/// <summary>
+		/// Starts a C++ build task
+		/// </summary>
+		/// <param name="channel">Channel to write to</param>
+		/// <param name="locator">Locator for the node containing the task definition</param>
+		/// <param name="replyChannelId">Channel to reply to with future messages</param>
+		/// <param name="cancellationToken">Cancellation token for the operation</param>
+		public static async Task CppStartAsync(this IComputeChannel channel, NodeLocator locator, int replyChannelId, CancellationToken cancellationToken)
+		{
+			using (IComputeMessageWriter writer = channel.CreateMessage(ComputeMessageType.CppBegin))
+			{
+				writer.WriteNodeLocator(locator);
+				writer.WriteUnsignedVarInt(replyChannelId);
+				await writer.SendAsync(cancellationToken);
+			}
+		}
+
+		/// <summary>
+		/// 
+		/// </summary>
+		/// <param name="message"></param>
+		/// <returns></returns>
+		public static CppBeginMessage AsCppBegin(this IComputeMessage message)
+		{
+			MemoryReader reader = new MemoryReader(message.Data);
+			NodeLocator locator = reader.ReadNodeLocator();
+			int replyChannelId = (int)reader.ReadUnsignedVarInt();
+			return new CppBeginMessage(locator, replyChannelId);
+		}
+
+		/// <summary>
+		/// Sends the result for a C++ compute task
+		/// </summary>
+		/// <param name="message"></param>
+		public static NodeLocator AsCppSuccess(this IComputeMessage message)
+		{
+			MemoryReader reader = new MemoryReader(message.Data);
+			return reader.ReadNodeLocator();
+		}
+
+		/// <summary>
+		/// Sends the result for a C++ compute task
+		/// </summary>
+		/// <param name="channel">Channel to write to</param>
+		/// <param name="locator">Locator for the result</param>
+		/// <param name="cancellationToken">Cancellation token for the operation</param>
+		public static async Task CppSuccessAsync(this IComputeChannel channel, NodeLocator locator, CancellationToken cancellationToken)
+		{
+			using (IComputeMessageWriter writer = channel.CreateMessage(ComputeMessageType.CppSuccess))
+			{
+				writer.WriteNodeLocator(locator);
+				await writer.SendAsync(cancellationToken);
+			}
+		}
+
+		/// <summary>
+		/// Sends the result for a C++ compute task
+		/// </summary>
+		/// <param name="message"></param>
+		public static string AsCppFailure(this IComputeMessage message)
+		{
+			MemoryReader reader = new MemoryReader(message.Data);
+			return reader.ReadString();
+		}
+
+		/// <summary>
+		/// Sends the failure error for a C++ compute task
+		/// </summary>
+		/// <param name="channel">Channel to write to</param>
+		/// <param name="error">Error string to send to the caller</param>
+		/// <param name="cancellationToken">Cancellation token for the operation</param>
+		public static async Task CppFailureAsync(this IComputeChannel channel, string error, CancellationToken cancellationToken)
+		{
+			using (IComputeMessageWriter writer = channel.CreateMessage(ComputeMessageType.CppFailure))
+			{
+				writer.WriteString(error);
+				await writer.SendAsync(cancellationToken);
+			}
+		}
+
+		/// <summary>
+		/// Completes a C++ build task
+		/// </summary>
+		/// <param name="channel">Channel to write to</param>
+		/// <param name="cancellationToken">Cancellation token for the operation</param>
+		public static async Task CppFinishAsync(this IComputeChannel channel, CancellationToken cancellationToken)
+		{
+			using (IComputeMessageWriter writer = channel.CreateMessage(ComputeMessageType.CppEnd))
+			{
+				await writer.SendAsync(cancellationToken);
+			}
+		}
+
+		/// <summary>
+		/// 
+		/// </summary>
+		/// <param name="message"></param>
+		/// <returns></returns>
+		public static CppBlobReadMessage AsCppBlobRead(this IComputeMessage message)
+		{
+			MemoryReader reader = new MemoryReader(message.Data);
+			BlobLocator locator = reader.ReadBlobLocator();
+			int offset = (int)reader.ReadUnsignedVarInt();
+			int length = (int)reader.ReadUnsignedVarInt();
+			return new CppBlobReadMessage(locator, offset, length);
+		}
+
+		/// <summary>
+		/// 
+		/// </summary>
+		/// <param name="channel"></param>
+		/// <param name="locator"></param>
+		/// <param name="offset"></param>
+		/// <param name="length"></param>
+		/// <param name="cancellationToken"></param>
+		/// <returns></returns>
+		public static async Task CppBlobReadAsync(this IComputeChannel channel, BlobLocator locator, int offset, int length, CancellationToken cancellationToken)
+		{
+			using (IComputeMessageWriter writer = channel.CreateMessage(ComputeMessageType.CppBlobRead))
+			{
+				writer.WriteBlobLocator(locator);
+				writer.WriteUnsignedVarInt(offset);
+				writer.WriteUnsignedVarInt(length);
+				await writer.SendAsync(cancellationToken);
+			}
+		}
+
+		/// <summary>
+		/// Writes blob data to a compute channel
+		/// </summary>
+		/// <param name="channel">Channel to write to</param>
+		/// <param name="message">The read request</param>
+		/// <param name="storage">Storage client to retrieve the blob from</param>
+		/// <param name="cancellationToken">Cancellation token for the operation</param>
+		public static Task CppBlobDataAsync(this IComputeChannel channel, CppBlobReadMessage message, IStorageClient storage, CancellationToken cancellationToken)
+		{
+			return CppBlobDataAsync(channel, message.Locator, message.Offset, message.Length, storage, cancellationToken);
+		}
+
+		/// <summary>
+		/// Writes blob data to a compute channel
+		/// </summary>
+		/// <param name="channel">Channel to write to</param>
+		/// <param name="locator">Locator for the blob to send</param>
+		/// <param name="offset">Starting offset of the data</param>
+		/// <param name="length">Length of the data</param>
+		/// <param name="storage">Storage client to retrieve the blob from</param>
+		/// <param name="cancellationToken">Cancellation token for the operation</param>
+		public static async Task CppBlobDataAsync(this IComputeChannel channel, BlobLocator locator, int offset, int length, IStorageClient storage, CancellationToken cancellationToken)
+		{
+			byte[] data;
+			if (offset == 0 && length == 0)
+			{
+				using (Stream stream = await storage.ReadBlobAsync(locator, cancellationToken))
+				{
+					using MemoryStream target = new MemoryStream();
+					await stream.CopyToAsync(target, cancellationToken);
+					data = target.ToArray();
+				}
+			}
+			else
+			{
+				using (Stream stream = await storage.ReadBlobRangeAsync(locator, offset, length, cancellationToken))
+				{
+					using MemoryStream target = new MemoryStream();
+					await stream.CopyToAsync(target, cancellationToken);
+					data = target.ToArray();
+				}
+			}
+
+			using (IComputeMessageWriter writer = channel.CreateMessage(ComputeMessageType.CppBlobData, length + 128))
+			{
+				writer.WriteFixedLengthBytes(data);
+				await writer.SendAsync(cancellationToken);
+			}
+		}
+	}
+}
diff --git a/Engine/Source/Programs/Shared/EpicGames.Horde/Compute/IComputeMessage.cs b/Engine/Source/Programs/Shared/EpicGames.Horde/Compute/IComputeMessage.cs
--- a/Engine/Source/Programs/Shared/EpicGames.Horde/Compute/IComputeMessage.cs
+++ b/Engine/Source/Programs/Shared/EpicGames.Horde/Compute/IComputeMessage.cs
@@ -97,23 +97,6 @@ namespace EpicGames.Horde.Compute
 	/// <param name="Value">Value to XOR with</param>
 	public record struct XorRequestMessage(ReadOnlyMemory<byte> Data, byte Value);
 
-	/// <summary>
-	/// Message for starting a C++ build
-	/// </summary>
-	/// <param name="Locator">Locator for the task description</param>
-	/// <param name="ReplyChannelId">Channel to reply on</param>
-	public record struct CppBeginMessage(NodeLocator Locator, int ReplyChannelId);
-
-	/// <summary>
-	/// Creates a blob read request
-	/// </summary>
-	public record struct CppBlobReadMessage(BlobLocator Locator, int Offset, int Length);
-
-	/// <summary>
-	/// Output from a C++ build request
-	/// </summary>
-	public record struct CppOutputMessage(NodeLocator Locator);
-
 	/// <summary>
 	/// Wraps various requests across compute channels
 	/// </summary>
@@ -160,182 +143,5 @@ namespace EpicGames.Horde.Compute
 			ReadOnlyMemory<byte> data = message.Data;
 			return new XorRequestMessage(data[0..^1], data.Span[^1]);
 		}
-
-		/// <summary>
-		/// Starts a C++ build task
-		/// </summary>
-		/// <param name="channel">Channel to write to</param>
-		/// <param name="locator">Locator for the node containing the task definition</param>
-		/// <param name="replyChannelId">Channel to reply to with future messages</param>
-		/// <param name="cancellationToken">Cancellation token for the operation</param>
-		public static async Task CppStartAsync(this IComputeChannel channel, NodeLocator locator, int replyChannelId, CancellationToken cancellationToken)
-		{
-			using (IComputeMessageWriter writer = channel.CreateMessage(ComputeMessageType.CppBegin))
-			{
-				writer.WriteNodeLocator(locator);
-				writer.WriteUnsignedVarInt(replyChannelId);
-				await writer.SendAsync(cancellationToken);
-			}
-		}
-
-		/// <summary>
-		/// 
-		/// </summary>
-		/// <param name="message"></param>
-		/// <returns></returns>
-		public static CppBeginMessage AsCppBegin(this IComputeMessage message)
-		{
-			MemoryReader reader = new MemoryReader(message.Data);
-			NodeLocator locator = reader.ReadNodeLocator();
-			int replyChannelId = (int)reader.ReadUnsignedVarInt();
-			return new CppBeginMessage(locator, replyChannelId);
-		}
-
-		/// <summary>
-		/// Sends the result for a C++ compute task
-		/// </summary>
-		/// <param name="message"></param>
-		public static NodeLocator AsCppSuccess(this IComputeMessage message)
-		{
-			MemoryReader reader = new MemoryReader(message.Data);
-			return reader.ReadNodeLocator();
-		}
-
-		/// <summary>
-		/// Sends the result for a C++ compute task
-		/// </summary>
-		/// <param name="channel">Channel to write to</param>
-		/// <param name="locator">Locator for the result</param>
-		/// <param name="cancellationToken">Cancellation token for the operation</param>
-		public static async Task CppSuccessAsync(this IComputeChannel channel, NodeLocator locator, CancellationToken cancellationToken)
-		{
-			using (IComputeMessageWriter writer = channel.CreateMessage(ComputeMessageType.CppSuccess))
-			{
-				writer.WriteNodeLocator(locator);
-				await writer.SendAsync(cancellationToken);
-			}
-		}
-
-		/// <summary>
-		/// Sends the result for a C++ compute task
-		/// </summary>
-		/// <param name="message"></param>
-		public static string AsCppFailure(this IComputeMessage message)
-		{
-			MemoryReader reader = new MemoryReader(message.Data);
-			return reader.ReadString();
-		}
-
-		/// <summary>
-		/// Sends the failure error for a C++ compute task
-		/// </summary>
-		/// <param name="channel">Channel to write to</param>
-		/// <param name="error">Error string to send to the caller</param>
-		/// <param name="cancellationToken">Cancellation token for the operation</param>
-		public static async Task CppFailureAsync(this IComputeChannel channel, string error, CancellationToken cancellationToken)
-		{
-			using (IComputeMessageWriter writer = channel.CreateMessage(ComputeMessageType.CppFailure))
-			{
-				writer.WriteString(error);
-				await writer.SendAsync(cancellationToken);
-			}
-		}
-
-		/// <summary>
-		/// Completes a C++ build task
-		/// </summary>
-		/// <param name="channel">Channel to write to</param>
-		/// <param name="cancellationToken">Cancellation token for the operation</param>
-		public static async Task CppFinishAsync(this IComputeChannel channel, CancellationToken cancellationToken)
-		{
-			using (IComputeMessageWriter writer = channel.CreateMessage(ComputeMessageType.CppEnd))
-			{
-				await writer.SendAsync(cancellationToken);
-			}
-		}
-
-		/// <summary>
-		/// 
-		/// </summary>
-		/// <param name="message"></param>
-		/// <returns></returns>
-		public static CppBlobReadMessage AsCppBlobRead(this IComputeMessage message)
-		{
-			MemoryReader reader = new MemoryReader(message.Data);
-			BlobLocator locator = reader.ReadBlobLocator();
-			int offset = (int)reader.ReadUnsignedVarInt();
-			int length = (int)reader.ReadUnsignedVarInt();
-			return new CppBlobReadMessage(locator, offset, length);
-		}
-
-		/// <summary>
-		/// 
-		/// </summary>
-		/// <param name="channel"></param>
-		/// <param name="locator"></param>
-		/// <param name="offset"></param>
-		/// <param name="length"></param>
-		/// <param name="cancellationToken"></param>
-		/// <returns></returns>
-		public static async Task CppBlobReadAsync(this IComputeChannel channel, BlobLocator locator, int offset, int length, CancellationToken cancellationToken)
-		{
-			using (IComputeMessageWriter writer = channel.CreateMessage(ComputeMessageType.CppBlobRead))
-			{
-				writer.WriteBlobLocator(locator);
-				writer.WriteUnsignedVarInt(offset);
-				writer.WriteUnsignedVarInt(length);
-				await writer.SendAsync(cancellationToken);
-			}
-		}
-
-		/// <summary>
-		/// Writes blob data to a compute channel
-		/// </summary>
-		/// <param name="channel">Channel to write to</param>
-		/// <param name="message">The read request</param>
-		/// <param name="storage">Storage client to retrieve the blob from</param>
-		/// <param name="cancellationToken">Cancellation token for the operation</param>
-		public static Task CppBlobDataAsync(this IComputeChannel channel, CppBlobReadMessage message, IStorageClient storage, CancellationToken cancellationToken)
-		{
-			return CppBlobDataAsync(channel, message.Locator, message.Offset, message.Length, storage, cancellationToken);
-		}
-
-		/// <summary>
-		/// Writes blob data to a compute channel
-		/// </summary>
-		/// <param name="channel">Channel to write to</param>
-		/// <param name="locator">Locator for the blob to send</param>
-		/// <param name="offset">Starting offset of the data</param>
-		/// <param name="length">Length of the data</param>
-		/// <param name="storage">Storage client to retrieve the blob from</param>
-		/// <param name="cancellationToken">Cancellation token for the operation</param>
-		public static async Task CppBlobDataAsync(this IComputeChannel channel, BlobLocator locator, int offset, int length, IStorageClient storage, CancellationToken cancellationToken)
-		{
-			byte[] data;
-			if (offset == 0 && length == 0)
-			{
-				using (Stream stream = await storage.ReadBlobAsync(locator, cancellationToken))
-				{
-					using MemoryStream target = new MemoryStream();
-					await stream.CopyToAsync(target, cancellationToken);
-					data = target.ToArray();
-				}
-			}
-			else
-			{
-				using (Stream stream = await storage.ReadBlobRangeAsync(locator, offset, length, cancellationToken))
-				{
-					using MemoryStream target = new MemoryStream();
-					await stream.CopyToAsync(target, cancellationToken);
-					data = target.ToArray();
-				}
-			}
-
-			using (IComputeMessageWriter writer = channel.CreateMessage(ComputeMessageType.CppBlobData, length + 128))
-			{
-				writer.WriteFixedLengthBytes(data);
-				await writer.SendAsync(cancellationToken);
-			}
-		}
 	}
 }
