Horde: Reimplement compute message passing on top of lower-level primitives to support reading/writing directly into shared memory buffers, to make interfacing with native code significantly more efficient.

diff --git a/Engine/Source/Programs/Horde/Horde.Agent/Commands/Compute/CppComputeCommand.cs b/Engine/Source/Programs/Horde/Horde.Agent/Commands/Compute/CppComputeCommand.cs
--- a/Engine/Source/Programs/Horde/Horde.Agent/Commands/Compute/CppComputeCommand.cs
+++ b/Engine/Source/Programs/Horde/Horde.Agent/Commands/Compute/CppComputeCommand.cs
@@ -1,6 +1,7 @@
 // Copyright Epic Games, Inc. All Rights Reserved.
 
 using System;
+using System.Buffers;
 using System.Collections.Generic;
 using System.Net.Http;
 using System.Net.Http.Headers;
@@ -74,16 +75,17 @@ namespace Horde.Agent.Commands
 
 		async Task<bool> HandleRequestAsync(IComputeLease lease, CancellationToken cancellationToken)
 		{
-			IComputeChannel channel = lease.OpenChannel(1);
+			using IComputeChannel defaultChannel = lease.CreateChannel(0);
+			using IComputeChannel channel = lease.CreateChannel(1);
 
 			MemoryStorageClient storage = new MemoryStorageClient();
 
 			NodeLocator node = await CreateComputeNodeAsync(TaskFile, storage, cancellationToken);
-			await lease.DefaultChannel.CppStartAsync(node, channel.Id, cancellationToken);
+			await defaultChannel.CppStartAsync(node, channel.Id, cancellationToken);
 
 			for (; ; )
 			{
-				IComputeMessage message = await channel.ReadAsync(cancellationToken);
+				using IComputeMessage message = await channel.ReadAsync(cancellationToken);
 				switch (message.Type)
 				{
 					case ComputeMessageType.None:
@@ -115,7 +117,7 @@ namespace Horde.Agent.Commands
 							}
 
 							await channel.CppFinishAsync(cancellationToken);
-							await lease.DefaultChannel.CloseAsync(cancellationToken);
+							await defaultChannel.CloseAsync(cancellationToken);
 						}
 						return false;
 					case ComputeMessageType.CppFailure:
diff --git a/Engine/Source/Programs/Horde/Horde.Agent/Commands/Compute/DdcComputeCommand.cs b/Engine/Source/Programs/Horde/Horde.Agent/Commands/Compute/DdcComputeCommand.cs
--- a/Engine/Source/Programs/Horde/Horde.Agent/Commands/Compute/DdcComputeCommand.cs
+++ b/Engine/Source/Programs/Horde/Horde.Agent/Commands/Compute/DdcComputeCommand.cs
@@ -55,7 +55,7 @@ namespace Horde.Agent.Commands
 
 		async Task<object?> TestCommandAsync(IComputeLease lease, CancellationToken cancellationToken)
 		{
-			IComputeChannel channel = lease.DefaultChannel;
+			using IComputeChannel channel = lease.CreateChannel(0);
 
 			await channel.XorRequestAsync(new byte[] { 1, 2, 3, 4, 5 }, (byte)123, cancellationToken);
 			IComputeMessage response = await channel.ReadAsync(cancellationToken);
diff --git a/Engine/Source/Programs/Horde/Horde.Agent/Leases/Handlers/ComputeHandler.cs b/Engine/Source/Programs/Horde/Horde.Agent/Leases/Handlers/ComputeHandler.cs
--- a/Engine/Source/Programs/Horde/Horde.Agent/Leases/Handlers/ComputeHandler.cs
+++ b/Engine/Source/Programs/Horde/Horde.Agent/Leases/Handlers/ComputeHandler.cs
@@ -57,7 +57,7 @@ namespace Horde.Agent.Leases.Handlers
 					return LeaseResult.Success;
 				}
 
-				await using (ComputeLease lease = new ComputeLease(tcpClient.Client, computeTask.Key.Span, computeTask.Nonce.Span.Slice(0, ComputeLease.NonceLength), computeTask.Resources))
+				await using (ComputeLease lease = new ComputeLease(new TcpComputeTransport(tcpClient.Client), computeTask.Resources))
 				{
 					await RunAsync(lease, cancellationToken);
 					return LeaseResult.Success;
@@ -76,7 +76,7 @@ namespace Horde.Agent.Leases.Handlers
 
 		public async Task RunAsync(IComputeLease lease, CancellationToken cancellationToken)
 		{
-			IComputeChannel channel = lease.DefaultChannel;
+			using IComputeChannel channel = lease.CreateChannel(0);
 			for (; ; )
 			{
 				using IComputeMessage message = await channel.ReadAsync(cancellationToken);
@@ -94,7 +94,8 @@ namespace Horde.Agent.Leases.Handlers
 					case ComputeMessageType.CppBegin:
 						{
 							CppBeginMessage cppBegin = message.AsCppBegin();
-							await RunCppAsync(lease.OpenChannel(cppBegin.ReplyChannelId), cppBegin.Locator, cancellationToken);
+							using IComputeChannel replyChannel = lease.CreateChannel(cppBegin.ReplyChannelId);
+							await RunCppAsync(replyChannel, cppBegin.Locator, cancellationToken);
 						}
 						break;
 					default:
diff --git a/Engine/Source/Programs/Shared/EpicGames.Horde/Compute/ComputeBuffer.cs b/Engine/Source/Programs/Shared/EpicGames.Horde/Compute/ComputeBuffer.cs
new file mode 100644
--- /dev/null
+++ b/Engine/Source/Programs/Shared/EpicGames.Horde/Compute/ComputeBuffer.cs
@@ -0,0 +1,189 @@
+// Copyright Epic Games, Inc. All Rights Reserved.
+
+using System;
+using System.Buffers;
+using System.Collections.Generic;
+using System.Threading;
+using System.Threading.Tasks;
+
+namespace EpicGames.Horde.Compute
+{
+	/// <summary>
+	/// Base implementation of <see cref="IComputeInputBuffer"/> and <see cref="IComputeOutputBuffer"/>. Can be derived from to intercept write calls and add special handling for change events.
+	/// </summary>
+	class ComputeBuffer : IComputeInputBuffer, IComputeOutputBuffer, IDisposable
+	{
+		/// <inheritdoc/>
+		public int Id => _id;
+
+		/// <inheritdoc/>
+		public Memory<byte> Data { get; }
+
+		/// <inheritdoc/>
+		ReadOnlyMemory<byte> IComputeInputBuffer.Data => Data;
+
+		/// <inheritdoc/>
+		public long ReadPosition => _readPosition;
+
+		/// <inheritdoc/>
+		public long WritePosition => _writePosition;
+
+		readonly int _id;
+		readonly IMemoryOwner<byte> _memoryOwner;
+		readonly object _lockObject = new object();
+
+		long _readPosition;
+		long _writePosition;
+		bool _complete;
+		readonly List<TaskCompletionSource> _waiters = new List<TaskCompletionSource>();
+		TaskCompletionSource? _consumedTcs;
+
+		/// <summary>
+		/// Constructor
+		/// </summary>
+		/// <param name="id">Id for the buffer</param>
+		/// <param name="memoryOwner">The underlying block of memory</param>
+		public ComputeBuffer(int id, IMemoryOwner<byte> memoryOwner)
+		{
+			_id = id;
+			_memoryOwner = memoryOwner;
+
+			Data = _memoryOwner.Memory;
+		}
+
+		/// <inheritdoc/>
+		public void Dispose()
+		{
+			_memoryOwner.Dispose();
+		}
+
+		/// <inheritdoc/>
+		public virtual ValueTask FlushAsync(CancellationToken cancellationToken) => new ValueTask();
+
+		/// <inheritdoc/>
+		public virtual async ValueTask ResetWritePositionAsync(CancellationToken cancellationToken)
+		{
+			await FlushAsync(cancellationToken);
+			ResetWritePosition();
+		}
+
+		void ReleaseWaiters()
+		{
+			for (int idx = 0; idx < _waiters.Count; idx++)
+			{
+				_waiters[idx].TrySetResult();
+			}
+			_waiters.Clear();
+		}
+
+		/// <summary>
+		/// Marks this buffer as complete, indicating that no more data will be written to it.
+		/// </summary>
+		public void MarkComplete()
+		{
+			lock (_lockObject)
+			{
+				_complete = true;
+				ReleaseWaiters();
+			}
+		}
+
+		/// <summary>
+		/// Reset the current write position.
+		/// </summary>
+		public void ResetWritePosition()
+		{
+			_readPosition = 0;
+			_writePosition = 0;
+		}
+
+		/// <inheritdoc/>
+		public void AdvanceReadPosition(long size)
+		{
+			lock (_lockObject)
+			{
+				_readPosition += size;
+				if (_readPosition == _writePosition && _consumedTcs != null)
+				{
+					_consumedTcs.SetResult();
+				}
+			}
+		}
+
+		/// <inheritdoc/>
+		public virtual void AdvanceWritePosition(long size)
+		{
+			if (size < 0)
+			{
+				throw new ArgumentException("New length cannot be less than current length");
+			}
+
+			if (size > 0)
+			{
+				lock (_lockObject)
+				{
+					if (_complete)
+					{
+						throw new InvalidOperationException("Cannot add data to a buffer after it is marked complete");
+					}
+
+					_writePosition += size;
+					ReleaseWaiters();
+				}
+			}
+		}
+
+		/// <summary>
+		/// Waits for all data to be read from the buffer
+		/// </summary>
+		/// <param name="cancellationToken"></param>
+		/// <returns></returns>
+		protected async ValueTask WaitForAllReadDataAsync(CancellationToken cancellationToken)
+		{
+			TaskCompletionSource tcs;
+			lock (_lockObject)
+			{
+				if (_complete || _readPosition == _writePosition)
+				{
+					return;
+				}
+				else
+				{
+					tcs = _consumedTcs ??= new TaskCompletionSource(TaskCreationOptions.RunContinuationsAsynchronously);
+				}
+			}
+
+			using (IDisposable disposable = cancellationToken.Register(() => tcs.SetCanceled(cancellationToken)))
+			{
+				await tcs.Task;
+			}
+		}
+
+		/// <inheritdoc/>
+		public async ValueTask WaitForWrittenDataAsync(long currentLength, CancellationToken cancellationToken)
+		{
+			if (_writePosition != currentLength || _complete)
+			{
+				return;
+			}
+
+			TaskCompletionSource tcs = new TaskCompletionSource(TaskCreationOptions.RunContinuationsAsynchronously);
+			lock (_lockObject)
+			{
+				if (_writePosition != currentLength)
+				{
+					return;
+				}
+				else
+				{
+					_waiters.Add(tcs);
+				}
+			}
+
+			using (IDisposable disposable = cancellationToken.Register(tcs.SetCanceled))
+			{
+				await tcs.Task;
+			}
+		}
+	}
+}
diff --git a/Engine/Source/Programs/Shared/EpicGames.Horde/Compute/ComputeChannel.cs b/Engine/Source/Programs/Shared/EpicGames.Horde/Compute/ComputeChannel.cs
new file mode 100644
--- /dev/null
+++ b/Engine/Source/Programs/Shared/EpicGames.Horde/Compute/ComputeChannel.cs
@@ -0,0 +1,169 @@
+// Copyright Epic Games, Inc. All Rights Reserved.
+
+using System;
+using System.Buffers;
+using System.Buffers.Binary;
+using System.Threading;
+using System.Threading.Tasks;
+using EpicGames.Core;
+
+namespace EpicGames.Horde.Compute
+{
+	/// <summary>
+	/// Implementation of a compute channel
+	/// </summary>
+	class ComputeChannel : IComputeChannel, IDisposable
+	{
+		/// <summary>
+		/// Standard implementation of a message
+		/// </summary>
+		sealed class Message : IComputeMessage
+		{
+			/// <inheritdoc/>
+			public ComputeMessageType Type { get; }
+
+			/// <inheritdoc/>
+			public ReadOnlyMemory<byte> Data { get; }
+
+			readonly IMemoryOwner<byte> _memoryOwner;
+
+			public Message(ComputeMessageType type, ReadOnlyMemory<byte> data)
+			{
+				_memoryOwner = MemoryPool<byte>.Shared.Rent(data.Length);
+				data.CopyTo(_memoryOwner.Memory);
+
+				Type = type;
+				Data = _memoryOwner.Memory.Slice(0, data.Length);
+			}
+
+			/// <inheritdoc/>
+			public void Dispose()
+			{
+				_memoryOwner.Dispose();
+			}
+		}
+
+		/// <summary>
+		/// Allows creating new messages in rented memory
+		/// </summary>
+		class Writer : IComputeMessageWriter
+		{
+			readonly ComputeChannel _channel;
+			readonly IComputeOutputBuffer _buffer;
+			int _length;
+
+			/// <inheritdoc/>
+			public int Length => _length;
+
+			public Writer(ComputeChannel channel, IComputeOutputBuffer buffer, ComputeMessageType type)
+			{
+				_channel = channel;
+				_buffer = buffer;
+
+				this.WriteUInt8((byte)type);
+				this.WriteUInt32(0); // size placeholder
+			}
+
+			public void Dispose()
+			{
+				if (_channel._currentWriter == this)
+				{
+					_channel._currentWriter = null;
+				}
+			}
+
+			public async ValueTask SendAsync(CancellationToken cancellationToken)
+			{
+				BinaryPrimitives.WriteInt32LittleEndian(_buffer.GetWriteSpan().Slice(1), _length);
+				await _channel.SendAsync(_length, cancellationToken);
+			}
+
+			/// <inheritdoc/>
+			public void Advance(int count) => _length += count;
+
+			/// <inheritdoc/>
+			public Memory<byte> GetMemory(int sizeHint = 0) => _buffer.Data.Slice((int)_buffer.WritePosition + _length);
+
+			/// <inheritdoc/>
+			public Span<byte> GetSpan(int sizeHint = 0) => GetMemory(sizeHint).Span;
+		}
+
+		/// <inheritdoc/>
+		public int Id => _id;
+
+		readonly int _id;
+		readonly IComputeInputBuffer _inputBuffer;
+		readonly IComputeOutputBuffer _outputBuffer;
+		readonly long _flushLength;
+
+		Writer? _currentWriter;
+
+		/// <summary>
+		/// Constructor
+		/// </summary>
+		/// <param name="id">The channel id</param>
+		/// <param name="inputBuffer">Buffer to read incoming messages from</param>
+		/// <param name="outputBuffer">Buffer to write outgoing messages to</param>
+		/// <param name="maxMessageSize">Maximum size of a message</param>
+		public ComputeChannel(int id, IComputeInputBuffer inputBuffer, IComputeOutputBuffer outputBuffer, int maxMessageSize = 1024 * 64)
+		{
+			_id = id;
+			_inputBuffer = inputBuffer;
+			_outputBuffer = outputBuffer;
+			_flushLength = outputBuffer.Data.Length - maxMessageSize;
+		}
+
+		/// <inheritdoc/>
+		public void Dispose()
+		{
+			_currentWriter?.Dispose();
+		}
+
+		/// <inheritdoc/>
+		public async ValueTask<IComputeMessage> ReadAsync(CancellationToken cancellationToken)
+		{
+			for (; ; )
+			{
+				long readPosition = _inputBuffer.ReadPosition;
+				ReadOnlyMemory<byte> memory = _inputBuffer.GetReadMemory();
+
+				const int HeaderLength = 5;
+				if (memory.Length >= HeaderLength)
+				{
+					int length = BinaryPrimitives.ReadInt32LittleEndian(memory.Span.Slice(1, 4));
+					if (memory.Length >= length)
+					{
+						ComputeMessageType type = (ComputeMessageType)memory.Span[0];
+						Message message = new Message(type, memory.Slice(HeaderLength, length - HeaderLength));
+						_inputBuffer.AdvanceReadPosition(length);
+						return message;
+					}
+				}
+
+				await _inputBuffer.WaitForWrittenDataAsync(readPosition + memory.Length, cancellationToken);
+			}
+		}
+
+		async Task SendAsync(int length, CancellationToken cancellationToken)
+		{
+			_outputBuffer.AdvanceWritePosition(length);
+
+			if (_outputBuffer.WritePosition > _flushLength)
+			{
+				await _outputBuffer.ResetWritePositionAsync(cancellationToken);
+			}
+		}
+
+		/// <inheritdoc/>
+		public IComputeMessageWriter CreateMessage(ComputeMessageType type, int sizeHint = 0)
+		{
+			if (_currentWriter != null)
+			{
+				throw new InvalidOperationException("Only one writer can be active at a time. Dispose of the previous writer first.");
+			}
+
+			_currentWriter = new Writer(this, _outputBuffer, type);
+			return _currentWriter;
+		}
+	}
+}
diff --git a/Engine/Source/Programs/Shared/EpicGames.Horde/Compute/ComputeLease.cs b/Engine/Source/Programs/Shared/EpicGames.Horde/Compute/ComputeLease.cs
--- a/Engine/Source/Programs/Shared/EpicGames.Horde/Compute/ComputeLease.cs
+++ b/Engine/Source/Programs/Shared/EpicGames.Horde/Compute/ComputeLease.cs
@@ -4,372 +4,264 @@ using System;
 using System.Buffers;
 using System.Buffers.Binary;
 using System.Collections.Generic;
-using System.Net.Sockets;
-using System.Security.Cryptography;
 using System.Threading;
+using System.Threading.Channels;
 using System.Threading.Tasks;
-using EpicGames.Core;
-using EpicGames.Serialization;
 
 namespace EpicGames.Horde.Compute
 {
 	/// <summary>
-	/// Handles a compute lease with a particular remote, multiplexing messages across a single socket.
+	/// Implementation of a compute lease using sockets to transfer data
 	/// </summary>
-	public sealed class ComputeLease : IComputeLease
+	public sealed class ComputeLease : IComputeLease, IAsyncDisposable
 	{
-		/// <summary>
-		/// Length of the required encrption key. 
-		/// </summary>
-		public const int KeyLength = 32;
-
-		/// <summary>
-		/// Length of the nonce. This should be a cryptographically random number, and does not have to be secret.
-		/// </summary>
-		public const int NonceLength = 12;
-
-		/// <summary>
-		/// Maximum length of a single message
-		/// </summary>
-		public const int MaxMessageLength = 1024 * 1024;
+		record struct Packet(int Id, ReadOnlyMemory<byte> Data, TaskCompletionSource? Fence);
 
-		sealed class Message : IComputeMessage
-		{
-			internal IMemoryOwner<byte>? _owner;
-			internal readonly TaskCompletionSource<Message> _next = new TaskCompletionSource<Message>();
-
-			/// <summary>
-			/// Type of the message
-			/// </summary>
-			public ComputeMessageType Type { get; }
-
-			/// <summary>
-			/// Message data
-			/// </summary>
-			public ReadOnlyMemory<byte> Data { get; private set; }
-
-			public Message(IMemoryOwner<byte>? owner, ComputeMessageType type, ReadOnlyMemory<byte> data)
-			{
-				_owner = owner;
-
-				Type = type;
-				Data = data;
-			}
-
-			/// <inheritdoc/>
-			public void Dispose()
-			{
-				if (_owner != null)
-				{
-					_owner.Dispose();
-					_owner = null;
-				}
-			}
-		}
-
-		/// <summary>
-		/// Allows creating new messages in rented memory
-		/// </summary>
-		class Writer : IComputeMessageWriter
+		class Buffer : ComputeBuffer
 		{
 			readonly ComputeLease _owner;
-			PooledMemoryWriter _writer;
-
-			public int Length => _writer.Length;
 
-			public Writer(ComputeLease owner, int channelId, ComputeMessageType type, PooledMemoryWriter writer)
+			public Buffer(ComputeLease owner, int id, IMemoryOwner<byte> memoryOwner)
+				: base(id, memoryOwner)
 			{
 				_owner = owner;
-				_writer = writer;
-
-				_writer.WriteUInt32(0); // Placeholder for size
-				_writer.WriteUnsignedVarInt(channelId);
-				_writer.WriteUInt8((byte)type);
 			}
 
-			public void Dispose()
+			public override async ValueTask FlushAsync(CancellationToken cancellationToken)
 			{
-				if (_writer != null)
+				TaskCompletionSource tcs = new TaskCompletionSource(TaskCreationOptions.RunContinuationsAsynchronously);
+				using (IDisposable disposable = cancellationToken.Register(() => tcs.SetCanceled(cancellationToken)))
 				{
-					_owner.ReleaseWriter(_writer);
-					_writer = null!;
+					_owner._writeQueue.Writer.TryWrite(new Packet(Id, ReadOnlyMemory<byte>.Empty, tcs));
+					await tcs.Task;
 				}
 			}
 
-			public ValueTask SendAsync(CancellationToken cancellationToken)
+			public override async ValueTask ResetWritePositionAsync(CancellationToken cancellationToken)
 			{
-				int bodySize = _writer.Length - HeaderLength;
-				_writer.Advance(FooterLength);
-
-				Span<byte> message = _writer.WrittenMemory.Span;
-
-				Span<byte> header = message.Slice(0, HeaderLength);
-				BinaryPrimitives.WriteInt32LittleEndian(header, bodySize);
-
-				Span<byte> body = message.Slice(HeaderLength, bodySize);
-				Span<byte> footer = message.Slice(header.Length + body.Length);
-				_owner._aesGcm.Encrypt(_owner._writeNonce, body, body, footer, header);
-				IncrementNonce(_owner._writeNonce);
-
-				return _owner.SendMessageAsync(_writer.WrittenMemory, cancellationToken);
-			}
-
-			/// <inheritdoc/>
-			public void Advance(int count) => _writer.Advance(count);
-
-			/// <inheritdoc/>
-			public Memory<byte> GetMemory(int sizeHint = 0) => _writer.GetMemory(sizeHint);
-
-			/// <inheritdoc/>
-			public Span<byte> GetSpan(int sizeHint = 0) => _writer.GetSpan(sizeHint);
-		}
-
-		/// <summary>
-		/// Multiplexes data onto the owner's socket
-		/// </summary>
-		class Channel : IComputeChannel
-		{
-			readonly ComputeLease _owner;
-			readonly int _channelId;
+				// Flush any pending writes
+				await FlushAsync(cancellationToken);
 
-			TaskCompletionSource<Message> _readTail;
-			Task<Message> _readNext;
+				// Write an empty packet telling the remote to flush any reads before continuing
+				_owner._writeQueue.Writer.TryWrite(new Packet(Id, ReadOnlyMemory<byte>.Empty, null));
 
-			public int Id => _channelId;
-
-			public Channel(ComputeLease owner, int channelId)
-			{
-				_owner = owner;
-				_channelId = channelId;
-				_readTail = new TaskCompletionSource<Message>();
-				_readNext = _readTail.Task;
+				// Execute the regular logic to make sure we don't have any readers
+				ResetWritePosition();
 			}
 
-			public IComputeMessageWriter CreateMessage(ComputeMessageType type, int reserve) => _owner.CreateWriter(type, _channelId, reserve);
-
-			public void Dispose()
+			public async ValueTask ClientResetWritePositionAsync(CancellationToken cancellationToken)
 			{
-				_owner.RemoveChannel(_channelId);
+				await WaitForAllReadDataAsync(cancellationToken);
+				ResetWritePosition();
 			}
 
-			public void MarkComplete()
+			public override void AdvanceWritePosition(long size)
 			{
-#pragma warning disable CA2000 // Dispose objects before losing scope
-				_readTail.SetResult(new Message(null, ComputeMessageType.None, ReadOnlyMemory<byte>.Empty));
-#pragma warning restore CA2000 // Dispose objects before losing scope
-			}
+				if (size > 0)
+				{
+					ReadOnlyMemory<byte> data = Data.Slice((int)WritePosition, (int)size);
+					_owner._writeQueue.Writer.TryWrite(new Packet(Id, data, null));
+				}
 
-			public void QueueMessage(Message message)
-			{
-				_readTail.SetResult(message);
-				_readTail = message._next;
+				base.AdvanceWritePosition(size);
 			}
 
-			/// <inheritdoc/>
-			public async ValueTask<IComputeMessage> ReadAsync(CancellationToken cancellationToken)
+			public void AdvanceWritePositionAfterRecv(long size)
 			{
-				Message message = await _readNext;
-				_readNext = message._next.Task;
-				return message;
+				base.AdvanceWritePosition(size);
 			}
 		}
 
-		/// <inheritdoc/>
-		public IComputeChannel DefaultChannel => _defaultChannel;
-
 		/// <inheritdoc/>
 		public IReadOnlyDictionary<string, int> AssignedResources { get; }
 
-		const int HeaderLength = sizeof(int); // Unencrypted size of the message
-		const int FooterLength = 16; // 16-byte auth tag for encryption
-
-		readonly Socket _socket;
-		readonly MemoryPool<byte> _pool;
-		readonly AesGcm _aesGcm;
-		readonly BackgroundTask _backgroundTask;
-		readonly Channel _defaultChannel;
-
 		readonly object _lockObject = new object();
-		readonly Dictionary<int, Channel> _channels = new Dictionary<int, Channel>();
 
-		readonly byte[] _readNonce;
-		readonly byte[] _writeNonce;
+		bool _complete;
+		readonly Task _readTask;
+		readonly Task _writeTask;
+		readonly CancellationTokenSource _cts = new CancellationTokenSource();
 
-		readonly Stack<PooledMemoryWriter> _freeWriters = new Stack<PooledMemoryWriter>();
+		readonly Dictionary<int, Buffer> _inputBuffers = new Dictionary<int, Buffer>();
+		readonly Dictionary<int, Buffer> _outputBuffers = new Dictionary<int, Buffer>();
 
-		readonly SemaphoreSlim _writeSemaphore;
+		readonly Channel<Packet> _writeQueue = Channel.CreateUnbounded<Packet>();
 
 		/// <summary>
 		/// Constructor
 		/// </summary>
-		/// <param name="socket">Socket for communication</param>
-		/// <param name="key">AES encryption key (256 bits / 32 bytes)</param>
-		/// <param name="nonce">Cryptographic nonce to identify the connection. Must be longer than <see cref="NonceLength"/>.</param>
+		/// <param name="transport">Transport to communicate with the remote</param>
 		/// <param name="assignedResources">Resources assigned to this lease</param>
-		public ComputeLease(Socket socket, ReadOnlySpan<byte> key, ReadOnlySpan<byte> nonce, IReadOnlyDictionary<string, int> assignedResources)
-			: this(socket, key, nonce, assignedResources, MemoryPool<byte>.Shared)
+		public ComputeLease(IComputeTransport transport, IReadOnlyDictionary<string, int> assignedResources)
 		{
+			AssignedResources = assignedResources;
+
+			_readTask = RunReaderAsync(transport, _cts.Token);
+			_writeTask = RunWriterAsync(transport, _cts.Token);
 		}
 
 		/// <summary>
-		/// Constructor
+		/// Attempt to gracefully close the current connection and shutdown both ends of the transport
 		/// </summary>
-		/// <param name="socket">Socket for communication</param>
-		/// <param name="key">AES encryption key (256 bits / 32 bytes)</param>
-		/// <param name="nonce">Cryptographic nonce to identify the connection. Must be longer than <see cref="NonceLength"/>.</param>
-		/// <param name="assignedResources">Resources assigned to this lease</param>
-		/// <param name="pool">Memory pool from which to allocate message buffers</param>
-		public ComputeLease(Socket socket, ReadOnlySpan<byte> key, ReadOnlySpan<byte> nonce, IReadOnlyDictionary<string, int> assignedResources, MemoryPool<byte> pool)
+		public async ValueTask CloseAsync(CancellationToken cancellationToken)
 		{
-			if (key.Length != KeyLength)
-			{
-				throw new ArgumentException($"Key must be {KeyLength} bytes", nameof(key));
-			}
-
-			_socket = socket;
-			_pool = pool;
-			_aesGcm = new AesGcm(key);
-
-			_readNonce = nonce.Slice(0, NonceLength).ToArray();
-			_writeNonce = nonce.Slice(0, NonceLength).ToArray();
-
-			_writeSemaphore = new SemaphoreSlim(1);
+			_writeQueue.Writer.TryComplete();
 
-			_defaultChannel = new Channel(this, 0);
-			_channels.Add(0, _defaultChannel);
+			Task cancellationTask = Task.Delay(-1, cancellationToken);
 
-			_backgroundTask = BackgroundTask.StartNew(BackgroundReadAsync);
+			await Task.WhenAny(_writeTask, cancellationTask);
+			cancellationToken.ThrowIfCancellationRequested();
 
-			AssignedResources = assignedResources;
+			await Task.WhenAny(_readTask, cancellationTask);
+			cancellationToken.ThrowIfCancellationRequested();
 		}
 
 		/// <inheritdoc/>
 		public async ValueTask DisposeAsync()
 		{
-			await _backgroundTask.DisposeAsync();
+			_writeQueue.Writer.TryComplete();
 
-			foreach (Channel channel in _channels.Values)
-			{
-				channel.MarkComplete();
-			}
+			_cts.Cancel();
 
-			_writeSemaphore.Dispose();
-			_aesGcm.Dispose();
-		}
+			await _writeTask;
+			await _readTask;
 
-		/// <summary>
-		/// Creates a new encryption key for the compute lead
-		/// </summary>
-		/// <returns>New key bytes</returns>
-		public static byte[] CreateKey() => RandomNumberGenerator.GetBytes(KeyLength);
+			_cts.Dispose();
 
-		/// <inheritdoc/>
-		public IComputeChannel OpenChannel(int channelId)
-		{
-			Channel channel = new Channel(this, channelId);
-			lock (_lockObject)
+			foreach (Buffer inputBuffer in _inputBuffers.Values)
+			{
+				inputBuffer.Dispose();
+			}
+
+			foreach (Buffer outputBuffer in _outputBuffers.Values)
 			{
-				_channels.Add(channelId, channel);
+				outputBuffer.Dispose();
 			}
-			return channel;
 		}
 
-		async Task BackgroundReadAsync(CancellationToken cancellationToken)
+		async Task RunReaderAsync(IComputeTransport transport, CancellationToken cancellationToken)
 		{
-			byte[] sizeBuffer = new byte[4];
-			while (!cancellationToken.IsCancellationRequested)
+			byte[] header = new byte[8];
+			try
 			{
-				if (!await _socket.TryReceiveMessageAsync(sizeBuffer, SocketFlags.None, cancellationToken))
-				{
-					break;
-				}
+				// Maintain a local cache of buffers to be able to query for them without having to acquire a global lock
+				Dictionary<int, Buffer> buffers = new Dictionary<int, Buffer>();
 
-				int length = BinaryPrimitives.ReadInt32LittleEndian(sizeBuffer);
-				if (length < 0 || length > MaxMessageLength)
+				// Process messages from the remote
+				for(; ;)
 				{
-					throw new InvalidOperationException();
-				}
+					await transport.ReadAsync(header, cancellationToken);
 
-				IMemoryOwner<byte> owner = _pool.Rent(length + FooterLength);
-				try
-				{
-					// Decrypt the message
-					await _socket.ReceiveMessageAsync(owner.Memory.Slice(0, length + FooterLength), SocketFlags.None, cancellationToken);
-					_aesGcm.Decrypt(_readNonce, owner.Memory.Span.Slice(0, length), owner.Memory.Span.Slice(length, FooterLength), owner.Memory.Span.Slice(0, length), sizeBuffer);
-					IncrementNonce(_readNonce);
-
-					// Decode the message data
-					ReadOnlyMemory<byte> data = owner.Memory.Slice(0, length);
-
-					int channelId = (int)VarInt.ReadUnsigned(data.Span, out int channelIdByteCount);
-					data = data.Slice(channelIdByteCount);
-
-					ComputeMessageType type = (ComputeMessageType)data.Span[0];
-					data = data.Slice(1);
-
-#pragma warning disable CA2000 // Dispose objects before losing scope
-					Message message = new Message(owner, type, data);
-#pragma warning restore CA2000 // Dispose objects before losing scope
+					int id = BinaryPrimitives.ReadInt32LittleEndian(header);
+					int size = BinaryPrimitives.ReadInt32LittleEndian(header.AsSpan(4));
+					if (size < 0)
+					{
+						break;
+					}
 
-					Channel? channel;
-					lock (_lockObject)
+					Buffer? buffer;
+					if (!buffers.TryGetValue(id, out buffer))
 					{
-						channel = _channels[channelId];
+						lock (_lockObject)
+						{
+							if (!_inputBuffers.TryGetValue(id, out buffer))
+							{
+								throw new NotImplementedException();
+							}
+							buffers.Add(id, buffer);
+						}
 					}
 
-					channel.QueueMessage(message);
-				}
-				catch
-				{
-					owner.Dispose();
-					throw;
+					if (size == 0)
+					{
+						// We're being asked to reset the read position back to zero, so we need to know that the read offset is equal to length
+						await buffer.ClientResetWritePositionAsync(cancellationToken);
+					}
+					else
+					{
+						await transport.ReadAsync(buffer.Data.Slice((int)buffer.WritePosition, size), cancellationToken);
+						buffer.AdvanceWritePositionAfterRecv(size);
+					}
 				}
 			}
-		}
+			catch (OperationCanceledException)
+			{
+			}
 
-		void RemoveChannel(int channelId)
-		{
+			// Mark all buffers as complete
 			lock (_lockObject)
 			{
-				_channels.Remove(channelId);
+				_complete = true;
+				foreach (Buffer buffer in _outputBuffers.Values)
+				{
+					buffer.MarkComplete();
+				}
 			}
 		}
 
-		Writer CreateWriter(ComputeMessageType type, int channelId, int reserve)
+		async Task RunWriterAsync(IComputeTransport transport, CancellationToken cancellationToken)
 		{
-			PooledMemoryWriter? writer;
-			if (_freeWriters.TryPop(out writer))
+			// Parse the queue of write requests
+			byte[] header = new byte[8];
+			try
 			{
-				writer.Clear(reserve);
+				while (await _writeQueue.Reader.WaitToReadAsync(cancellationToken))
+				{
+					Packet packet = await _writeQueue.Reader.ReadAsync(cancellationToken);
+					if (packet.Fence != null)
+					{
+						packet.Fence.SetResult();
+						continue;
+					}
+
+					BinaryPrimitives.WriteInt32LittleEndian(header, packet.Id);
+					BinaryPrimitives.WriteInt32LittleEndian(header.AsSpan(4), packet.Data.Length);
+					await transport.WriteAsync(header, cancellationToken);
+
+					if (packet.Data.Length > 0)
+					{
+						await transport.WriteAsync(packet.Data, cancellationToken);
+					}
+				}
+
+				// Send a close packet to the remote
+				BinaryPrimitives.WriteInt32LittleEndian(header, 0);
+				BinaryPrimitives.WriteInt32LittleEndian(header.AsSpan(4), -1);
+				await transport.WriteAsync(header, cancellationToken);
 			}
-			else
+			catch (OperationCanceledException)
 			{
-				writer = new PooledMemoryWriter(reserve);
 			}
-			return new Writer(this, channelId, type, writer);
-		}
-
-		void ReleaseWriter(PooledMemoryWriter writer)
-		{
-			_freeWriters.Push(writer);
 		}
 
-		async ValueTask SendMessageAsync(ReadOnlyMemory<byte> message, CancellationToken cancellationToken)
+		/// <inheritdoc/>
+		public IComputeInputBuffer CreateInputBuffer(int id, IMemoryOwner<byte> memoryOwner)
 		{
-			await _writeSemaphore.WaitAsync(cancellationToken);
-			try
-			{
-				await _socket.SendMessageAsync(message, SocketFlags.None, cancellationToken);
-			}
-			finally
+			Buffer inputBuffer = new Buffer(this, id, memoryOwner);
+			lock (_lockObject)
 			{
-				_writeSemaphore.Release();
+				if (_complete)
+				{
+					inputBuffer.MarkComplete();
+				}
+				_inputBuffers.Add(id, inputBuffer);
 			}
+			return inputBuffer;
 		}
 
-		static void IncrementNonce(byte[] nonce)
+		/// <inheritdoc/>
+		public IComputeOutputBuffer CreateOutputBuffer(int id, IMemoryOwner<byte> memoryOwner)
 		{
-			BinaryPrimitives.WriteInt64LittleEndian(nonce, BinaryPrimitives.ReadInt64LittleEndian(nonce) + 1);
+			Buffer outputBuffer = new Buffer(this, id, memoryOwner);
+			lock (_lockObject)
+			{
+				if (_complete)
+				{
+					outputBuffer.MarkComplete();
+				}
+				_outputBuffers.Add(id, outputBuffer);
+			}
+			return outputBuffer;
 		}
 	}
 }
diff --git a/Engine/Source/Programs/Shared/EpicGames.Horde/Compute/ComputeStorageClient.cs b/Engine/Source/Programs/Shared/EpicGames.Horde/Compute/ComputeStorageClient.cs
--- a/Engine/Source/Programs/Shared/EpicGames.Horde/Compute/ComputeStorageClient.cs
+++ b/Engine/Source/Programs/Shared/EpicGames.Horde/Compute/ComputeStorageClient.cs
@@ -1,6 +1,7 @@
 // Copyright Epic Games, Inc. All Rights Reserved.
 
 using System;
+using System.Buffers;
 using System.Collections.Generic;
 using System.IO;
 using System.Threading;
diff --git a/Engine/Source/Programs/Shared/EpicGames.Horde/Compute/IComputeBuffer.cs b/Engine/Source/Programs/Shared/EpicGames.Horde/Compute/IComputeBuffer.cs
new file mode 100644
--- /dev/null
+++ b/Engine/Source/Programs/Shared/EpicGames.Horde/Compute/IComputeBuffer.cs
@@ -0,0 +1,105 @@
+// Copyright Epic Games, Inc. All Rights Reserved.
+
+using System;
+using System.Threading;
+using System.Threading.Tasks;
+
+namespace EpicGames.Horde.Compute
+{
+	/// <summary>
+	/// Base interface for input and output buffers
+	/// </summary>
+	public interface IComputeBuffer
+	{
+		/// <summary>
+		/// Identifier for the buffer
+		/// </summary>
+		int Id { get; }
+	}
+
+	/// <summary>
+	/// Read interface for a compute buffer
+	/// </summary>
+	public interface IComputeInputBuffer : IComputeBuffer
+	{
+		/// <summary>
+		/// The shared memory for the buffer
+		/// </summary>
+		ReadOnlyMemory<byte> Data { get; }
+
+		/// <summary>
+		/// Offset within the data that is still to be used. Can be left at zero for a buffer that is only appended to, but must be updated to indicate data
+		/// that has been processed for message buffers (in order to coordinate when the buffer contents can be discarded during a reset).
+		/// </summary>
+		long ReadPosition { get; }
+
+		/// <summary>
+		/// Position in the buffer at which new data will be added
+		/// </summary>
+		long WritePosition { get; }
+
+		/// <summary>
+		/// Updates the current read position within the buffer
+		/// </summary>
+		void AdvanceReadPosition(long size);
+
+		/// <summary>
+		/// Waits until new data is written to the buffer
+		/// </summary>
+		/// <param name="currentWritePosition">Current write position. This is passed as an argument to avoid races with <see cref="WritePosition"/> being updated asynchronously.</param>
+		/// <param name="cancellationToken">Cancellation token for the operation</param>
+		ValueTask WaitForWrittenDataAsync(long currentWritePosition, CancellationToken cancellationToken);
+	}
+
+	/// <summary>
+	/// Read-only view of shared memory that can be incrementally replicated to a remote machine.
+	/// </summary>
+	public interface IComputeOutputBuffer : IComputeBuffer
+	{
+		/// <summary>
+		/// The shared memory for the buffer
+		/// </summary>
+		Memory<byte> Data { get; }
+
+		/// <summary>
+		/// Current write position/used length of the buffer
+		/// </summary>
+		long WritePosition { get; }
+
+		/// <summary>
+		/// Updates the current write position within the buffer
+		/// </summary>
+		void AdvanceWritePosition(long size);
+
+		/// <summary>
+		/// Flushes all data that needs to be written.
+		/// </summary>
+		ValueTask FlushAsync(CancellationToken cancellationToken);
+
+		/// <summary>
+		/// Flushes all writes, and resets the current write position.
+		/// </summary>
+		ValueTask ResetWritePositionAsync(CancellationToken cancellationToken);
+	}
+
+	/// <summary>
+	/// Extension methods for compute buffers
+	/// </summary>
+	static class ComputeBufferExtensions
+	{
+		/// <summary>
+		/// Gets a span for the next data to read
+		/// </summary>
+		public static ReadOnlySpan<byte> GetReadSpan(this IComputeInputBuffer buffer) => GetReadMemory(buffer).Span;
+
+		/// <summary>
+		/// Gets a span for the next data to read
+		/// </summary>
+		public static ReadOnlyMemory<byte> GetReadMemory(this IComputeInputBuffer buffer) => buffer.Data.Slice((int)buffer.ReadPosition, (int)(buffer.WritePosition - buffer.ReadPosition));
+
+		/// <summary>
+		/// Gets a span to write new data to
+		/// </summary>
+		public static Span<byte> GetWriteSpan(this IComputeOutputBuffer buffer) => buffer.Data.Slice((int)buffer.WritePosition).Span;
+	}
+}
diff --git a/Engine/Source/Programs/Shared/EpicGames.Horde/Compute/IComputeChannel.cs b/Engine/Source/Programs/Shared/EpicGames.Horde/Compute/IComputeChannel.cs
--- a/Engine/Source/Programs/Shared/EpicGames.Horde/Compute/IComputeChannel.cs
+++ b/Engine/Source/Programs/Shared/EpicGames.Horde/Compute/IComputeChannel.cs
@@ -1,6 +1,7 @@
 // Copyright Epic Games, Inc. All Rights Reserved.
 
 using System;
+using System.Buffers;
 using System.Threading;
 using System.Threading.Tasks;
 using EpicGames.Core;
@@ -10,10 +11,10 @@ namespace EpicGames.Horde.Compute
 	/// <summary>
 	/// Full-duplex channel for sending and receiving messages
 	/// </summary>
-	public interface IComputeChannel
+	public interface IComputeChannel : IDisposable
 	{
 		/// <summary>
-		/// Identifier for the current channel
+		/// Identifier for the channel
 		/// </summary>
 		int Id { get; }
 
@@ -44,4 +45,26 @@ namespace EpicGames.Horde.Compute
 		/// <param name="cancellationToken">Cancellation token for the operation</param>
 		ValueTask SendAsync(CancellationToken cancellationToken);
 	}
+
+	/// <summary>
+	/// Extension methods to allow creating channels from leases
+	/// </summary>
+	public static class ComputeChannelExtensions
+	{
+		/// <summary>
+		/// Creates a compute channel with the given identifier
+		/// </summary>
+		/// <param name="lease">The lease to create a channel for</param>
+		/// <param name="id">Identifier for the channel</param>
+		public static IComputeChannel CreateChannel(this IComputeLease lease, int id)
+		{
+			const int MaxMessageSize = 64 * 1024;
+			const int BufferSize = MaxMessageSize * 3;
+
+			IComputeInputBuffer inputBuffer = lease.CreateInputBuffer(id, MemoryPool<byte>.Shared.Rent(BufferSize));
+			IComputeOutputBuffer outputBuffer = lease.CreateOutputBuffer(id, MemoryPool<byte>.Shared.Rent(BufferSize));
+
+			return new ComputeChannel(id, inputBuffer, outputBuffer, MaxMessageSize);
+		}
+	}
 }
diff --git a/Engine/Source/Programs/Shared/EpicGames.Horde/Compute/IComputeLease.cs b/Engine/Source/Programs/Shared/EpicGames.Horde/Compute/IComputeLease.cs
--- a/Engine/Source/Programs/Shared/EpicGames.Horde/Compute/IComputeLease.cs
+++ b/Engine/Source/Programs/Shared/EpicGames.Horde/Compute/IComputeLease.cs
@@ -1,7 +1,10 @@
 // Copyright Epic Games, Inc. All Rights Reserved.
 
 using System;
+using System.Buffers;
 using System.Collections.Generic;
+using System.Threading;
+using System.Threading.Tasks;
 
 namespace EpicGames.Horde.Compute
 {
@@ -11,20 +14,27 @@ namespace EpicGames.Horde.Compute
 	public interface IComputeLease : IAsyncDisposable
 	{
 		/// <summary>
-		/// The control channel, used to transmit messages before additional channels have been created
+		/// Resources assigned to this lease
 		/// </summary>
-		IComputeChannel DefaultChannel { get; }
+		IReadOnlyDictionary<string, int> AssignedResources { get; }
 
 		/// <summary>
-		/// Resources assigned to this lease
+		/// Closes the underlying transport stream gracefully
 		/// </summary>
-		IReadOnlyDictionary<string, int> AssignedResources { get; }
+		ValueTask CloseAsync(CancellationToken cancellationToken);
+
+		/// <summary>
+		/// Creates an input buffer from the given memory. A corresponding output buffer with the same size and id must be created at the remote end.
+		/// </summary>
+		/// <param name="id">Identifier for the buffer. This much match up with an output buffer on the remote for a connection to be established.</param>
+		/// <param name="memory">Memory to read from</param>
+		IComputeInputBuffer CreateInputBuffer(int id, IMemoryOwner<byte> memory);
 
 		/// <summary>
-		/// Opens a new channel with the given identifier.
+		/// Creates an output buffer from the given memory. A corresponding output buffer with the same size and id must be created at the remote end.
 		/// </summary>
-		/// <param name="id">The pipe identifier. These values are arbitrary, as long as both ends agree on a channel to communicate on. Channel 0 is reserved for the control channel.</param>
-		/// <returns>New channel instance</returns>
-		IComputeChannel OpenChannel(int id);
+		/// <param name="id">Identifier for the buffer. This much match up with an input buffer on the remote for a connection to be established.</param>
+		/// <param name="memory">Memory to read from</param>
+		IComputeOutputBuffer CreateOutputBuffer(int id, IMemoryOwner<byte> memory);
 	}
 }
diff --git a/Engine/Source/Programs/Shared/EpicGames.Horde/Compute/IComputeMessage.cs b/Engine/Source/Programs/Shared/EpicGames.Horde/Compute/IComputeMessage.cs
--- a/Engine/Source/Programs/Shared/EpicGames.Horde/Compute/IComputeMessage.cs
+++ b/Engine/Source/Programs/Shared/EpicGames.Horde/Compute/IComputeMessage.cs
@@ -1,12 +1,9 @@
 // Copyright Epic Games, Inc. All Rights Reserved.
 
 using System;
-using System.IO;
 using System.Threading;
 using System.Threading.Tasks;
 using EpicGames.Core;
-using EpicGames.Horde.Storage;
-using Microsoft.CodeAnalysis.CSharp.Syntax;
 
 namespace EpicGames.Horde.Compute
 {
@@ -100,7 +97,7 @@ namespace EpicGames.Horde.Compute
 	/// <summary>
 	/// Wraps various requests across compute channels
 	/// </summary>
-	public static class ComputeChannelExtensions
+	public static class ComputeMessageExtensions
 	{
 		/// <summary>
 		/// Send a close message to the remote
diff --git a/Engine/Source/Programs/Shared/EpicGames.Horde/Compute/LoopbackComputeClient.cs b/Engine/Source/Programs/Shared/EpicGames.Horde/Compute/LoopbackComputeClient.cs
--- a/Engine/Source/Programs/Shared/EpicGames.Horde/Compute/LoopbackComputeClient.cs
+++ b/Engine/Source/Programs/Shared/EpicGames.Horde/Compute/LoopbackComputeClient.cs
@@ -16,8 +16,6 @@ namespace EpicGames.Horde.Compute
 	/// </summary>
 	public sealed class LoopbackComputeClient : IComputeClient
 	{
-		readonly byte[] _key;
-		readonly byte[] _nonce;
 		readonly BackgroundTask _listenerTask;
 		readonly Socket _listener;
 		readonly Socket _socket;
@@ -29,14 +27,11 @@ namespace EpicGames.Horde.Compute
 		/// <param name="port">Port to connect on</param>
 		public LoopbackComputeClient(Func<IComputeLease, CancellationToken, Task> serverFunc, int port = 2000)
 		{
-			_key = ComputeLease.CreateKey();
-			_nonce = RandomNumberGenerator.GetBytes(ComputeLease.NonceLength);
-
 			_listener = new Socket(SocketType.Stream, ProtocolType.IP);
 			_listener.Bind(new IPEndPoint(IPAddress.Loopback, port));
 			_listener.Listen();
 
-			_listenerTask = BackgroundTask.StartNew(ctx => RunListenerAsync(_listener, _key, _nonce, serverFunc, ctx));
+			_listenerTask = BackgroundTask.StartNew(ctx => RunListenerAsync(_listener, serverFunc, ctx));
 
 			_socket = new Socket(SocketType.Stream, ProtocolType.IP);
 			_socket.Connect(IPAddress.Loopback, port);
@@ -53,19 +48,22 @@ namespace EpicGames.Horde.Compute
 		/// <summary>
 		/// Sets up the loopback listener and calls the server method
 		/// </summary>
-		static async Task RunListenerAsync(Socket listener, ReadOnlyMemory<byte> key, ReadOnlyMemory<byte> nonce, Func<IComputeLease, CancellationToken, Task> serverFunc, CancellationToken cancellationToken)
+		static async Task RunListenerAsync(Socket listener, Func<IComputeLease, CancellationToken, Task> serverFunc, CancellationToken cancellationToken)
 		{
 			using Socket socket = await listener.AcceptAsync(cancellationToken);
 
-			await using ComputeLease lease = new ComputeLease(socket, key.Span, nonce.Span, new Dictionary<string, int>());
+			await using ComputeLease lease = new ComputeLease(new TcpComputeTransport(socket), new Dictionary<string, int>());
 			await serverFunc(lease, cancellationToken);
+			await lease.CloseAsync(cancellationToken);
 		}
 
 		/// <inheritdoc/>
 		public async Task<TResult> ExecuteAsync<TResult>(ClusterId clusterId, Requirements? requirements, Func<IComputeLease, CancellationToken, Task<TResult>> handler, CancellationToken cancellationToken)
 		{
-			await using ComputeLease lease = new ComputeLease(_socket, _key, _nonce, new Dictionary<string, int>());
-			return await handler(lease, cancellationToken);
+			await using ComputeLease lease = new ComputeLease(new TcpComputeTransport(_socket), new Dictionary<string, int>());
+			TResult result = await handler(lease, cancellationToken);
+			await lease.CloseAsync(cancellationToken);
+			return result;
 		}
 	}
 }
diff --git a/Engine/Source/Programs/Shared/EpicGames.Horde/Compute/ServerComputeClient.cs b/Engine/Source/Programs/Shared/EpicGames.Horde/Compute/ServerComputeClient.cs
--- a/Engine/Source/Programs/Shared/EpicGames.Horde/Compute/ServerComputeClient.cs
+++ b/Engine/Source/Programs/Shared/EpicGames.Horde/Compute/ServerComputeClient.cs
@@ -129,7 +129,7 @@ namespace EpicGames.Horde.Compute
 			// Pass the rest of the call over to the handler
 			byte[] key = StringUtils.ParseHexString(responseMessage.Key);
 
-			await using ComputeLease channel = new ComputeLease(socket, key, nonce, responseMessage.AssignedResources);
+			await using ComputeLease channel = new ComputeLease(new TcpComputeTransport(socket), responseMessage.AssignedResources);
 			return await handler(channel, cancellationToken);
 		}
 	}
diff --git a/Engine/Source/Programs/Shared/EpicGames.Horde/Compute/SharedMemory.cs b/Engine/Source/Programs/Shared/EpicGames.Horde/Compute/SharedMemory.cs
new file mode 100644
--- /dev/null
+++ b/Engine/Source/Programs/Shared/EpicGames.Horde/Compute/SharedMemory.cs
@@ -0,0 +1,98 @@
+// Copyright Epic Games, Inc. All Rights Reserved.
+
+using Microsoft.Win32.SafeHandles;
+using System;
+using System.Buffers;
+using System.IO.MemoryMappedFiles;
+
+namespace EpicGames.Horde.Compute
+{
+	/// <summary>
+	/// Wraps the resources required for an unnamed memory mapped file and allows accessing it from C# as a <see cref="Memory{Byte}"/> object.
+	/// </summary>
+	public sealed unsafe class SharedMemory : IMemoryOwner<byte>
+	{
+		sealed unsafe class MemoryWrapper : MemoryManager<byte>
+		{
+			private readonly byte* _pointer;
+			private readonly int _length;
+
+			public MemoryWrapper(byte* pointer, int length)
+			{
+				_pointer = pointer;
+				_length = length;
+			}
+
+			/// <inheritdoc/>
+			public override Span<byte> GetSpan() => new Span<byte>(_pointer, _length);
+
+			/// <inheritdoc/>
+			public override MemoryHandle Pin(int elementIndex) => new MemoryHandle(_pointer + elementIndex);
+
+			/// <inheritdoc/>
+			public override void Unpin() { }
+
+			/// <inheritdoc/>
+			protected override void Dispose(bool disposing) { }
+		}
+
+		readonly MemoryMappedFile _file;
+		readonly MemoryMappedViewAccessor _viewAccessor;
+		byte* _data;
+
+		/// <summary>
+		/// Access the memory allocated through this buffer
+		/// </summary>
+		public Memory<byte> Memory { get; }
+
+		/// <summary>
+		/// Gets a handle for the underlying memory mapped file
+		/// </summary>
+		public SafeMemoryMappedFileHandle SafeMemoryMappedFileHandle => _file.SafeMemoryMappedFileHandle;
+
+		/// <summary>
+		/// Constructor
+		/// </summary>
+		/// <param name="capacity">Size of the memory to allocate. Data will not be committed until used.</param>
+		public SharedMemory(long capacity)
+		{
+			_file = MemoryMappedFile.CreateNew(null, capacity);
+
+			_viewAccessor = _file.CreateViewAccessor(0, capacity);
+			_viewAccessor.SafeMemoryMappedViewHandle.AcquirePointer(ref _data);
+
+			MemoryWrapper wrapper = new MemoryWrapper(_data, (int)capacity);
+			Memory = wrapper.Memory;
+		}
+
+		/// <summary>
+		/// Finalizer
+		/// </summary>
+		~SharedMemory()
+		{
+			Dispose(false);
+		}
+
+		/// <inheritdoc/>
+		public void Dispose()
+		{
+			Dispose(true);
+			GC.SuppressFinalize(this);
+		}
+
+		void Dispose(bool disposing)
+		{
+			if (_data != null)
+			{
+				_viewAccessor.SafeMemoryMappedViewHandle.ReleasePointer();
+				_data = null;
+			}
+
+			if (disposing)
+			{
+				_viewAccessor.Dispose();
+				_file.Dispose();
+			}
+		}
+	}
+}
