Horde: Add a pluggable architecture for compute transports. Add implementations using regular C# streams, TCP sockets, and an AES-encrypted wrapper over any underlying transport.

diff --git a/Engine/Source/Programs/Horde/Horde.Agent/Leases/Handlers/ComputeHandler.cs b/Engine/Source/Programs/Horde/Horde.Agent/Leases/Handlers/ComputeHandler.cs
--- a/Engine/Source/Programs/Horde/Horde.Agent/Leases/Handlers/ComputeHandler.cs
+++ b/Engine/Source/Programs/Horde/Horde.Agent/Leases/Handlers/ComputeHandler.cs
@@ -10,6 +10,7 @@ using System.Threading.Tasks;
 using EpicGames.Core;
 using EpicGames.Horde.Compute;
 using EpicGames.Horde.Compute.Cpp;
+using EpicGames.Horde.Compute.Transports;
 using EpicGames.Horde.Logs;
 using EpicGames.Horde.Storage;
 using EpicGames.Horde.Storage.Backends;
@@ -57,7 +58,7 @@ namespace Horde.Agent.Leases.Handlers
 					return LeaseResult.Success;
 				}
 
-				await using (ComputeLease lease = new ComputeLease(new TcpComputeTransport(tcpClient.Client), computeTask.Resources))
+				await using (ComputeLease lease = new ComputeLease(new TcpTransport(tcpClient.Client), computeTask.Resources))
 				{
 					await RunAsync(lease, cancellationToken);
 					return LeaseResult.Success;
diff --git a/Engine/Source/Programs/Shared/EpicGames.Horde.Tests/TransportTests.cs b/Engine/Source/Programs/Shared/EpicGames.Horde.Tests/TransportTests.cs
new file mode 100644
--- /dev/null
+++ b/Engine/Source/Programs/Shared/EpicGames.Horde.Tests/TransportTests.cs
@@ -0,0 +1,98 @@
+// Copyright Epic Games, Inc. All Rights Reserved.
+
+using System;
+using System.IO;
+using System.Linq;
+using System.Security.Cryptography;
+using System.Threading;
+using System.Threading.Tasks;
+using EpicGames.Horde.Compute;
+using EpicGames.Horde.Compute.Transports;
+using Microsoft.VisualStudio.TestTools.UnitTesting;
+
+namespace EpicGames.Horde.Tests
+{
+	[TestClass]
+	public class TransportTests
+	{
+		[TestMethod]
+		public async Task TestStreamTransport()
+		{
+			byte[] input = new byte[256 * 1024];
+			new Random(0).NextBytes(input);
+
+			byte[] payload;
+			{
+				using MemoryStream memoryStream = new MemoryStream();
+				StreamTransport streamTransport = new StreamTransport(memoryStream);
+				await streamTransport.WriteAsync(input, CancellationToken.None);
+				payload = memoryStream.ToArray();
+			}
+
+			byte[] output = new byte[input.Length];
+			{
+				using MemoryStream memoryStream = new MemoryStream(payload);
+				StreamTransport streamTransport = new StreamTransport(memoryStream);
+				await CopyToAsync(streamTransport, output);
+			}
+
+			for (int idx = 0; idx < input.Length; idx++)
+			{
+				Assert.IsTrue(input[idx] == output[idx]);
+			}
+			Assert.IsTrue(input.SequenceEqual(output));
+
+		}
+
+		[TestMethod]
+		public async Task TestAesTransport()
+		{
+			byte[] key = new byte[AesTransport.KeyLength];
+			new Random(1).NextBytes(key);
+
+			byte[] nonce = new byte[AesTransport.NonceLength];
+			new Random(2).NextBytes(nonce);
+
+			byte[] input = RandomNumberGenerator.GetBytes(256 * 1024);
+			new Random(0).NextBytes(input);
+
+			byte[] encrypted;
+
+			{
+				using MemoryStream memoryStream = new MemoryStream();
+				StreamTransport streamTransport = new StreamTransport(memoryStream);
+				await using AesTransport aesTransport = new AesTransport(streamTransport, key, nonce);
+				await aesTransport.WriteAsync(input, CancellationToken.None);
+				encrypted = memoryStream.ToArray();
+			}
+
+			byte[] output = new byte[input.Length];
+			{
+				using MemoryStream memoryStream = new MemoryStream(encrypted);
+				StreamTransport streamTransport = new StreamTransport(memoryStream);
+				await using AesTransport aesTransport = new AesTransport(streamTransport, key, nonce);
+
+				for (int offset = 0; offset < output.Length;)
+				{
+					offset += await aesTransport.ReadAsync(output.AsMemory(offset), CancellationToken.None);
+				}
+			}
+
+			for (int idx = 0; idx < input.Length; idx++)
+			{
+				Assert.IsTrue(input[idx] == output[idx]);
+			}
+			Assert.IsTrue(input.SequenceEqual(output));
+		}
+
+		public static async Task CopyToAsync(IComputeTransport transport, Memory<byte> buffer)
+		{
+			for (int offset = 0; offset < buffer.Length;)
+			{
+				int size = await transport.ReadAsync(buffer.Slice(offset), CancellationToken.None);
+				Assert.IsTrue(size > 0);
+				offset += size;
+			}
+		}
+	}
+}
diff --git a/Engine/Source/Programs/Shared/EpicGames.Horde/Compute/IComputeTransport.cs b/Engine/Source/Programs/Shared/EpicGames.Horde/Compute/IComputeTransport.cs
new file mode 100644
--- /dev/null
+++ b/Engine/Source/Programs/Shared/EpicGames.Horde/Compute/IComputeTransport.cs
@@ -0,0 +1,45 @@
+// Copyright Epic Games, Inc. All Rights Reserved.
+
+using System;
+using System.Buffers;
+using System.Threading;
+using System.Threading.Tasks;
+
+namespace EpicGames.Horde.Compute
+{
+	/// <summary>
+	/// Interface for transferring data on a compute channel
+	/// </summary>
+	public interface IComputeTransport
+	{
+		/// <summary>
+		/// Reads data from the underlying transport into an output buffer
+		/// </summary>
+		/// <param name="buffer">Buffer to read into</param>
+		/// <param name="cancellationToken">Cancellation token for the operation</param>
+		public ValueTask<int> ReadAsync(Memory<byte> buffer, CancellationToken cancellationToken);
+
+		/// <summary>
+		/// Writes data to the underlying transport
+		/// </summary>
+		/// <param name="buffer">Buffer to be written</param>
+		/// <param name="cancellationToken">Cancellation token for the operation</param>
+		public ValueTask WriteAsync(ReadOnlySequence<byte> buffer, CancellationToken cancellationToken);
+	}
+
+	/// <summary>
+	/// Extension methods for <see cref="IComputeTransport"/>
+	/// </summary>
+	public static class ComputeTransportExtensions
+	{
+		/// <summary>
+		/// Writes data to the underlying transport
+		/// </summary>
+		/// <param name="transport">Transport instance</param>
+		/// <param name="buffer">Buffer to be written</param>
+		/// <param name="cancellationToken">Cancellation token for the operation</param>
+		public static ValueTask WriteAsync(this IComputeTransport transport, ReadOnlyMemory<byte> buffer, CancellationToken cancellationToken)
+			=> transport.WriteAsync(new ReadOnlySequence<byte>(buffer), cancellationToken);
+	}
+}
+
diff --git a/Engine/Source/Programs/Shared/EpicGames.Horde/Compute/LoopbackComputeClient.cs b/Engine/Source/Programs/Shared/EpicGames.Horde/Compute/LoopbackComputeClient.cs
--- a/Engine/Source/Programs/Shared/EpicGames.Horde/Compute/LoopbackComputeClient.cs
+++ b/Engine/Source/Programs/Shared/EpicGames.Horde/Compute/LoopbackComputeClient.cs
@@ -8,6 +8,7 @@ using System.Security.Cryptography;
 using System.Threading;
 using System.Threading.Tasks;
 using EpicGames.Core;
+using EpicGames.Horde.Compute.Transports;
 
 namespace EpicGames.Horde.Compute
 {
@@ -52,7 +53,7 @@ namespace EpicGames.Horde.Compute
 		{
 			using Socket socket = await listener.AcceptAsync(cancellationToken);
 
-			await using ComputeLease lease = new ComputeLease(new TcpComputeTransport(socket), new Dictionary<string, int>());
+			await using ComputeLease lease = new ComputeLease(new TcpTransport(socket), new Dictionary<string, int>());
 			await serverFunc(lease, cancellationToken);
 			await lease.CloseAsync(cancellationToken);
 		}
@@ -60,7 +61,7 @@ namespace EpicGames.Horde.Compute
 		/// <inheritdoc/>
 		public async Task<TResult> ExecuteAsync<TResult>(ClusterId clusterId, Requirements? requirements, Func<IComputeLease, CancellationToken, Task<TResult>> handler, CancellationToken cancellationToken)
 		{
-			await using ComputeLease lease = new ComputeLease(new TcpComputeTransport(_socket), new Dictionary<string, int>());
+			await using ComputeLease lease = new ComputeLease(new TcpTransport(_socket), new Dictionary<string, int>());
 			TResult result = await handler(lease, cancellationToken);
 			await lease.CloseAsync(cancellationToken);
 			return result;
diff --git a/Engine/Source/Programs/Shared/EpicGames.Horde/Compute/ServerComputeClient.cs b/Engine/Source/Programs/Shared/EpicGames.Horde/Compute/ServerComputeClient.cs
--- a/Engine/Source/Programs/Shared/EpicGames.Horde/Compute/ServerComputeClient.cs
+++ b/Engine/Source/Programs/Shared/EpicGames.Horde/Compute/ServerComputeClient.cs
@@ -1,6 +1,7 @@
 // Copyright Epic Games, Inc. All Rights Reserved.
 
 using EpicGames.Core;
+using EpicGames.Horde.Compute.Transports;
 using Microsoft.Extensions.Logging;
 using System;
 using System.Collections.Generic;
@@ -129,7 +130,7 @@ namespace EpicGames.Horde.Compute
 			// Pass the rest of the call over to the handler
 			byte[] key = StringUtils.ParseHexString(responseMessage.Key);
 
-			await using ComputeLease channel = new ComputeLease(new TcpComputeTransport(socket), responseMessage.AssignedResources);
+			await using ComputeLease channel = new ComputeLease(new TcpTransport(socket), responseMessage.AssignedResources);
 			return await handler(channel, cancellationToken);
 		}
 	}
diff --git a/Engine/Source/Programs/Shared/EpicGames.Horde/Compute/Transports/AesTransport.cs b/Engine/Source/Programs/Shared/EpicGames.Horde/Compute/Transports/AesTransport.cs
new file mode 100644
--- /dev/null
+++ b/Engine/Source/Programs/Shared/EpicGames.Horde/Compute/Transports/AesTransport.cs
@@ -0,0 +1,216 @@
+// Copyright Epic Games, Inc. All Rights Reserved.
+
+using System;
+using System.Buffers;
+using System.Buffers.Binary;
+using System.IO.Pipelines;
+using System.Security.Cryptography;
+using System.Threading;
+using System.Threading.Tasks;
+using EpicGames.Core;
+
+namespace EpicGames.Horde.Compute.Transports
+{
+	/// <summary>
+	/// Transport layer that adds AES encryption on top of an underlying transport implementation. Key must be exchanged separately
+	/// (eg. via the HTTPS request to negotiate a lease with the server).
+	/// </summary>
+	public sealed class AesTransport : IComputeTransport, IAsyncDisposable
+	{
+		/// <summary>
+		/// Length of the required encrption key. 
+		/// </summary>
+		public const int KeyLength = 32;
+
+		/// <summary>
+		/// Length of the nonce. This should be a cryptographically random number, and does not have to be secret.
+		/// </summary>
+		public const int NonceLength = 12;
+
+		const int HeaderLength = sizeof(int); // Unencrypted size of the message
+		const int FooterLength = 16; // 16-byte auth tag for encryption
+
+		readonly IComputeTransport _inner;
+		readonly Pipe _readPipe;
+		readonly AesGcm _aesGcm;
+
+		ReadOnlySequence<byte> _readBuffer;
+
+		readonly byte[] _readNonce;
+		readonly byte[] _writeNonce;
+
+		readonly IMemoryOwner<byte> _writeBufferOwner;
+		Task _lastWriteTask = Task.CompletedTask;
+		int _writeBufferOffset = 0;
+
+		const int WritePacketSize = 64 * 1024;
+		readonly BackgroundTask _backgroundReadTask;
+
+		/// <summary>
+		/// Constructor
+		/// </summary>
+		/// <param name="inner">The underlying transport implementation</param>
+		/// <param name="key">AES encryption key (256 bits / 32 bytes)</param>
+		/// <param name="nonce">Cryptographic nonce to identify the connection. Must be longer than <see cref="NonceLength"/>.</param>
+		public AesTransport(IComputeTransport inner, ReadOnlySpan<byte> key, ReadOnlySpan<byte> nonce)
+		{
+			if (key.Length != KeyLength)
+			{
+				throw new ArgumentException($"Key must be {KeyLength} bytes", nameof(key));
+			}
+			if (nonce.Length < NonceLength)
+			{
+				throw new ArgumentException($"Nonce must be at least {NonceLength} bytes", nameof(nonce));
+			}
+
+			_inner = inner;
+			_readPipe = new Pipe();
+			_aesGcm = new AesGcm(key);
+			_readNonce = nonce.Slice(0, NonceLength).ToArray();
+			_writeNonce = nonce.Slice(0, NonceLength).ToArray();
+
+			_backgroundReadTask = BackgroundTask.StartNew(BackgroundReadAsync);
+
+			_writeBufferOwner = MemoryPool<byte>.Shared.Rent(WritePacketSize * 2);
+		}
+
+		/// <inheritdoc/>
+		public async ValueTask DisposeAsync()
+		{
+			await _lastWriteTask;
+			await _backgroundReadTask.DisposeAsync();
+			_writeBufferOwner.Dispose();
+			_aesGcm.Dispose();
+		}
+
+		/// <summary>
+		/// Creates an encryption key
+		/// </summary>
+		public static byte[] CreateKey() => RandomNumberGenerator.GetBytes(KeyLength);
+
+		async Task BackgroundReadAsync(CancellationToken cancellationToken)
+		{
+			using IMemoryOwner<byte> buffer = MemoryPool<byte>.Shared.Rent(128 * 1024);
+			Memory<byte> memory = buffer.Memory;
+
+			int size = 0;
+			for (; ; )
+			{
+				// Read more data into the buffer
+				int read = await _inner.ReadAsync(memory.Slice(size), cancellationToken);
+				if (read == 0)
+				{
+					break;
+				}
+				size += read;
+
+				// Wait until we've got a complete encryption packet
+				int prevSize = size;
+				size = Decrypt(memory.Slice(0, size).Span);
+				memory.Slice(prevSize - size, size).CopyTo(memory);
+			}
+
+			await _readPipe.Writer.CompleteAsync();
+		}
+
+		int Decrypt(Span<byte> span)
+		{
+			while (span.Length >= HeaderLength)
+			{
+				ReadOnlySpan<byte> header = span.Slice(0, HeaderLength);
+
+				int length = BinaryPrimitives.ReadInt32LittleEndian(header);
+				if (span.Length < HeaderLength + length + FooterLength)
+				{
+					break;
+				}
+
+				ReadOnlySpan<byte> cipherText = span.Slice(HeaderLength, length);
+				ReadOnlySpan<byte> tag = span.Slice(HeaderLength + length, FooterLength);
+				Span<byte> plainText = _readPipe.Writer.GetSpan(length).Slice(0, length);
+
+				_aesGcm.Decrypt(_readNonce, cipherText, tag, plainText, header);
+				IncrementNonce(_readNonce);
+
+				_readPipe.Writer.Advance(length);
+
+				span = span.Slice(HeaderLength + length + FooterLength);
+			}
+			return span.Length;
+		}
+
+		/// <inheritdoc/>
+		public async ValueTask<int> ReadAsync(Memory<byte> buffer, CancellationToken cancellationToken)
+		{
+			int sizeRead = 0;
+			while (sizeRead < buffer.Length)
+			{
+				// Try to get more data into the read buffer
+				if (_readBuffer.Length == 0)
+				{
+					ReadResult result;
+					if (sizeRead == 0)
+					{
+						result = await _readPipe.Reader.ReadAsync(cancellationToken);
+					}
+					else if (!_readPipe.Reader.TryRead(out result) || result.IsCompleted)
+					{
+						return sizeRead;
+					}
+					_readBuffer = result.Buffer;
+				}
+
+				// Copy as much of the next chunk as we can
+				int copySize = Math.Min(buffer.Length, _readBuffer.First.Length);
+				_readBuffer.First.Slice(0, copySize).CopyTo(buffer);
+				buffer = buffer.Slice(copySize);
+				_readBuffer = _readBuffer.Slice(copySize);
+				sizeRead += copySize;
+			}
+			return sizeRead;
+		}
+
+		/// <inheritdoc/>
+		public async ValueTask WriteAsync(ReadOnlySequence<byte> buffer, CancellationToken cancellationToken)
+		{
+			while (buffer.Length > 0)
+			{
+				ReadOnlyMemory<byte> plainText = buffer.First;
+
+				int maxSize = WritePacketSize - HeaderLength - FooterLength;
+				if (plainText.Length > maxSize)
+				{
+					plainText = plainText.Slice(0, maxSize);
+				}
+
+				Memory<byte> writeBuffer = _writeBufferOwner.Memory.Slice(_writeBufferOffset, HeaderLength + plainText.Length + FooterLength);
+				Encrypt(plainText.Span, writeBuffer.Span);
+
+				await _lastWriteTask;
+				_lastWriteTask = _inner.WriteAsync(writeBuffer, cancellationToken).AsTask();
+
+				buffer = buffer.Slice(plainText.Length);
+				_writeBufferOffset ^= WritePacketSize;
+			}
+		}
+
+		void Encrypt(ReadOnlySpan<byte> input, Span<byte> output)
+		{
+			int length = input.Length;
+
+			Span<byte> header = output.Slice(0, HeaderLength);
+			BinaryPrimitives.WriteInt32LittleEndian(header, input.Length);
+
+			Span<byte> cipherText = output.Slice(HeaderLength, input.Length);
+			Span<byte> tag = output.Slice(HeaderLength + length, FooterLength);
+
+			_aesGcm.Encrypt(_writeNonce, input, cipherText, tag, header);
+			IncrementNonce(_writeNonce);
+		}
+
+		static void IncrementNonce(byte[] nonce)
+		{
+			BinaryPrimitives.WriteInt64LittleEndian(nonce, BinaryPrimitives.ReadInt64LittleEndian(nonce) + 1);
+		}
+	}
+}
diff --git a/Engine/Source/Programs/Shared/EpicGames.Horde/Compute/Transports/StreamTransport.cs b/Engine/Source/Programs/Shared/EpicGames.Horde/Compute/Transports/StreamTransport.cs
new file mode 100644
--- /dev/null
+++ b/Engine/Source/Programs/Shared/EpicGames.Horde/Compute/Transports/StreamTransport.cs
@@ -0,0 +1,36 @@
+// Copyright Epic Games, Inc. All Rights Reserved.
+
+using System;
+using System.Buffers;
+using System.IO;
+using System.Threading;
+using System.Threading.Tasks;
+
+namespace EpicGames.Horde.Compute.Transports
+{
+	/// <summary>
+	/// Compute transport which wraps an underlying stream
+	/// </summary>
+	class StreamTransport : IComputeTransport
+	{
+		readonly Stream _stream;
+
+		/// <summary>
+		/// Constructor
+		/// </summary>
+		/// <param name="stream">Stream to use for the transferring data</param>
+		public StreamTransport(Stream stream) => _stream = stream;
+
+		/// <inheritdoc/>
+		public async ValueTask<int> ReadAsync(Memory<byte> buffer, CancellationToken cancellationToken) => await _stream.ReadAsync(buffer, cancellationToken);
+
+		/// <inheritdoc/>
+		public async ValueTask WriteAsync(ReadOnlySequence<byte> buffer, CancellationToken cancellationToken)
+		{
+			foreach (ReadOnlyMemory<byte> memory in buffer)
+			{
+				await _stream.WriteAsync(memory, cancellationToken);
+			}
+		}
+	}
+}
diff --git a/Engine/Source/Programs/Shared/EpicGames.Horde/Compute/Transports/TcpTransport.cs b/Engine/Source/Programs/Shared/EpicGames.Horde/Compute/Transports/TcpTransport.cs
new file mode 100644
--- /dev/null
+++ b/Engine/Source/Programs/Shared/EpicGames.Horde/Compute/Transports/TcpTransport.cs
@@ -0,0 +1,37 @@
+// Copyright Epic Games, Inc. All Rights Reserved.
+
+using System;
+using System.Buffers;
+using System.Net.Sockets;
+using System.Threading;
+using System.Threading.Tasks;
+using EpicGames.Core;
+
+namespace EpicGames.Horde.Compute.Transports
+{
+	/// <summary>
+	/// Implementation of <see cref="IComputeTransport"/> for communicating over a socket
+	/// </summary>
+	public class TcpTransport : IComputeTransport
+	{
+		readonly Socket _socket;
+
+		/// <summary>
+		/// Constructor
+		/// </summary>
+		/// <param name="socket">Socket to communicate over</param>
+		public TcpTransport(Socket socket) => _socket = socket;
+
+		/// <inheritdoc/>
+		public ValueTask<int> ReadAsync(Memory<byte> buffer, CancellationToken cancellationToken) => _socket.ReceiveAsync(buffer, SocketFlags.None, cancellationToken);
+
+		/// <inheritdoc/>
+		public async ValueTask WriteAsync(ReadOnlySequence<byte> buffer, CancellationToken cancellationToken)
+		{
+			foreach (ReadOnlyMemory<byte> memory in buffer)
+			{
+				await _socket.SendMessageAsync(memory, SocketFlags.None, cancellationToken);
+			}
+		}
+	}
+}
