Horde: Allow disposing and re-creating compute buffers with the same id, without having to dispose of the entire lease.

diff --git a/Engine/Source/Programs/Shared/EpicGames.Horde/Compute/ComputeBuffer.cs b/Engine/Source/Programs/Shared/EpicGames.Horde/Compute/ComputeBuffer.cs
--- a/Engine/Source/Programs/Shared/EpicGames.Horde/Compute/ComputeBuffer.cs
+++ b/Engine/Source/Programs/Shared/EpicGames.Horde/Compute/ComputeBuffer.cs
@@ -54,7 +54,19 @@ namespace EpicGames.Horde.Compute
 		/// <inheritdoc/>
 		public void Dispose()
 		{
-			_memoryOwner.Dispose();
+			Dispose(true);
+			GC.SuppressFinalize(this);
+		}
+
+		/// <summary>
+		/// Standard dispose pattern
+		/// </summary>
+		protected virtual void Dispose(bool disposing)
+		{
+			if (disposing)
+			{
+				_memoryOwner.Dispose();
+			}
 		}
 
 		/// <inheritdoc/>
diff --git a/Engine/Source/Programs/Shared/EpicGames.Horde/Compute/ComputeChannel.cs b/Engine/Source/Programs/Shared/EpicGames.Horde/Compute/ComputeChannel.cs
--- a/Engine/Source/Programs/Shared/EpicGames.Horde/Compute/ComputeChannel.cs
+++ b/Engine/Source/Programs/Shared/EpicGames.Horde/Compute/ComputeChannel.cs
@@ -117,6 +117,8 @@ namespace EpicGames.Horde.Compute
 		public void Dispose()
 		{
 			_currentWriter?.Dispose();
+			_inputBuffer.Dispose();
+			_outputBuffer.Dispose();
 		}
 
 		/// <inheritdoc/>
diff --git a/Engine/Source/Programs/Shared/EpicGames.Horde/Compute/ComputeLease.cs b/Engine/Source/Programs/Shared/EpicGames.Horde/Compute/ComputeLease.cs
--- a/Engine/Source/Programs/Shared/EpicGames.Horde/Compute/ComputeLease.cs
+++ b/Engine/Source/Programs/Shared/EpicGames.Horde/Compute/ComputeLease.cs
@@ -27,6 +27,16 @@ namespace EpicGames.Horde.Compute
 				_owner = owner;
 			}
 
+			protected override void Dispose(bool disposing)
+			{
+				base.Dispose(disposing);
+
+				if (disposing)
+				{
+					_owner.ReleaseBuffer(this);
+				}
+			}
+
 			public override async ValueTask FlushAsync(CancellationToken cancellationToken)
 			{
 				TaskCompletionSource tcs = new TaskCompletionSource(TaskCreationOptions.RunContinuationsAsynchronously);
@@ -127,16 +137,6 @@ namespace EpicGames.Horde.Compute
 			await _readTask;
 
 			_cts.Dispose();
-
-			foreach (Buffer inputBuffer in _inputBuffers.Values)
-			{
-				inputBuffer.Dispose();
-			}
-
-			foreach (Buffer outputBuffer in _outputBuffers.Values)
-			{
-				outputBuffer.Dispose();
-			}
 		}
 
 		async Task RunReaderAsync(IComputeTransport transport, CancellationToken cancellationToken)
@@ -263,5 +263,21 @@ namespace EpicGames.Horde.Compute
 			}
 			return outputBuffer;
 		}
+
+		void ReleaseBuffer(Buffer buffer)
+		{
+			lock (_lockObject)
+			{
+				Buffer? existingBuffer;
+				if (_inputBuffers.TryGetValue(buffer.Id, out existingBuffer) && ReferenceEquals(existingBuffer, buffer))
+				{
+					_inputBuffers.Remove(buffer.Id);
+				}
+				if (_outputBuffers.TryGetValue(buffer.Id, out existingBuffer) && ReferenceEquals(existingBuffer, buffer))
+				{
+					_outputBuffers.Remove(buffer.Id);
+				}
+			}
+		}
 	}
 }
diff --git a/Engine/Source/Programs/Shared/EpicGames.Horde/Compute/IComputeBuffer.cs b/Engine/Source/Programs/Shared/EpicGames.Horde/Compute/IComputeBuffer.cs
--- a/Engine/Source/Programs/Shared/EpicGames.Horde/Compute/IComputeBuffer.cs
+++ b/Engine/Source/Programs/Shared/EpicGames.Horde/Compute/IComputeBuffer.cs
@@ -20,7 +20,7 @@ namespace EpicGames.Horde.Compute
 	/// <summary>
 	/// Read interface for a compute buffer
 	/// </summary>
-	public interface IComputeInputBuffer : IComputeBuffer
+	public interface IComputeInputBuffer : IComputeBuffer, IDisposable
 	{
 		/// <summary>
 		/// The shared memory for the buffer
@@ -54,7 +54,7 @@ namespace EpicGames.Horde.Compute
 	/// <summary>
 	/// Read-only view of shared memory that can be incrementally replicated to a remote machine.
 	/// </summary>
-	public interface IComputeOutputBuffer : IComputeBuffer
+	public interface IComputeOutputBuffer : IComputeBuffer, IDisposable
 	{
 		/// <summary>
 		/// The shared memory for the buffer
