Horde: Append to arbitrary buffers as part of compute protocol, and rework message passing API to be built on top of this. Any data appended to buffers will be sent to the remote and appended to a buffer with a matching id there. This mechanism allows significantly more efficient throughput to native apps, without having to marshal any data across the managed/unmanaged boundary.

diff --git a/Engine/Source/Programs/Horde/Horde.Agent/Commands/Compute/ComputeCommand.cs b/Engine/Source/Programs/Horde/Horde.Agent/Commands/Compute/ComputeCommand.cs
new file mode 100644
--- /dev/null
+++ b/Engine/Source/Programs/Horde/Horde.Agent/Commands/Compute/ComputeCommand.cs
@@ -0,0 +1,94 @@
+// Copyright Epic Games, Inc. All Rights Reserved.
+
+using System;
+using System.Net.Http;
+using System.Net.Http.Headers;
+using System.Threading;
+using System.Threading.Tasks;
+using EpicGames.Core;
+using EpicGames.Horde.Common;
+using EpicGames.Horde.Compute;
+using Horde.Agent.Leases.Handlers;
+using Microsoft.Extensions.DependencyInjection;
+using Microsoft.Extensions.Logging;
+using Microsoft.Extensions.Options;
+
+namespace Horde.Agent.Commands.Compute
+{
+	abstract class ComputeCommand : Command
+	{
+		[CommandLine("-Cluster")]
+		public string ClusterId { get; set; } = "default";
+
+		[CommandLine("-Requirements=", Description = "Match the agent to run on")]
+		public string? Requirements { get; set; }
+
+		[CommandLine("-Loopback")]
+		public bool Loopback { get; set; }
+
+		readonly IServiceProvider _serviceProvider;
+		readonly IHttpClientFactory _httpClientFactory;
+		readonly IOptions<AgentSettings> _settings;
+		readonly ILogger _logger;
+
+		public ComputeCommand(IServiceProvider serviceProvider, ILogger logger)
+		{
+			_serviceProvider = serviceProvider;
+			_httpClientFactory = serviceProvider.GetRequiredService<IHttpClientFactory>();
+			_settings = serviceProvider.GetRequiredService<IOptions<AgentSettings>>();
+			_logger = logger;
+		}
+
+		/// <inheritdoc/>
+		public override async Task<int> ExecuteAsync(ILogger logger)
+		{
+			await using IComputeClient client = CreateClient();
+
+			Requirements? requirements = null;
+			if (Requirements != null)
+			{
+				requirements = new Requirements(Condition.Parse(Requirements));
+			}
+
+			bool result = await client.ExecuteAsync(new ClusterId(ClusterId), requirements, HandleRequestAsync, CancellationToken.None);
+
+			return result ? 0 : 1;
+		}
+
+		/// <summary>
+		/// Callback for executing a request
+		/// </summary>
+		/// <param name="lease"></param>
+		/// <param name="cancellationToken"></param>
+		/// <returns></returns>
+		protected abstract Task<bool> HandleRequestAsync(IComputeLease lease, CancellationToken cancellationToken);
+
+		IComputeClient CreateClient()
+		{
+			if (Loopback)
+			{
+				return new LoopbackComputeClient(RunListenerAsync);
+			}
+			else
+			{
+				return new ServerComputeClient(CreateHttpClient, _logger);
+			}
+		}
+
+		HttpClient CreateHttpClient()
+		{
+			ServerProfile profile = _settings.Value.GetCurrentServerProfile();
+
+			HttpClient client = _httpClientFactory.CreateClient();
+			client.BaseAddress = profile.Url;
+			client.DefaultRequestHeaders.Authorization = new AuthenticationHeaderValue("Bearer", profile.Token);
+			return client;
+		}
+
+		async Task RunListenerAsync(IComputeLease lease, CancellationToken cancellationToken)
+		{
+			ComputeHandler handler = ActivatorUtilities.CreateInstance<ComputeHandler>(_serviceProvider);
+			await handler.RunAsync(lease, cancellationToken);
+		}
+	}
+}
diff --git a/Engine/Source/Programs/Horde/Horde.Agent/Commands/Compute/CppComputeCommand.cs b/Engine/Source/Programs/Horde/Horde.Agent/Commands/Compute/CppComputeCommand.cs
--- a/Engine/Source/Programs/Horde/Horde.Agent/Commands/Compute/CppComputeCommand.cs
+++ b/Engine/Source/Programs/Horde/Horde.Agent/Commands/Compute/CppComputeCommand.cs
@@ -1,13 +1,9 @@
 // Copyright Epic Games, Inc. All Rights Reserved.
 
 using System;
-using System.Buffers;
 using System.Collections.Generic;
-using System.Net.Http;
-using System.Net.Http.Headers;
 using System.Text.Json;
 using System.Threading;
-using System.Threading.Channels;
 using System.Threading.Tasks;
 using EpicGames.Core;
 using EpicGames.Horde.Compute;
@@ -16,18 +12,15 @@ using EpicGames.Horde.Logs;
 using EpicGames.Horde.Storage;
 using EpicGames.Horde.Storage.Backends;
 using EpicGames.Horde.Storage.Nodes;
-using Horde.Agent.Leases.Handlers;
-using Microsoft.Extensions.DependencyInjection;
 using Microsoft.Extensions.Logging;
-using Microsoft.Extensions.Options;
 
-namespace Horde.Agent.Commands
+namespace Horde.Agent.Commands.Compute
 {
 	/// <summary>
 	/// Installs the agent as a service
 	/// </summary>
 	[Command("cppcompute", "Executes a command through the C++ Compute API")]
-	class CppComputeCommand : Command
+	class CppComputeCommand : ComputeCommand
 	{
 		class JsonComputeTask
 		{
@@ -38,54 +31,36 @@ namespace Horde.Agent.Commands
 			public List<string> OutputPaths { get; set; } = new List<string>();
 		}
 
-		[CommandLine("-Cluster")]
-		public string ClusterId { get; set; } = "default";
-
-		[CommandLine("-Loopback")]
-		public bool Loopback { get; set; }
-
 		[CommandLine("-Task=", Required = true)]
 		FileReference TaskFile { get; set; } = null!;
 
-		readonly IServiceProvider _serviceProvider;
-		readonly IHttpClientFactory _httpClientFactory;
-		readonly AgentSettings _settings;
 		readonly ILogger _logger;
 
 		/// <summary>
 		/// Constructor
 		/// </summary>
-		public CppComputeCommand(IServiceProvider serviceProvider, IHttpClientFactory httpClientFactory, IOptions<AgentSettings> settings, ILogger<CppComputeCommand> logger)
+		public CppComputeCommand(IServiceProvider serviceProvider, ILogger<CppComputeCommand> logger)
+			: base(serviceProvider, logger)
 		{
-			_serviceProvider = serviceProvider;
-			_httpClientFactory = httpClientFactory;
-			_settings = settings.Value;
 			_logger = logger;
 		}
 
 		/// <inheritdoc/>
-		public override async Task<int> ExecuteAsync(ILogger logger)
+		protected override async Task<bool> HandleRequestAsync(IComputeLease lease, CancellationToken cancellationToken)
 		{
-			await using IComputeClient client = CreateClient();
-
-			bool result = await client.ExecuteAsync(new ClusterId(ClusterId), null, HandleRequestAsync, CancellationToken.None);
-
-			return result? 0 : 1;
-		}
+			await using IComputeChannel defaultChannel = lease.CreateChannel(0);
 
-		async Task<bool> HandleRequestAsync(IComputeLease lease, CancellationToken cancellationToken)
-		{
-			using IComputeChannel defaultChannel = lease.CreateChannel(0);
-			using IComputeChannel channel = lease.CreateChannel(1);
+			const int ReplyChannelId = 1;
+			await using IComputeChannel channel = lease.CreateChannel(ReplyChannelId);
 
 			MemoryStorageClient storage = new MemoryStorageClient();
 
 			NodeLocator node = await CreateComputeNodeAsync(TaskFile, storage, cancellationToken);
-			await defaultChannel.CppStartAsync(node, channel.Id, cancellationToken);
+			defaultChannel.CppStart(node, ReplyChannelId);
 
 			for (; ; )
 			{
-				using IComputeMessage message = await channel.ReadAsync(cancellationToken);
+				using IComputeMessage message = await channel.ReceiveAsync(cancellationToken);
 				switch (message.Type)
 				{
 					case ComputeMessageType.None:
@@ -116,8 +91,7 @@ namespace Horde.Agent.Commands
 								}
 							}
 
-							await channel.CppFinishAsync(cancellationToken);
-							await defaultChannel.CloseAsync(cancellationToken);
+							channel.CppFinish();
 						}
 						return false;
 					case ComputeMessageType.CppFailure:
@@ -155,33 +129,5 @@ namespace Horde.Agent.Commands
 			NodeHandle handle = await writer.FlushAsync(computeNode, cancellationToken);
 			return handle.Locator;
 		}
-
-		IComputeClient CreateClient()
-		{
-			if (Loopback)
-			{
-				return new LoopbackComputeClient(RunListenerAsync);
-			}
-			else
-			{
-				return new ServerComputeClient(CreateHttpClient, _logger);
-			}
-		}
-
-		HttpClient CreateHttpClient()
-		{
-			ServerProfile profile = _settings.GetCurrentServerProfile();
-
-			HttpClient client = _httpClientFactory.CreateClient();
-			client.BaseAddress = profile.Url;
-			client.DefaultRequestHeaders.Authorization = new AuthenticationHeaderValue("Bearer", profile.Token);
-			return client;
-		}
-
-		async Task RunListenerAsync(IComputeLease lease, CancellationToken cancellationToken)
-		{
-			ComputeHandler handler = ActivatorUtilities.CreateInstance<ComputeHandler>(_serviceProvider);
-			await handler.RunAsync(lease, cancellationToken);
-		}
 	}
 }
diff --git a/Engine/Source/Programs/Horde/Horde.Agent/Commands/Compute/DdcComputeCommand.cs b/Engine/Source/Programs/Horde/Horde.Agent/Commands/Compute/DdcComputeCommand.cs
--- a/Engine/Source/Programs/Horde/Horde.Agent/Commands/Compute/DdcComputeCommand.cs
+++ b/Engine/Source/Programs/Horde/Horde.Agent/Commands/Compute/DdcComputeCommand.cs
@@ -1,74 +1,60 @@
 // Copyright Epic Games, Inc. All Rights Reserved.
 
 using System;
-using System.Diagnostics;
-using System.Net.Http;
+using System.Linq;
 using System.Threading;
 using System.Threading.Tasks;
 using EpicGames.Core;
 using EpicGames.Horde.Compute;
 using Microsoft.Extensions.Logging;
-using Microsoft.Extensions.Options;
 
-namespace Horde.Agent.Commands
+namespace Horde.Agent.Commands.Compute
 {
 	/// <summary>
 	/// Installs the agent as a service
 	/// </summary>
 	[Command("ddccompute", "Executes a command through the Horde Compute API")]
-	class DdcComputeCommand : Command
+	class DdcComputeCommand : ComputeCommand
 	{
-		[CommandLine("-Cluster=", Description = "Cluster to execute on")]
-		public ClusterId ClusterId { get; set; } = new ClusterId("default");
-
-		[CommandLine("-Condition=", Description = "Match the agent to run on")]
-		public string? Condition { get; set; }
-
-		readonly AgentSettings _settings;
+		readonly ILogger _logger;
 
 		/// <summary>
 		/// Constructor
 		/// </summary>
-		public DdcComputeCommand(IOptions<AgentSettings> settings)
+		public DdcComputeCommand(IServiceProvider serviceProvider, ILogger<DdcComputeCommand> logger)
+			: base(serviceProvider, logger)
 		{
-			_settings = settings.Value;
+			_logger = logger;
 		}
 
 		/// <inheritdoc/>
-		public override async Task<int> ExecuteAsync(ILogger logger)
-		{
-			ServerProfile profile = _settings.GetCurrentServerProfile();
-			logger.LogInformation("Connecting to server: {Server}", profile.Name);
-
-			await using ServerComputeClient client = new ServerComputeClient(CreateHttpClient, logger);
-			await client.ExecuteAsync(ClusterId, null, TestCommandAsync, CancellationToken.None);
-
-			return 0;
-		}
-
-		HttpClient CreateHttpClient()
+		protected override async Task<bool> HandleRequestAsync(IComputeLease lease, CancellationToken cancellationToken)
 		{
-			HttpClient client = new HttpClient();
-			client.BaseAddress = _settings.GetCurrentServerProfile().Url;
-			return client;
-		}
-
-		async Task<object?> TestCommandAsync(IComputeLease lease, CancellationToken cancellationToken)
-		{
-			using IComputeChannel channel = lease.CreateChannel(0);
-
-			await channel.XorRequestAsync(new byte[] { 1, 2, 3, 4, 5 }, (byte)123, cancellationToken);
-			IComputeMessage response = await channel.ReadAsync(cancellationToken);
-
-			byte[] result = response.Data.ToArray();
-			Debug.Assert(result[0] == (byte)(1 ^ 123));
-			Debug.Assert(result[1] == (byte)(2 ^ 123));
-			Debug.Assert(result[2] == (byte)(3 ^ 123));
-			Debug.Assert(result[3] == (byte)(4 ^ 123));
-			Debug.Assert(result[4] == (byte)(5 ^ 123));
-
-			await channel.CloseAsync(cancellationToken);
-			return null;
+			await using (IComputeChannel channel = lease.CreateChannel(0))
+			{
+				_logger.LogInformation("Sending XOR request");
+				channel.XorRequest(new byte[] { 1, 2, 3, 4, 5 }, (byte)123);
+
+				_logger.LogInformation("Waiting for response...");
+				IComputeMessage response = await channel.ReceiveAsync(cancellationToken);
+
+				byte[] result = response.Data.ToArray();
+				byte[] expectedResult = new byte[] { 1 ^ 123, 2 ^ 123, 3 ^ 123, 4 ^ 123, 5 ^ 123 };
+
+				if (result.SequenceEqual(expectedResult))
+				{
+					_logger.LogInformation("Received response; data is correct.");
+				}
+				else
+				{
+					throw new Exception("Incorrect response data");
+				}
+
+				_logger.LogInformation("Closing channel");
+			}
+
+			_logger.LogInformation("Closed");
+			return true;
 		}
 	}
 }
diff --git a/Engine/Source/Programs/Horde/Horde.Agent/Commands/Compute/XorComputeCommand.cs b/Engine/Source/Programs/Horde/Horde.Agent/Commands/Compute/XorComputeCommand.cs
new file mode 100644
--- /dev/null
+++ b/Engine/Source/Programs/Horde/Horde.Agent/Commands/Compute/XorComputeCommand.cs
@@ -0,0 +1,61 @@
+// Copyright Epic Games, Inc. All Rights Reserved.
+
+using System;
+using System.Diagnostics;
+using System.Linq;
+using System.Threading;
+using System.Threading.Tasks;
+using EpicGames.Core;
+using EpicGames.Horde.Compute;
+using Microsoft.Extensions.Logging;
+
+namespace Horde.Agent.Commands.Compute
+{
+	/// <summary>
+	/// Runs a simple compute command
+	/// </summary>
+	[Command("xorcompute", "Executes a simple XOR command through the compute API")]
+	class XorComputeCommand : ComputeCommand
+	{
+		readonly ILogger _logger;
+
+		/// <summary>
+		/// Constructor
+		/// </summary>
+		public XorComputeCommand(IServiceProvider serviceProvider, ILogger<XorComputeCommand> logger)
+			: base(serviceProvider, logger)
+		{
+			_logger = logger;
+		}
+
+		/// <inheritdoc/>
+		protected override async Task<bool> HandleRequestAsync(IComputeLease lease, CancellationToken cancellationToken)
+		{
+			await using (IComputeChannel channel = lease.CreateChannel(0))
+			{
+				_logger.LogInformation("Sending XOR request");
+				channel.XorRequest(new byte[] { 1, 2, 3, 4, 5 }, (byte)123);
+
+				_logger.LogInformation("Waiting for response...");
+				IComputeMessage response = await channel.ReceiveAsync(cancellationToken);
+
+				byte[] result = response.Data.ToArray();
+				byte[] expectedResult = new byte[] { 1 ^ 123, 2 ^ 123, 3 ^ 123, 4 ^ 123, 5 ^ 123 };
+
+				if (result.SequenceEqual(expectedResult))
+				{
+					_logger.LogInformation("Received response; data is correct.");
+				}
+				else
+				{
+					throw new Exception("Incorrect response data");
+				}
+
+				_logger.LogInformation("Closing channel");
+			}
+
+			_logger.LogInformation("Closed");
+			return true;
+		}
+	}
+}
diff --git a/Engine/Source/Programs/Horde/Horde.Agent/Leases/Handlers/ComputeHandler.cs b/Engine/Source/Programs/Horde/Horde.Agent/Leases/Handlers/ComputeHandler.cs
--- a/Engine/Source/Programs/Horde/Horde.Agent/Leases/Handlers/ComputeHandler.cs
+++ b/Engine/Source/Programs/Horde/Horde.Agent/Leases/Handlers/ComputeHandler.cs
@@ -77,25 +77,24 @@ namespace Horde.Agent.Leases.Handlers
 
 		public async Task RunAsync(IComputeLease lease, CancellationToken cancellationToken)
 		{
-			using IComputeChannel channel = lease.CreateChannel(0);
+			await using IComputeChannel channel = lease.CreateChannel(0);
 			for (; ; )
 			{
-				using IComputeMessage message = await channel.ReadAsync(cancellationToken);
+				using IComputeMessage message = await channel.ReceiveAsync(cancellationToken);
 				switch (message.Type)
 				{
 					case ComputeMessageType.None:
-					case ComputeMessageType.Close:
 						return;
 					case ComputeMessageType.XorRequest:
 						{
 							XorRequestMessage xorRequest = message.AsXorRequest();
-							await RunXorAsync(channel, xorRequest.Data, xorRequest.Value, cancellationToken);
+							RunXor(channel, xorRequest.Data, xorRequest.Value);
 						}
 						break;
 					case ComputeMessageType.CppBegin:
 						{
 							CppBeginMessage cppBegin = message.AsCppBegin();
-							using IComputeChannel replyChannel = lease.CreateChannel(cppBegin.ReplyChannelId);
+							await using IComputeChannel replyChannel = lease.CreateChannel(cppBegin.ReplyChannelId);
 							await RunCppAsync(replyChannel, cppBegin.Locator, cancellationToken);
 						}
 						break;
@@ -105,12 +104,12 @@ namespace Horde.Agent.Leases.Handlers
 			}
 		}
 
-		static async Task RunXorAsync(IComputeChannel channel, ReadOnlyMemory<byte> source, byte value, CancellationToken cancellationToken)
+		static void RunXor(IComputeChannel channel, ReadOnlyMemory<byte> source, byte value)
 		{
-			using (IComputeMessageWriter writer = channel.CreateMessage(ComputeMessageType.XorResponse, source.Length))
+			using (IComputeMessageBuilder writer = channel.CreateMessage(ComputeMessageType.XorResponse, source.Length))
 			{
 				XorData(source.Span, writer.GetSpanAndAdvance(source.Length), value);
-				await writer.SendAsync(cancellationToken);
+				writer.Send();
 			}
 		}
 
@@ -130,13 +129,13 @@ namespace Horde.Agent.Leases.Handlers
 			}
 			catch (Exception ex)
 			{
-				await channel.CppFailureAsync(ex.ToString(), cancellationToken);
+				channel.CppFailure(ex.ToString());
 			}
 		}
 
 		async Task RunCppInternalAsync(IComputeChannel channel, NodeLocator locator, CancellationToken cancellationToken)
 		{
-			DirectoryReference sandboxDir = DirectoryReference.Combine(Program.DataDir, "Sandbox", channel.Id.ToString());
+			DirectoryReference sandboxDir = DirectoryReference.Combine(Program.DataDir, "Sandbox", Guid.NewGuid().ToString());
 
 			using ComputeStorageClient store = new ComputeStorageClient(channel);
 			TreeReader reader = new TreeReader(store, _memoryCache, _logger);
@@ -196,12 +195,12 @@ namespace Horde.Agent.Leases.Handlers
 				CppComputeOutputNode outputNode = new CppComputeOutputNode(exitCode, logNodeRef, new TreeNodeRef<DirectoryNode>(outputTree));
 				NodeHandle outputHandle = await writer.FlushAsync(outputNode, cancellationToken);
 
-				await channel.CppSuccessAsync(outputHandle.Locator, cancellationToken);
+				channel.CppSuccess(outputHandle.Locator);
 			}
 
 			for (; ; )
 			{
-				using IComputeMessage message = await channel.ReadAsync(cancellationToken);
+				using IComputeMessage message = await channel.ReceiveAsync(cancellationToken);
 				switch (message.Type)
 				{
 					case ComputeMessageType.CppEnd:
diff --git a/Engine/Source/Programs/Shared/EpicGames.Horde/Compute/Buffers/LocalBuffer.cs b/Engine/Source/Programs/Shared/EpicGames.Horde/Compute/Buffers/LocalBuffer.cs
new file mode 100644
--- /dev/null
+++ b/Engine/Source/Programs/Shared/EpicGames.Horde/Compute/Buffers/LocalBuffer.cs
@@ -0,0 +1,215 @@
+// Copyright Epic Games, Inc. All Rights Reserved.
+
+using System;
+using System.Buffers;
+using System.Threading;
+using System.Threading.Tasks;
+
+namespace EpicGames.Horde.Compute.Buffers
+{
+	/// <summary>
+	/// In-process buffer used for compute messages
+	/// </summary>
+	public class LocalBuffer : IComputeBuffer
+	{
+		class ReaderImpl : IComputeBufferReader
+		{
+			readonly LocalBuffer _owner;
+
+			public ReaderImpl(LocalBuffer owner) => _owner = owner;
+
+			/// <inheritdoc/>
+			public bool IsComplete => _owner.FinishedReading();
+
+			/// <inheritdoc/>
+			public void Advance(int size) => _owner.AdvanceReadPosition(size);
+
+			/// <inheritdoc/>
+			public ReadOnlyMemory<byte> GetMemory() => _owner.GetReadMemory();
+
+			/// <inheritdoc/>
+			public ValueTask WaitAsync(int currentLength, CancellationToken cancellationToken) => _owner.WaitAsync(currentLength, cancellationToken);
+		}
+
+		class WriterImpl : IComputeBufferWriter
+		{
+			readonly LocalBuffer _owner;
+
+			public WriterImpl(LocalBuffer owner) => _owner = owner;
+
+			/// <inheritdoc/>
+			public void MarkComplete() => _owner.FinishWriting();
+
+			/// <inheritdoc/>
+			public void Advance(int size) => _owner.AdvanceWritePosition(size);
+
+			/// <inheritdoc/>
+			public Memory<byte> GetMemory() => _owner.GetWriteMemory();
+
+			/// <inheritdoc/>
+			public ValueTask ResetAsync(CancellationToken cancellationToken) => _owner.ResetWritePositionAsync(cancellationToken);
+
+			/// <inheritdoc/>
+			public ValueTask FlushAsync(CancellationToken cancellationToken) => _owner.FlushWritesAsync(cancellationToken);
+		}
+
+		readonly Memory<byte> _memory;
+		int _readPosition;
+		int _writePosition;
+		bool _finishedWriting;
+		TaskCompletionSource? _writtenTcs;
+		TaskCompletionSource? _flushedTcs;
+
+		/// <inheritdoc/>
+		public IComputeBufferReader Reader { get; }
+
+		/// <inheritdoc/>
+		public IComputeBufferWriter Writer { get; }
+
+		/// <summary>
+		/// Creates a local buffer with the given capacity
+		/// </summary>
+		/// <param name="memory"></param>
+		public LocalBuffer(Memory<byte> memory)
+		{
+			_memory = memory;
+
+			Reader = new ReaderImpl(this);
+			Writer = new WriterImpl(this);
+		}
+
+		#region Reader
+
+		bool FinishedReading() => _finishedWriting && _readPosition == _writePosition;
+
+		void AdvanceReadPosition(int size)
+		{
+			_readPosition += size;
+
+			if (_readPosition == _writePosition)
+			{
+				_flushedTcs?.TrySetResult();
+			}
+		}
+
+		ReadOnlyMemory<byte> GetReadMemory() => _memory.Slice(_readPosition, _writePosition - _readPosition);
+
+		async ValueTask WaitAsync(int currentLength, CancellationToken cancellationToken)
+		{
+			int initialWritePosition = _readPosition + currentLength;
+
+			// Early out if we're already past this
+			if (_finishedWriting || _writePosition > initialWritePosition)
+			{
+				return;
+			}
+
+			TaskCompletionSource tcs = new TaskCompletionSource();
+			using (IDisposable registration = cancellationToken.Register(x => tcs.SetCanceled((CancellationToken)x!), cancellationToken))
+			{
+				try
+				{
+					_writtenTcs = tcs;
+					for (; ; )
+					{
+						if (_finishedWriting || _writePosition > initialWritePosition)
+						{
+							break;
+						}
+						else
+						{
+							await _writtenTcs.Task;
+						}
+					}
+				}
+				finally
+				{
+					_writtenTcs = null;
+				}
+			}
+		}
+
+		#endregion
+		#region Writer
+
+		void FinishWriting()
+		{
+			_finishedWriting = true;
+			_writtenTcs?.TrySetResult();
+		}
+
+		void AdvanceWritePosition(int size)
+		{
+			if (_finishedWriting)
+			{
+				throw new InvalidOperationException("Cannot update write position after marking as complete");
+			}
+
+			_writePosition += size;
+			_writtenTcs?.TrySetResult();
+		}
+
+		Memory<byte> GetWriteMemory() => _memory.Slice(_writePosition);
+
+		async ValueTask ResetWritePositionAsync(CancellationToken cancellationToken)
+		{
+			await FlushWritesAsync(cancellationToken);
+			_writePosition = 0;
+		}
+
+		async ValueTask FlushWritesAsync(CancellationToken cancellationToken)
+		{
+			if (_readPosition < _writePosition)
+			{
+				TaskCompletionSource tcs = new TaskCompletionSource();
+				using (_ = cancellationToken.Register(x => tcs.SetCanceled((CancellationToken)x!), cancellationToken, false))
+				{
+					TaskCompletionSource? originalTcs = Interlocked.CompareExchange(ref _flushedTcs, tcs, null);
+					if (originalTcs != null)
+					{
+						_ = originalTcs.Task.ContinueWith(x => tcs.TrySetResult(), cancellationToken, TaskContinuationOptions.None, TaskScheduler.Default);
+					}
+				}
+				if (_readPosition < _writePosition)
+				{
+					await tcs.Task;
+				}
+			}
+		}
+
+		#endregion
+	}
+
+	/// <summary>
+	/// Implementation of <see cref="LocalBuffer"/> that owns a rented memory resource
+	/// </summary>
+	class RentedLocalBuffer : LocalBuffer, IDisposable
+	{
+		readonly IMemoryOwner<byte> _memoryOwner;
+
+		/// <summary>
+		/// Constructor
+		/// </summary>
+		/// <param name="capacity">Size of data to allocate</param>
+		public RentedLocalBuffer(int capacity)
+			: this(MemoryPool<byte>.Shared.Rent(capacity))
+		{
+		}
+
+		/// <summary>
+		/// Constructor
+		/// </summary>
+		/// <param name="memoryOwner">Memory to manage</param>
+		public RentedLocalBuffer(IMemoryOwner<byte> memoryOwner)
+			: base(memoryOwner.Memory)
+		{
+			_memoryOwner = memoryOwner;
+		}
+
+		/// <inheritdoc/>
+		public void Dispose()
+		{
+			_memoryOwner.Dispose();
+		}
+	}
+}
diff --git a/Engine/Source/Programs/Shared/EpicGames.Horde/Compute/ComputeBuffer.cs b/Engine/Source/Programs/Shared/EpicGames.Horde/Compute/ComputeBuffer.cs
deleted file mode 100644
diff --git a/Engine/Source/Programs/Shared/EpicGames.Horde/Compute/ComputeChannel.cs b/Engine/Source/Programs/Shared/EpicGames.Horde/Compute/ComputeChannel.cs
--- a/Engine/Source/Programs/Shared/EpicGames.Horde/Compute/ComputeChannel.cs
+++ b/Engine/Source/Programs/Shared/EpicGames.Horde/Compute/ComputeChannel.cs
@@ -5,15 +5,17 @@ using System.Buffers;
 using System.Buffers.Binary;
 using System.Threading;
 using System.Threading.Tasks;
-using EpicGames.Core;
 
 namespace EpicGames.Horde.Compute
 {
 	/// <summary>
 	/// Implementation of a compute channel
 	/// </summary>
-	class ComputeChannel : IComputeChannel, IDisposable
+	class ComputeChannel : IComputeChannel, IAsyncDisposable
 	{
+		// Length of a message header. Consists of a 1 byte type field, followed by 4 byte length field.
+		const int HeaderLength = 5;
+
 		/// <summary>
 		/// Standard implementation of a message
 		/// </summary>
@@ -46,126 +48,117 @@ namespace EpicGames.Horde.Compute
 		/// <summary>
 		/// Allows creating new messages in rented memory
 		/// </summary>
-		class Writer : IComputeMessageWriter
+		class MessageBuilder : IComputeMessageBuilder
 		{
 			readonly ComputeChannel _channel;
-			readonly IComputeOutputBuffer _buffer;
+			readonly IComputeBufferWriter _sendBufferWriter;
+			readonly ComputeMessageType _type;
 			int _length;
 
 			/// <inheritdoc/>
 			public int Length => _length;
 
-			public Writer(ComputeChannel channel, IComputeOutputBuffer buffer, ComputeMessageType type)
+			public MessageBuilder(ComputeChannel channel, IComputeBufferWriter sendBufferWriter, ComputeMessageType type)
 			{
 				_channel = channel;
-				_buffer = buffer;
-
-				this.WriteUInt8((byte)type);
-				this.WriteUInt32(0); // size placeholder
+				_sendBufferWriter = sendBufferWriter;
+				_type = type;
+				_length = 0;
 			}
 
 			public void Dispose()
 			{
-				if (_channel._currentWriter == this)
+				if (_channel._currentBuilder == this)
 				{
-					_channel._currentWriter = null;
+					_channel._currentBuilder = null;
 				}
 			}
 
-			public async ValueTask SendAsync(CancellationToken cancellationToken)
+			public void Send()
 			{
-				BinaryPrimitives.WriteInt32LittleEndian(_buffer.GetWriteSpan().Slice(1), _length);
-				await _channel.SendAsync(_length, cancellationToken);
+				Span<byte> header = _sendBufferWriter.GetMemory().Span;
+				header[0] = (byte)_type;
+				BinaryPrimitives.WriteInt32LittleEndian(header.Slice(1, 4), _length);
+
+				_sendBufferWriter.Advance(HeaderLength + _length);
+				_length = 0;
 			}
 
 			/// <inheritdoc/>
 			public void Advance(int count) => _length += count;
 
 			/// <inheritdoc/>
-			public Memory<byte> GetMemory(int sizeHint = 0) => _buffer.Data.Slice((int)_buffer.WritePosition + _length);
+			public Memory<byte> GetMemory(int sizeHint = 0) => _sendBufferWriter.GetMemory().Slice(HeaderLength + _length);
 
 			/// <inheritdoc/>
 			public Span<byte> GetSpan(int sizeHint = 0) => GetMemory(sizeHint).Span;
 		}
 
-		/// <inheritdoc/>
-		public int Id => _id;
-
-		readonly int _id;
-		readonly IComputeInputBuffer _inputBuffer;
-		readonly IComputeOutputBuffer _outputBuffer;
-		readonly long _flushLength;
+		readonly IComputeBufferReader _receiveBufferReader;
+		readonly IComputeBufferWriter _sendBufferWriter;
 
-		Writer? _currentWriter;
+		MessageBuilder? _currentBuilder;
 
 		/// <summary>
 		/// Constructor
 		/// </summary>
-		/// <param name="id">The channel id</param>
-		/// <param name="inputBuffer">Buffer to read incoming messages from</param>
-		/// <param name="outputBuffer">Buffer to write outgoing messages to</param>
-		/// <param name="maxMessageSize">Maximum size of a message</param>
-		public ComputeChannel(int id, IComputeInputBuffer inputBuffer, IComputeOutputBuffer outputBuffer, int maxMessageSize = 1024 * 64)
+		/// <param name="receiveBufferReader">Reader for incoming messages</param>
+		/// <param name="sendBufferWriter">Writer for outgoing messages</param>
+		public ComputeChannel(IComputeBufferReader receiveBufferReader, IComputeBufferWriter sendBufferWriter)
 		{
-			_id = id;
-			_inputBuffer = inputBuffer;
-			_outputBuffer = outputBuffer;
-			_flushLength = outputBuffer.Data.Length - maxMessageSize;
+			_receiveBufferReader = receiveBufferReader;
+			_sendBufferWriter = sendBufferWriter;
 		}
 
 		/// <inheritdoc/>
-		public void Dispose()
+		public async ValueTask DisposeAsync()
+		{
+			await DisposeAsync(true);
+			GC.SuppressFinalize(this);
+		}
+
+		/// <summary>
+		/// Overridable dispose method
+		/// </summary>
+		protected virtual async ValueTask DisposeAsync(bool disposing)
 		{
-			_currentWriter?.Dispose();
-			_inputBuffer.Dispose();
-			_outputBuffer.Dispose();
+			_currentBuilder?.Dispose();
+			_sendBufferWriter.MarkComplete();
+			await _sendBufferWriter.FlushAsync(CancellationToken.None);
 		}
 
 		/// <inheritdoc/>
-		public async ValueTask<IComputeMessage> ReadAsync(CancellationToken cancellationToken)
+		public async ValueTask<IComputeMessage> ReceiveAsync(CancellationToken cancellationToken)
 		{
-			for (; ; )
+			while (!_receiveBufferReader.IsComplete)
 			{
-				long readPosition = _inputBuffer.ReadPosition;
-				ReadOnlyMemory<byte> memory = _inputBuffer.GetReadMemory();
-
-				const int HeaderLength = 5;
+				ReadOnlyMemory<byte> memory = _receiveBufferReader.GetMemory();
 				if (memory.Length >= HeaderLength)
 				{
 					int length = BinaryPrimitives.ReadInt32LittleEndian(memory.Span.Slice(1, 4));
-					if (memory.Length >= length)
+					if (memory.Length >= HeaderLength + length)
 					{
 						ComputeMessageType type = (ComputeMessageType)memory.Span[0];
-						Message message = new Message(type, memory.Slice(HeaderLength, length - HeaderLength));
-						_inputBuffer.AdvanceReadPosition(length);
+						Message message = new Message(type, memory.Slice(HeaderLength, length));
+						_receiveBufferReader.Advance(HeaderLength + length);
 						return message;
 					}
 				}
-
-				await _inputBuffer.WaitForWrittenDataAsync(readPosition + memory.Length, cancellationToken);
-			}
-		}
-
-		async Task SendAsync(int length, CancellationToken cancellationToken)
-		{
-			_outputBuffer.AdvanceWritePosition(length);
-
-			if (_outputBuffer.WritePosition > _flushLength)
-			{
-				await _outputBuffer.ResetWritePositionAsync(cancellationToken);
+				await _receiveBufferReader.WaitAsync(memory.Length, cancellationToken);
 			}
+			return new Message(ComputeMessageType.None, ReadOnlyMemory<byte>.Empty);
 		}
 
 		/// <inheritdoc/>
-		public IComputeMessageWriter CreateMessage(ComputeMessageType type, int sizeHint = 0)
+		public IComputeMessageBuilder CreateMessage(ComputeMessageType type, int sizeHint = 0)
 		{
-			if (_currentWriter != null)
+			if (_currentBuilder != null)
 			{
 				throw new InvalidOperationException("Only one writer can be active at a time. Dispose of the previous writer first.");
 			}
 
-			_currentWriter = new Writer(this, _outputBuffer, type);
-			return _currentWriter;
+			_currentBuilder = new MessageBuilder(this, _sendBufferWriter, type);
+			return _currentBuilder;
 		}
 	}
 }
diff --git a/Engine/Source/Programs/Shared/EpicGames.Horde/Compute/ComputeLease.cs b/Engine/Source/Programs/Shared/EpicGames.Horde/Compute/ComputeLease.cs
--- a/Engine/Source/Programs/Shared/EpicGames.Horde/Compute/ComputeLease.cs
+++ b/Engine/Source/Programs/Shared/EpicGames.Horde/Compute/ComputeLease.cs
@@ -5,8 +5,8 @@ using System.Buffers;
 using System.Buffers.Binary;
 using System.Collections.Generic;
 using System.Threading;
-using System.Threading.Channels;
 using System.Threading.Tasks;
+using EpicGames.Core;
 
 namespace EpicGames.Horde.Compute
 {
@@ -15,88 +15,24 @@ namespace EpicGames.Horde.Compute
 	/// </summary>
 	public sealed class ComputeLease : IComputeLease, IAsyncDisposable
 	{
-		record struct Packet(int Id, ReadOnlyMemory<byte> Data, TaskCompletionSource? Fence);
-
-		class Buffer : ComputeBuffer
-		{
-			readonly ComputeLease _owner;
-
-			public Buffer(ComputeLease owner, int id, IMemoryOwner<byte> memoryOwner)
-				: base(id, memoryOwner)
-			{
-				_owner = owner;
-			}
-
-			protected override void Dispose(bool disposing)
-			{
-				base.Dispose(disposing);
-
-				if (disposing)
-				{
-					_owner.ReleaseBuffer(this);
-				}
-			}
-
-			public override async ValueTask FlushAsync(CancellationToken cancellationToken)
-			{
-				TaskCompletionSource tcs = new TaskCompletionSource(TaskCreationOptions.RunContinuationsAsynchronously);
-				using (IDisposable disposable = cancellationToken.Register(() => tcs.SetCanceled(cancellationToken)))
-				{
-					_owner._writeQueue.Writer.TryWrite(new Packet(Id, ReadOnlyMemory<byte>.Empty, tcs));
-					await tcs.Task;
-				}
-			}
-
-			public override async ValueTask ResetWritePositionAsync(CancellationToken cancellationToken)
-			{
-				// Flush any pending writes
-				await FlushAsync(cancellationToken);
-
-				// Write an empty packet telling the remote to flush any reads before continuing
-				_owner._writeQueue.Writer.TryWrite(new Packet(Id, ReadOnlyMemory<byte>.Empty, null));
-
-				// Execute the regular logic to make sure we don't have any readers
-				ResetWritePosition();
-			}
+		readonly object _lockObject = new object();
 
-			public async ValueTask ClientResetWritePositionAsync(CancellationToken cancellationToken)
-			{
-				await WaitForAllReadDataAsync(cancellationToken);
-				ResetWritePosition();
-			}
+		bool _complete;
 
-			public override void AdvanceWritePosition(long size)
-			{
-				if (size > 0)
-				{
-					ReadOnlyMemory<byte> data = Data.Slice((int)WritePosition, (int)size);
-					_owner._writeQueue.Writer.TryWrite(new Packet(Id, data, null));
-				}
+		readonly IComputeTransport _transport;
+		readonly Task _readTask;
+		readonly CancellationTokenSource _cancellationSource = new CancellationTokenSource();
 
-				base.AdvanceWritePosition(size);
-			}
+		readonly Dictionary<int, IComputeBufferWriter> _receiveBuffers = new Dictionary<int, IComputeBufferWriter>();
+		readonly AsyncEvent _receiveBuffersChangedEvent = new AsyncEvent();
 
-			public void AdvanceWritePositionAfterRecv(long size)
-			{
-				base.AdvanceWritePosition(size);
-			}
-		}
+		readonly SemaphoreSlim _sendSemaphore = new SemaphoreSlim(1, 1);
+		readonly Dictionary<int, IComputeBufferReader> _sendBuffers = new Dictionary<int, IComputeBufferReader>();
+		readonly Dictionary<int, Task> _sendTasks = new Dictionary<int, Task>();
 
 		/// <inheritdoc/>
 		public IReadOnlyDictionary<string, int> AssignedResources { get; }
 
-		readonly object _lockObject = new object();
-
-		bool _complete;
-		readonly Task _readTask;
-		readonly Task _writeTask;
-		readonly CancellationTokenSource _cts = new CancellationTokenSource();
-
-		readonly Dictionary<int, Buffer> _inputBuffers = new Dictionary<int, Buffer>();
-		readonly Dictionary<int, Buffer> _outputBuffers = new Dictionary<int, Buffer>();
-
-		readonly Channel<Packet> _writeQueue = Channel.CreateUnbounded<Packet>();
-
 		/// <summary>
 		/// Constructor
 		/// </summary>
@@ -106,8 +42,8 @@ namespace EpicGames.Horde.Compute
 		{
 			AssignedResources = assignedResources;
 
-			_readTask = RunReaderAsync(transport, _cts.Token);
-			_writeTask = RunWriterAsync(transport, _cts.Token);
+			_transport = transport;
+			_readTask = RunReaderAsync(transport, _cancellationSource.Token);
 		}
 
 		/// <summary>
@@ -115,13 +51,18 @@ namespace EpicGames.Horde.Compute
 		/// </summary>
 		public async ValueTask CloseAsync(CancellationToken cancellationToken)
 		{
-			_writeQueue.Writer.TryComplete();
-
 			Task cancellationTask = Task.Delay(-1, cancellationToken);
 
-			await Task.WhenAny(_writeTask, cancellationTask);
+			// Wait for all the individual send tasks to complete
+			await Task.WhenAll(_sendTasks.Values);
 			cancellationToken.ThrowIfCancellationRequested();
 
+			// Send a final message indicating that the lease is done. This will allow the senders on the remote end to terminate, and trigger
+			// a message to be sent back to our read task allowing it to shut down gracefully.
+			byte[] header = new byte[8];
+			await _transport.WriteAsync(header, cancellationToken);
+
+			// Wait for the reader to stop
 			await Task.WhenAny(_readTask, cancellationTask);
 			cancellationToken.ThrowIfCancellationRequested();
 		}
@@ -129,14 +70,13 @@ namespace EpicGames.Horde.Compute
 		/// <inheritdoc/>
 		public async ValueTask DisposeAsync()
 		{
-			_writeQueue.Writer.TryComplete();
+			_cancellationSource.Cancel();
 
-			_cts.Cancel();
-
-			await _writeTask;
+			await Task.WhenAll(_sendTasks.Values);
 			await _readTask;
 
-			_cts.Dispose();
+			_cancellationSource.Dispose();
+			_sendSemaphore.Dispose();
 		}
 
 		async Task RunReaderAsync(IComputeTransport transport, CancellationToken cancellationToken)
@@ -145,42 +85,35 @@ namespace EpicGames.Horde.Compute
 			try
 			{
 				// Maintain a local cache of buffers to be able to query for them without having to acquire a global lock
-				Dictionary<int, Buffer> buffers = new Dictionary<int, Buffer>();
+				Dictionary<int, IComputeBufferWriter> cachedWriters = new Dictionary<int, IComputeBufferWriter>();
 
 				// Process messages from the remote
 				for(; ;)
 				{
+					// Read the next packet header
 					await transport.ReadAsync(header, cancellationToken);
 
+					// Parse the target buffer and packet size
 					int id = BinaryPrimitives.ReadInt32LittleEndian(header);
 					int size = BinaryPrimitives.ReadInt32LittleEndian(header.AsSpan(4));
-					if (size < 0)
-					{
-						break;
-					}
 
-					Buffer? buffer;
-					if (!buffers.TryGetValue(id, out buffer))
+					// If the size if negative, we're closing the entire connection
+					if (size == 0)
 					{
-						lock (_lockObject)
-						{
-							if (!_inputBuffers.TryGetValue(id, out buffer))
-							{
-								throw new NotImplementedException();
-							}
-							buffers.Add(id, buffer);
-						}
+						break;
 					}
-
-					if (size == 0)
+					
+					// Dispatch it to the correct place
+					IComputeBufferWriter writer = await GetReceiveBufferAsync(cachedWriters, id);
+					if (size < 0)
 					{
-						// We're being asked to reset the read position back to zero, so we need to know that the read offset is equal to length
-						await buffer.ClientResetWritePositionAsync(cancellationToken);
+						writer.MarkComplete();
+						DetachReceiveBuffer(cachedWriters, id);
 					}
 					else
 					{
-						await transport.ReadAsync(buffer.Data.Slice((int)buffer.WritePosition, size), cancellationToken);
-						buffer.AdvanceWritePositionAfterRecv(size);
+						await transport.ReadAsync(writer.GetMemory(), cancellationToken);
+						writer.Advance(size);
 					}
 				}
 			}
@@ -192,91 +125,138 @@ namespace EpicGames.Horde.Compute
 			lock (_lockObject)
 			{
 				_complete = true;
-				foreach (Buffer buffer in _outputBuffers.Values)
+				foreach (IComputeBufferWriter writer in _receiveBuffers.Values)
 				{
-					buffer.MarkComplete();
+					writer.MarkComplete();
 				}
 			}
 		}
 
-		async Task RunWriterAsync(IComputeTransport transport, CancellationToken cancellationToken)
+		async Task<IComputeBufferWriter> GetReceiveBufferAsync(Dictionary<int, IComputeBufferWriter> cachedWriters, int id)
 		{
-			// Parse the queue of write requests
-			byte[] header = new byte[8];
-			try
+			IComputeBufferWriter? writer;
+			if (cachedWriters.TryGetValue(id, out writer))
 			{
-				while (await _writeQueue.Reader.WaitToReadAsync(cancellationToken))
+				return writer;
+			}
+
+			for(; ;)
+			{
+				Task waitTask;
+				lock (_lockObject)
 				{
-					Packet packet = await _writeQueue.Reader.ReadAsync(cancellationToken);
-					if (packet.Fence != null)
+					if (_receiveBuffers.TryGetValue(id, out writer))
 					{
-						packet.Fence.SetResult();
-						continue;
+						cachedWriters.Add(id, writer);
+						return writer;
 					}
+					waitTask = _receiveBuffersChangedEvent.Task;
+				}
+				await waitTask;
+			}
+		}
+
+		class SendSegment : ReadOnlySequenceSegment<byte>
+		{
+			public void Set(ReadOnlyMemory<byte> memory, ReadOnlySequenceSegment<byte>? next, long runningIndex)
+			{
+				Memory = memory;
+				Next = next;
+				RunningIndex = runningIndex;
+			}
+		}
+
+		async Task SendFromBufferAsync(int id, IComputeBufferReader reader, CancellationToken cancellationToken)
+		{
+			byte[] header = new byte[8];
+
+			BinaryPrimitives.WriteInt32LittleEndian(header, id);
+
+			SendSegment headerSegment = new SendSegment();
+			SendSegment bodySegment = new SendSegment();
+			headerSegment.Set(header, bodySegment, 0);
+
+			for (; ; )
+			{
+				// Wait for something to read
+				await reader.WaitAsync(0, cancellationToken);
 
-					BinaryPrimitives.WriteInt32LittleEndian(header, packet.Id);
-					BinaryPrimitives.WriteInt32LittleEndian(header.AsSpan(4), packet.Data.Length);
-					await transport.WriteAsync(header, cancellationToken);
+				// Acquire the semaphore
+				await _sendSemaphore.WaitAsync(cancellationToken);
+				try
+				{
+					ReadOnlyMemory<byte> memory = reader.GetMemory();
+					if (memory.Length > 0)
+					{
+						BinaryPrimitives.WriteInt32LittleEndian(header.AsSpan(4), memory.Length);
+
+						bodySegment.Set(memory, null, header.Length);
+						ReadOnlySequence<byte> sequence = new ReadOnlySequence<byte>(headerSegment, 0, bodySegment, memory.Length);
 
-					if (packet.Data.Length > 0)
+						await _transport.WriteAsync(sequence, cancellationToken);
+
+						reader.Advance(memory.Length);
+					}
+					else if (reader.IsComplete)
 					{
-						await transport.WriteAsync(packet.Data, cancellationToken);
+						BinaryPrimitives.WriteInt32LittleEndian(header.AsSpan(4), -1);
+						await _transport.WriteAsync(header, cancellationToken);
+						break;
 					}
 				}
-
-				// Send a close packet to the remote
-				BinaryPrimitives.WriteInt32LittleEndian(header, 0);
-				BinaryPrimitives.WriteInt32LittleEndian(header.AsSpan(4), -1);
-				await transport.WriteAsync(header, cancellationToken);
+				finally
+				{
+					_sendSemaphore.Release();
+				}
 			}
-			catch (OperationCanceledException)
+
+			lock (_lockObject)
 			{
+				_sendBuffers.Remove(id);
+				_sendTasks.Remove(id);
 			}
 		}
 
 		/// <inheritdoc/>
-		public IComputeInputBuffer CreateInputBuffer(int id, IMemoryOwner<byte> memoryOwner)
+		public void AttachSendBuffer(int id, IComputeBufferReader reader)
 		{
-			Buffer inputBuffer = new Buffer(this, id, memoryOwner);
 			lock (_lockObject)
 			{
-				if (_complete)
-				{
-					inputBuffer.MarkComplete();
-				}
-				_inputBuffers.Add(id, inputBuffer);
+				_sendBuffers.Add(id, reader);
+				_sendTasks.Add(id, Task.Run(() => SendFromBufferAsync(id, reader, _cancellationSource.Token), _cancellationSource.Token));
 			}
-			return inputBuffer;
 		}
 
 		/// <inheritdoc/>
-		public IComputeOutputBuffer CreateOutputBuffer(int id, IMemoryOwner<byte> memoryOwner)
+		public void AttachReceiveBuffer(int id, IComputeBufferWriter writer)
 		{
-			Buffer outputBuffer = new Buffer(this, id, memoryOwner);
+			bool complete;
 			lock (_lockObject)
 			{
-				if (_complete)
+				complete = _complete;
+
+				if(!complete)
 				{
-					outputBuffer.MarkComplete();
+					_receiveBuffers.Add(id, writer);
 				}
-				_outputBuffers.Add(id, outputBuffer);
 			}
-			return outputBuffer;
+
+			if (complete)
+			{
+				writer.MarkComplete();
+			}
+			else
+			{
+				_receiveBuffersChangedEvent.Set();
+			}
 		}
 
-		void ReleaseBuffer(Buffer buffer)
+		void DetachReceiveBuffer(Dictionary<int, IComputeBufferWriter> cachedWriters, int id)
 		{
+			cachedWriters.Remove(id);
 			lock (_lockObject)
 			{
-				Buffer? existingBuffer;
-				if (_inputBuffers.TryGetValue(buffer.Id, out existingBuffer) && ReferenceEquals(existingBuffer, buffer))
-				{
-					_inputBuffers.Remove(buffer.Id);
-				}
-				if (_outputBuffers.TryGetValue(buffer.Id, out existingBuffer) && ReferenceEquals(existingBuffer, buffer))
-				{
-					_outputBuffers.Remove(buffer.Id);
-				}
+				_receiveBuffers.Remove(id);
 			}
 		}
 	}
diff --git a/Engine/Source/Programs/Shared/EpicGames.Horde/Compute/ComputeStorageClient.cs b/Engine/Source/Programs/Shared/EpicGames.Horde/Compute/ComputeStorageClient.cs
--- a/Engine/Source/Programs/Shared/EpicGames.Horde/Compute/ComputeStorageClient.cs
+++ b/Engine/Source/Programs/Shared/EpicGames.Horde/Compute/ComputeStorageClient.cs
@@ -71,9 +71,9 @@ namespace EpicGames.Horde.Compute
 			await _semaphore.WaitAsync(cancellationToken);
 			try
 			{
-				await _channel.CppBlobReadAsync(locator, offset, length, cancellationToken);
+				_channel.CppBlobRead(locator, offset, length);
 
-				IComputeMessage response = await _channel.ReadAsync(cancellationToken);
+				IComputeMessage response = await _channel.ReceiveAsync(cancellationToken);
 				if (response.Type != ComputeMessageType.CppBlobData)
 				{
 					response.Dispose();
diff --git a/Engine/Source/Programs/Shared/EpicGames.Horde/Compute/Cpp/CppMessages.cs b/Engine/Source/Programs/Shared/EpicGames.Horde/Compute/Cpp/CppMessages.cs
--- a/Engine/Source/Programs/Shared/EpicGames.Horde/Compute/Cpp/CppMessages.cs
+++ b/Engine/Source/Programs/Shared/EpicGames.Horde/Compute/Cpp/CppMessages.cs
@@ -36,14 +36,13 @@ namespace EpicGames.Horde.Compute.Cpp
 		/// <param name="channel">Channel to write to</param>
 		/// <param name="locator">Locator for the node containing the task definition</param>
 		/// <param name="replyChannelId">Channel to reply to with future messages</param>
-		/// <param name="cancellationToken">Cancellation token for the operation</param>
-		public static async Task CppStartAsync(this IComputeChannel channel, NodeLocator locator, int replyChannelId, CancellationToken cancellationToken)
+		public static void CppStart(this IComputeChannel channel, NodeLocator locator, int replyChannelId)
 		{
-			using (IComputeMessageWriter writer = channel.CreateMessage(ComputeMessageType.CppBegin))
+			using (IComputeMessageBuilder writer = channel.CreateMessage(ComputeMessageType.CppBegin))
 			{
 				writer.WriteNodeLocator(locator);
 				writer.WriteUnsignedVarInt(replyChannelId);
-				await writer.SendAsync(cancellationToken);
+				writer.Send();
 			}
 		}
 
@@ -75,13 +74,12 @@ namespace EpicGames.Horde.Compute.Cpp
 		/// </summary>
 		/// <param name="channel">Channel to write to</param>
 		/// <param name="locator">Locator for the result</param>
-		/// <param name="cancellationToken">Cancellation token for the operation</param>
-		public static async Task CppSuccessAsync(this IComputeChannel channel, NodeLocator locator, CancellationToken cancellationToken)
+		public static void CppSuccess(this IComputeChannel channel, NodeLocator locator)
 		{
-			using (IComputeMessageWriter writer = channel.CreateMessage(ComputeMessageType.CppSuccess))
+			using (IComputeMessageBuilder writer = channel.CreateMessage(ComputeMessageType.CppSuccess))
 			{
 				writer.WriteNodeLocator(locator);
-				await writer.SendAsync(cancellationToken);
+				writer.Send();
 			}
 		}
 
@@ -100,13 +98,12 @@ namespace EpicGames.Horde.Compute.Cpp
 		/// </summary>
 		/// <param name="channel">Channel to write to</param>
 		/// <param name="error">Error string to send to the caller</param>
-		/// <param name="cancellationToken">Cancellation token for the operation</param>
-		public static async Task CppFailureAsync(this IComputeChannel channel, string error, CancellationToken cancellationToken)
+		public static void CppFailure(this IComputeChannel channel, string error)
 		{
-			using (IComputeMessageWriter writer = channel.CreateMessage(ComputeMessageType.CppFailure))
+			using (IComputeMessageBuilder writer = channel.CreateMessage(ComputeMessageType.CppFailure))
 			{
 				writer.WriteString(error);
-				await writer.SendAsync(cancellationToken);
+				writer.Send();
 			}
 		}
 
@@ -114,12 +111,11 @@ namespace EpicGames.Horde.Compute.Cpp
 		/// Completes a C++ build task
 		/// </summary>
 		/// <param name="channel">Channel to write to</param>
-		/// <param name="cancellationToken">Cancellation token for the operation</param>
-		public static async Task CppFinishAsync(this IComputeChannel channel, CancellationToken cancellationToken)
+		public static void CppFinish(this IComputeChannel channel)
 		{
-			using (IComputeMessageWriter writer = channel.CreateMessage(ComputeMessageType.CppEnd))
+			using (IComputeMessageBuilder writer = channel.CreateMessage(ComputeMessageType.CppEnd))
 			{
-				await writer.SendAsync(cancellationToken);
+				writer.Send();
 			}
 		}
 
@@ -144,16 +140,15 @@ namespace EpicGames.Horde.Compute.Cpp
 		/// <param name="locator"></param>
 		/// <param name="offset"></param>
 		/// <param name="length"></param>
-		/// <param name="cancellationToken"></param>
 		/// <returns></returns>
-		public static async Task CppBlobReadAsync(this IComputeChannel channel, BlobLocator locator, int offset, int length, CancellationToken cancellationToken)
+		public static void CppBlobRead(this IComputeChannel channel, BlobLocator locator, int offset, int length)
 		{
-			using (IComputeMessageWriter writer = channel.CreateMessage(ComputeMessageType.CppBlobRead))
+			using (IComputeMessageBuilder writer = channel.CreateMessage(ComputeMessageType.CppBlobRead))
 			{
 				writer.WriteBlobLocator(locator);
 				writer.WriteUnsignedVarInt(offset);
 				writer.WriteUnsignedVarInt(length);
-				await writer.SendAsync(cancellationToken);
+				writer.Send();
 			}
 		}
 
@@ -200,10 +195,10 @@ namespace EpicGames.Horde.Compute.Cpp
 				}
 			}
 
-			using (IComputeMessageWriter writer = channel.CreateMessage(ComputeMessageType.CppBlobData, length + 128))
+			using (IComputeMessageBuilder writer = channel.CreateMessage(ComputeMessageType.CppBlobData, length + 128))
 			{
 				writer.WriteFixedLengthBytes(data);
-				await writer.SendAsync(cancellationToken);
+				writer.Send();
 			}
 		}
 	}
diff --git a/Engine/Source/Programs/Shared/EpicGames.Horde/Compute/IComputeBuffer.cs b/Engine/Source/Programs/Shared/EpicGames.Horde/Compute/IComputeBuffer.cs
--- a/Engine/Source/Programs/Shared/EpicGames.Horde/Compute/IComputeBuffer.cs
+++ b/Engine/Source/Programs/Shared/EpicGames.Horde/Compute/IComputeBuffer.cs
@@ -12,94 +12,71 @@ namespace EpicGames.Horde.Compute
 	public interface IComputeBuffer
 	{
 		/// <summary>
-		/// Identifier for the buffer
+		/// Reader from the buffer
 		/// </summary>
-		int Id { get; }
+		IComputeBufferReader Reader { get; }
+
+		/// <summary>
+		/// Writer to the buffer
+		/// </summary>
+		IComputeBufferWriter Writer { get; }
 	}
 
 	/// <summary>
 	/// Read interface for a compute buffer
 	/// </summary>
-	public interface IComputeInputBuffer : IComputeBuffer, IDisposable
+	public interface IComputeBufferReader
 	{
 		/// <summary>
-		/// The shared memory for the buffer
+		/// Whether this buffer is complete (no more data will be added)
 		/// </summary>
-		ReadOnlyMemory<byte> Data { get; }
+		bool IsComplete { get; }
 
 		/// <summary>
-		/// Offset within the data that is still to be used. Can be left at zero for a buffer that is only appended to, but must be updated to indicate data
-		/// that has been processed for message buffers (in order to coordinate when the buffer contents can be discarded during a reset).
+		/// Updates the read position
 		/// </summary>
-		long ReadPosition { get; }
+		/// <param name="size">Size of data that was read</param>
+		void Advance(int size);
 
 		/// <summary>
-		/// Position in the buffer at which new data will be added
+		/// Gets the next data to read
 		/// </summary>
-		long WritePosition { get; }
+		/// <returns>Memory to read from</returns>
+		ReadOnlyMemory<byte> GetMemory();
 
 		/// <summary>
-		/// Updates the current read position within the buffer
+		/// Wait for data to be available, or for the buffer to be marked as complete
 		/// </summary>
-		void AdvanceReadPosition(long size);
-
-		/// <summary>
-		/// Waits until new data is written to the buffer
-		/// </summary>
-		/// <param name="currentWritePosition">Current write position. This is passed as an argument to avoid races with <see cref="WritePosition"/> being updated asynchronously.</param>
+		/// <param name="currentLength">Current length of the buffer</param>
 		/// <param name="cancellationToken">Cancellation token for the operation</param>
-		ValueTask WaitForWrittenDataAsync(long currentWritePosition, CancellationToken cancellationToken);
+		ValueTask WaitAsync(int currentLength, CancellationToken cancellationToken);
 	}
 
 	/// <summary>
-	/// Read-only view of shared memory that can be incrementally replicated to a remote machine.
+	/// Buffer that can receive data from a remote machine.
 	/// </summary>
-	public interface IComputeOutputBuffer : IComputeBuffer, IDisposable
+	public interface IComputeBufferWriter
 	{
 		/// <summary>
-		/// The shared memory for the buffer
+		/// Mark the output to this buffer as complete
 		/// </summary>
-		Memory<byte> Data { get; }
-
-		/// <summary>
-		/// Current write position/used length of the buffer
-		/// </summary>
-		long WritePosition { get; }
+		void MarkComplete();
 
 		/// <summary>
 		/// Updates the current write position within the buffer
 		/// </summary>
-		void AdvanceWritePosition(long size);
+		void Advance(int size);
 
 		/// <summary>
-		/// Flushes all data that needs to be written.
+		/// Gets memory to write to
 		/// </summary>
-		ValueTask FlushAsync(CancellationToken cancellationToken);
+		/// <returns>Memory to be written to</returns>
+		Memory<byte> GetMemory();
 
 		/// <summary>
-		/// Flushes all writes, and resets the current write position.
+		/// Waits until all data in this buffer has been written
 		/// </summary>
-		ValueTask ResetWritePositionAsync(CancellationToken cancellationToken);
-	}
-
-	/// <summary>
-	/// Extension methods for compute buffers
-	/// </summary>
-	static class ComputeBufferExtensions
-	{
-		/// <summary>
-		/// Gets a span for the next data to read
-		/// </summary>
-		public static ReadOnlySpan<byte> GetReadSpan(this IComputeInputBuffer buffer) => GetReadMemory(buffer).Span;
-
-		/// <summary>
-		/// Gets a span for the next data to read
-		/// </summary>
-		public static ReadOnlyMemory<byte> GetReadMemory(this IComputeInputBuffer buffer) => buffer.Data.Slice((int)buffer.ReadPosition, (int)(buffer.WritePosition - buffer.ReadPosition));
-
-		/// <summary>
-		/// Gets a span to write new data to
-		/// </summary>
-		public static Span<byte> GetWriteSpan(this IComputeOutputBuffer buffer) => buffer.Data.Slice((int)buffer.WritePosition).Span;
+		/// <param name="cancellationToken">Cancellation token for the operation</param>
+		ValueTask FlushAsync(CancellationToken cancellationToken);
 	}
 }
diff --git a/Engine/Source/Programs/Shared/EpicGames.Horde/Compute/IComputeChannel.cs b/Engine/Source/Programs/Shared/EpicGames.Horde/Compute/IComputeChannel.cs
--- a/Engine/Source/Programs/Shared/EpicGames.Horde/Compute/IComputeChannel.cs
+++ b/Engine/Source/Programs/Shared/EpicGames.Horde/Compute/IComputeChannel.cs
@@ -5,25 +5,21 @@ using System.Buffers;
 using System.Threading;
 using System.Threading.Tasks;
 using EpicGames.Core;
+using EpicGames.Horde.Compute.Buffers;
 
 namespace EpicGames.Horde.Compute
 {
 	/// <summary>
 	/// Full-duplex channel for sending and receiving messages
 	/// </summary>
-	public interface IComputeChannel : IDisposable
+	public interface IComputeChannel : IAsyncDisposable
 	{
-		/// <summary>
-		/// Identifier for the channel
-		/// </summary>
-		int Id { get; }
-
 		/// <summary>
 		/// Reads a message from the channel
 		/// </summary>
 		/// <param name="cancellationToken">Cancellation token for the operation</param>
 		/// <returns>Data for a message that was read. Must be disposed.</returns>
-		ValueTask<IComputeMessage> ReadAsync(CancellationToken cancellationToken);
+		ValueTask<IComputeMessage> ReceiveAsync(CancellationToken cancellationToken);
 
 		/// <summary>
 		/// Creates a new builder for a message
@@ -31,19 +27,18 @@ namespace EpicGames.Horde.Compute
 		/// <param name="type">Type of the message</param>
 		/// <param name="sizeHint">Hint for the expected message size</param>
 		/// <returns>New builder for messages</returns>
-		IComputeMessageWriter CreateMessage(ComputeMessageType type, int sizeHint = 0);
+		IComputeMessageBuilder CreateMessage(ComputeMessageType type, int sizeHint = 0);
 	}
 
 	/// <summary>
 	/// Writer for compute messages
 	/// </summary>
-	public interface IComputeMessageWriter : IMemoryWriter, IDisposable
+	public interface IComputeMessageBuilder : IMemoryWriter, IDisposable
 	{
 		/// <summary>
 		/// Sends the current message
 		/// </summary>
-		/// <param name="cancellationToken">Cancellation token for the operation</param>
-		ValueTask SendAsync(CancellationToken cancellationToken);
+		void Send();
 	}
 
 	/// <summary>
@@ -51,6 +46,31 @@ namespace EpicGames.Horde.Compute
 	/// </summary>
 	public static class ComputeChannelExtensions
 	{
+		class RentedBufferChannel : ComputeChannel
+		{
+			readonly RentedLocalBuffer _receiveBuffer;
+			readonly RentedLocalBuffer _sendBuffer;
+
+			public RentedBufferChannel(RentedLocalBuffer receiveBuffer, RentedLocalBuffer sendBuffer)
+				: base(receiveBuffer.Reader, sendBuffer.Writer)
+			{
+				_sendBuffer = sendBuffer;
+				_receiveBuffer = receiveBuffer;
+			}
+
+			/// <inheritdoc/>
+			protected override async ValueTask DisposeAsync(bool disposing)
+			{
+				await base.DisposeAsync(disposing);
+
+				if (disposing)
+				{
+					_sendBuffer.Dispose();
+					_receiveBuffer.Dispose();
+				}
+			}
+		}
+
 		/// <summary>
 		/// Creates a compute channel with the given identifier
 		/// </summary>
@@ -61,10 +81,13 @@ namespace EpicGames.Horde.Compute
 			const int MaxMessageSize = 64 * 1024;
 			const int BufferSize = MaxMessageSize * 3;
 
-			IComputeInputBuffer inputBuffer = lease.CreateInputBuffer(id, MemoryPool<byte>.Shared.Rent(BufferSize));
-			IComputeOutputBuffer outputBuffer = lease.CreateOutputBuffer(id, MemoryPool<byte>.Shared.Rent(BufferSize));
+			RentedLocalBuffer receiveBuffer = new RentedLocalBuffer(BufferSize);
+			lease.AttachReceiveBuffer(id, receiveBuffer.Writer);
+
+			RentedLocalBuffer sendBuffer = new RentedLocalBuffer(BufferSize);
+			lease.AttachSendBuffer(id, sendBuffer.Reader);
 
-			return new ComputeChannel(id, inputBuffer, outputBuffer, MaxMessageSize);
+			return new RentedBufferChannel(receiveBuffer, sendBuffer);
 		}
 	}
 }
diff --git a/Engine/Source/Programs/Shared/EpicGames.Horde/Compute/IComputeLease.cs b/Engine/Source/Programs/Shared/EpicGames.Horde/Compute/IComputeLease.cs
--- a/Engine/Source/Programs/Shared/EpicGames.Horde/Compute/IComputeLease.cs
+++ b/Engine/Source/Programs/Shared/EpicGames.Horde/Compute/IComputeLease.cs
@@ -1,7 +1,6 @@
 // Copyright Epic Games, Inc. All Rights Reserved.
 
 using System;
-using System.Buffers;
 using System.Collections.Generic;
 using System.Threading;
 using System.Threading.Tasks;
@@ -24,17 +23,17 @@ namespace EpicGames.Horde.Compute
 		ValueTask CloseAsync(CancellationToken cancellationToken);
 
 		/// <summary>
-		/// Creates an input buffer from the given memory. A corresponding output buffer with the same size and id must be created at the remote end.
+		/// Attaches a send buffer to this lease. Data will be read from this buffer and replicated a receive buffer attached with the same id on the remote.
 		/// </summary>
-		/// <param name="id">Identifier for the buffer. This much match up with an output buffer on the remote for a connection to be established.</param>
-		/// <param name="memory">Memory to read from</param>
-		IComputeInputBuffer CreateInputBuffer(int id, IMemoryOwner<byte> memory);
+		/// <param name="id">Identifier for the buffer</param>
+		/// <param name="buffer">The buffer to attach</param>
+		void AttachSendBuffer(int id, IComputeBufferReader buffer);
 
 		/// <summary>
-		/// Creates an output buffer from the given memory. A corresponding output buffer with the same size and id must be created at the remote end.
+		/// Attaches a receive buffer to this lease. Data will be read into this buffer from the other end of the lease.
 		/// </summary>
-		/// <param name="id">Identifier for the buffer. This much match up with an input buffer on the remote for a connection to be established.</param>
-		/// <param name="memory">Memory to read from</param>
-		IComputeOutputBuffer CreateOutputBuffer(int id, IMemoryOwner<byte> memory);
+		/// <param name="id">Identifier for the buffer</param>
+		/// <param name="buffer">The buffer to attach</param>
+		void AttachReceiveBuffer(int id, IComputeBufferWriter buffer);
 	}
 }
diff --git a/Engine/Source/Programs/Shared/EpicGames.Horde/Compute/IComputeMessage.cs b/Engine/Source/Programs/Shared/EpicGames.Horde/Compute/IComputeMessage.cs
--- a/Engine/Source/Programs/Shared/EpicGames.Horde/Compute/IComputeMessage.cs
+++ b/Engine/Source/Programs/Shared/EpicGames.Horde/Compute/IComputeMessage.cs
@@ -17,11 +17,6 @@ namespace EpicGames.Horde.Compute
 		/// </summary>
 		None = 0x00,
 
-		/// <summary>
-		/// Gracefully close the current connection
-		/// </summary>
-		Close = 0x01,
-
 		#region Test Requests
 
 		/// <summary>
@@ -100,41 +95,21 @@ namespace EpicGames.Horde.Compute
 	public static class ComputeMessageExtensions
 	{
 		/// <summary>
-		/// Send a close message to the remote
-		/// </summary>
-		/// <param name="channel">Channel to write to</param>
-		/// <param name="cancellationToken">Cancellation token for the operation</param>
-		public static async Task CloseAsync(this IComputeChannel channel, CancellationToken cancellationToken)
-		{
-			using (IComputeMessageWriter writer = channel.CreateMessage(ComputeMessageType.Close))
-			{
-				await writer.SendAsync(cancellationToken);
-			}
-		}
-
-		/// <summary>
-		/// 
+		/// Send a message to request that a byte string be xor'ed with a particular value
 		/// </summary>
-		/// <param name="channel"></param>
-		/// <param name="data"></param>
-		/// <param name="value"></param>
-		/// <param name="cancellationToken"></param>
-		/// <returns></returns>
-		public static async Task XorRequestAsync(this IComputeChannel channel, ReadOnlyMemory<byte> data, byte value, CancellationToken cancellationToken)
+		public static void XorRequest(this IComputeChannel channel, ReadOnlyMemory<byte> data, byte value)
 		{
-			using (IComputeMessageWriter writer = channel.CreateMessage(ComputeMessageType.XorRequest))
+			using (IComputeMessageBuilder builder = channel.CreateMessage(ComputeMessageType.XorRequest))
 			{
-				writer.WriteFixedLengthBytes(data.Span);
-				writer.WriteUInt8(value);
-				await writer.SendAsync(cancellationToken);
+				builder.WriteFixedLengthBytes(data.Span);
+				builder.WriteUInt8(value);
+				builder.Send();
 			}
 		}
 
 		/// <summary>
-		/// 
+		/// Parse a message as an XOR request
 		/// </summary>
-		/// <param name="message"></param>
-		/// <returns></returns>
 		public static XorRequestMessage AsXorRequest(this IComputeMessage message)
 		{
 			ReadOnlyMemory<byte> data = message.Data;
diff --git a/Engine/Source/Programs/Shared/EpicGames.Horde/Compute/IComputeTransport.cs b/Engine/Source/Programs/Shared/EpicGames.Horde/Compute/IComputeTransport.cs
--- a/Engine/Source/Programs/Shared/EpicGames.Horde/Compute/IComputeTransport.cs
+++ b/Engine/Source/Programs/Shared/EpicGames.Horde/Compute/IComputeTransport.cs
@@ -8,7 +8,7 @@ using System.Threading.Tasks;
 namespace EpicGames.Horde.Compute
 {
 	/// <summary>
-	/// Interface for transferring data on a compute channel
+	/// Low-level interface for transferring data
 	/// </summary>
 	public interface IComputeTransport
 	{
diff --git a/Engine/Source/Programs/Shared/EpicGames.Horde/Compute/LoopbackComputeClient.cs b/Engine/Source/Programs/Shared/EpicGames.Horde/Compute/LoopbackComputeClient.cs
--- a/Engine/Source/Programs/Shared/EpicGames.Horde/Compute/LoopbackComputeClient.cs
+++ b/Engine/Source/Programs/Shared/EpicGames.Horde/Compute/LoopbackComputeClient.cs
@@ -4,7 +4,6 @@ using System;
 using System.Collections.Generic;
 using System.Net;
 using System.Net.Sockets;
-using System.Security.Cryptography;
 using System.Threading;
 using System.Threading.Tasks;
 using EpicGames.Core;
diff --git a/Engine/Source/Programs/Shared/EpicGames.Horde/Compute/SharedMemory.cs b/Engine/Source/Programs/Shared/EpicGames.Horde/Compute/SharedMemory.cs
--- a/Engine/Source/Programs/Shared/EpicGames.Horde/Compute/SharedMemory.cs
+++ b/Engine/Source/Programs/Shared/EpicGames.Horde/Compute/SharedMemory.cs
@@ -4,6 +4,7 @@ using Microsoft.Win32.SafeHandles;
 using System;
 using System.Buffers;
 using System.IO.MemoryMappedFiles;
+using System.Reflection;
 
 namespace EpicGames.Horde.Compute
 {
@@ -55,16 +56,41 @@ namespace EpicGames.Horde.Compute
 		/// </summary>
 		/// <param name="capacity">Size of the memory to allocate. Data will not be committed until used.</param>
 		public SharedMemory(long capacity)
+			: this(MemoryMappedFile.CreateNew(null, capacity))
 		{
-			_file = MemoryMappedFile.CreateNew(null, capacity);
+		}
 
-			_viewAccessor = _file.CreateViewAccessor(0, capacity);
+		/// <summary>
+		/// Creates a memory mapped file from the given file object and capacity
+		/// </summary>
+		/// <param name="file">File to create from</param>
+		private SharedMemory(MemoryMappedFile file)
+		{
+			_file = file;
+
+			_viewAccessor = _file.CreateViewAccessor();
 			_viewAccessor.SafeMemoryMappedViewHandle.AcquirePointer(ref _data);
 
-			MemoryWrapper wrapper = new MemoryWrapper(_data, (int)capacity);
+			MemoryWrapper wrapper = new MemoryWrapper(_data, (int)_viewAccessor.Capacity);
 			Memory = wrapper.Memory;
 		}
 
+		/// <summary>
+		/// Opens a shared memory region from a handle
+		/// </summary>
+		/// <param name="handle">Handle to the underlying memory mapped file</param>
+		public static SharedMemory FromHandle(SafeMemoryMappedFileHandle handle)
+		{
+			ConstructorInfo? constructorInfo = typeof(MemoryMappedFile).GetConstructor(BindingFlags.Instance | BindingFlags.NonPublic, null, new[] { typeof(SafeMemoryMappedFileHandle) }, null);
+			if (constructorInfo == null)
+			{
+				throw new InvalidOperationException("Cannot find private constructor for memory mapped file");
+			}
+
+			MemoryMappedFile file = (MemoryMappedFile)constructorInfo.Invoke(new object?[] { handle });
+			return new SharedMemory(file);
+		}
+
 		/// <summary>
 		/// Finalizer
 		/// </summary>
