Add the ability to compute a hash from a slice of an array.

diff --git a/Engine/Source/Programs/Shared/EpicGames.Core/ContentHash.cs b/Engine/Source/Programs/Shared/EpicGames.Core/ContentHash.cs
--- a/Engine/Source/Programs/Shared/EpicGames.Core/ContentHash.cs
+++ b/Engine/Source/Programs/Shared/EpicGames.Core/ContentHash.cs
@@ -121,6 +121,19 @@ namespace EpicGames.Core
 			return new ContentHash(algorithm.ComputeHash(data));
 		}
 
+		/// <summary>
+		/// Creates a content hash for a block of data, using a given algorithm.
+		/// </summary>
+		/// <param name="data">Data to compute the hash for</param>
+		/// <param name="offset">Offset to the start of the data</param>
+		/// <param name="count">Length of the data</param>
+		/// <param name="algorithm">Algorithm to use to create the hash</param>
+		/// <returns>New content hash instance containing the hash of the data</returns>
+		public static ContentHash Compute(byte[] data, int offset, int count, HashAlgorithm algorithm)
+		{
+			return new ContentHash(algorithm.ComputeHash(data, offset, count));
+		}
+
 		/// <summary>
 		/// Creates a content hash for a string, using a given algorithm.
 		/// </summary>
@@ -159,6 +172,21 @@ namespace EpicGames.Core
 			}
 		}
 
+		/// <summary>
+		/// Creates a content hash for a block of data using MD5
+		/// </summary>
+		/// <param name="data">Data to compute the hash for</param>
+		/// <param name="offset">Offset to the start of the data</param>
+		/// <param name="count">Length of the data</param>
+		/// <returns>New content hash instance containing the hash of the data</returns>
+		public static ContentHash MD5(byte[] data, int offset, int count)
+		{
+			using (MD5 algorithm = System.Security.Cryptography.MD5.Create())
+			{
+				return Compute(data, offset, count, algorithm);
+			}
+		}
+
 		/// <summary>
 		/// Creates a content hash for a string using MD5.
 		/// </summary>
