Horde: Add a compute buffer type that uses shared memory and OS events to communicate between processes.

diff --git a/Engine/Source/Programs/Shared/EpicGames.Core/MemoryMappedView.cs b/Engine/Source/Programs/Shared/EpicGames.Core/MemoryMappedView.cs
new file mode 100644
--- /dev/null
+++ b/Engine/Source/Programs/Shared/EpicGames.Core/MemoryMappedView.cs
@@ -0,0 +1,73 @@
+// Copyright Epic Games, Inc. All Rights Reserved.
+
+using System;
+using System.Buffers;
+using System.IO.MemoryMappedFiles;
+
+namespace EpicGames.Core
+{
+	/// <summary>
+	/// Implements an unmarshlled view of a memory mapped file
+	/// </summary>
+	public sealed unsafe class MemoryMappedView : IDisposable
+	{
+		sealed unsafe class MemoryWrapper : MemoryManager<byte>
+		{
+			private readonly byte* _pointer;
+			private readonly int _length;
+
+			public MemoryWrapper(byte* pointer, int length)
+			{
+				_pointer = pointer;
+				_length = length;
+			}
+
+			/// <inheritdoc/>
+			public override Span<byte> GetSpan() => new Span<byte>(_pointer, _length);
+
+			/// <inheritdoc/>
+			public override MemoryHandle Pin(int elementIndex) => new MemoryHandle(_pointer + elementIndex);
+
+			/// <inheritdoc/>
+			public override void Unpin() { }
+
+			/// <inheritdoc/>
+			protected override void Dispose(bool disposing) { }
+		}
+
+		readonly MemoryMappedViewAccessor _memoryMappedViewAccessor;
+		byte* _data;
+
+		/// <summary>
+		/// Constructor
+		/// </summary>
+		/// <param name="memoryMappedViewAccessor"></param>
+		public MemoryMappedView(MemoryMappedViewAccessor memoryMappedViewAccessor)
+		{
+			_memoryMappedViewAccessor = memoryMappedViewAccessor;
+			memoryMappedViewAccessor.SafeMemoryMappedViewHandle.AcquirePointer(ref _data);
+		}
+
+		/// <summary>
+		/// Gets a memory object for the given range
+		/// </summary>
+		/// <param name="offset"></param>
+		/// <param name="length"></param>
+		/// <returns></returns>
+		public Memory<byte> GetMemory(long offset, int length)
+		{
+			MemoryWrapper wrapper = new MemoryWrapper(_data + offset, length);
+			return wrapper.Memory;
+		}
+
+		/// <inheritdoc/>
+		public void Dispose()
+		{
+			if (_data != null)
+			{
+				_memoryMappedViewAccessor.SafeMemoryMappedViewHandle.ReleasePointer();
+				_data = null;
+			}
+		}
+	}
+}
diff --git a/Engine/Source/Programs/Shared/EpicGames.Horde.Tests/BufferTests.cs b/Engine/Source/Programs/Shared/EpicGames.Horde.Tests/BufferTests.cs
new file mode 100644
--- /dev/null
+++ b/Engine/Source/Programs/Shared/EpicGames.Horde.Tests/BufferTests.cs
@@ -0,0 +1,80 @@
+// Copyright Epic Games, Inc. All Rights Reserved.
+
+using System;
+using System.Diagnostics;
+using System.Linq;
+using System.Runtime.InteropServices;
+using System.Security.Cryptography;
+using System.Threading;
+using System.Threading.Tasks;
+using EpicGames.Horde.Compute;
+using EpicGames.Horde.Compute.Buffers;
+using Microsoft.VisualStudio.TestTools.UnitTesting;
+
+namespace EpicGames.Horde.Tests
+{
+	[TestClass]
+	public class BufferTests
+	{
+		[TestMethod]
+		public async Task TestLocalBuffer()
+		{
+			using RentedLocalBuffer buffer = new RentedLocalBuffer(8000);
+			await TestProducerConsumerAsync(buffer.Writer, buffer.Reader);
+		}
+
+		[TestMethod]
+		public async Task TestIpcBuffer()
+		{
+			if (RuntimeInformation.IsOSPlatform(OSPlatform.Windows))
+			{
+				using IpcBuffer buffer1 = IpcBuffer.CreateNew(16384);
+				string info = buffer1.DuplicateIntoProcess(Process.GetCurrentProcess().Handle);
+				using IpcBuffer buffer2 = IpcBuffer.OpenExisting(info);
+				await TestProducerConsumerAsync(buffer1.Writer, buffer2.Reader);
+			}
+		}
+
+		static async Task TestProducerConsumerAsync(IComputeBufferWriter writer, IComputeBufferReader reader)
+		{
+			const int Length = 8000;
+
+			byte[] input = RandomNumberGenerator.GetBytes(Length);
+			Task producerTask = RunProducerAsync(writer, input);
+
+			byte[] output = new byte[Length];
+			await RunConsumerAsync(reader, output);
+
+			await producerTask;
+			Assert.IsTrue(input.SequenceEqual(output));
+		}
+
+		static async Task RunProducerAsync(IComputeBufferWriter writer, ReadOnlyMemory<byte> input)
+		{
+			int offset = 0;
+			while (offset < input.Length)
+			{
+				int length = Math.Min(input.Length - offset, 100);
+				input.Slice(offset, length).CopyTo(writer.GetMemory());
+				writer.Advance(length);
+				await Task.Delay(10);
+				offset += length;
+			}
+			writer.MarkComplete();
+		}
+
+		static async Task RunConsumerAsync(IComputeBufferReader reader, Memory<byte> output)
+		{
+			int offset = 0;
+			while (!reader.IsComplete)
+			{
+				ReadOnlyMemory<byte> memory = reader.GetMemory();
+				int length = Math.Min(memory.Length, 7);
+				memory.Slice(0, length).CopyTo(output.Slice(offset));
+				reader.Advance(length);
+				await reader.WaitAsync(memory.Length - length, CancellationToken.None);
+				offset += length;
+			}
+		}
+	}
+}
diff --git a/Engine/Source/Programs/Shared/EpicGames.Horde/Compute/Buffers/HeapBuffer.cs b/Engine/Source/Programs/Shared/EpicGames.Horde/Compute/Buffers/HeapBuffer.cs
new file mode 100644
--- /dev/null
+++ b/Engine/Source/Programs/Shared/EpicGames.Horde/Compute/Buffers/HeapBuffer.cs
@@ -0,0 +1,135 @@
+// Copyright Epic Games, Inc. All Rights Reserved.
+
+using System;
+using System.Threading;
+using System.Threading.Tasks;
+
+namespace EpicGames.Horde.Compute.Buffers
+{
+	/// <summary>
+	/// In-process buffer used for compute messages
+	/// </summary>
+	public class MemoryBuffer : ComputeBuffer
+	{
+		readonly Memory<byte> _memory;
+
+		int _readPosition;
+		int _writePosition;
+		bool _finishedWriting;
+		TaskCompletionSource? _writtenTcs;
+		TaskCompletionSource? _flushedTcs;
+
+		/// <summary>
+		/// Creates a local buffer with the given capacity
+		/// </summary>
+		/// <param name="memory"></param>
+		public MemoryBuffer(Memory<byte> memory)
+		{
+			_memory = memory;
+		}
+
+		#region Reader
+
+		/// <inheritdoc/>
+		protected override bool FinishedReading() => _finishedWriting && _readPosition == _writePosition;
+
+		/// <inheritdoc/>
+		protected override void AdvanceReadPosition(int size)
+		{
+			_readPosition += size;
+
+			if (_readPosition == _writePosition)
+			{
+				_flushedTcs?.TrySetResult();
+			}
+		}
+
+		/// <inheritdoc/>
+		protected override ReadOnlyMemory<byte> GetReadMemory() => _memory.Slice(_readPosition, _writePosition - _readPosition);
+
+		/// <inheritdoc/>
+		protected override async ValueTask WaitAsync(int currentLength, CancellationToken cancellationToken)
+		{
+			int initialWritePosition = _readPosition + currentLength;
+
+			// Early out if we're already past this
+			if (_finishedWriting || _writePosition > initialWritePosition)
+			{
+				return;
+			}
+
+			TaskCompletionSource tcs = new TaskCompletionSource();
+			using (IDisposable registration = cancellationToken.Register(x => tcs.SetCanceled((CancellationToken)x!), cancellationToken))
+			{
+				try
+				{
+					_writtenTcs = tcs;
+					for (; ; )
+					{
+						if (_finishedWriting || _writePosition > initialWritePosition)
+						{
+							break;
+						}
+						else
+						{
+							await _writtenTcs.Task;
+						}
+					}
+				}
+				finally
+				{
+					_writtenTcs = null;
+				}
+			}
+		}
+
+		#endregion
+
+		#region Writer
+
+		/// <inheritdoc/>
+		protected override void FinishWriting()
+		{
+			_finishedWriting = true;
+			_writtenTcs?.TrySetResult();
+		}
+
+		/// <inheritdoc/>
+		protected override void AdvanceWritePosition(int size)
+		{
+			if (_finishedWriting)
+			{
+				throw new InvalidOperationException("Cannot update write position after marking as complete");
+			}
+
+			_writePosition += size;
+			_writtenTcs?.TrySetResult();
+		}
+
+		/// <inheritdoc/>
+		protected override Memory<byte> GetWriteMemory() => _memory.Slice(_writePosition);
+
+		/// <inheritdoc/>
+		protected override async ValueTask FlushWritesAsync(CancellationToken cancellationToken)
+		{
+			if (_readPosition < _writePosition)
+			{
+				TaskCompletionSource tcs = new TaskCompletionSource();
+				using (_ = cancellationToken.Register(x => tcs.SetCanceled((CancellationToken)x!), cancellationToken, false))
+				{
+					TaskCompletionSource? originalTcs = Interlocked.CompareExchange(ref _flushedTcs, tcs, null);
+					if (originalTcs != null)
+					{
+						_ = originalTcs.Task.ContinueWith(x => tcs.TrySetResult(), cancellationToken, TaskContinuationOptions.None, TaskScheduler.Default);
+					}
+				}
+				if (_readPosition < _writePosition)
+				{
+					await tcs.Task;
+				}
+			}
+		}
+
+		#endregion
+	}
+}
diff --git a/Engine/Source/Programs/Shared/EpicGames.Horde/Compute/Buffers/IpcBuffer.cs b/Engine/Source/Programs/Shared/EpicGames.Horde/Compute/Buffers/IpcBuffer.cs
new file mode 100644
--- /dev/null
+++ b/Engine/Source/Programs/Shared/EpicGames.Horde/Compute/Buffers/IpcBuffer.cs
@@ -0,0 +1,259 @@
+// Copyright Epic Games, Inc. All Rights Reserved.
+
+using System;
+using System.Buffers.Binary;
+using System.Globalization;
+using System.Runtime.InteropServices;
+using System.Threading;
+using System.Threading.Tasks;
+using EpicGames.Core;
+using Microsoft.Win32.SafeHandles;
+using System.Runtime.Versioning;
+using System.IO.MemoryMappedFiles;
+using System.Reflection;
+
+namespace EpicGames.Horde.Compute.Buffers
+{
+	/// <summary>
+	/// Implementation of <see cref="IComputeBuffer"/> suitable for cross-process communication
+	/// </summary>
+	[SupportedOSPlatform("windows")]
+	public sealed class IpcBuffer : ComputeBuffer, IDisposable
+	{
+		class NativeEvent : WaitHandle
+		{
+			public NativeEvent()
+			{
+				SafeWaitHandle = CreateEvent(IntPtr.Zero, false, false, null);
+			}
+
+			public NativeEvent(IntPtr handle, bool ownsHandle)
+			{
+				SafeWaitHandle = new SafeWaitHandle(handle, ownsHandle);
+			}
+
+			public void Set() => SetEvent(SafeWaitHandle);
+
+			[DllImport("kernel32.dll")]
+			static extern SafeWaitHandle CreateEvent(IntPtr lpEventAttributes, bool bManualReset, bool bInitialState, string? lpName);
+
+			[DllImport("kernel32.dll")]
+			static extern bool SetEvent(SafeWaitHandle hEvent);
+		}
+
+		const int HeaderLength = 32;
+
+		long ReadPosition
+		{
+			get => BinaryPrimitives.ReadInt64LittleEndian(_header.Span.Slice(8));
+			set => BinaryPrimitives.WriteInt64LittleEndian(_header.Span.Slice(8), value);
+		}
+
+		long WritePosition
+		{
+			get => BinaryPrimitives.ReadInt64LittleEndian(_header.Span.Slice(16));
+			set => BinaryPrimitives.WriteInt64LittleEndian(_header.Span.Slice(16), value);
+		}
+
+		bool FinishedWriting
+		{
+			get => _header.Span[0] != 0;
+			set => _header.Span[0] = value? (byte)1 : (byte)0;
+		}
+
+		readonly MemoryMappedFile _memoryMappedFile;
+		readonly MemoryMappedViewAccessor _memoryMappedViewAccessor;
+		readonly MemoryMappedView _memoryMappedView;
+
+		readonly Memory<byte> _header;
+		readonly Memory<byte> _memory;
+
+		readonly NativeEvent _writtenEvent;
+		readonly NativeEvent _flushedEvent;
+
+		private IpcBuffer(MemoryMappedFile memoryMappedFile, NativeEvent writtenEvent, NativeEvent flushedEvent)
+		{
+			_memoryMappedFile = memoryMappedFile;
+			_memoryMappedViewAccessor = _memoryMappedFile.CreateViewAccessor();
+			_memoryMappedView = new MemoryMappedView(_memoryMappedViewAccessor);
+
+			Memory<byte> memory = _memoryMappedView.GetMemory(0, (int)_memoryMappedViewAccessor.Capacity);
+			_header = memory.Slice(0, HeaderLength);
+			_memory = memory.Slice(HeaderLength);
+
+			_writtenEvent = writtenEvent;
+			_flushedEvent = flushedEvent;
+		}
+
+		/// <summary>
+		/// Creates a new IPC buffer with the given capacity
+		/// </summary>
+		/// <param name="capacity">Capacity of the buffer</param>
+		public static IpcBuffer CreateNew(long capacity)
+		{
+			MemoryMappedFile memoryMappedFile = MemoryMappedFile.CreateNew(null, HeaderLength + capacity);
+
+			NativeEvent writtenEvent = new NativeEvent();
+			NativeEvent flushedEvent = new NativeEvent();
+
+			return new IpcBuffer(memoryMappedFile, writtenEvent, flushedEvent);
+		}
+
+		/// <summary>
+		/// Opens an IpcBuffer from a string passed in from another process
+		/// </summary>
+		/// <param name="str">Descriptor for the buffer to open</param>
+		public static IpcBuffer OpenExisting(string str)
+		{
+			string[] values = str.Split(',');
+			if (values.Length != 3)
+			{
+				throw new ArgumentException("Expected three handle values for IPC buffer", nameof(str));
+			}
+
+			IntPtr memoryHandle = new IntPtr((long)UInt64.Parse(values[0], NumberStyles.None));
+			MemoryMappedFile memoryMappedFile = OpenMemoryMappedFile(memoryHandle);
+
+			IntPtr writtenEventHandle = new IntPtr((long)UInt64.Parse(values[1], NumberStyles.None));
+			NativeEvent writtenEvent = new NativeEvent(writtenEventHandle, true);
+
+			IntPtr flushedEventHandle = new IntPtr((long)UInt64.Parse(values[2], NumberStyles.None));
+			NativeEvent flushedEvent = new NativeEvent(flushedEventHandle, true);
+
+			return new IpcBuffer(memoryMappedFile, writtenEvent, flushedEvent);
+		}
+
+		static MemoryMappedFile OpenMemoryMappedFile(IntPtr handle)
+		{
+			MethodInfo? setHandleInfo = typeof(SafeMemoryMappedFileHandle).GetMethod("SetHandle", BindingFlags.Instance | BindingFlags.NonPublic, new[] { typeof(IntPtr) });
+			if (setHandleInfo == null)
+			{
+				throw new InvalidOperationException("Cannot find SetHandle method for SafeMemoryMappedFileHandle");
+			}
+
+			ConstructorInfo? constructorInfo = typeof(MemoryMappedFile).GetConstructor(BindingFlags.Instance | BindingFlags.NonPublic, null, new[] { typeof(SafeMemoryMappedFileHandle) }, null);
+			if (constructorInfo == null)
+			{
+				throw new InvalidOperationException("Cannot find private constructor for memory mapped file");
+			}
+
+			SafeMemoryMappedFileHandle safeHandle = new SafeMemoryMappedFileHandle();
+			setHandleInfo.Invoke(safeHandle, new object[] { handle });
+
+			return (MemoryMappedFile)constructorInfo.Invoke(new object[] { safeHandle });
+		}
+
+		/// <summary>
+		/// Duplicates handles for this buffer into the given process, and returns a string that the target process can use to open it.
+		/// </summary>
+		/// <param name="targetProcessHandle">Handle to the target process</param>
+		/// <returns>String that the target process can pass into a call to <see cref="OpenExisting(String)"/> to open the buffer.</returns>
+		public string DuplicateIntoProcess(IntPtr targetProcessHandle)
+		{
+			IntPtr sourceProcessHandle = GetCurrentProcess();
+
+			IntPtr targetMemoryHandle;
+			DuplicateHandle(sourceProcessHandle, _memoryMappedFile.SafeMemoryMappedFileHandle.DangerousGetHandle(), targetProcessHandle, out targetMemoryHandle, 0, false, DUPLICATE_SAME_ACCESS);
+
+			IntPtr targetWrittenEventHandle;
+			DuplicateHandle(sourceProcessHandle, _writtenEvent.SafeWaitHandle.DangerousGetHandle(), targetProcessHandle, out targetWrittenEventHandle, 0, false, DUPLICATE_SAME_ACCESS);
+
+			IntPtr targetFlushedEventHandle;
+			DuplicateHandle(sourceProcessHandle, _flushedEvent.SafeWaitHandle.DangerousGetHandle(), targetProcessHandle, out targetFlushedEventHandle, 0, false, DUPLICATE_SAME_ACCESS);
+
+			return $"{(ulong)targetMemoryHandle.ToInt64()},{(ulong)targetWrittenEventHandle.ToInt64()},{(ulong)targetFlushedEventHandle.ToInt64()}";
+		}
+
+		const uint DUPLICATE_SAME_ACCESS = 2;
+
+		[DllImport("kernel32.dll", SetLastError = true)]
+		static extern IntPtr GetCurrentProcess();
+
+		[DllImport("kernel32.dll", SetLastError = true)]
+		[return: MarshalAs(UnmanagedType.Bool)]
+		static extern bool DuplicateHandle(IntPtr hSourceProcessHandle, IntPtr hSourceHandle, IntPtr hTargetProcessHandle, out IntPtr lpTargetHandle, uint dwDesiredAccess, [MarshalAs(UnmanagedType.Bool)] bool bInheritHandle, uint dwOptions);
+
+		/// <inheritdoc/>
+		public void Dispose()
+		{
+			_flushedEvent.Dispose();
+			_writtenEvent.Dispose();
+			_memoryMappedView.Dispose();
+			_memoryMappedViewAccessor.Dispose();
+			_memoryMappedFile.Dispose();
+		}
+
+		#region Reader
+
+		/// <inheritdoc/>
+		protected override bool FinishedReading() => FinishedWriting && ReadPosition == WritePosition;
+
+		/// <inheritdoc/>
+		protected override void AdvanceReadPosition(int size)
+		{
+			ReadPosition += size;
+
+			if (ReadPosition == WritePosition)
+			{
+				_flushedEvent.Set();
+			}
+		}
+
+		/// <inheritdoc/>
+		protected override ReadOnlyMemory<byte> GetReadMemory() => _memory.Slice((int)ReadPosition, (int)(WritePosition - ReadPosition));
+
+		/// <inheritdoc/>
+		protected override async ValueTask WaitAsync(int currentLength, CancellationToken cancellationToken)
+		{
+			long initialWritePosition = ReadPosition + currentLength;
+			for (; ; )
+			{
+				if (FinishedWriting || WritePosition > initialWritePosition)
+				{
+					return;
+				}
+				else
+				{
+					await _writtenEvent.WaitOneAsync(cancellationToken);
+				}
+			}
+		}
+
+		#endregion
+
+		#region Writer
+
+		/// <inheritdoc/>
+		protected override void FinishWriting()
+		{
+			FinishedWriting = true;
+			_writtenEvent.Set();
+		}
+
+		/// <inheritdoc/>
+		protected override void AdvanceWritePosition(int size)
+		{
+			if (FinishedWriting)
+			{
+				throw new InvalidOperationException("Cannot update write position after marking as complete");
+			}
+
+			WritePosition += size;
+			_writtenEvent.Set();
+		}
+
+		/// <inheritdoc/>
+		protected override Memory<byte> GetWriteMemory() => _memory.Slice((int)WritePosition);
+
+		/// <inheritdoc/>
+		protected override async ValueTask FlushWritesAsync(CancellationToken cancellationToken)
+		{
+			while (ReadPosition < WritePosition)
+			{
+				await _flushedEvent.WaitOneAsync(cancellationToken);
+			}
+		}
+
+		#endregion	
+	}
+}
diff --git a/Engine/Source/Programs/Shared/EpicGames.Horde/Compute/Buffers/LocalBuffer.cs b/Engine/Source/Programs/Shared/EpicGames.Horde/Compute/Buffers/LocalBuffer.cs
deleted file mode 100644
diff --git a/Engine/Source/Programs/Shared/EpicGames.Horde/Compute/Buffers/RentedBuffer.cs b/Engine/Source/Programs/Shared/EpicGames.Horde/Compute/Buffers/RentedBuffer.cs
new file mode 100644
--- /dev/null
+++ b/Engine/Source/Programs/Shared/EpicGames.Horde/Compute/Buffers/RentedBuffer.cs
@@ -0,0 +1,40 @@
+// Copyright Epic Games, Inc. All Rights Reserved.
+
+using System;
+using System.Buffers;
+
+namespace EpicGames.Horde.Compute.Buffers
+{
+	/// <summary>
+	/// Implementation of <see cref="MemoryBuffer"/> that owns a rented memory resource
+	/// </summary>
+	class RentedLocalBuffer : MemoryBuffer, IDisposable
+	{
+		readonly IMemoryOwner<byte> _memoryOwner;
+
+		/// <summary>
+		/// Constructor
+		/// </summary>
+		/// <param name="capacity">Size of data to allocate</param>
+		public RentedLocalBuffer(int capacity)
+			: this(MemoryPool<byte>.Shared.Rent(capacity))
+		{
+		}
+
+		/// <summary>
+		/// Constructor
+		/// </summary>
+		/// <param name="memoryOwner">Memory to manage</param>
+		public RentedLocalBuffer(IMemoryOwner<byte> memoryOwner)
+			: base(memoryOwner.Memory)
+		{
+			_memoryOwner = memoryOwner;
+		}
+
+		/// <inheritdoc/>
+		public void Dispose()
+		{
+			_memoryOwner.Dispose();
+		}
+	}
+}
diff --git a/Engine/Source/Programs/Shared/EpicGames.Horde/Compute/ComputeBuffer.cs b/Engine/Source/Programs/Shared/EpicGames.Horde/Compute/ComputeBuffer.cs
new file mode 100644
--- /dev/null
+++ b/Engine/Source/Programs/Shared/EpicGames.Horde/Compute/ComputeBuffer.cs
@@ -0,0 +1,98 @@
+// Copyright Epic Games, Inc. All Rights Reserved.
+
+using System;
+using System.Threading;
+using System.Threading.Tasks;
+
+namespace EpicGames.Horde.Compute
+{
+	/// <summary>
+	/// Base implementation for <see cref="IComputeBuffer"/>, which offers reader and writer objects that redirect to the owner
+	/// </summary>
+	public abstract class ComputeBuffer : IComputeBuffer
+	{
+		class ReaderImpl : IComputeBufferReader
+		{
+			readonly ComputeBuffer _owner;
+
+			public ReaderImpl(ComputeBuffer owner) => _owner = owner;
+
+			/// <inheritdoc/>
+			public bool IsComplete => _owner.FinishedReading();
+
+			/// <inheritdoc/>
+			public void Advance(int size) => _owner.AdvanceReadPosition(size);
+
+			/// <inheritdoc/>
+			public ReadOnlyMemory<byte> GetMemory() => _owner.GetReadMemory();
+
+			/// <inheritdoc/>
+			public ValueTask WaitAsync(int currentLength, CancellationToken cancellationToken) => _owner.WaitAsync(currentLength, cancellationToken);
+		}
+
+		class WriterImpl : IComputeBufferWriter
+		{
+			readonly ComputeBuffer _owner;
+
+			public WriterImpl(ComputeBuffer owner) => _owner = owner;
+
+			/// <inheritdoc/>
+			public void MarkComplete() => _owner.FinishWriting();
+
+			/// <inheritdoc/>
+			public void Advance(int size) => _owner.AdvanceWritePosition(size);
+
+			/// <inheritdoc/>
+			public Memory<byte> GetMemory() => _owner.GetWriteMemory();
+
+			/// <inheritdoc/>
+			public ValueTask FlushAsync(CancellationToken cancellationToken) => _owner.FlushWritesAsync(cancellationToken);
+		}
+
+		/// <inheritdoc/>
+		public IComputeBufferReader Reader { get; }
+
+		/// <inheritdoc/>
+		public IComputeBufferWriter Writer { get; }
+
+		/// <summary>
+		/// Constructor
+		/// </summary>
+		protected ComputeBuffer()
+		{
+			Reader = new ReaderImpl(this);
+			Writer = new WriterImpl(this);
+		}
+
+		#region Reader
+
+		/// <inheritdoc cref="IComputeBufferReader.IsComplete"/>
+		protected abstract bool FinishedReading();
+
+		/// <inheritdoc cref="IComputeBufferReader.Advance(Int32)"/>
+		protected abstract void AdvanceReadPosition(int size);
+
+		/// <inheritdoc cref="IComputeBufferReader.GetMemory"/>
+		protected abstract ReadOnlyMemory<byte> GetReadMemory();
+
+		/// <inheritdoc cref="IComputeBufferReader.WaitAsync(Int32, CancellationToken)"/>
+		protected abstract ValueTask WaitAsync(int currentLength, CancellationToken cancellationToken);
+
+		#endregion
+		#region Writer
+
+		/// <inheritdoc cref="IComputeBufferWriter.MarkComplete"/>
+		protected abstract void FinishWriting();
+
+		/// <inheritdoc cref="IComputeBufferWriter.Advance(Int32)"/>
+		protected abstract void AdvanceWritePosition(int size);
+
+		/// <inheritdoc cref="IComputeBufferWriter.GetMemory()"/>
+		protected abstract Memory<byte> GetWriteMemory();
+
+		/// <inheritdoc cref="IComputeBufferWriter.FlushAsync(CancellationToken)"/>
+		protected abstract ValueTask FlushWritesAsync(CancellationToken cancellationToken);
+
+		#endregion
+	}
+}
