Horde: Inherit handles to IPC objects by default.

diff --git a/Engine/Source/Programs/Shared/EpicGames.Horde/Compute/Buffers/HeapBuffer.cs b/Engine/Source/Programs/Shared/EpicGames.Horde/Compute/Buffers/HeapBuffer.cs
--- a/Engine/Source/Programs/Shared/EpicGames.Horde/Compute/Buffers/HeapBuffer.cs
+++ b/Engine/Source/Programs/Shared/EpicGames.Horde/Compute/Buffers/HeapBuffer.cs
@@ -9,7 +9,7 @@ namespace EpicGames.Horde.Compute.Buffers
 	/// <summary>
 	/// In-process buffer used for compute messages
 	/// </summary>
-	public class MemoryBuffer : ComputeBuffer
+	public class HeapBuffer : ComputeBuffer
 	{
 		readonly Memory<byte> _memory;
 
@@ -23,7 +23,7 @@ namespace EpicGames.Horde.Compute.Buffers
 		/// Creates a local buffer with the given capacity
 		/// </summary>
 		/// <param name="memory"></param>
-		public MemoryBuffer(Memory<byte> memory)
+		public HeapBuffer(Memory<byte> memory)
 		{
 			_memory = memory;
 		}
diff --git a/Engine/Source/Programs/Shared/EpicGames.Horde/Compute/Buffers/IpcBuffer.cs b/Engine/Source/Programs/Shared/EpicGames.Horde/Compute/Buffers/IpcBuffer.cs
--- a/Engine/Source/Programs/Shared/EpicGames.Horde/Compute/Buffers/IpcBuffer.cs
+++ b/Engine/Source/Programs/Shared/EpicGames.Horde/Compute/Buffers/IpcBuffer.cs
@@ -11,6 +11,7 @@ using Microsoft.Win32.SafeHandles;
 using System.Runtime.Versioning;
 using System.IO.MemoryMappedFiles;
 using System.Reflection;
+using System.IO;
 
 namespace EpicGames.Horde.Compute.Buffers
 {
@@ -20,11 +21,22 @@ namespace EpicGames.Horde.Compute.Buffers
 	[SupportedOSPlatform("windows")]
 	public sealed class IpcBuffer : ComputeBuffer, IDisposable
 	{
+		[StructLayout(LayoutKind.Sequential)]
+		class SECURITY_ATTRIBUTES
+		{
+			public int nLength;
+			public IntPtr lpSecurityDescriptor;
+			public int bInheritHandle;
+		}
+
 		class NativeEvent : WaitHandle
 		{
-			public NativeEvent()
+			public NativeEvent(HandleInheritability handleInheritability)
 			{
-				SafeWaitHandle = CreateEvent(IntPtr.Zero, false, false, null);
+				SECURITY_ATTRIBUTES securityAttributes = new SECURITY_ATTRIBUTES();
+				securityAttributes.nLength = Marshal.SizeOf(securityAttributes);
+				securityAttributes.bInheritHandle = (handleInheritability == HandleInheritability.Inheritable)? 1 : 0;
+				SafeWaitHandle = CreateEvent(securityAttributes, false, false, null);
 			}
 
 			public NativeEvent(IntPtr handle, bool ownsHandle)
@@ -35,7 +47,7 @@ namespace EpicGames.Horde.Compute.Buffers
 			public void Set() => SetEvent(SafeWaitHandle);
 
 			[DllImport("kernel32.dll")]
-			static extern SafeWaitHandle CreateEvent(IntPtr lpEventAttributes, bool bManualReset, bool bInitialState, string? lpName);
+			static extern SafeWaitHandle CreateEvent(SECURITY_ATTRIBUTES lpEventAttributes, bool bManualReset, bool bInitialState, string? lpName);
 
 			[DllImport("kernel32.dll")]
 			static extern bool SetEvent(SafeWaitHandle hEvent);
@@ -91,10 +103,10 @@ namespace EpicGames.Horde.Compute.Buffers
 		/// <param name="capacity">Capacity of the buffer</param>
 		public static IpcBuffer CreateNew(long capacity)
 		{
-			MemoryMappedFile memoryMappedFile = MemoryMappedFile.CreateNew(null, HeaderLength + capacity);
+			MemoryMappedFile memoryMappedFile = MemoryMappedFile.CreateNew(null, HeaderLength + capacity, MemoryMappedFileAccess.ReadWrite, MemoryMappedFileOptions.None, HandleInheritability.Inheritable);
 
-			NativeEvent writtenEvent = new NativeEvent();
-			NativeEvent flushedEvent = new NativeEvent();
+			NativeEvent writtenEvent = new NativeEvent(HandleInheritability.Inheritable);
+			NativeEvent flushedEvent = new NativeEvent(HandleInheritability.Inheritable);
 
 			return new IpcBuffer(memoryMappedFile, writtenEvent, flushedEvent);
 		}
@@ -102,13 +114,13 @@ namespace EpicGames.Horde.Compute.Buffers
 		/// <summary>
 		/// Opens an IpcBuffer from a string passed in from another process
 		/// </summary>
-		/// <param name="str">Descriptor for the buffer to open</param>
-		public static IpcBuffer OpenExisting(string str)
+		/// <param name="handle">Descriptor for the buffer to open</param>
+		public static IpcBuffer OpenExisting(string handle)
 		{
-			string[] values = str.Split(',');
+			string[] values = handle.Split(',');
 			if (values.Length != 3)
 			{
-				throw new ArgumentException("Expected three handle values for IPC buffer", nameof(str));
+				throw new ArgumentException("Expected three handle values for IPC buffer", nameof(handle));
 			}
 
 			IntPtr memoryHandle = new IntPtr((long)UInt64.Parse(values[0], NumberStyles.None));
@@ -143,6 +155,18 @@ namespace EpicGames.Horde.Compute.Buffers
 			return (MemoryMappedFile)constructorInfo.Invoke(new object[] { safeHandle });
 		}
 
+		/// <summary>
+		/// Gets the descriptor for the current buffer. This can be used for calls to <see cref="OpenExisting(String)"/> by newly spawned processes that inherit handles.
+		/// </summary>
+		/// <returns>Descriptor string for the buffer</returns>
+		public string GetHandle()
+		{
+			ulong targetMemoryHandle = (ulong)_memoryMappedFile.SafeMemoryMappedFileHandle.DangerousGetHandle().ToInt64();
+			ulong targetWrittenEventHandle = (ulong)_writtenEvent.SafeWaitHandle.DangerousGetHandle();
+			ulong targetFlushedEventHandle = (ulong)_flushedEvent.SafeWaitHandle.DangerousGetHandle();
+			return $"{targetMemoryHandle},{targetWrittenEventHandle},{targetFlushedEventHandle}";
+		}
+
 		/// <summary>
 		/// Duplicates handles for this buffer into the given process, and returns a string that the target process can use to open it.
 		/// </summary>
@@ -254,6 +278,6 @@ namespace EpicGames.Horde.Compute.Buffers
 			}
 		}
 
-		#endregion	
+		#endregion
 	}
 }
diff --git a/Engine/Source/Programs/Shared/EpicGames.Horde/Compute/Buffers/RentedBuffer.cs b/Engine/Source/Programs/Shared/EpicGames.Horde/Compute/Buffers/RentedBuffer.cs
--- a/Engine/Source/Programs/Shared/EpicGames.Horde/Compute/Buffers/RentedBuffer.cs
+++ b/Engine/Source/Programs/Shared/EpicGames.Horde/Compute/Buffers/RentedBuffer.cs
@@ -6,9 +6,9 @@ using System.Buffers;
 namespace EpicGames.Horde.Compute.Buffers
 {
 	/// <summary>
-	/// Implementation of <see cref="MemoryBuffer"/> that owns a rented memory resource
+	/// Implementation of <see cref="HeapBuffer"/> that owns a rented memory resource
 	/// </summary>
-	class RentedLocalBuffer : MemoryBuffer, IDisposable
+	class RentedLocalBuffer : HeapBuffer, IDisposable
 	{
 		readonly IMemoryOwner<byte> _memoryOwner;
 
diff --git a/Engine/Source/Programs/Shared/EpicGames.Horde/Compute/ComputeChannel.cs b/Engine/Source/Programs/Shared/EpicGames.Horde/Compute/ComputeChannel.cs
--- a/Engine/Source/Programs/Shared/EpicGames.Horde/Compute/ComputeChannel.cs
+++ b/Engine/Source/Programs/Shared/EpicGames.Horde/Compute/ComputeChannel.cs
@@ -11,7 +11,7 @@ namespace EpicGames.Horde.Compute
 	/// <summary>
 	/// Implementation of a compute channel
 	/// </summary>
-	class ComputeChannel : IComputeChannel, IAsyncDisposable
+	public class ComputeChannel : IComputeChannel, IAsyncDisposable
 	{
 		// Length of a message header. Consists of a 1 byte type field, followed by 4 byte length field.
 		const int HeaderLength = 5;
