Horde: Fix data being ignored if multiple requests are sent in succession.

diff --git a/Engine/Source/Programs/Shared/EpicGames.Horde/Compute/ComputeLease.cs b/Engine/Source/Programs/Shared/EpicGames.Horde/Compute/ComputeLease.cs
--- a/Engine/Source/Programs/Shared/EpicGames.Horde/Compute/ComputeLease.cs
+++ b/Engine/Source/Programs/Shared/EpicGames.Horde/Compute/ComputeLease.cs
@@ -112,7 +112,7 @@ namespace EpicGames.Horde.Compute
 					}
 					else
 					{
-						await transport.ReadAsync(writer.GetMemory(), cancellationToken);
+						await transport.ReadAsync(writer.GetMemory().Slice(0, size), cancellationToken);
 						writer.Advance(size);
 					}
 				}
