LogParser: Continue if a matcher can not be loaded

diff --git a/Engine/Source/Programs/Shared/EpicGames.Core/LogEventParser.cs b/Engine/Source/Programs/Shared/EpicGames.Core/LogEventParser.cs
--- a/Engine/Source/Programs/Shared/EpicGames.Core/LogEventParser.cs
+++ b/Engine/Source/Programs/Shared/EpicGames.Core/LogEventParser.cs
@@ -216,8 +216,15 @@ namespace EpicGames.Core
 				if (type.IsClass && typeof(ILogEventMatcher).IsAssignableFrom(type))
 				{
 					_logger.LogDebug("Adding event matcher: {Type}", type.Name);
-					ILogEventMatcher matcher = (ILogEventMatcher)Activator.CreateInstance(type)!;
-					Matchers.Add(matcher);
+					try
+					{
+						ILogEventMatcher matcher = (ILogEventMatcher)Activator.CreateInstance(type)!;
+						Matchers.Add(matcher);
+					}
+					catch (Exception ex)
+					{
+						_logger.LogDebug("Failed to add event matcher for {Type}: {Exception}", type.Name, ex.ToString());
+					}
 				}
 			}
 		}
