Ported C# UHT UhtBuffer to new objects that use ArrayPool.

diff --git a/Engine/Source/Programs/Shared/EpicGames.UHT/Exporters/CodeGen/UhtHeaderCodeGeneratorCppFile.cs b/Engine/Source/Programs/Shared/EpicGames.UHT/Exporters/CodeGen/UhtHeaderCodeGeneratorCppFile.cs
--- a/Engine/Source/Programs/Shared/EpicGames.UHT/Exporters/CodeGen/UhtHeaderCodeGeneratorCppFile.cs
+++ b/Engine/Source/Programs/Shared/EpicGames.UHT/Exporters/CodeGen/UhtHeaderCodeGeneratorCppFile.cs
@@ -339,7 +339,7 @@ namespace EpicGames.UHT.Exporters.CodeGen
 				builder.Append(EnableDeprecationWarnings).Append("\r\n");
 
 				{
-					using UhtBorrowBuffer borrowBuffer = new(builder);
+					using UhtRentedPoolBuffer<char> borrowBuffer = builder.RentPoolBuffer();
 					string cppFilePath = factory.MakePath(HeaderFile, ".gen.cpp");
 					StringView generatedBody = new(borrowBuffer.Buffer.Memory);
 					if (SaveExportedHeaders)
@@ -500,7 +500,7 @@ namespace EpicGames.UHT.Exporters.CodeGen
 				builder.Append("\t}\r\n");
 
 				{
-					using UhtBorrowBuffer borrowBuffer = new(builder, hashCodeBlockStart, builder.Length - hashCodeBlockStart);
+					using UhtRentedPoolBuffer<char> borrowBuffer = builder.RentPoolBuffer(hashCodeBlockStart, builder.Length - hashCodeBlockStart);
 					ObjectInfos[enumObj.ObjectTypeIndex].Hash = UhtHash.GenenerateTextHash(borrowBuffer.Buffer.Memory.Span);
 				}
 			}
@@ -696,7 +696,7 @@ namespace EpicGames.UHT.Exporters.CodeGen
 				builder.Append("\t}\r\n");
 			}
 
-			using (UhtBorrowBuffer borrowBuffer = new(builder, hashCodeBlockStart, builder.Length - hashCodeBlockStart))
+			using (UhtRentedPoolBuffer<char> borrowBuffer = builder.RentPoolBuffer(hashCodeBlockStart, builder.Length - hashCodeBlockStart))
 			{
 				ObjectInfos[scriptStruct.ObjectTypeIndex].Hash = UhtHash.GenenerateTextHash(borrowBuffer.Buffer.Memory.Span);
 			}
@@ -923,7 +923,7 @@ namespace EpicGames.UHT.Exporters.CodeGen
 
 			builder.AppendEndEditorOnlyGuard(function.FunctionFlags.HasAnyFlags(EFunctionFlags.EditorOnly));
 
-			using (UhtBorrowBuffer borrowBuffer = new(builder, hashCodeBlockStart, builder.Length - hashCodeBlockStart))
+			using (UhtRentedPoolBuffer<char> borrowBuffer = builder.RentPoolBuffer(hashCodeBlockStart, builder.Length - hashCodeBlockStart))
 			{
 				ObjectInfos[function.ObjectTypeIndex].Hash = UhtHash.GenenerateTextHash(borrowBuffer.Buffer.Memory.Span);
 			}
@@ -1701,7 +1701,7 @@ namespace EpicGames.UHT.Exporters.CodeGen
 
 				if (Session.IncludeDebugOutput)
 				{
-					using UhtBorrowBuffer borrowBuffer = new(hashBuilder, saveLength, hashBuilder.Length - saveLength);
+					using UhtRentedPoolBuffer<char> borrowBuffer = hashBuilder.RentPoolBuffer(saveLength, hashBuilder.Length - saveLength);
 					builder.Append("#if 0\r\n");
 					builder.Append(borrowBuffer.Buffer.Memory);
 					builder.Append("#endif\r\n");
@@ -1709,7 +1709,7 @@ namespace EpicGames.UHT.Exporters.CodeGen
 
 				// Calculate generated class initialization code hash so that we know when it changes after hot-reload
 				{
-					using UhtBorrowBuffer borrowBuffer = new(hashBuilder);
+					using UhtRentedPoolBuffer<char> borrowBuffer = hashBuilder.RentPoolBuffer();
 					ObjectInfos[classObj.ObjectTypeIndex].Hash = UhtHash.GenenerateTextHash(borrowBuffer.Buffer.Memory.Span);
 				}
 			}
diff --git a/Engine/Source/Programs/Shared/EpicGames.UHT/Utils/UhtBuffer.cs b/Engine/Source/Programs/Shared/EpicGames.UHT/Utils/UhtBuffer.cs
--- a/Engine/Source/Programs/Shared/EpicGames.UHT/Utils/UhtBuffer.cs
+++ b/Engine/Source/Programs/Shared/EpicGames.UHT/Utils/UhtBuffer.cs
@@ -1,11 +1,154 @@
 // Copyright Epic Games, Inc. All Rights Reserved.
 
 using System;
+using System.Buffers;
 using System.Numerics;
 using System.Text;
 
 namespace EpicGames.UHT.Utils
 {
+	/// <summary>
+	/// Rent and return pooled buffers
+	/// </summary>
+	public static class UhtPoolBuffers
+	{
+
+		/// <summary>
+		/// Rent a new buffer
+		/// </summary>
+		/// <param name="size">Minimum size of the buffer</param>
+		/// <returns></returns>
+		public static UhtPoolBuffer<T> Rent<T>(int size)
+		{
+			T[] block = ArrayPool<T>.Shared.Rent(size);
+			return new UhtPoolBuffer<T>(block, size);
+		}
+
+		/// <summary>
+		/// Return a buffer
+		/// </summary>
+		/// <param name="buffer">Buffer being returned</param>
+		public static void Return<T>(UhtPoolBuffer<T> buffer)
+		{
+			if (buffer.IsSet)
+			{
+				ArrayPool<T>.Shared.Return(buffer.GetBlock());
+			}
+		}
+	}
+
+	/// <summary>
+	/// Pooled buffers
+	/// </summary>
+	/// <typeparam name="T"></typeparam>
+	public struct UhtPoolBuffer<T>
+	{
+
+		/// <summary>
+		/// THe backing array
+		/// </summary>
+		private readonly T[]? _block = null;
+
+		/// <summary>
+		/// Memory region sized to the requested size
+		/// </summary>
+		public Memory<T> Memory { get; set; }
+
+		/// <summary>
+		/// Return true if we have a pool block
+		/// </summary>
+		public bool IsSet => _block != null;
+
+		/// <summary>
+		/// Create a pool buffer
+		/// </summary>
+		/// <param name="block">Array block</param>
+		/// <param name="size">Size of the requested block</param>
+		public UhtPoolBuffer(T[] block, int size)
+		{
+			_block = block;
+			Memory = new Memory<T>(block, 0, size);
+		}
+
+		/// <summary>
+		/// The backing array.  The size of the array will normally be larger than the requested size.
+		/// </summary>
+		public T[] GetBlock()
+		{
+			return _block ?? Array.Empty<T>();
+		}
+
+		/// <summary>
+		/// Reset the memory region to the given size
+		/// </summary>
+		/// <param name="size"></param>
+		public void Reset(int size)
+		{
+			Memory = new Memory<T>(GetBlock(), 0, size);
+		}
+	}
+
+	/// <summary>
+	/// Helper class for using pattern to borrow and return a buffer.
+	/// </summary>
+	public struct UhtRentedPoolBuffer<T> : IDisposable
+	{
+
+		/// <summary>
+		/// The borrowed buffer
+		/// </summary>
+		public UhtPoolBuffer<T> Buffer { get; set; }
+
+		/// <summary>
+		/// Borrow a buffer with the given size
+		/// </summary>
+		/// <param name="size">The size to borrow</param>
+		public UhtRentedPoolBuffer(int size)
+		{
+			Buffer = UhtPoolBuffers.Rent<T>(size);
+		}
+
+		/// <summary>
+		/// Return the borrowed buffer to the cache
+		/// </summary>
+		public void Dispose()
+		{
+			UhtPoolBuffers.Return<T>(Buffer);
+		}
+	}
+
+	/// <summary>
+	/// Collection of helper methods to convert string builder to borrow buffers
+	/// </summary>
+	public static class UhtPoolBufferStringBuilderExtensions
+	{
+		/// <summary>
+		/// Return a buffer initialized with the string builder.
+		/// </summary>
+		/// <param name="builder">Source builder content</param>
+		/// <returns>Buffer that should be returned with a call to Return</returns>
+		public static UhtRentedPoolBuffer<char> RentPoolBuffer(this StringBuilder builder)
+		{
+			int length = builder.Length;
+			UhtRentedPoolBuffer<char> buffer = new(length);
+			builder.CopyTo(0, buffer.Buffer.Memory.Span, length);
+			return buffer;
+		}
+
+		/// <summary>
+		/// Return a buffer initialized with the string builder sub string.
+		/// </summary>
+		/// <param name="builder">Source builder content</param>
+		/// <param name="startIndex">Starting index in the builder</param>
+		/// <param name="length">Length of the content</param>
+		/// <returns>Buffer that should be returned with a call to Return</returns>
+		public static UhtRentedPoolBuffer<char> RentPoolBuffer(this StringBuilder builder, int startIndex, int length)
+		{
+			UhtRentedPoolBuffer<char> buffer = new(length);
+			builder.CopyTo(startIndex, buffer.Buffer.Memory.Span, length);
+			return buffer;
+		}
+	}
 
 	/// <summary>
 	/// Cached character buffer system.
@@ -13,6 +156,7 @@ namespace EpicGames.UHT.Utils
 	/// Invoke UhtBuffer.Borrow method to get a buffer of the given size.
 	/// Invoke UhtBuffer.Return to return the buffer to the cache.
 	/// </summary>
+	[Obsolete("Use UhtPoolBuffer<T> instead of UhtBuffer")]
 	public class UhtBuffer
 	{
 		/// <summary>
@@ -174,6 +318,7 @@ namespace EpicGames.UHT.Utils
 	/// <summary>
 	/// Helper class for using pattern to borrow and return a buffer.
 	/// </summary>
+	[Obsolete("Use UhtRentedPoolBuffer<T> instead of UhtBorrowBuffer")]
 	public struct UhtBorrowBuffer : IDisposable
 	{
 
@@ -226,7 +371,7 @@ namespace EpicGames.UHT.Utils
 	/// Invoke UhtBuffer.Borrow method to get a buffer of the given size.
 	/// Invoke UhtBuffer.Return to return the buffer to the cache.
 	/// </summary>
-	// TODO: Refactor to share implementation and possibly support ArrayPool<T>.
+	[Obsolete("Use UhtPoolBuffer<T> instead of UhtByteBuffer")]
 	public class UhtByteBuffer
 	{
 		/// <summary>
@@ -361,6 +506,7 @@ namespace EpicGames.UHT.Utils
 	/// <summary>
 	/// Helper class for using pattern to borrow and return a buffer.
 	/// </summary>
+	[Obsolete("Use UhtRentedPoolBuffer<T> instead of UhtBorrowByteBuffer")]
 	public struct UhtBorrowByteBuffer : IDisposable
 	{
 
diff --git a/Engine/Source/Programs/Shared/EpicGames.UHT/Utils/UhtFileManager.cs b/Engine/Source/Programs/Shared/EpicGames.UHT/Utils/UhtFileManager.cs
--- a/Engine/Source/Programs/Shared/EpicGames.UHT/Utils/UhtFileManager.cs
+++ b/Engine/Source/Programs/Shared/EpicGames.UHT/Utils/UhtFileManager.cs
@@ -1,6 +1,7 @@
 // Copyright Epic Games, Inc. All Rights Reserved.
 
 using System;
+using System.Buffers;
 using System.IO;
 using System.Text;
 using EpicGames.Core;
@@ -33,8 +34,17 @@ namespace EpicGames.UHT.Utils
 		/// </summary>
 		/// <param name="filePath">File path</param>
 		/// <returns>Buffer containing the read data or null if not found.  The returned buffer must be returned to the cache via a call to UhtBuffer.Return</returns>
+		[Obsolete("Use the new ReadOutput with UhtPoolBuffer")]
 		public UhtBuffer? ReadOutput(string filePath);
 
+		/// <summary>
+		/// Read the given source file
+		/// </summary>
+		/// <param name="filePath">File path</param>
+		/// <param name="output">Read output</param>
+		/// <returns>True if the file was read</returns>
+		public bool ReadOutput(string filePath, out UhtPoolBuffer<char> output);
+
 		/// <summary>
 		/// Write the given contents to the file
 		/// </summary>
@@ -83,6 +93,7 @@ namespace EpicGames.UHT.Utils
 		}
 
 		/// <inheritdoc/>
+		[Obsolete("Use the new ReadOutput with UhtPoolBuffer")]
 		public UhtBuffer? ReadOutput(string filePath)
 		{
 			// Exceptions are very expensive.  Don't bother trying to open the file if it doesn't exist
@@ -113,6 +124,39 @@ namespace EpicGames.UHT.Utils
 			}
 		}
 
+		/// <inheritdoc/>
+		public bool ReadOutput(string filePath, out UhtPoolBuffer<char> output)
+		{
+			// Exceptions are very expensive.  Don't bother trying to open the file if it doesn't exist
+			if (!File.Exists(filePath))
+			{
+				output = default;
+				return false;
+			}
+
+			try
+			{
+				using FileStream fs = new(filePath, FileMode.Open, FileAccess.Read, FileShare.Read, 4 * 1024, FileOptions.SequentialScan);
+				using UhtRentedPoolBuffer<byte> byteBuffer = new((int)fs.Length);
+				Span<byte> byteSpan = byteBuffer.Buffer.Memory.Span;
+				int readLength = fs.Read(byteSpan);
+				byteSpan = byteSpan.Slice(0, readLength);
+
+				Encoding encoding = GetEncoding(byteSpan, out int skipBytes);
+				byteSpan = byteSpan.Slice(skipBytes);
+
+				int charCount = encoding.GetCharCount(byteBuffer.Buffer.Memory.Span);
+				output = UhtPoolBuffers.Rent<char>(charCount);
+				encoding.GetChars(byteBuffer.Buffer.Memory.Span, output.Memory.Span);
+				return true;
+			}
+			catch (IOException)
+			{
+				output = default;
+				return false;
+			}
+		}
+
 		/// <inheritdoc/>
 		public bool WriteOutput(string filePath, ReadOnlySpan<char> contents)
 		{
@@ -165,7 +209,7 @@ namespace EpicGames.UHT.Utils
 			try
 			{
 				using FileStream fs = new(filePath, FileMode.Open, FileAccess.Read, FileShare.Read, 4 * 1024, FileOptions.SequentialScan);
-				using UhtBorrowByteBuffer byteBuffer = new((int)fs.Length);
+				using UhtRentedPoolBuffer<byte> byteBuffer = new((int)fs.Length);
 				Span<byte> byteSpan = byteBuffer.Buffer.Memory.Span;
 				int readLength = fs.Read(byteSpan);
 				byteSpan = byteSpan.Slice(0, readLength);
diff --git a/Engine/Source/Programs/Shared/EpicGames.UHT/Utils/UhtSession.cs b/Engine/Source/Programs/Shared/EpicGames.UHT/Utils/UhtSession.cs
--- a/Engine/Source/Programs/Shared/EpicGames.UHT/Utils/UhtSession.cs
+++ b/Engine/Source/Programs/Shared/EpicGames.UHT/Utils/UhtSession.cs
@@ -143,7 +143,7 @@ namespace EpicGames.UHT.Utils
 		/// <param name="builder">Source for the content</param>
 		public void CommitOutput(string filePath, StringBuilder builder)
 		{
-			using UhtBorrowBuffer borrowBuffer = new(builder);
+			using UhtRentedPoolBuffer<char> borrowBuffer = builder.RentPoolBuffer();
 			string tempFilePath = filePath + ".tmp";
 			SaveIfChanged(filePath, tempFilePath, new StringView(borrowBuffer.Buffer.Memory));
 		}
@@ -288,15 +288,14 @@ namespace EpicGames.UHT.Utils
 				{
 					string message = String.Empty;
 					string refPath = Path.Combine(ReferenceDirectory, fileName);
-					UhtBuffer? existingRef = Session.ReadSourceToBuffer(refPath);
-					if (existingRef != null)
+					if (Session.ReadSource(refPath, out UhtPoolBuffer<char> existingRef))
 					{
 						ReadOnlySpan<char> existingSpan = existingRef.Memory.Span;
 						if (existingSpan.CompareTo(exportedSpan, StringComparison.Ordinal) != 0)
 						{
 							message = $"********************************* {fileName} has changed.";
 						}
-						UhtBuffer.Return(existingRef);
+						UhtPoolBuffers.Return<char>(existingRef);
 					}
 					else
 					{
@@ -315,9 +314,8 @@ namespace EpicGames.UHT.Utils
 			}
 
 			// Check to see if the contents have changed
-			UhtBuffer? original = Session.ReadSourceToBuffer(filePath);
-			bool save = original == null;
-			if (original != null)
+			bool save = false;
+			if (Session.ReadSource(filePath, out UhtPoolBuffer<char> original))
 			{
 				ReadOnlySpan<char> originalSpan = original.Memory.Span;
 				if (originalSpan.CompareTo(exportedSpan, StringComparison.Ordinal) != 0)
@@ -332,7 +330,11 @@ namespace EpicGames.UHT.Utils
 					}
 					save = true;
 				}
-				UhtBuffer.Return(original);
+				UhtPoolBuffers.Return<char>(original);
+			}
+			else
+			{
+				save = true;
 			}
 
 			// If changed of the original didn't exist, then save the new version
@@ -1121,11 +1123,23 @@ namespace EpicGames.UHT.Utils
 		/// </summary>
 		/// <param name="filePath">Full or relative file path</param>
 		/// <returns>Buffer containing the read data or null if not found.  The returned buffer must be returned to the cache via a call to UhtBuffer.Return</returns>
+		[Obsolete("Use the ReadSource method with the UhtPoolBuffer<char> output parameter")]
 		public UhtBuffer? ReadSourceToBuffer(string filePath)
 		{
 			return FileManager!.ReadOutput(filePath);
 		}
 
+		/// <summary>
+		/// Read the given source file
+		/// </summary>
+		/// <param name="filePath">Full or relative file path</param>
+		/// <param name="output">Read source file</param>
+		/// <returns>True if the file was read</returns>
+		public bool ReadSource(string filePath, out UhtPoolBuffer<char> output)
+		{
+			return FileManager!.ReadOutput(filePath, out output);
+		}
+
 		/// <summary>
 		/// Write the given contents to the file
 		/// </summary>
@@ -2440,11 +2454,10 @@ namespace EpicGames.UHT.Utils
 		{
 			if (_projectJson == null && ProjectDirectory != null && ProjectFile != null)
 			{
-				UhtBuffer? contents = ReadSourceToBuffer(ProjectFile);
-				if (contents != null)
+				if (ReadSource(ProjectFile, out UhtPoolBuffer<char> contents))
 				{
 					_projectJson = JsonDocument.Parse(contents.Memory);
-					UhtBuffer.Return(contents);
+					UhtPoolBuffers.Return<char>(contents);
 				}
 			}
 
diff --git a/Engine/Source/Programs/UnrealBuildTool/Modes/UnrealHeaderToolMode.cs b/Engine/Source/Programs/UnrealBuildTool/Modes/UnrealHeaderToolMode.cs
--- a/Engine/Source/Programs/UnrealBuildTool/Modes/UnrealHeaderToolMode.cs
+++ b/Engine/Source/Programs/UnrealBuildTool/Modes/UnrealHeaderToolMode.cs
@@ -594,11 +594,19 @@ namespace UnrealBuildTool.Modes
 		}
 
 		/// <inheritdoc/>
+		[Obsolete("Use the new ReadOutput with UhtPoolBuffer")]
 		public UhtBuffer? ReadOutput(string FilePath)
 		{
 			return null;
 		}
 
+		/// <inheritdoc/>
+		public bool ReadOutput(string FilePath, out UhtPoolBuffer<char> Output)
+		{
+			Output = default;
+			return false;
+		}
+
 		/// <inheritdoc/>
 		public bool WriteOutput(string FilePath, ReadOnlySpan<char> Contents)
 		{
