Horde: Tweaks to compute layer.

* Added trace level debug output for sockets and channels, which display the contents of messages sent and received.
* Merged reader functionality with IComputeMessage.
* Separated the lease and socket classes. Sockets are now constructed for each end of the connection, whereas leases only exist at the initiating end.
* Added an exception message type, for reporting errors back to the initiator.
* Separated commands into more fundamental building blocks; upload files, execute process, delete files, etc...
* Added support for attaching buffers from child processes.
* Removed cpp compute node types.

diff --git a/Engine/Source/Programs/Horde/Horde.Agent/Commands/Compute/ComputeCommand.cs b/Engine/Source/Programs/Horde/Horde.Agent/Commands/Compute/ComputeCommand.cs
--- a/Engine/Source/Programs/Horde/Horde.Agent/Commands/Compute/ComputeCommand.cs
+++ b/Engine/Source/Programs/Horde/Horde.Agent/Commands/Compute/ComputeCommand.cs
@@ -8,6 +8,7 @@ using System.Threading.Tasks;
 using EpicGames.Core;
 using EpicGames.Horde.Common;
 using EpicGames.Horde.Compute;
+using EpicGames.Horde.Compute.Clients;
 using Horde.Agent.Leases.Handlers;
 using Microsoft.Extensions.DependencyInjection;
 using Microsoft.Extensions.Logging;
@@ -29,14 +30,12 @@ namespace Horde.Agent.Commands.Compute
 		readonly IServiceProvider _serviceProvider;
 		readonly IHttpClientFactory _httpClientFactory;
 		readonly IOptions<AgentSettings> _settings;
-		readonly ILogger _logger;
 
-		public ComputeCommand(IServiceProvider serviceProvider, ILogger logger)
+		public ComputeCommand(IServiceProvider serviceProvider)
 		{
 			_serviceProvider = serviceProvider;
 			_httpClientFactory = serviceProvider.GetRequiredService<IHttpClientFactory>();
 			_settings = serviceProvider.GetRequiredService<IOptions<AgentSettings>>();
-			_logger = logger;
 		}
 
 		/// <inheritdoc/>
@@ -65,13 +64,14 @@ namespace Horde.Agent.Commands.Compute
 
 		IComputeClient CreateClient()
 		{
+			ILoggerFactory loggerFactory = _serviceProvider.GetRequiredService<ILoggerFactory>();
 			if (Loopback)
 			{
-				return new LoopbackComputeClient(RunListenerAsync);
+				return new LoopbackComputeClient(RunListenerAsync, 2000, loggerFactory);
 			}
 			else
 			{
-				return new ServerComputeClient(CreateHttpClient, _logger);
+				return new ServerComputeClient(CreateHttpClient, loggerFactory);
 			}
 		}
 
@@ -85,10 +85,10 @@ namespace Horde.Agent.Commands.Compute
 			return client;
 		}
 
-		async Task RunListenerAsync(IComputeLease lease, CancellationToken cancellationToken)
+		async Task RunListenerAsync(IComputeSocket socket, CancellationToken cancellationToken)
 		{
 			ComputeHandler handler = ActivatorUtilities.CreateInstance<ComputeHandler>(_serviceProvider);
-			await handler.RunAsync(lease, cancellationToken);
+			await handler.RunAsync(socket, cancellationToken);
 		}
 	}
 }
diff --git a/Engine/Source/Programs/Horde/Horde.Agent/Commands/Compute/CppComputeCommand.cs b/Engine/Source/Programs/Horde/Horde.Agent/Commands/Compute/CppComputeCommand.cs
--- a/Engine/Source/Programs/Horde/Horde.Agent/Commands/Compute/CppComputeCommand.cs
+++ b/Engine/Source/Programs/Horde/Horde.Agent/Commands/Compute/CppComputeCommand.cs
@@ -7,11 +7,10 @@ using System.Threading;
 using System.Threading.Tasks;
 using EpicGames.Core;
 using EpicGames.Horde.Compute;
-using EpicGames.Horde.Compute.Cpp;
-using EpicGames.Horde.Logs;
 using EpicGames.Horde.Storage;
 using EpicGames.Horde.Storage.Backends;
 using EpicGames.Horde.Storage.Nodes;
+using Horde.Agent.Parser;
 using Microsoft.Extensions.Logging;
 
 namespace Horde.Agent.Commands.Compute
@@ -27,7 +26,7 @@ namespace Horde.Agent.Commands.Compute
 			public string Executable { get; set; } = null!;
 			public List<string> Arguments { get; set; } = new List<string>();
 			public string WorkingDir { get; set; } = String.Empty;
-			public Dictionary<string, string> EnvVars { get; set; } = new Dictionary<string, string>();
+			public Dictionary<string, string?> EnvVars { get; set; } = new Dictionary<string, string?>();
 			public List<string> OutputPaths { get; set; } = new List<string>();
 		}
 
@@ -40,7 +39,7 @@ namespace Horde.Agent.Commands.Compute
 		/// Constructor
 		/// </summary>
 		public CppComputeCommand(IServiceProvider serviceProvider, ILogger<CppComputeCommand> logger)
-			: base(serviceProvider, logger)
+			: base(serviceProvider)
 		{
 			_logger = logger;
 		}
@@ -48,85 +47,37 @@ namespace Horde.Agent.Commands.Compute
 		/// <inheritdoc/>
 		protected override async Task<bool> HandleRequestAsync(IComputeLease lease, CancellationToken cancellationToken)
 		{
-			await using IComputeChannel defaultChannel = lease.CreateChannel(0);
-
-			const int ReplyChannelId = 1;
-			await using IComputeChannel channel = lease.CreateChannel(ReplyChannelId);
+			// Read the task definition
+			byte[] data = await FileReference.ReadAllBytesAsync(TaskFile, cancellationToken);
+			JsonComputeTask jsonComputeTask = JsonSerializer.Deserialize<JsonComputeTask>(data, new JsonSerializerOptions { AllowTrailingCommas = true, PropertyNameCaseInsensitive = true, PropertyNamingPolicy = JsonNamingPolicy.CamelCase })!;
 
+			// Create a sandbox from the data to be uploaded
 			MemoryStorageClient storage = new MemoryStorageClient();
+			NodeLocator sandbox = await CreateSandboxAsync(TaskFile, storage, cancellationToken);
 
-			NodeLocator node = await CreateComputeNodeAsync(TaskFile, storage, cancellationToken);
-			defaultChannel.CppStart(node, ReplyChannelId);
-
-			for (; ; )
+			// Open a socket and upload the sandbox
+			await using (IComputeChannel channel = lease.Socket.AttachMessageChannel(0))
 			{
-				using IComputeMessage message = await channel.ReceiveAsync(cancellationToken);
-				switch (message.Type)
-				{
-					case ComputeMessageType.None:
-						return false;
-					case ComputeMessageType.CppBlobRead:
-						{
-							CppBlobReadMessage cppBlobRead = message.AsCppBlobRead();
-							await channel.CppBlobDataAsync(cppBlobRead, storage, cancellationToken);
-						}
-						break;
-					case ComputeMessageType.CppSuccess:
-						{
-							NodeLocator cppOutputLocator = message.AsCppSuccess();
+				await channel.UploadFilesAsync("", sandbox, storage, cancellationToken);
 
-							using ComputeStorageClient remoteStorage = new ComputeStorageClient(channel);
-							TreeReader reader = new TreeReader(remoteStorage, null, _logger);
-
-							CppComputeOutputNode computeOutput = await reader.ReadNodeAsync<CppComputeOutputNode>(cppOutputLocator, cancellationToken);
-							_logger.LogInformation("Exit code: {ExitCode}", computeOutput.ExitCode);
-
-							LogNode logNode = await computeOutput.Log.ExpandAsync(reader, cancellationToken);
-							foreach (LogChunkRef logChunkRef in logNode.TextChunkRefs)
-							{
-								LogChunkNode logChunk = await logChunkRef.ExpandAsync(reader, cancellationToken);
-								foreach (Utf8String line in logChunk.Lines)
-								{
-									_logger.LogInformation("Output: {Line}", line);
-								}
-							}
-
-							channel.CppFinish();
-						}
-						return false;
-					case ComputeMessageType.CppFailure:
-						{
-							string error = message.AsCppFailure();
-							_logger.LogError("Failed: {Error}", error);
-						}
-						return false;
-					default:
-						throw new NotImplementedException();
+				// Execute the program
+				using (LogParser parser = new LogParser(_logger, new List<string>()))
+				{
+					await channel.ExecuteProcessAsync(jsonComputeTask.Executable, jsonComputeTask.Arguments, jsonComputeTask.WorkingDir, jsonComputeTask.EnvVars, parser.WriteData, cancellationToken);
 				}
 			}
+
+			return true;
 		}
 
-		static async Task<NodeLocator> CreateComputeNodeAsync(FileReference taskFile, IStorageClient storage, CancellationToken cancellationToken)
+		static async Task<NodeLocator> CreateSandboxAsync(FileReference taskFile, IStorageClient storage, CancellationToken cancellationToken)
 		{
 			using TreeWriter writer = new TreeWriter(storage);
 
-			byte[] data = await FileReference.ReadAllBytesAsync(taskFile, cancellationToken);
-			JsonComputeTask jsonComputeTask = JsonSerializer.Deserialize<JsonComputeTask>(data, new JsonSerializerOptions { AllowTrailingCommas = true, PropertyNameCaseInsensitive = true, PropertyNamingPolicy = JsonNamingPolicy.CamelCase })!;
-
 			DirectoryNode sandbox = new DirectoryNode();
 			await sandbox.CopyFromDirectoryAsync(taskFile.Directory.ToDirectoryInfo(), new ChunkingOptions(), writer, null, cancellationToken);
 
-			CppComputeNode computeNode = new CppComputeNode(jsonComputeTask.Executable, jsonComputeTask.Arguments, jsonComputeTask.WorkingDir, new TreeNodeRef<DirectoryNode>(sandbox));
-			foreach ((string key, string value) in jsonComputeTask.EnvVars)
-			{
-				computeNode.EnvVars.Add(key, value);
-			}
-			foreach (string outputPath in jsonComputeTask.OutputPaths)
-			{
-				computeNode.OutputPaths.Add(outputPath);
-			}
-
-			NodeHandle handle = await writer.FlushAsync(computeNode, cancellationToken);
+			NodeHandle handle = await writer.FlushAsync(sandbox, cancellationToken);
 			return handle.Locator;
 		}
 	}
diff --git a/Engine/Source/Programs/Horde/Horde.Agent/Commands/Compute/DdcComputeCommand.cs b/Engine/Source/Programs/Horde/Horde.Agent/Commands/Compute/DdcComputeCommand.cs
--- a/Engine/Source/Programs/Horde/Horde.Agent/Commands/Compute/DdcComputeCommand.cs
+++ b/Engine/Source/Programs/Horde/Horde.Agent/Commands/Compute/DdcComputeCommand.cs
@@ -22,7 +22,7 @@ namespace Horde.Agent.Commands.Compute
 		/// Constructor
 		/// </summary>
 		public DdcComputeCommand(IServiceProvider serviceProvider, ILogger<DdcComputeCommand> logger)
-			: base(serviceProvider, logger)
+			: base(serviceProvider)
 		{
 			_logger = logger;
 		}
@@ -30,7 +30,8 @@ namespace Horde.Agent.Commands.Compute
 		/// <inheritdoc/>
 		protected override async Task<bool> HandleRequestAsync(IComputeLease lease, CancellationToken cancellationToken)
 		{
-			await using (IComputeChannel channel = lease.CreateChannel(0))
+			IComputeSocket socket = lease.Socket;
+			await using (IComputeChannel channel = socket.AttachMessageChannel(0))
 			{
 				_logger.LogInformation("Sending XOR request");
 				channel.XorRequest(new byte[] { 1, 2, 3, 4, 5 }, (byte)123);
diff --git a/Engine/Source/Programs/Horde/Horde.Agent/Commands/Compute/XorComputeCommand.cs b/Engine/Source/Programs/Horde/Horde.Agent/Commands/Compute/XorComputeCommand.cs
--- a/Engine/Source/Programs/Horde/Horde.Agent/Commands/Compute/XorComputeCommand.cs
+++ b/Engine/Source/Programs/Horde/Horde.Agent/Commands/Compute/XorComputeCommand.cs
@@ -1,7 +1,6 @@
 // Copyright Epic Games, Inc. All Rights Reserved.
 
 using System;
-using System.Diagnostics;
 using System.Linq;
 using System.Threading;
 using System.Threading.Tasks;
@@ -23,7 +22,7 @@ namespace Horde.Agent.Commands.Compute
 		/// Constructor
 		/// </summary>
 		public XorComputeCommand(IServiceProvider serviceProvider, ILogger<XorComputeCommand> logger)
-			: base(serviceProvider, logger)
+			: base(serviceProvider)
 		{
 			_logger = logger;
 		}
@@ -31,12 +30,13 @@ namespace Horde.Agent.Commands.Compute
 		/// <inheritdoc/>
 		protected override async Task<bool> HandleRequestAsync(IComputeLease lease, CancellationToken cancellationToken)
 		{
-			await using (IComputeChannel mainChannel = lease.CreateChannel(0))
+			IComputeSocket socket = lease.Socket;
+			await using (IComputeChannel mainChannel = socket.AttachMessageChannel(0))
 			{
 				_logger.LogInformation("Forking compute channel...");
-				mainChannel.Fork(1);
 
-				await using IComputeChannel channel = lease.CreateChannel(1);
+				await using IComputeChannel channel = socket.AttachMessageChannel(1);
+				mainChannel.Fork(1);
 
 				_logger.LogInformation("Sending XOR request");
 				channel.XorRequest(new byte[] { 1, 2, 3, 4, 5 }, (byte)123);
diff --git a/Engine/Source/Programs/Horde/Horde.Agent/Leases/Handlers/ComputeHandler.cs b/Engine/Source/Programs/Horde/Horde.Agent/Leases/Handlers/ComputeHandler.cs
--- a/Engine/Source/Programs/Horde/Horde.Agent/Leases/Handlers/ComputeHandler.cs
+++ b/Engine/Source/Programs/Horde/Horde.Agent/Leases/Handlers/ComputeHandler.cs
@@ -3,17 +3,16 @@
 using System;
 using System.Collections.Generic;
 using System.Diagnostics;
-using System.Linq;
 using System.Net.Sockets;
+using System.Runtime.CompilerServices;
+using System.Runtime.InteropServices;
 using System.Threading;
 using System.Threading.Tasks;
 using EpicGames.Core;
 using EpicGames.Horde.Compute;
-using EpicGames.Horde.Compute.Cpp;
+using EpicGames.Horde.Compute.Buffers;
 using EpicGames.Horde.Compute.Transports;
-using EpicGames.Horde.Logs;
 using EpicGames.Horde.Storage;
-using EpicGames.Horde.Storage.Backends;
 using EpicGames.Horde.Storage.Nodes;
 using Horde.Agent.Services;
 using HordeCommon.Rpc.Tasks;
@@ -29,16 +28,20 @@ namespace Horde.Agent.Leases.Handlers
 	{
 		readonly ComputeListenerService _listenerService;
 		readonly IMemoryCache _memoryCache;
+		readonly DirectoryReference _sandboxDir;
+		readonly ILoggerFactory _loggerFactory;
 		readonly ILogger _logger;
 
 		/// <summary>
 		/// Constructor
 		/// </summary>
-		public ComputeHandler(ComputeListenerService listenerService, IMemoryCache memoryCache, ILogger<ComputeHandler> logger)
+		public ComputeHandler(ComputeListenerService listenerService, IMemoryCache memoryCache, ILoggerFactory loggerFactory)
 		{
 			_listenerService = listenerService;
 			_memoryCache = memoryCache;
-			_logger = logger;
+			_sandboxDir = DirectoryReference.Combine(Program.DataDir, "Sandbox");
+			_loggerFactory = loggerFactory;
+			_logger = loggerFactory.CreateLogger<ComputeHandler>();
 		}
 
 		/// <inheritdoc/>
@@ -58,9 +61,10 @@ namespace Horde.Agent.Leases.Handlers
 					return LeaseResult.Success;
 				}
 
-				await using (ComputeLease lease = new ComputeLease(new TcpTransport(tcpClient.Client), computeTask.Resources))
+				await using (ComputeSocket socket = new ComputeSocket(new TcpTransport(tcpClient.Client), true, _loggerFactory))
 				{
-					await RunAsync(lease, cancellationToken);
+					await RunAsync(socket, cancellationToken);
+					await socket.CloseAsync(cancellationToken);
 					return LeaseResult.Success;
 				}
 			}
@@ -75,20 +79,20 @@ namespace Horde.Agent.Leases.Handlers
 			}
 		}
 
-		public async Task RunAsync(IComputeLease lease, CancellationToken cancellationToken)
+		public async Task RunAsync(IComputeSocket socket, CancellationToken cancellationToken)
 		{
-			await RunAsync(lease, 0, cancellationToken);
+			await RunAsync(socket, 0, cancellationToken);
 		}
 
-		async Task RunAsync(IComputeLease lease, int id, CancellationToken cancellationToken)
+		async Task RunAsync(IComputeSocket socket, int channelId, CancellationToken cancellationToken)
 		{
-			await using (IComputeChannel channel = lease.CreateChannel(id))
+			await using (IComputeChannel channel = socket.AttachMessageChannel(channelId))
 			{
 				List<Task> childTasks = new List<Task>();
 				for (; ; )
 				{
 					using IComputeMessage message = await channel.ReceiveAsync(cancellationToken);
-					_logger.LogTrace("Compute Channel {ChannelId}: {MessageType}", id, message.Type);
+					_logger.LogTrace("Compute Channel {ChannelId}: {MessageType}", channelId, message.Type);
 
 					switch (message.Type)
 					{
@@ -97,25 +101,36 @@ namespace Horde.Agent.Leases.Handlers
 							return;
 						case ComputeMessageType.Fork:
 							{
-								ForkMessage fork = message.AsForkMessage();
-								childTasks.Add(RunAsync(lease, fork.Id, cancellationToken));
+								ForkMessage fork = message.ParseForkMessage();
+								childTasks.Add(RunAsync(socket, fork.ChannelId, cancellationToken));
 							}
 							break;
-						case ComputeMessageType.XorRequest:
+						case ComputeMessageType.WriteFiles:
 							{
-								XorRequestMessage xorRequest = message.AsXorRequest();
-								RunXor(channel, xorRequest.Data, xorRequest.Value);
+								UploadFilesMessage writeFiles = message.ParseUploadFilesMessage();
+								await WriteFilesAsync(channel, writeFiles.Name, writeFiles.Locator, cancellationToken);
+							}
+							break;
+						case ComputeMessageType.DeleteFiles:
+							{
+								DeleteFilesMessage deleteFiles = message.ParseDeleteFilesMessage();
+								DeleteFiles(deleteFiles.Filter);
 							}
 							break;
-						case ComputeMessageType.CppBegin:
+						case ComputeMessageType.ExecuteProcess:
+							{
+								ExecuteProcessMessage executeProcess = message.ParseExecuteProcessMessage();
+								await ExecuteProcessAsync(socket, channel, executeProcess.Executable, executeProcess.Arguments, executeProcess.WorkingDir, executeProcess.EnvVars, cancellationToken);  
+							}
+							break;
+						case ComputeMessageType.XorRequest:
 							{
-								CppBeginMessage cppBegin = message.AsCppBegin();
-								await using IComputeChannel replyChannel = lease.CreateChannel(cppBegin.ReplyChannelId);
-								await RunCppAsync(replyChannel, cppBegin.Locator, cancellationToken);
+								XorRequestMessage xorRequest = message.AsXorRequest();
+								RunXor(channel, xorRequest.Data, xorRequest.Value);
 							}
 							break;
 						default:
-							throw new NotImplementedException();
+							throw new ComputeInvalidMessageException(message);
 					}
 				}
 			}
@@ -138,151 +153,152 @@ namespace Horde.Agent.Leases.Handlers
 			}
 		}
 
-		public async Task RunCppAsync(IComputeChannel channel, NodeLocator locator, CancellationToken cancellationToken)
+		async Task WriteFilesAsync(IComputeChannel channel, string path, NodeLocator locator, CancellationToken cancellationToken)
 		{
-			try
+			using ComputeStorageClient store = new ComputeStorageClient(channel);
+			TreeReader reader = new TreeReader(store, _memoryCache, _logger);
+
+			DirectoryNode directoryNode = await reader.ReadNodeAsync<DirectoryNode>(locator, cancellationToken);
+
+			DirectoryReference outputDir = DirectoryReference.Combine(_sandboxDir, path);
+			if (!outputDir.IsUnderDirectory(_sandboxDir))
 			{
-				await RunCppInternalAsync(channel, locator, cancellationToken);
+				throw new InvalidOperationException("Cannot write files outside sandbox");
 			}
-			catch (Exception ex)
+
+			await directoryNode.CopyToDirectoryAsync(reader, outputDir.ToDirectoryInfo(), _logger, cancellationToken);
+
+			using (IComputeMessageBuilder builder = channel.CreateMessage(ComputeMessageType.WriteFilesResponse))
 			{
-				channel.CppFailure(ex.ToString());
+				builder.Send();
 			}
 		}
 
-		async Task RunCppInternalAsync(IComputeChannel channel, NodeLocator locator, CancellationToken cancellationToken)
+		void DeleteFiles(IReadOnlyList<string> deleteFiles)
 		{
-			DirectoryReference sandboxDir = DirectoryReference.Combine(Program.DataDir, "Sandbox", Guid.NewGuid().ToString());
+			FileFilter filter = new FileFilter(deleteFiles);
 
-			using ComputeStorageClient store = new ComputeStorageClient(channel);
-			TreeReader reader = new TreeReader(store, _memoryCache, _logger);
+			List<FileReference> files = filter.ApplyToDirectory(_sandboxDir, false);
+			foreach (FileReference file in files)
+			{
+				FileUtils.ForceDeleteFile(file);
+			}
+		}
 
-			CppComputeNode node = await reader.ReadNodeAsync<CppComputeNode>(locator, cancellationToken);
-			DirectoryNode directoryNode = await node.Sandbox.ExpandAsync(reader, cancellationToken);
+		async Task ExecuteProcessAsync(IComputeSocket socket, IComputeChannel channel, string executable, IReadOnlyList<string> arguments, string? workingDir, IReadOnlyDictionary<string, string?>? envVars, CancellationToken cancellationToken)
+		{
+			if (RuntimeInformation.IsOSPlatform(OSPlatform.Windows))
+			{
+				await ExecuteProcessWindowsAsync(socket, channel, executable, arguments, workingDir, envVars, cancellationToken);
+			}
+			else
+			{
+				await ExecuteProcessInternalAsync(channel, executable, arguments, workingDir, envVars, cancellationToken);
+			}
+		}
 
-			await directoryNode.CopyToDirectoryAsync(reader, sandboxDir.ToDirectoryInfo(), _logger, cancellationToken);
+		async Task ExecuteProcessWindowsAsync(IComputeSocket socket, IComputeChannel channel, string executable, IReadOnlyList<string> arguments, string? workingDir, IReadOnlyDictionary<string, string?>? envVars, CancellationToken cancellationToken)
+		{
+			using IpcBuffer sendBuffer = IpcBuffer.CreateNew(16384); // Messages that we send to child process
+			using IpcBuffer receiveBuffer = IpcBuffer.CreateNew(16384); // Messages that child process sends to us
 
-			MemoryStorageClient storage = new MemoryStorageClient();
-			using (TreeWriter writer = new TreeWriter(storage))
+			Dictionary<string, string?> newEnvVars = new Dictionary<string, string?>();
+			if (envVars != null)
 			{
-				string executable = FileReference.Combine(sandboxDir, node.Executable).FullName;
-				string commandLine = CommandLineArguments.Join(node.Arguments);
-				string workingDir = DirectoryReference.Combine(sandboxDir, node.WorkingDirectory).FullName;
-
-				Dictionary<string, string>? envVars = null;
-				if (node.EnvVars.Count > 0)
+				foreach ((string name, string? value) in envVars)
 				{
-					envVars = ManagedProcess.GetCurrentEnvVars();
-					foreach ((string key, string value) in node.EnvVars)
-					{
-						envVars[key] = value;
-					}
+					newEnvVars.Add(name, value);
 				}
+			}
+
+			newEnvVars[ComputeEnvVarNames.DefaultSendBuffer] = receiveBuffer.GetHandle(); // Note: opposite ends that we're sending/receiving to
+			newEnvVars[ComputeEnvVarNames.DefaultRecvBuffer] = sendBuffer.GetHandle();
+
+			await using ComputeChannel controlChannel = new ComputeChannel(receiveBuffer.Reader, sendBuffer.Writer, _loggerFactory.CreateLogger<ComputeChannel>());
 
-				int exitCode;
+			Task processTask = ExecuteProcessInternalAsync(channel, executable, arguments, workingDir, envVars, cancellationToken);
+			processTask = processTask.ContinueWith(x => sendBuffer.Writer.MarkComplete(), cancellationToken, TaskContinuationOptions.None, TaskScheduler.Default);
 
-				LogWriter logWriter = new LogWriter(writer);
-				using (ManagedProcessGroup group = new ManagedProcessGroup())
+			List<IpcBuffer> buffers = new List<IpcBuffer>();
+			try
+			{
+				for (; ; )
 				{
-					using (ManagedProcess process = new ManagedProcess(group, executable, commandLine, workingDir, envVars, null, ProcessPriorityClass.Normal))
+					using IComputeMessage message = await controlChannel.ReceiveAsync(cancellationToken);
+					switch (message.Type)
 					{
-						byte[] buffer = new byte[1024];
-						for (; ; )
-						{
-							int length = await process.ReadAsync(buffer, 0, buffer.Length, cancellationToken);
-							if (length == 0)
+						case ComputeMessageType.None:
+							await processTask;
+							return;
+						case ComputeMessageType.CreateBufferRequest:
 							{
-								process.WaitForExit();
-								exitCode = process.ExitCode;
-								break;
+								CreateBufferRequest createBuffer = message.ParseCreateBufferRequest();
+
+								IpcBuffer buffer = IpcBuffer.CreateNew(createBuffer.Capacity);
+								buffers.Add(buffer);
+
+								if (createBuffer.Send)
+								{
+									socket.AttachSendBuffer(createBuffer.ChannelId, buffer.Reader);
+								}
+								else
+								{
+									socket.AttachReceiveBuffer(createBuffer.ChannelId, buffer.Writer);
+								}
 							}
-							await logWriter.AppendAsync(buffer.AsMemory(0, length), cancellationToken);
-						}
+							break;
+						default:
+							throw new ComputeInvalidMessageException(message);
 					}
 				}
-
-				TreeNodeRef<LogNode> logNodeRef = await logWriter.CompleteAsync(cancellationToken);
-
-				FileFilter filter = new FileFilter(node.OutputPaths.Select(x => x.ToString()));
-				List<FileReference> outputFiles = filter.ApplyToDirectory(sandboxDir, false);
-
-				DirectoryNode outputTree = new DirectoryNode();
-				await outputTree.CopyFilesAsync(sandboxDir, outputFiles, new ChunkingOptions(), writer, null, cancellationToken);
-
-				CppComputeOutputNode outputNode = new CppComputeOutputNode(exitCode, logNodeRef, new TreeNodeRef<DirectoryNode>(outputTree));
-				NodeHandle outputHandle = await writer.FlushAsync(outputNode, cancellationToken);
-
-				channel.CppSuccess(outputHandle.Locator);
 			}
-
-			for (; ; )
+			finally
 			{
-				using IComputeMessage message = await channel.ReceiveAsync(cancellationToken);
-				switch (message.Type)
+				foreach (IpcBuffer buffer in buffers)
 				{
-					case ComputeMessageType.CppEnd:
-						FileUtils.ForceDeleteDirectory(sandboxDir);
-						return;
-					case ComputeMessageType.CppBlobRead:
-						{
-							CppBlobReadMessage cppBlobRead = message.AsCppBlobRead();
-							await channel.CppBlobDataAsync(cppBlobRead, storage, cancellationToken);
-						}
-						break;
-					default:
-						throw new NotImplementedException();
+					buffer.Dispose();
 				}
 			}
 		}
 
-		class LogWriter
+		async Task ExecuteProcessInternalAsync(IComputeChannel channel, string executable, IReadOnlyList<string> arguments, string? workingDir, IReadOnlyDictionary<string, string?>? envVars, CancellationToken cancellationToken)
 		{
-			const int MaxChunkSize = 128 * 1024;
-			const int FlushChunkSize = 100 * 1024;
-
-			readonly TreeWriter _writer;
-			int _lineCount;
-			long _offset;
-			readonly LogChunkBuilder _chunkBuilder = new LogChunkBuilder(MaxChunkSize);
-			readonly List<LogChunkRef> _chunks = new List<LogChunkRef>();
-
-			public IReadOnlyList<LogChunkRef> Chunks => _chunks;
-
-			public LogWriter(TreeWriter writer)
-			{
-				_writer = writer;
-			}
+			string resolvedExecutable = FileReference.Combine(_sandboxDir, executable).FullName;
+			string resolvedCommandLine = CommandLineArguments.Join(arguments);
+			string resolvedWorkingDir = DirectoryReference.Combine(_sandboxDir, workingDir ?? String.Empty).FullName;
 
-			public async ValueTask AppendAsync(ReadOnlyMemory<byte> data, CancellationToken cancellationToken)
+			Dictionary<string, string> resolvedEnvVars = ManagedProcess.GetCurrentEnvVars();
+			if (envVars != null)
 			{
-				_chunkBuilder.Append(data.Span);
-
-				if (_chunkBuilder.Length > FlushChunkSize)
+				foreach ((string key, string? value) in envVars)
 				{
-					await FlushAsync(cancellationToken);
+					if (value == null)
+					{
+						resolvedEnvVars.Remove(key);
+					}
+					else
+					{
+						resolvedEnvVars[key] = value;
+					}
 				}
 			}
 
-			async Task FlushAsync(CancellationToken cancellationToken)
-			{
-				LogChunkNode chunkNode = _chunkBuilder.ToLogChunk();
-				LogChunkRef chunkRef = new LogChunkRef(_lineCount, _offset, chunkNode);
-				await _writer.WriteAsync(chunkRef, cancellationToken);
-
-				_chunks.Add(chunkRef);
-				_lineCount += chunkNode.LineCount;
-				_offset += chunkNode.Length;
-			}
-
-			public async Task<TreeNodeRef<LogNode>> CompleteAsync(CancellationToken cancellationToken)
+			using (ManagedProcessGroup group = new ManagedProcessGroup())
 			{
-				await FlushAsync(cancellationToken);
-				LogIndexNode index = new LogIndexNode(new NgramSet(Array.Empty<ushort>()), 0, Array.Empty<LogChunkRef>());
-
-				TreeNodeRef<LogNode> logNodeRef = new TreeNodeRef<LogNode>(new LogNode(LogFormat.Text, _lineCount, _offset, _chunks, new TreeNodeRef<LogIndexNode>(index), true));
-				await _writer.WriteAsync(logNodeRef, cancellationToken);
-
-				return logNodeRef;
+				using (ManagedProcess process = new ManagedProcess(group, resolvedExecutable, resolvedCommandLine, resolvedWorkingDir, resolvedEnvVars, null, ProcessPriorityClass.Normal))
+				{
+					byte[] buffer = new byte[1024];
+					for (; ; )
+					{
+						int length = await process.ReadAsync(buffer, 0, buffer.Length, cancellationToken);
+						if (length == 0)
+						{
+							channel.SendExecuteProcessResponse(process.ExitCode);
+							return;
+						}
+						channel.SendProcessOutput(buffer);
+					}
+				}
 			}
 		}
 	}
diff --git a/Engine/Source/Programs/Horde/Horde.Agent/Services/ComputeListenerService.cs b/Engine/Source/Programs/Horde/Horde.Agent/Services/ComputeListenerService.cs
--- a/Engine/Source/Programs/Horde/Horde.Agent/Services/ComputeListenerService.cs
+++ b/Engine/Source/Programs/Horde/Horde.Agent/Services/ComputeListenerService.cs
@@ -7,10 +7,10 @@ using System.Net;
 using System.Net.Sockets;
 using System.Threading;
 using System.Threading.Tasks;
-using EpicGames.Horde.Compute;
 using Microsoft.Extensions.Hosting;
 using Microsoft.Extensions.Logging;
 using Microsoft.Extensions.Options;
+using EpicGames.Horde.Compute.Clients;
 
 namespace Horde.Agent.Services
 {
diff --git a/Engine/Source/Programs/Horde/Horde.Server.Tests/TunnelServiceTests.cs b/Engine/Source/Programs/Horde/Horde.Server.Tests/TunnelServiceTests.cs
--- a/Engine/Source/Programs/Horde/Horde.Server.Tests/TunnelServiceTests.cs
+++ b/Engine/Source/Programs/Horde/Horde.Server.Tests/TunnelServiceTests.cs
@@ -8,6 +8,7 @@ using System.Net.Sockets;
 using System.Threading;
 using System.Threading.Tasks;
 using EpicGames.Horde.Compute;
+using EpicGames.Horde.Compute.Clients;
 using Horde.Server.Compute;
 using Microsoft.Extensions.DependencyInjection;
 using Microsoft.Extensions.Options;
diff --git a/Engine/Source/Programs/Horde/Horde.Server/Compute/ComputeTaskSource.cs b/Engine/Source/Programs/Horde/Horde.Server/Compute/ComputeTaskSource.cs
--- a/Engine/Source/Programs/Horde/Horde.Server/Compute/ComputeTaskSource.cs
+++ b/Engine/Source/Programs/Horde/Horde.Server/Compute/ComputeTaskSource.cs
@@ -9,6 +9,7 @@ using System.Threading;
 using System.Threading.Tasks;
 using EpicGames.Core;
 using EpicGames.Horde.Compute;
+using EpicGames.Horde.Compute.Clients;
 using EpicGames.Horde.Compute.Transports;
 using Google.Protobuf;
 using Google.Protobuf.WellKnownTypes;
diff --git a/Engine/Source/Programs/Horde/Horde.Server/Compute/TunnelService.cs b/Engine/Source/Programs/Horde/Horde.Server/Compute/TunnelService.cs
--- a/Engine/Source/Programs/Horde/Horde.Server/Compute/TunnelService.cs
+++ b/Engine/Source/Programs/Horde/Horde.Server/Compute/TunnelService.cs
@@ -14,6 +14,7 @@ using EpicGames.Horde.Compute;
 using System.IO;
 using System.Buffers;
 using System.Buffers.Binary;
+using EpicGames.Horde.Compute.Clients;
 
 namespace Horde.Server.Compute
 {
diff --git a/Engine/Source/Programs/Shared/EpicGames.Core/LogEvent.cs b/Engine/Source/Programs/Shared/EpicGames.Core/LogEvent.cs
--- a/Engine/Source/Programs/Shared/EpicGames.Core/LogEvent.cs
+++ b/Engine/Source/Programs/Shared/EpicGames.Core/LogEvent.cs
@@ -480,7 +480,10 @@ namespace EpicGames.Core
 			{
 				foreach ((string name, object? value) in Properties)
 				{
-					yield return new KeyValuePair<string, object?>(name, value?.ToString());
+					if (!name.Equals(MessageTemplate.FormatPropertyName, StringComparison.OrdinalIgnoreCase))
+					{
+						yield return new KeyValuePair<string, object?>(name, value?.ToString());
+					}
 				}
 			}
 		}
diff --git a/Engine/Source/Programs/Shared/EpicGames.Horde/Compute/Buffers/HeapBuffer.cs b/Engine/Source/Programs/Shared/EpicGames.Horde/Compute/Buffers/HeapBuffer.cs
--- a/Engine/Source/Programs/Shared/EpicGames.Horde/Compute/Buffers/HeapBuffer.cs
+++ b/Engine/Source/Programs/Shared/EpicGames.Horde/Compute/Buffers/HeapBuffer.cs
@@ -1,6 +1,7 @@
 // Copyright Epic Games, Inc. All Rights Reserved.
 
 using System;
+using System.Buffers;
 using System.Threading;
 using System.Threading.Tasks;
 
@@ -11,6 +12,7 @@ namespace EpicGames.Horde.Compute.Buffers
 	/// </summary>
 	public class HeapBuffer : ComputeBuffer
 	{
+		readonly IMemoryOwner<byte> _memoryOwner;
 		readonly Memory<byte> _memory;
 
 		int _readPosition;
@@ -22,10 +24,22 @@ namespace EpicGames.Horde.Compute.Buffers
 		/// <summary>
 		/// Creates a local buffer with the given capacity
 		/// </summary>
-		/// <param name="memory"></param>
-		public HeapBuffer(Memory<byte> memory)
+		/// <param name="capacity">Capacity of the buffer</param>
+		public HeapBuffer(int capacity)
 		{
-			_memory = memory;
+			_memoryOwner = MemoryPool<byte>.Shared.Rent(capacity);
+			_memory = _memoryOwner.Memory.Slice(0, capacity);
+		}
+
+		/// <inheritdoc/>
+		protected override void Dispose(bool disposing)
+		{
+			base.Dispose(disposing);
+
+			if (disposing)
+			{
+				_memoryOwner.Dispose();
+			}
 		}
 
 		#region Reader
diff --git a/Engine/Source/Programs/Shared/EpicGames.Horde/Compute/Buffers/IpcBuffer.cs b/Engine/Source/Programs/Shared/EpicGames.Horde/Compute/Buffers/IpcBuffer.cs
--- a/Engine/Source/Programs/Shared/EpicGames.Horde/Compute/Buffers/IpcBuffer.cs
+++ b/Engine/Source/Programs/Shared/EpicGames.Horde/Compute/Buffers/IpcBuffer.cs
@@ -18,7 +18,6 @@ namespace EpicGames.Horde.Compute.Buffers
 	/// <summary>
 	/// Implementation of <see cref="IComputeBuffer"/> suitable for cross-process communication
 	/// </summary>
-	[SupportedOSPlatform("windows")]
 	public sealed class IpcBuffer : ComputeBuffer, IDisposable
 	{
 		[StructLayout(LayoutKind.Sequential)]
@@ -198,13 +197,18 @@ namespace EpicGames.Horde.Compute.Buffers
 		static extern bool DuplicateHandle(IntPtr hSourceProcessHandle, IntPtr hSourceHandle, IntPtr hTargetProcessHandle, out IntPtr lpTargetHandle, uint dwDesiredAccess, [MarshalAs(UnmanagedType.Bool)] bool bInheritHandle, uint dwOptions);
 
 		/// <inheritdoc/>
-		public void Dispose()
+		protected override void Dispose(bool disposing)
 		{
-			_flushedEvent.Dispose();
-			_writtenEvent.Dispose();
-			_memoryMappedView.Dispose();
-			_memoryMappedViewAccessor.Dispose();
-			_memoryMappedFile.Dispose();
+			base.Dispose(disposing);
+
+			if (disposing)
+			{
+				_flushedEvent.Dispose();
+				_writtenEvent.Dispose();
+				_memoryMappedView.Dispose();
+				_memoryMappedViewAccessor.Dispose();
+				_memoryMappedFile.Dispose();
+			}
 		}
 
 		#region Reader
diff --git a/Engine/Source/Programs/Shared/EpicGames.Horde/Compute/Buffers/RentedBuffer.cs b/Engine/Source/Programs/Shared/EpicGames.Horde/Compute/Buffers/RentedBuffer.cs
deleted file mode 100644
diff --git a/Engine/Source/Programs/Shared/EpicGames.Horde/Compute/Clients/LoopbackComputeClient.cs b/Engine/Source/Programs/Shared/EpicGames.Horde/Compute/Clients/LoopbackComputeClient.cs
new file mode 100644
--- /dev/null
+++ b/Engine/Source/Programs/Shared/EpicGames.Horde/Compute/Clients/LoopbackComputeClient.cs
@@ -0,0 +1,78 @@
+// Copyright Epic Games, Inc. All Rights Reserved.
+
+using System;
+using System.Collections.Generic;
+using System.Linq;
+using System.Net;
+using System.Net.Sockets;
+using System.Threading;
+using System.Threading.Tasks;
+using EpicGames.Core;
+using EpicGames.Horde.Compute.Transports;
+using Microsoft.Extensions.Logging;
+
+namespace EpicGames.Horde.Compute.Clients
+{
+	/// <summary>
+	/// Implementation of <see cref="IComputeChannel"/> which marshals data over a loopback connection to a method running on a background task.
+	/// </summary>
+	public sealed class LoopbackComputeClient : IComputeClient
+	{
+		readonly BackgroundTask _listenerTask;
+		readonly Socket _listener;
+		readonly Socket _socket;
+		readonly ILoggerFactory _loggerFactory;
+
+		/// <summary>
+		/// Constructor
+		/// </summary>
+		/// <param name="serverFunc">Callback to process the server side of the connection</param>
+		/// <param name="port">Port to connect on</param>
+		/// <param name="loggerFactory">Logger for diagnostic output</param>
+		public LoopbackComputeClient(Func<IComputeSocket, CancellationToken, Task> serverFunc, int port, ILoggerFactory loggerFactory)
+		{
+			_loggerFactory = loggerFactory;
+
+			_listener = new Socket(SocketType.Stream, ProtocolType.IP);
+			_listener.Bind(new IPEndPoint(IPAddress.Loopback, port));
+			_listener.Listen();
+
+			_listenerTask = BackgroundTask.StartNew(ctx => RunListenerAsync(_listener, serverFunc, loggerFactory, ctx));
+
+			_socket = new Socket(SocketType.Stream, ProtocolType.IP);
+			_socket.Connect(IPAddress.Loopback, port);
+		}
+
+		/// <inheritdoc/>
+		public async ValueTask DisposeAsync()
+		{
+			_socket.Dispose();
+			await _listenerTask.DisposeAsync();
+			_listener.Dispose();
+		}
+
+		/// <summary>
+		/// Sets up the loopback listener and calls the server method
+		/// </summary>
+		static async Task RunListenerAsync(Socket listener, Func<IComputeSocket, CancellationToken, Task> serverFunc, ILoggerFactory loggerFactory, CancellationToken cancellationToken)
+		{
+			using Socket tcpSocket = await listener.AcceptAsync(cancellationToken);
+
+			await using (ComputeSocket socket = new ComputeSocket(new TcpTransport(tcpSocket), true, loggerFactory))
+			{
+				await serverFunc(socket, cancellationToken);
+				await socket.CloseAsync(cancellationToken);
+			}
+		}
+
+		/// <inheritdoc/>
+		public async Task<TResult> ExecuteAsync<TResult>(ClusterId clusterId, Requirements? requirements, Func<IComputeLease, CancellationToken, Task<TResult>> handler, CancellationToken cancellationToken)
+		{
+			await using ComputeSocket socket = new ComputeSocket(new TcpTransport(_socket), false, _loggerFactory);
+			await using ComputeLease lease = new ComputeLease(new Dictionary<string, int>(), socket);
+			TResult result = await handler(lease, cancellationToken);
+			await socket.CloseAsync(cancellationToken);
+			return result;
+		}
+	}
+}
diff --git a/Engine/Source/Programs/Shared/EpicGames.Horde/Compute/Clients/ServerComputeClient.cs b/Engine/Source/Programs/Shared/EpicGames.Horde/Compute/Clients/ServerComputeClient.cs
new file mode 100644
--- /dev/null
+++ b/Engine/Source/Programs/Shared/EpicGames.Horde/Compute/Clients/ServerComputeClient.cs
@@ -0,0 +1,140 @@
+// Copyright Epic Games, Inc. All Rights Reserved.
+
+using EpicGames.Core;
+using EpicGames.Horde.Compute.Transports;
+using Microsoft.Extensions.Logging;
+using System;
+using System.Collections.Generic;
+using System.Net;
+using System.Net.Http;
+using System.Net.Http.Json;
+using System.Net.Sockets;
+using System.Threading;
+using System.Threading.Tasks;
+
+namespace EpicGames.Horde.Compute.Clients
+{
+	/// <summary>
+	/// Helper class to enlist remote resources to perform compute-intensive tasks.
+	/// </summary>
+	public sealed class ServerComputeClient : IComputeClient
+	{
+		/// <summary>
+		/// Length of the nonce sent as part of handshaking between initiator and remote
+		/// </summary>
+		public const int NonceLength = 64;
+
+		class AssignComputeRequest
+		{
+			public Requirements? Requirements { get; set; }
+		}
+
+		class AssignComputeResponse
+		{
+			public string Ip { get; set; } = String.Empty;
+			public int Port { get; set; }
+			public string Nonce { get; set; } = String.Empty;
+			public string Key { get; set; } = String.Empty;
+			public Dictionary<string, int> AssignedResources { get; set; } = new Dictionary<string, int>();
+		}
+
+		readonly HttpClient? _defaultHttpClient;
+		readonly Func<HttpClient> _createHttpClient;
+		readonly CancellationTokenSource _cancellationSource = new CancellationTokenSource();
+		readonly ILoggerFactory _loggerFactory;
+		readonly ILogger _logger;
+
+		/// <summary>
+		/// Constructor
+		/// </summary>
+		/// <param name="serverUri">Uri of the server to connect to</param>
+		/// <param name="loggerFactory">Logger for diagnostic messages</param>
+		public ServerComputeClient(Uri serverUri, ILoggerFactory loggerFactory)
+		{
+#pragma warning disable CA2000 // Dispose objects before losing scope
+			// This is disposed via HttpClient
+			SocketsHttpHandler handler = new SocketsHttpHandler();
+			handler.PooledConnectionLifetime = TimeSpan.FromMinutes(2.0);
+
+			_defaultHttpClient = new HttpClient(handler, true);
+			_defaultHttpClient.BaseAddress = serverUri;
+#pragma warning restore CA2000 // Dispose objects before losing scope
+
+			_createHttpClient = GetDefaultHttpClient;
+			_loggerFactory = loggerFactory;
+			_logger = loggerFactory.CreateLogger<ServerComputeClient>();
+		}
+
+		/// <summary>
+		/// Constructor
+		/// </summary>
+		/// <param name="createHttpClient">Creates an HTTP client with the correct base address for the server</param>
+		/// <param name="loggerFactory">Logger for diagnostic messages</param>
+		public ServerComputeClient(Func<HttpClient> createHttpClient, ILoggerFactory loggerFactory)
+		{
+			_createHttpClient = createHttpClient;
+			_loggerFactory = loggerFactory;
+			_logger = loggerFactory.CreateLogger<ServerComputeClient>();
+		}
+
+		/// <summary>
+		/// Gets the default http client
+		/// </summary>
+		/// <returns></returns>
+		HttpClient GetDefaultHttpClient() => _defaultHttpClient!;
+
+		/// <inheritdoc/>
+		public ValueTask DisposeAsync()
+		{
+			Dispose();
+			return new ValueTask();
+		}
+
+		/// <inheritdoc/>
+		public void Dispose()
+		{
+			_cancellationSource.Dispose();
+			_defaultHttpClient?.Dispose();
+		}
+
+		/// <inheritdoc/>
+		public async Task<TResult> ExecuteAsync<TResult>(ClusterId clusterId, Requirements? requirements, Func<IComputeLease, CancellationToken, Task<TResult>> handler, CancellationToken cancellationToken)
+		{
+			_logger.LogInformation("Requesting compute resource");
+
+			// Assign a compute worker
+			HttpClient client = _createHttpClient();
+
+			AssignComputeRequest request = new AssignComputeRequest();
+			request.Requirements = requirements;
+
+			AssignComputeResponse? responseMessage;
+			using (HttpResponseMessage response = await client.PostAsync($"api/v2/compute/{clusterId}", request, _cancellationSource.Token))
+			{
+				response.EnsureSuccessStatusCode();
+				responseMessage = await response.Content.ReadFromJsonAsync<AssignComputeResponse>(cancellationToken: cancellationToken);
+				if (responseMessage == null)
+				{
+					throw new InvalidOperationException();
+				}
+			}
+
+			_logger.LogInformation("Connecting to {Ip} with nonce {Nonce}...", responseMessage.Ip, responseMessage.Nonce);
+
+			// Connect to the remote machine
+			using Socket socket = new Socket(SocketType.Stream, ProtocolType.Tcp);
+			await socket.ConnectAsync(IPAddress.Parse(responseMessage.Ip), responseMessage.Port, cancellationToken);
+
+			// Send the nonce
+			byte[] nonce = StringUtils.ParseHexString(responseMessage.Nonce);
+			await socket.SendMessageAsync(nonce, SocketFlags.None, cancellationToken);
+			_logger.LogInformation("Connection established.");
+
+			// Pass the rest of the call over to the handler
+			byte[] key = StringUtils.ParseHexString(responseMessage.Key);
+
+			await using ComputeLease lease = new ComputeLease(responseMessage.AssignedResources, new ComputeSocket(new TcpTransport(socket), false, _loggerFactory));
+			return await handler(lease, cancellationToken);
+		}
+	}
+}
diff --git a/Engine/Source/Programs/Shared/EpicGames.Horde/Compute/ComputeBuffer.cs b/Engine/Source/Programs/Shared/EpicGames.Horde/Compute/ComputeBuffer.cs
--- a/Engine/Source/Programs/Shared/EpicGames.Horde/Compute/ComputeBuffer.cs
+++ b/Engine/Source/Programs/Shared/EpicGames.Horde/Compute/ComputeBuffer.cs
@@ -64,6 +64,20 @@ namespace EpicGames.Horde.Compute
 			Writer = new WriterImpl(this);
 		}
 
+		/// <inheritdoc/>
+		public void Dispose()
+		{
+			Dispose(true);
+			GC.SuppressFinalize(this);
+		}
+
+		/// <summary>
+		/// Standard dispose pattern
+		/// </summary>
+		protected virtual void Dispose(bool disposing)
+		{
+		}
+
 		#region Reader
 
 		/// <inheritdoc cref="IComputeBufferReader.IsComplete"/>
diff --git a/Engine/Source/Programs/Shared/EpicGames.Horde/Compute/ComputeChannel.cs b/Engine/Source/Programs/Shared/EpicGames.Horde/Compute/ComputeChannel.cs
--- a/Engine/Source/Programs/Shared/EpicGames.Horde/Compute/ComputeChannel.cs
+++ b/Engine/Source/Programs/Shared/EpicGames.Horde/Compute/ComputeChannel.cs
@@ -3,8 +3,10 @@
 using System;
 using System.Buffers;
 using System.Buffers.Binary;
+using System.Text;
 using System.Threading;
 using System.Threading.Tasks;
+using Microsoft.Extensions.Logging;
 
 namespace EpicGames.Horde.Compute
 {
@@ -28,6 +30,7 @@ namespace EpicGames.Horde.Compute
 			public ReadOnlyMemory<byte> Data { get; }
 
 			readonly IMemoryOwner<byte> _memoryOwner;
+			int _position;
 
 			public Message(ComputeMessageType type, ReadOnlyMemory<byte> data)
 			{
@@ -43,6 +46,12 @@ namespace EpicGames.Horde.Compute
 			{
 				_memoryOwner.Dispose();
 			}
+
+			/// <inheritdoc/>
+			public ReadOnlyMemory<byte> GetMemory(int minSize = 1) => Data.Slice(_position);
+
+			/// <inheritdoc/>
+			public void Advance(int length) => _position += length;
 		}
 
 		/// <summary>
@@ -80,6 +89,11 @@ namespace EpicGames.Horde.Compute
 				header[0] = (byte)_type;
 				BinaryPrimitives.WriteInt32LittleEndian(header.Slice(1, 4), _length);
 
+				if (_channel._logger.IsEnabled(LogLevel.Trace))
+				{
+					_channel.LogMessageInfo("SEND", _type, _sendBufferWriter.GetMemory().Slice(HeaderLength, _length).Span);
+				}
+
 				_sendBufferWriter.Advance(HeaderLength + _length);
 				_length = 0;
 			}
@@ -96,6 +110,7 @@ namespace EpicGames.Horde.Compute
 
 		readonly IComputeBufferReader _receiveBufferReader;
 		readonly IComputeBufferWriter _sendBufferWriter;
+		readonly ILogger _logger;
 
 #pragma warning disable CA2213
 		MessageBuilder? _currentBuilder;
@@ -106,10 +121,12 @@ namespace EpicGames.Horde.Compute
 		/// </summary>
 		/// <param name="receiveBufferReader">Reader for incoming messages</param>
 		/// <param name="sendBufferWriter">Writer for outgoing messages</param>
-		public ComputeChannel(IComputeBufferReader receiveBufferReader, IComputeBufferWriter sendBufferWriter)
+		/// <param name="logger">Logger for diagnostic output</param>
+		public ComputeChannel(IComputeBufferReader receiveBufferReader, IComputeBufferWriter sendBufferWriter, ILogger logger)
 		{
 			_receiveBufferReader = receiveBufferReader;
 			_sendBufferWriter = sendBufferWriter;
+			_logger = logger;
 		}
 
 		/// <inheritdoc/>
@@ -142,6 +159,10 @@ namespace EpicGames.Horde.Compute
 					{
 						ComputeMessageType type = (ComputeMessageType)memory.Span[0];
 						Message message = new Message(type, memory.Slice(HeaderLength, length));
+						if (_logger.IsEnabled(LogLevel.Trace))
+						{
+							LogMessageInfo("RECV", message.Type, message.Data.Span);
+						}
 						_receiveBufferReader.Advance(HeaderLength + length);
 						return message;
 					}
@@ -151,6 +172,20 @@ namespace EpicGames.Horde.Compute
 			return new Message(ComputeMessageType.None, ReadOnlyMemory<byte>.Empty);
 		}
 
+		void LogMessageInfo(string verb, ComputeMessageType type, ReadOnlySpan<byte> data)
+		{
+			StringBuilder bytes = new StringBuilder();
+			for (int offset = 0; offset < 16 && offset < data.Length; offset++)
+			{
+				bytes.Append($" {data[offset]:X2}");
+			}
+			if (data.Length > 16)
+			{
+				bytes.Append("..");
+			}
+			_logger.LogTrace("{Verb} {Type,-20} [{Length,8:n0}] = {Bytes}", verb, type, data.Length, bytes.ToString());
+		}
+
 		/// <inheritdoc/>
 		public IComputeMessageBuilder CreateMessage(ComputeMessageType type, int sizeHint = 0)
 		{
diff --git a/Engine/Source/Programs/Shared/EpicGames.Horde/Compute/ComputeExceptions.cs b/Engine/Source/Programs/Shared/EpicGames.Horde/Compute/ComputeExceptions.cs
new file mode 100644
--- /dev/null
+++ b/Engine/Source/Programs/Shared/EpicGames.Horde/Compute/ComputeExceptions.cs
@@ -0,0 +1,62 @@
+// Copyright Epic Games, Inc. All Rights Reserved.
+
+using System;
+
+namespace EpicGames.Horde.Compute
+{
+	/// <summary>
+	/// Base class for compute errors
+	/// </summary>
+	public abstract class ComputeException : Exception
+	{
+		/// <summary>
+		/// Constructor
+		/// </summary>
+		protected ComputeException(string message)
+			: base(message)
+		{
+		}
+	}
+
+	/// <summary>
+	/// Exception thrown on a remote machine
+	/// </summary>
+	public sealed class ComputeRemoteException : ComputeException
+	{
+		readonly string _description;
+
+		/// <summary>
+		/// Constructor
+		/// </summary>
+		public ComputeRemoteException(ExceptionMessage message)
+			: this(message.Message, message.Description)
+		{
+		}
+
+		/// <summary>
+		/// Constructor
+		/// </summary>
+		public ComputeRemoteException(string message, string description)
+			: base(message)
+		{
+			_description = description;
+		}
+
+		/// <inheritdoc/>
+		public override string ToString() => _description;
+	}
+
+	/// <summary>
+	/// Exception thrown when an invalid message is received
+	/// </summary>
+	public sealed class ComputeInvalidMessageException : ComputeException
+	{
+		/// <summary>
+		/// Constructor
+		/// </summary>
+		public ComputeInvalidMessageException(IComputeMessage message)
+			: base($"Unexpected message {message.Type}")
+		{
+		}
+	}
+}
diff --git a/Engine/Source/Programs/Shared/EpicGames.Horde/Compute/ComputeLease.cs b/Engine/Source/Programs/Shared/EpicGames.Horde/Compute/ComputeLease.cs
--- a/Engine/Source/Programs/Shared/EpicGames.Horde/Compute/ComputeLease.cs
+++ b/Engine/Source/Programs/Shared/EpicGames.Horde/Compute/ComputeLease.cs
@@ -1,12 +1,8 @@
 // Copyright Epic Games, Inc. All Rights Reserved.
 
 using System;
-using System.Buffers;
-using System.Buffers.Binary;
 using System.Collections.Generic;
-using System.Threading;
 using System.Threading.Tasks;
-using EpicGames.Core;
 
 namespace EpicGames.Horde.Compute
 {
@@ -15,249 +11,24 @@ namespace EpicGames.Horde.Compute
 	/// </summary>
 	public sealed class ComputeLease : IComputeLease, IAsyncDisposable
 	{
-		readonly object _lockObject = new object();
-
-		bool _complete;
-
-		readonly IComputeTransport _transport;
-		readonly Task _readTask;
-		readonly CancellationTokenSource _cancellationSource = new CancellationTokenSource();
-
-		readonly Dictionary<int, IComputeBufferWriter> _receiveBuffers = new Dictionary<int, IComputeBufferWriter>();
-		readonly AsyncEvent _receiveBuffersChangedEvent = new AsyncEvent();
-
-		readonly SemaphoreSlim _sendSemaphore = new SemaphoreSlim(1, 1);
-		readonly Dictionary<int, IComputeBufferReader> _sendBuffers = new Dictionary<int, IComputeBufferReader>();
-		readonly Dictionary<int, Task> _sendTasks = new Dictionary<int, Task>();
-
 		/// <inheritdoc/>
 		public IReadOnlyDictionary<string, int> AssignedResources { get; }
 
+		/// <inheritdoc/>
+		public IComputeSocket Socket { get; }
+
 		/// <summary>
 		/// Constructor
 		/// </summary>
-		/// <param name="transport">Transport to communicate with the remote</param>
 		/// <param name="assignedResources">Resources assigned to this lease</param>
-		public ComputeLease(IComputeTransport transport, IReadOnlyDictionary<string, int> assignedResources)
+		/// <param name="socket">Socket for communication</param>
+		public ComputeLease(IReadOnlyDictionary<string, int> assignedResources, IComputeSocket socket)
 		{
 			AssignedResources = assignedResources;
-
-			_transport = transport;
-			_readTask = RunReaderAsync(transport, _cancellationSource.Token);
-		}
-
-		/// <summary>
-		/// Attempt to gracefully close the current connection and shutdown both ends of the transport
-		/// </summary>
-		public async ValueTask CloseAsync(CancellationToken cancellationToken)
-		{
-			Task cancellationTask = Task.Delay(-1, cancellationToken);
-
-			// Wait for all the individual send tasks to complete
-			await Task.WhenAll(_sendTasks.Values);
-			cancellationToken.ThrowIfCancellationRequested();
-
-			// Send a final message indicating that the lease is done. This will allow the senders on the remote end to terminate, and trigger
-			// a message to be sent back to our read task allowing it to shut down gracefully.
-			byte[] header = new byte[8];
-			await _transport.WriteAsync(header, cancellationToken);
-
-			// Wait for the reader to stop
-			await Task.WhenAny(_readTask, cancellationTask);
-			cancellationToken.ThrowIfCancellationRequested();
-		}
-
-		/// <inheritdoc/>
-		public async ValueTask DisposeAsync()
-		{
-			_cancellationSource.Cancel();
-
-			await Task.WhenAll(_sendTasks.Values);
-			await _readTask;
-
-			_cancellationSource.Dispose();
-			_sendSemaphore.Dispose();
-		}
-
-		async Task RunReaderAsync(IComputeTransport transport, CancellationToken cancellationToken)
-		{
-			byte[] header = new byte[8];
-			try
-			{
-				// Maintain a local cache of buffers to be able to query for them without having to acquire a global lock
-				Dictionary<int, IComputeBufferWriter> cachedWriters = new Dictionary<int, IComputeBufferWriter>();
-
-				// Process messages from the remote
-				for(; ;)
-				{
-					// Read the next packet header
-					await transport.ReadAsync(header, cancellationToken);
-
-					// Parse the target buffer and packet size
-					int id = BinaryPrimitives.ReadInt32LittleEndian(header);
-					int size = BinaryPrimitives.ReadInt32LittleEndian(header.AsSpan(4));
-
-					// If the size if negative, we're closing the entire connection
-					if (size == 0)
-					{
-						break;
-					}
-					
-					// Dispatch it to the correct place
-					IComputeBufferWriter writer = await GetReceiveBufferAsync(cachedWriters, id);
-					if (size < 0)
-					{
-						writer.MarkComplete();
-						DetachReceiveBuffer(cachedWriters, id);
-					}
-					else
-					{
-						await transport.ReadAsync(writer.GetMemory().Slice(0, size), cancellationToken);
-						writer.Advance(size);
-					}
-				}
-			}
-			catch (OperationCanceledException)
-			{
-			}
-
-			// Mark all buffers as complete
-			lock (_lockObject)
-			{
-				_complete = true;
-				foreach (IComputeBufferWriter writer in _receiveBuffers.Values)
-				{
-					writer.MarkComplete();
-				}
-			}
-		}
-
-		async Task<IComputeBufferWriter> GetReceiveBufferAsync(Dictionary<int, IComputeBufferWriter> cachedWriters, int id)
-		{
-			IComputeBufferWriter? writer;
-			if (cachedWriters.TryGetValue(id, out writer))
-			{
-				return writer;
-			}
-
-			for(; ;)
-			{
-				Task waitTask;
-				lock (_lockObject)
-				{
-					if (_receiveBuffers.TryGetValue(id, out writer))
-					{
-						cachedWriters.Add(id, writer);
-						return writer;
-					}
-					waitTask = _receiveBuffersChangedEvent.Task;
-				}
-				await waitTask;
-			}
-		}
-
-		class SendSegment : ReadOnlySequenceSegment<byte>
-		{
-			public void Set(ReadOnlyMemory<byte> memory, ReadOnlySequenceSegment<byte>? next, long runningIndex)
-			{
-				Memory = memory;
-				Next = next;
-				RunningIndex = runningIndex;
-			}
-		}
-
-		async Task SendFromBufferAsync(int id, IComputeBufferReader reader, CancellationToken cancellationToken)
-		{
-			byte[] header = new byte[8];
-
-			BinaryPrimitives.WriteInt32LittleEndian(header, id);
-
-			SendSegment headerSegment = new SendSegment();
-			SendSegment bodySegment = new SendSegment();
-			headerSegment.Set(header, bodySegment, 0);
-
-			for (; ; )
-			{
-				// Wait for something to read
-				await reader.WaitAsync(0, cancellationToken);
-
-				// Acquire the semaphore
-				await _sendSemaphore.WaitAsync(cancellationToken);
-				try
-				{
-					ReadOnlyMemory<byte> memory = reader.GetMemory();
-					if (memory.Length > 0)
-					{
-						BinaryPrimitives.WriteInt32LittleEndian(header.AsSpan(4), memory.Length);
-
-						bodySegment.Set(memory, null, header.Length);
-						ReadOnlySequence<byte> sequence = new ReadOnlySequence<byte>(headerSegment, 0, bodySegment, memory.Length);
-
-						await _transport.WriteAsync(sequence, cancellationToken);
-
-						reader.Advance(memory.Length);
-					}
-					else if (reader.IsComplete)
-					{
-						BinaryPrimitives.WriteInt32LittleEndian(header.AsSpan(4), -1);
-						await _transport.WriteAsync(header, cancellationToken);
-						break;
-					}
-				}
-				finally
-				{
-					_sendSemaphore.Release();
-				}
-			}
-
-			lock (_lockObject)
-			{
-				_sendBuffers.Remove(id);
-				_sendTasks.Remove(id);
-			}
-		}
-
-		/// <inheritdoc/>
-		public void AttachSendBuffer(int id, IComputeBufferReader reader)
-		{
-			lock (_lockObject)
-			{
-				_sendBuffers.Add(id, reader);
-				_sendTasks.Add(id, Task.Run(() => SendFromBufferAsync(id, reader, _cancellationSource.Token), _cancellationSource.Token));
-			}
+			Socket = socket;
 		}
 
 		/// <inheritdoc/>
-		public void AttachReceiveBuffer(int id, IComputeBufferWriter writer)
-		{
-			bool complete;
-			lock (_lockObject)
-			{
-				complete = _complete;
-
-				if(!complete)
-				{
-					_receiveBuffers.Add(id, writer);
-				}
-			}
-
-			if (complete)
-			{
-				writer.MarkComplete();
-			}
-			else
-			{
-				_receiveBuffersChangedEvent.Set();
-			}
-		}
-
-		void DetachReceiveBuffer(Dictionary<int, IComputeBufferWriter> cachedWriters, int id)
-		{
-			cachedWriters.Remove(id);
-			lock (_lockObject)
-			{
-				_receiveBuffers.Remove(id);
-			}
-		}
+		public ValueTask DisposeAsync() => Socket.DisposeAsync();
 	}
 }
diff --git a/Engine/Source/Programs/Shared/EpicGames.Horde/Compute/ComputeSocket.cs b/Engine/Source/Programs/Shared/EpicGames.Horde/Compute/ComputeSocket.cs
new file mode 100644
--- /dev/null
+++ b/Engine/Source/Programs/Shared/EpicGames.Horde/Compute/ComputeSocket.cs
@@ -0,0 +1,320 @@
+// Copyright Epic Games, Inc. All Rights Reserved.
+
+using System;
+using System.Buffers;
+using System.Buffers.Binary;
+using System.Collections.Generic;
+using System.Diagnostics;
+using System.Threading;
+using System.Threading.Tasks;
+using EpicGames.Core;
+using Microsoft.Extensions.Logging;
+
+namespace EpicGames.Horde.Compute
+{
+	/// <summary>
+	/// Implementation of a compute lease using sockets to transfer data
+	/// </summary>
+	public sealed class ComputeSocket : IComputeSocket, IAsyncDisposable
+	{
+		readonly object _lockObject = new object();
+		
+		bool _complete;
+		int _nextChannelId;
+
+		readonly Queue<int> _freeChannelIds = new Queue<int>();
+		readonly IComputeTransport _transport;
+		readonly ILoggerFactory _loggerFactory;
+		readonly ILogger _logger;
+		readonly Task _readTask;
+		CancellationTokenSource _cancellationSource = new CancellationTokenSource();
+
+		readonly Dictionary<int, IComputeBufferWriter> _receiveBuffers = new Dictionary<int, IComputeBufferWriter>();
+		readonly AsyncEvent _receiveBuffersChangedEvent = new AsyncEvent();
+
+		readonly SemaphoreSlim _sendSemaphore = new SemaphoreSlim(1, 1);
+		readonly Dictionary<int, IComputeBufferReader> _sendBuffers = new Dictionary<int, IComputeBufferReader>();
+		readonly Dictionary<int, Task> _sendTasks = new Dictionary<int, Task>();
+
+		int Endpoint => (_nextChannelId & 1);
+
+		/// <summary>
+		/// Constructor
+		/// </summary>
+		/// <param name="transport">Transport to communicate with the remote</param>
+		/// <param name="isRemote">Whether this socket is the remote end of the connection</param>
+		/// <param name="loggerFactory">Logger for trace output</param>
+		public ComputeSocket(IComputeTransport transport, bool isRemote, ILoggerFactory loggerFactory)
+		{
+			_transport = transport;
+			_loggerFactory = loggerFactory;
+			_logger = loggerFactory.CreateLogger<ComputeSocket>();
+			_nextChannelId = isRemote ? 1 : 2;
+			_readTask = RunReaderAsync(transport, _cancellationSource.Token);
+		}
+
+		/// <summary>
+		/// Attempt to gracefully close the current connection and shutdown both ends of the transport
+		/// </summary>
+		public async ValueTask CloseAsync(CancellationToken cancellationToken)
+		{
+			Task cancellationTask = Task.Delay(-1, cancellationToken);
+
+			// Wait for all the individual send tasks to complete
+			await Task.WhenAll(_sendTasks.Values);
+			cancellationToken.ThrowIfCancellationRequested();
+
+			// Send a final message indicating that the lease is done. This will allow the senders on the remote end to terminate, and trigger
+			// a message to be sent back to our read task allowing it to shut down gracefully.
+			byte[] header = new byte[8];
+			await _transport.WriteAsync(header, cancellationToken);
+			_logger.LogTrace("{Endpoint}: Sent shutdown packet", Endpoint);
+
+			// Wait for the reader to stop
+			await Task.WhenAny(_readTask, cancellationTask);
+			cancellationToken.ThrowIfCancellationRequested();
+
+			_logger.LogTrace("{Endpoint}: Closed", Endpoint);
+		}
+
+		/// <inheritdoc/>
+		public async ValueTask DisposeAsync()
+		{
+			if (_cancellationSource != null)
+			{
+				_cancellationSource.Cancel();
+
+				await Task.WhenAll(_sendTasks.Values);
+				await _readTask;
+
+				_cancellationSource.Dispose();
+				_cancellationSource = null!;
+			}
+
+			_sendSemaphore.Dispose();
+		}
+
+		async Task RunReaderAsync(IComputeTransport transport, CancellationToken cancellationToken)
+		{
+			_logger.LogTrace("{Endpoint}: Started reader", Endpoint);
+
+			byte[] header = new byte[8];
+			try
+			{
+				// Maintain a local cache of buffers to be able to query for them without having to acquire a global lock
+				Dictionary<int, IComputeBufferWriter> cachedWriters = new Dictionary<int, IComputeBufferWriter>();
+
+				// Process messages from the remote
+				for (; ; )
+				{
+					// Read the next packet header
+					int bytesRead = await transport.ReadAsync(header, cancellationToken);
+					if (bytesRead == 0)
+					{
+						_logger.LogTrace("{Endpoint}: End of socket", Endpoint);
+						break;
+					}
+
+					// Parse the target buffer and packet size
+					int id = BinaryPrimitives.ReadInt32LittleEndian(header);
+					int size = BinaryPrimitives.ReadInt32LittleEndian(header.AsSpan(4));
+					_logger.LogTrace("{Endpoint}: Read {ChannelId} -> {Size} bytes", Endpoint, id, size);
+
+					// If the size if negative, we're closing the entire connection
+					if (size == 0)
+					{
+						_logger.LogTrace("{Endpoint}: Received shutdown packet", Endpoint);
+						break;
+					}
+
+					// Dispatch it to the correct place
+					if (size < 0)
+					{
+						_logger.LogTrace("{Endpoint}: Detaching buffer {Id}", Endpoint, id);
+						DetachReceiveBuffer(cachedWriters, id);
+					}
+					else
+					{
+						IComputeBufferWriter writer = await GetReceiveBufferAsync(cachedWriters, id);
+						await transport.ReadAsync(writer.GetMemory().Slice(0, size), cancellationToken);
+						writer.Advance(size);
+					}
+				}
+			}
+			catch (OperationCanceledException)
+			{
+			}
+
+			// Mark all buffers as complete
+			lock (_lockObject)
+			{
+				_complete = true;
+				foreach (IComputeBufferWriter writer in _receiveBuffers.Values)
+				{
+					writer.MarkComplete();
+				}
+			}
+
+			_logger.LogTrace("{Endpoint}: Closing reader", Endpoint);
+		}
+
+		async Task<IComputeBufferWriter> GetReceiveBufferAsync(Dictionary<int, IComputeBufferWriter> cachedWriters, int id)
+		{
+			IComputeBufferWriter? writer;
+			if (cachedWriters.TryGetValue(id, out writer))
+			{
+				return writer;
+			}
+
+			int delayTime = 1;
+			Task? delayTask = null;
+			Stopwatch? timer = null;
+
+			for (; ; )
+			{
+				Task waitTask;
+				lock (_lockObject)
+				{
+					if (_receiveBuffers.TryGetValue(id, out writer))
+					{
+						cachedWriters.Add(id, writer);
+						return writer;
+					}
+					waitTask = _receiveBuffersChangedEvent.Task;
+				}
+
+				timer ??= Stopwatch.StartNew();
+				delayTask ??= Task.Delay(TimeSpan.FromSeconds(delayTime));
+
+				while (!waitTask.IsCompleted)
+				{
+					if (await Task.WhenAny(waitTask, delayTask) == delayTask)
+					{
+						_logger.LogWarning("Socket stalled for {Time}s waiting for a buffer to be attached to channel {ChannelId}", (int)timer.Elapsed.TotalSeconds, id);
+						delayTime = Math.Min(15, delayTime + 2);
+						delayTask = Task.Delay(TimeSpan.FromSeconds(delayTime));
+					}
+				}
+
+				await waitTask;
+			}
+		}
+
+		class SendSegment : ReadOnlySequenceSegment<byte>
+		{
+			public void Set(ReadOnlyMemory<byte> memory, ReadOnlySequenceSegment<byte>? next, long runningIndex)
+			{
+				Memory = memory;
+				Next = next;
+				RunningIndex = runningIndex;
+			}
+		}
+
+		async Task SendFromBufferAsync(int id, IComputeBufferReader reader, CancellationToken cancellationToken)
+		{
+			byte[] header = new byte[8];
+
+			BinaryPrimitives.WriteInt32LittleEndian(header, id);
+
+			SendSegment headerSegment = new SendSegment();
+			SendSegment bodySegment = new SendSegment();
+			headerSegment.Set(header, bodySegment, 0);
+
+			for (; ; )
+			{
+				// Wait for something to read
+				await reader.WaitAsync(0, cancellationToken);
+
+				// Acquire the semaphore
+				await _sendSemaphore.WaitAsync(cancellationToken);
+				try
+				{
+					ReadOnlyMemory<byte> memory = reader.GetMemory();
+					if (memory.Length > 0)
+					{
+						BinaryPrimitives.WriteInt32LittleEndian(header.AsSpan(4), memory.Length);
+
+						bodySegment.Set(memory, null, header.Length);
+						ReadOnlySequence<byte> sequence = new ReadOnlySequence<byte>(headerSegment, 0, bodySegment, memory.Length);
+
+						await _transport.WriteAsync(sequence, cancellationToken);
+
+						reader.Advance(memory.Length);
+					}
+					else if (reader.IsComplete)
+					{
+						_logger.LogTrace("{Endpoint}: Sending complete packet for channel {ChannelId}", Endpoint, id);
+						BinaryPrimitives.WriteInt32LittleEndian(header.AsSpan(4), -1);
+						await _transport.WriteAsync(header, cancellationToken);
+						break;
+					}
+				}
+				finally
+				{
+					_sendSemaphore.Release();
+				}
+			}
+
+			lock (_lockObject)
+			{
+				_sendBuffers.Remove(id);
+				_sendTasks.Remove(id);
+			}
+		}
+
+		/// <inheritdoc/>
+		public void AttachBuffers(int channelId, IComputeBufferReader? sendBufferReader, IComputeBufferWriter? receiveBufferWriter)
+		{
+			bool complete;
+			lock (_lockObject)
+			{
+				complete = _complete;
+				if (sendBufferReader != null)
+				{
+					_sendBuffers.Add(channelId, sendBufferReader);
+					_sendTasks.Add(channelId, Task.Run(() => SendFromBufferAsync(channelId, sendBufferReader, _cancellationSource.Token), _cancellationSource.Token));
+				}
+				if (receiveBufferWriter != null && !complete)
+				{
+					_receiveBuffers.Add(channelId, receiveBufferWriter);
+				}
+			}
+
+			if (receiveBufferWriter != null)
+			{
+				if (complete)
+				{
+					receiveBufferWriter.MarkComplete();
+				}
+				else
+				{
+					_receiveBuffersChangedEvent.Set();
+				}
+			}
+		}
+
+		/// <inheritdoc/>
+		public IComputeChannel AttachMessageChannel(int channelId, IComputeBuffer sendBuffer, IComputeBuffer receiveBuffer)
+		{
+			AttachBuffers(channelId, sendBuffer.Reader, receiveBuffer.Writer);
+			return new ComputeChannel(receiveBuffer.Reader, sendBuffer.Writer, _loggerFactory.CreateLogger<ComputeChannel>());
+		}
+
+		void DetachReceiveBuffer(Dictionary<int, IComputeBufferWriter> cachedWriters, int id)
+		{
+			cachedWriters.Remove(id);
+
+			IComputeBufferWriter? writer;
+			lock (_lockObject)
+			{
+#pragma warning disable CA2000 // Dispose objects before losing scope
+				_receiveBuffers.Remove(id, out writer);
+#pragma warning restore CA2000 // Dispose objects before losing scope
+			}
+			if (writer != null)
+			{
+				writer.MarkComplete();
+			}
+		}
+	}
+}
diff --git a/Engine/Source/Programs/Shared/EpicGames.Horde/Compute/ComputeStorageClient.cs b/Engine/Source/Programs/Shared/EpicGames.Horde/Compute/ComputeStorageClient.cs
--- a/Engine/Source/Programs/Shared/EpicGames.Horde/Compute/ComputeStorageClient.cs
+++ b/Engine/Source/Programs/Shared/EpicGames.Horde/Compute/ComputeStorageClient.cs
@@ -1,13 +1,11 @@
 // Copyright Epic Games, Inc. All Rights Reserved.
 
 using System;
-using System.Buffers;
 using System.Collections.Generic;
 using System.IO;
 using System.Threading;
 using System.Threading.Tasks;
 using EpicGames.Core;
-using EpicGames.Horde.Compute.Cpp;
 using EpicGames.Horde.Storage;
 
 namespace EpicGames.Horde.Compute
@@ -38,27 +36,6 @@ namespace EpicGames.Horde.Compute
 
 		#region Blobs
 
-		sealed class BlobStream : ReadOnlyMemoryStream
-		{
-			readonly IComputeMessage _message;
-
-			public BlobStream(IComputeMessage message)
-				: base(message.Data)
-			{
-				_message = message;
-			}
-
-			protected override void Dispose(bool disposing)
-			{
-				base.Dispose(disposing);
-
-				if (disposing)
-				{
-					_message.Dispose();
-				}
-			}
-		}
-
 		/// <inheritdoc/>
 		public override async Task<Stream> ReadBlobAsync(BlobLocator locator, CancellationToken cancellationToken = default)
 		{
@@ -71,16 +48,7 @@ namespace EpicGames.Horde.Compute
 			await _semaphore.WaitAsync(cancellationToken);
 			try
 			{
-				_channel.CppBlobRead(locator, offset, length);
-
-				IComputeMessage response = await _channel.ReceiveAsync(cancellationToken);
-				if (response.Type != ComputeMessageType.CppBlobData)
-				{
-					response.Dispose();
-					throw new NotImplementedException();
-				}
-
-				return new BlobStream(response);
+				return await _channel.ReadBlobAsync(locator, offset, length, cancellationToken);
 			}
 			finally
 			{
diff --git a/Engine/Source/Programs/Shared/EpicGames.Horde/Compute/Cpp/CppComputeNode.cs b/Engine/Source/Programs/Shared/EpicGames.Horde/Compute/Cpp/CppComputeNode.cs
deleted file mode 100644
diff --git a/Engine/Source/Programs/Shared/EpicGames.Horde/Compute/Cpp/CppComputeOutputNode.cs b/Engine/Source/Programs/Shared/EpicGames.Horde/Compute/Cpp/CppComputeOutputNode.cs
deleted file mode 100644
diff --git a/Engine/Source/Programs/Shared/EpicGames.Horde/Compute/Cpp/CppMessages.cs b/Engine/Source/Programs/Shared/EpicGames.Horde/Compute/Cpp/CppMessages.cs
deleted file mode 100644
diff --git a/Engine/Source/Programs/Shared/EpicGames.Horde/Compute/IComputeBuffer.cs b/Engine/Source/Programs/Shared/EpicGames.Horde/Compute/IComputeBuffer.cs
--- a/Engine/Source/Programs/Shared/EpicGames.Horde/Compute/IComputeBuffer.cs
+++ b/Engine/Source/Programs/Shared/EpicGames.Horde/Compute/IComputeBuffer.cs
@@ -9,7 +9,7 @@ namespace EpicGames.Horde.Compute
 	/// <summary>
 	/// Base interface for input and output buffers
 	/// </summary>
-	public interface IComputeBuffer
+	public interface IComputeBuffer : IDisposable
 	{
 		/// <summary>
 		/// Reader from the buffer
@@ -79,4 +79,20 @@ namespace EpicGames.Horde.Compute
 		/// <param name="cancellationToken">Cancellation token for the operation</param>
 		ValueTask FlushAsync(CancellationToken cancellationToken);
 	}
+
+	/// <summary>
+	/// Default environment variable names for compute channels
+	/// </summary>
+	public static class ComputeEnvVarNames
+	{
+		/// <summary>
+		/// Default name for a channel that can be used to send messages to the initiator
+		/// </summary>
+		public const string DefaultSendBuffer = "UE_HORDE_SEND_BUFFER";
+
+		/// <summary>
+		/// Default name for a channel that can be used to receive messages from the initiator
+		/// </summary>
+		public const string DefaultRecvBuffer = "UE_HORDE_RECV_BUFFER";
+	}
 }
diff --git a/Engine/Source/Programs/Shared/EpicGames.Horde/Compute/IComputeChannel.cs b/Engine/Source/Programs/Shared/EpicGames.Horde/Compute/IComputeChannel.cs
--- a/Engine/Source/Programs/Shared/EpicGames.Horde/Compute/IComputeChannel.cs
+++ b/Engine/Source/Programs/Shared/EpicGames.Horde/Compute/IComputeChannel.cs
@@ -4,7 +4,6 @@ using System;
 using System.Threading;
 using System.Threading.Tasks;
 using EpicGames.Core;
-using EpicGames.Horde.Compute.Buffers;
 
 namespace EpicGames.Horde.Compute
 {
@@ -45,48 +44,18 @@ namespace EpicGames.Horde.Compute
 	/// </summary>
 	public static class ComputeChannelExtensions
 	{
-		class RentedBufferChannel : ComputeChannel
-		{
-			readonly RentedLocalBuffer _receiveBuffer;
-			readonly RentedLocalBuffer _sendBuffer;
-
-			public RentedBufferChannel(RentedLocalBuffer receiveBuffer, RentedLocalBuffer sendBuffer)
-				: base(receiveBuffer.Reader, sendBuffer.Writer)
-			{
-				_sendBuffer = sendBuffer;
-				_receiveBuffer = receiveBuffer;
-			}
-
-			/// <inheritdoc/>
-			protected override async ValueTask DisposeAsync(bool disposing)
-			{
-				await base.DisposeAsync(disposing);
-
-				if (disposing)
-				{
-					_sendBuffer.Dispose();
-					_receiveBuffer.Dispose();
-				}
-			}
-		}
-
 		/// <summary>
-		/// Creates a compute channel with the given identifier
+		/// Forwards an existing message across a channel
 		/// </summary>
-		/// <param name="lease">The lease to create a channel for</param>
-		/// <param name="id">Identifier for the channel</param>
-		public static IComputeChannel CreateChannel(this IComputeLease lease, int id)
+		/// <param name="channel">Channel to send on</param>
+		/// <param name="message">The message to be sent</param>
+		public static void Send(this IComputeChannel channel, IComputeMessage message)
 		{
-			const int MaxMessageSize = 64 * 1024;
-			const int BufferSize = MaxMessageSize * 3;
-
-			RentedLocalBuffer receiveBuffer = new RentedLocalBuffer(BufferSize);
-			lease.AttachReceiveBuffer(id, receiveBuffer.Writer);
-
-			RentedLocalBuffer sendBuffer = new RentedLocalBuffer(BufferSize);
-			lease.AttachSendBuffer(id, sendBuffer.Reader);
-
-			return new RentedBufferChannel(receiveBuffer, sendBuffer);
+			using (IComputeMessageBuilder builder = channel.CreateMessage(message.Type, message.Data.Length))
+			{
+				builder.WriteFixedLengthBytes(message.Data.Span);
+				builder.Send();
+			}
 		}
 	}
 }
diff --git a/Engine/Source/Programs/Shared/EpicGames.Horde/Compute/IComputeClient.cs b/Engine/Source/Programs/Shared/EpicGames.Horde/Compute/IComputeClient.cs
--- a/Engine/Source/Programs/Shared/EpicGames.Horde/Compute/IComputeClient.cs
+++ b/Engine/Source/Programs/Shared/EpicGames.Horde/Compute/IComputeClient.cs
@@ -16,7 +16,7 @@ namespace EpicGames.Horde.Compute
 		/// </summary>
 		/// <param name="clusterId">Cluster to execute the request</param>
 		/// <param name="requirements">Requirements for the agent</param>
-		/// <param name="handler">Handler for the connection</param>
+		/// <param name="handler">Handler for the lease</param>
 		/// <param name="cancellationToken">Cancellation token for the operation</param>
 		public Task<TResult> ExecuteAsync<TResult>(ClusterId clusterId, Requirements? requirements, Func<IComputeLease, CancellationToken, Task<TResult>> handler, CancellationToken cancellationToken);
 	}
diff --git a/Engine/Source/Programs/Shared/EpicGames.Horde/Compute/IComputeLease.cs b/Engine/Source/Programs/Shared/EpicGames.Horde/Compute/IComputeLease.cs
--- a/Engine/Source/Programs/Shared/EpicGames.Horde/Compute/IComputeLease.cs
+++ b/Engine/Source/Programs/Shared/EpicGames.Horde/Compute/IComputeLease.cs
@@ -2,8 +2,6 @@
 
 using System;
 using System.Collections.Generic;
-using System.Threading;
-using System.Threading.Tasks;
 
 namespace EpicGames.Horde.Compute
 {
@@ -18,22 +16,8 @@ namespace EpicGames.Horde.Compute
 		IReadOnlyDictionary<string, int> AssignedResources { get; }
 
 		/// <summary>
-		/// Closes the underlying transport stream gracefully
+		/// Socket to communicate with the remote
 		/// </summary>
-		ValueTask CloseAsync(CancellationToken cancellationToken);
-
-		/// <summary>
-		/// Attaches a send buffer to this lease. Data will be read from this buffer and replicated a receive buffer attached with the same id on the remote.
-		/// </summary>
-		/// <param name="id">Identifier for the buffer</param>
-		/// <param name="buffer">The buffer to attach</param>
-		void AttachSendBuffer(int id, IComputeBufferReader buffer);
-
-		/// <summary>
-		/// Attaches a receive buffer to this lease. Data will be read into this buffer from the other end of the lease.
-		/// </summary>
-		/// <param name="id">Identifier for the buffer</param>
-		/// <param name="buffer">The buffer to attach</param>
-		void AttachReceiveBuffer(int id, IComputeBufferWriter buffer);
+		IComputeSocket Socket { get; }
 	}
 }
diff --git a/Engine/Source/Programs/Shared/EpicGames.Horde/Compute/IComputeMessage.cs b/Engine/Source/Programs/Shared/EpicGames.Horde/Compute/IComputeMessage.cs
--- a/Engine/Source/Programs/Shared/EpicGames.Horde/Compute/IComputeMessage.cs
+++ b/Engine/Source/Programs/Shared/EpicGames.Horde/Compute/IComputeMessage.cs
@@ -2,9 +2,14 @@
 
 using System;
 using System.Buffers.Binary;
+using System.Collections.Generic;
+using System.IO;
 using System.Threading;
 using System.Threading.Tasks;
 using EpicGames.Core;
+using EpicGames.Horde.Compute.Buffers;
+using EpicGames.Horde.Storage;
+using EpicGames.Horde.Storage.Nodes;
 
 namespace EpicGames.Horde.Compute
 {
@@ -18,56 +23,99 @@ namespace EpicGames.Horde.Compute
 		/// </summary>
 		None = 0x00,
 
+		/// <summary>
+		/// Sent in place of a regular response if an error occurs on the remote
+		/// </summary>
+		Exception = 0x01,
+
 		/// <summary>
 		/// Fork a new request channel
 		/// </summary>
-		Fork = 0x01,
+		Fork = 0x02,
 
-		#region Test Requests
+		/// <summary>
+		/// Untyped user data
+		/// </summary>
+		UserData = 0x03,
+
+		#region Process Management
 
 		/// <summary>
-		/// Xor a block of data with a value
+		/// Extract files on the remote machine (Initiator -> Remote)
 		/// </summary>
-		XorRequest = 0x10,
+		WriteFiles = 0x10,
 
 		/// <summary>
-		/// Result from an <see cref="XorRequest"/> request.
+		/// Notification that files have been extracted (Remote -> Initiator)
 		/// </summary>
-		XorResponse = 0x11,
+		WriteFilesResponse = 0x11,
+
+		/// <summary>
+		/// Deletes files on the remote machine (Initiator -> Remote)
+		/// </summary>
+		DeleteFiles = 0x12,
+
+		/// <summary>
+		/// Delete a directory in the remote (Initiator -> Remote)
+		/// </summary>
+		DeleteDirectory = 0x13,
+
+		/// <summary>
+		/// Execute a process in a sandbox (Initiator -> Remote)
+		/// </summary>
+		ExecuteProcess = 0x16,
+
+		/// <summary>
+		/// Returns the process exit code (Remote -> Initiator)
+		/// </summary>
+		ExecuteProcessResponse = 0x17,
+
+		/// <summary>
+		/// Returns output from the child process to the caller (Remote -> Initiator)
+		/// </summary>
+		ProcessOutput = 0x18,
 
 		#endregion
 
-		#region C++ Builds
+		#region Storage
 
 		/// <summary>
-		/// Initiate a C++ build
+		/// Reads a blob from storage
 		/// </summary>
-		CppBegin = 0x20,
+		ReadBlob = 0x20,
 
 		/// <summary>
-		/// Gives output from a C++ build
+		/// Response to a <see cref="ReadBlob"/> request.
 		/// </summary>
-		CppSuccess = 0x21,
+		ReadBlobResponse = 0x21,
+
+		#endregion
+
+		#region Buffers
 
 		/// <summary>
-		/// Information about a failed C++ task
+		/// Creates a new receive buffer for messages
 		/// </summary>
-		CppFailure = 0x22,
+		CreateBufferRequest = 0x30,
 
 		/// <summary>
-		/// Finish a C++ build and close the current channel
+		/// Receives information about the new receive buffer
 		/// </summary>
-		CppEnd = 0x23,
+		CreateBufferResponse = 0x31,
+
+		#endregion
+
+		#region Test Requests
 
 		/// <summary>
-		/// Reads a blob pertaining to a C++ build
+		/// Xor a block of data with a value
 		/// </summary>
-		CppBlobRead = 0x24,
+		XorRequest = 0xf0,
 
 		/// <summary>
-		/// Response to a <see cref="CppBlobRead"/> request.
+		/// Result from an <see cref="XorRequest"/> request.
 		/// </summary>
-		CppBlobData = 0x25,
+		XorResponse = 0xf1,
 
 		#endregion
 	}
@@ -75,7 +123,7 @@ namespace EpicGames.Horde.Compute
 	/// <summary>
 	/// Information about a received compute message
 	/// </summary>
-	public interface IComputeMessage : IDisposable
+	public interface IComputeMessage : IMemoryReader, IDisposable
 	{
 		/// <summary>
 		/// Type of the message
@@ -89,10 +137,57 @@ namespace EpicGames.Horde.Compute
 	}
 
 	/// <summary>
-	/// Message for forking the message loop
+	/// Message for reporting an error
 	/// </summary>
-	/// <param name="Id">New channel id</param>
-	public record struct ForkMessage(int Id);
+	public record struct ExceptionMessage(string Message, string Description);
+
+	/// <summary>
+	/// Message to request forking the message loop
+	/// </summary>
+	public record struct ForkMessage(int ChannelId);
+
+	/// <summary>
+	/// Extract files from a bundle to a path in the remote sandbox
+	/// </summary>
+	/// <param name="Name">Path to extract the files to</param>
+	/// <param name="Locator">Locator for the tree to extract</param>
+	public record struct UploadFilesMessage(string Name, NodeLocator Locator);
+
+	/// <summary>
+	/// Deletes files or directories in the remote
+	/// </summary>
+	/// <param name="Filter">Filter for files to delete</param>
+	public record struct DeleteFilesMessage(IReadOnlyList<string> Filter);
+
+	/// <summary>
+	/// Message to execute a new child process
+	/// </summary>
+	/// <param name="Executable">Executable path</param>
+	/// <param name="Arguments">Arguments for the executable</param>
+	/// <param name="WorkingDir">Working directory to execute in</param>
+	/// <param name="EnvVars">Environment variables for the child process. Null values unset variables.</param>
+	public record struct ExecuteProcessMessage(string Executable, IReadOnlyList<string> Arguments, string? WorkingDir, IReadOnlyDictionary<string, string?> EnvVars);
+
+	/// <summary>
+	/// Response from executing a child process
+	/// </summary>
+	/// <param name="ExitCode">Exit code for the process</param>
+	public record struct ExecuteProcessResponseMessage(int ExitCode);
+
+	/// <summary>
+	/// Creates a blob read request
+	/// </summary>
+	public record struct ReadBlobMessage(BlobLocator Locator, int Offset, int Length);
+
+	/// <summary>
+	/// Request creation of a new buffer, and attaching it to the lease
+	/// </summary>
+	public record struct CreateBufferRequest(int ChannelId, long Capacity, bool Send);
+
+	/// <summary>
+	/// Response to create a new buffer
+	/// </summary>
+	public record struct CreateBufferResponse(string Handle);
 
 	/// <summary>
 	/// Message for running an XOR command
@@ -107,15 +202,39 @@ namespace EpicGames.Horde.Compute
 	public static class ComputeMessageExtensions
 	{
 		/// <summary>
-		/// Creates a new remote message loop using the given channel id
+		/// Sends an exception response to the remote
 		/// </summary>
-		/// <param name="channel">Current channel</param>
-		/// <param name="id">New channel id</param>
-		public static void Fork(this IComputeChannel channel, int id)
+		public static void Exception(this IComputeChannel channel, Exception ex) => Exception(channel, ex.Message, ex.StackTrace);
+
+		/// <summary>
+		/// Sends an exception response to the remote
+		/// </summary>
+		public static void Exception(this IComputeChannel channel, string description, string? trace)
+		{
+			using (IComputeMessageBuilder builder = channel.CreateMessage(ComputeMessageType.Exception))
+			{
+				builder.WriteString(description);
+				builder.WriteOptionalString(trace);
+				builder.Send();
+			}
+		}
+
+		/// <summary>
+		/// Parses a message as an <see cref="ExceptionMessage"/>
+		/// </summary>
+		public static ExceptionMessage ParseExceptionMessage(this IComputeMessage message)
+		{
+			string msg = message.ReadString();
+			string description = message.ReadString();
+			return new ExceptionMessage(msg, description);
+		}
+
+		/// <inheritdoc/>
+		public static void Fork(this IComputeChannel channel, int newChannelId)
 		{
 			using (IComputeMessageBuilder builder = channel.CreateMessage(ComputeMessageType.Fork))
 			{
-				builder.WriteInt32(id);
+				builder.WriteInt32(newChannelId);
 				builder.Send();
 			}
 		}
@@ -123,12 +242,422 @@ namespace EpicGames.Horde.Compute
 		/// <summary>
 		/// Parses a fork message from the given compute message
 		/// </summary>
-		public static ForkMessage AsForkMessage(this IComputeMessage message)
+		public static ForkMessage ParseForkMessage(this IComputeMessage message)
 		{
-			ReadOnlyMemory<byte> data = message.Data;
-			return new ForkMessage(BinaryPrimitives.ReadInt32LittleEndian(data.Span));
+			int channelId = message.ReadInt32();
+			return new ForkMessage(channelId);
 		}
 
+		/// <summary>
+		/// Sends untyped user data to the remote
+		/// </summary>
+		/// <param name="channel">Current channel</param>
+		/// <param name="span">Data to send</param>
+		public static void SendUserData(this IComputeChannel channel, ReadOnlySpan<byte> span)
+		{
+			using (IComputeMessageBuilder builder = channel.CreateMessage(ComputeMessageType.UserData))
+			{
+				builder.WriteFixedLengthBytes(span);
+				builder.Send();
+			}
+		}
+
+		#region Process
+
+		static async Task<IComputeMessage> RunStorageServer(this IComputeChannel channel, IStorageClient storage, CancellationToken cancellationToken)
+		{
+			for (; ; )
+			{
+				IComputeMessage message = await channel.ReceiveAsync(cancellationToken);
+				if (message.Type != ComputeMessageType.ReadBlob)
+				{
+					return message;
+				}
+
+				try
+				{
+					ReadBlobMessage readBlob = message.ParseReadBlobRequest();
+					await SendBlobDataAsync(channel, readBlob, storage, cancellationToken);
+				}
+				finally
+				{
+					message.Dispose();
+				}
+			}
+		}
+
+		/// <summary>
+		/// Creates a sandbox on the remote machine
+		/// </summary>
+		/// <param name="channel">Current channel</param>
+		/// <param name="path">Root directory to extract files within the sandbox</param>
+		/// <param name="locator">Location of a <see cref="DirectoryNode"/> describing contents of the sandbox</param>
+		/// <param name="storage">Storage for the sandbox data</param>
+		/// <param name="cancellationToken">Cancellation token for the operation</param>
+		public static async Task UploadFilesAsync(this IComputeChannel channel, string path, NodeLocator locator, IStorageClient storage, CancellationToken cancellationToken)
+		{
+			using (IComputeMessageBuilder writer = channel.CreateMessage(ComputeMessageType.WriteFiles))
+			{
+				writer.WriteString(path);
+				writer.WriteNodeLocator(locator);
+				writer.Send();
+			}
+
+			using IComputeMessage response = await RunStorageServer(channel, storage, cancellationToken);
+			if (response.Type != ComputeMessageType.WriteFilesResponse)
+			{
+				throw new NotSupportedException();
+			}
+		}
+
+		/// <summary>
+		/// Parses a message as a <see cref="UploadFilesMessage"/>
+		/// </summary>
+		public static UploadFilesMessage ParseUploadFilesMessage(this IComputeMessage message)
+		{
+			string name = message.ReadString();
+			NodeLocator locator = message.ReadNodeLocator();
+			return new UploadFilesMessage(name, locator);
+		}
+
+		/// <summary>
+		/// Destroys a sandbox on the remote machine
+		/// </summary>
+		/// <param name="channel">Current channel</param>
+		/// <param name="paths">Paths of files or directories to delete</param>
+		public static void DeleteFiles(this IComputeChannel channel, IReadOnlyList<string> paths)
+		{
+			using (IComputeMessageBuilder writer = channel.CreateMessage(ComputeMessageType.DeleteDirectory))
+			{
+				writer.WriteList(paths, MemoryWriterExtensions.WriteString);
+				writer.Send();
+			}
+		}
+
+		/// <summary>
+		/// Parses a message as a <see cref="DeleteFilesMessage"/>
+		/// </summary>
+		public static DeleteFilesMessage ParseDeleteFilesMessage(this IComputeMessage message)
+		{
+			List<string> files = message.ReadList(MemoryReaderExtensions.ReadString);
+			return new DeleteFilesMessage(files);
+		}
+
+		/// <summary>
+		/// Executes a remote process
+		/// </summary>
+		/// <param name="channel">Current channel</param>
+		/// <param name="executable">Executable to run, relative to the sandbox root</param>
+		/// <param name="arguments">Arguments for the child process</param>
+		/// <param name="workingDir">Working directory for the process</param>
+		/// <param name="outputHandler">Output callback for stdout</param>
+		/// <param name="envVars">Environment variables for the child process</param>
+		/// <param name="cancellationToken">Cancellation token for the operation</param>
+		public static async Task<int> ExecuteProcessAsync(this IComputeChannel channel, string executable, IReadOnlyList<string> arguments, string? workingDir, IReadOnlyDictionary<string, string?>? envVars, Action<ReadOnlyMemory<byte>> outputHandler, CancellationToken cancellationToken)
+		{
+			using (IComputeMessageBuilder writer = channel.CreateMessage(ComputeMessageType.ExecuteProcess))
+			{
+				writer.WriteString(executable);
+				writer.WriteList(arguments, MemoryWriterExtensions.WriteString);
+				writer.WriteOptionalString(workingDir);
+				writer.WriteDictionary(envVars ?? new Dictionary<string, string?>(), MemoryWriterExtensions.WriteString, MemoryWriterExtensions.WriteOptionalString);
+				writer.Send();
+			}
+
+			for (; ; )
+			{
+				using IComputeMessage message = await channel.ReceiveAsync(cancellationToken);
+				switch (message.Type)
+				{
+					case ComputeMessageType.Exception:
+						ExceptionMessage exception = message.ParseExceptionMessage();
+						throw new ComputeRemoteException(exception);
+					case ComputeMessageType.ProcessOutput:
+						outputHandler(message.Data);
+						break;
+					case ComputeMessageType.ExecuteProcessResponse:
+						ExecuteProcessResponseMessage executeProcessResponse = message.ParseExecuteProcessResponse();
+						return executeProcessResponse.ExitCode;
+					default:
+						throw new ComputeInvalidMessageException(message);
+				}
+			}
+		}
+
+		/// <summary>
+		/// Parses a message as a <see cref="ExecuteProcessMessage"/>
+		/// </summary>
+		public static ExecuteProcessMessage ParseExecuteProcessMessage(this IComputeMessage message)
+		{
+			string executable = message.ReadString();
+			List<string> arguments = message.ReadList(MemoryReaderExtensions.ReadString);
+			string? workingDir = message.ReadOptionalString();
+			Dictionary<string, string?> envVars = message.ReadDictionary(MemoryReaderExtensions.ReadString, MemoryReaderExtensions.ReadOptionalString);
+			return new ExecuteProcessMessage(executable, arguments, workingDir, envVars);
+		}
+
+		/// <summary>
+		/// Sends output from a child process
+		/// </summary>
+		public static void SendProcessOutput(this IComputeChannel channel, ReadOnlyMemory<byte> data)
+		{
+			using (IComputeMessageBuilder builder = channel.CreateMessage(ComputeMessageType.ProcessOutput))
+			{
+				builder.WriteFixedLengthBytes(data.Span);
+				builder.Send();
+			}
+		}
+
+		/// <summary>
+		/// Sends a response from executing a child process
+		/// </summary>
+		/// <param name="channel"></param>
+		/// <param name="exitCode">Exit code from the process</param>
+		public static void SendExecuteProcessResponse(this IComputeChannel channel, int exitCode)
+		{
+			using (IComputeMessageBuilder builder = channel.CreateMessage(ComputeMessageType.ExecuteProcessResponse))
+			{
+				builder.WriteInt32(exitCode);
+				builder.Send();
+			}
+		}
+
+		/// <summary>
+		/// Parses a message as a <see cref="ExecuteProcessMessage"/>
+		/// </summary>
+		public static ExecuteProcessResponseMessage ParseExecuteProcessResponse(this IComputeMessage message)
+		{
+			int exitCode = message.ReadInt32();
+			return new ExecuteProcessResponseMessage(exitCode);
+		}
+
+		#endregion
+
+		#region Storage
+
+		/// <summary>
+		/// 
+		/// </summary>
+		/// <param name="message"></param>
+		/// <returns></returns>
+		public static ReadBlobMessage ParseReadBlobRequest(this IComputeMessage message)
+		{
+			BlobLocator locator = message.ReadBlobLocator();
+			int offset = (int)message.ReadUnsignedVarInt();
+			int length = (int)message.ReadUnsignedVarInt();
+			return new ReadBlobMessage(locator, offset, length);
+		}
+
+		/// <summary>
+		/// Wraps a compute message containing blob data
+		/// </summary>
+		sealed class BlobDataStream : ReadOnlyMemoryStream
+		{
+			readonly IComputeMessage _message;
+
+			public BlobDataStream(IComputeMessage message)
+				: base(message.Data)
+			{
+				_message = message;
+			}
+
+			protected override void Dispose(bool disposing)
+			{
+				base.Dispose(disposing);
+
+				if (disposing)
+				{
+					_message.Dispose();
+				}
+			}
+		}
+
+		/// <summary>
+		/// Reads a blob from the remote
+		/// </summary>
+		/// <param name="channel">Channel to write to</param>
+		/// <param name="locator">Locator for the blob</param>
+		/// <param name="offset">Offset within the blob</param>
+		/// <param name="length">Length of data to return</param>
+		/// <param name="cancellationToken">Cancellation token for the operation</param>
+		/// <returns>Stream containing the blob data</returns>
+		public static async Task<Stream> ReadBlobAsync(this IComputeChannel channel, BlobLocator locator, int offset, int length, CancellationToken cancellationToken)
+		{
+			using (IComputeMessageBuilder writer = channel.CreateMessage(ComputeMessageType.ReadBlob))
+			{
+				writer.WriteBlobLocator(locator);
+				writer.WriteUnsignedVarInt(offset);
+				writer.WriteUnsignedVarInt(length);
+				writer.Send();
+			}
+
+			IComputeMessage response = await channel.ReceiveAsync(cancellationToken);
+			if (response.Type != ComputeMessageType.ReadBlobResponse)
+			{
+				response.Dispose();
+				throw new ComputeInvalidMessageException(response);
+			}
+
+			return new BlobDataStream(response);
+		}
+
+		/// <summary>
+		/// Writes blob data to a compute channel
+		/// </summary>
+		/// <param name="channel">Channel to write to</param>
+		/// <param name="message">The read request</param>
+		/// <param name="storage">Storage client to retrieve the blob from</param>
+		/// <param name="cancellationToken">Cancellation token for the operation</param>
+		public static Task SendBlobDataAsync(this IComputeChannel channel, ReadBlobMessage message, IStorageClient storage, CancellationToken cancellationToken)
+		{
+			return SendBlobDataAsync(channel, message.Locator, message.Offset, message.Length, storage, cancellationToken);
+		}
+
+		/// <summary>
+		/// Writes blob data to a compute channel
+		/// </summary>
+		/// <param name="channel">Channel to write to</param>
+		/// <param name="locator">Locator for the blob to send</param>
+		/// <param name="offset">Starting offset of the data</param>
+		/// <param name="length">Length of the data</param>
+		/// <param name="storage">Storage client to retrieve the blob from</param>
+		/// <param name="cancellationToken">Cancellation token for the operation</param>
+		public static async Task SendBlobDataAsync(this IComputeChannel channel, BlobLocator locator, int offset, int length, IStorageClient storage, CancellationToken cancellationToken)
+		{
+			byte[] data;
+			if (offset == 0 && length == 0)
+			{
+				using (Stream stream = await storage.ReadBlobAsync(locator, cancellationToken))
+				{
+					using MemoryStream target = new MemoryStream();
+					await stream.CopyToAsync(target, cancellationToken);
+					data = target.ToArray();
+				}
+			}
+			else
+			{
+				using (Stream stream = await storage.ReadBlobRangeAsync(locator, offset, length, cancellationToken))
+				{
+					using MemoryStream target = new MemoryStream();
+					await stream.CopyToAsync(target, cancellationToken);
+					data = target.ToArray();
+				}
+			}
+
+			using (IComputeMessageBuilder writer = channel.CreateMessage(ComputeMessageType.ReadBlobResponse, length + 128))
+			{
+				writer.WriteFixedLengthBytes(data);
+				writer.Send();
+			}
+		}
+
+		#endregion
+
+		#region Child process
+
+		/// <summary>
+		/// Wrapper around the lifetime of an <see cref="IpcBuffer"/> which only exposes the writer interface
+		/// </summary>
+		class WrappedBufferWriter : IComputeBufferWriter
+		{
+			readonly IpcBuffer _buffer;
+
+			public WrappedBufferWriter(IpcBuffer buffer) => _buffer = buffer;
+
+			public void Dispose() => _buffer.Dispose();
+
+			public void Advance(int size) => _buffer.Writer.Advance(size);
+
+			public ValueTask FlushAsync(CancellationToken cancellationToken) => _buffer.Writer.FlushAsync(cancellationToken);
+
+			public Memory<byte> GetMemory() => _buffer.Writer.GetMemory();
+
+			public void MarkComplete() => _buffer.Writer.MarkComplete();
+		}
+
+		/// <summary>
+		/// Creates a buffer which can be used to send data to the remote
+		/// </summary>
+		/// <param name="channel">Channel to write to</param>
+		/// <param name="id">Identifier for the new buffer</param>
+		/// <param name="capacity">Capacity of the buffer</param>
+		/// <param name="cancellationToken">Cancellation token for the operation</param>
+		/// <returns>New compute buffer</returns>
+		public static async Task<IComputeBufferWriter> CreateSendBufferAsync(this IComputeChannel channel, int id, long capacity, CancellationToken cancellationToken)
+		{
+			IpcBuffer buffer = await CreateBufferAsync(channel, id, capacity, true, cancellationToken);
+			return new WrappedBufferWriter(buffer);
+		}
+
+		/// <summary>
+		/// Wrapper around the lifetime of an <see cref="IpcBuffer"/> which only exposes the reader interface
+		/// </summary>
+		class WrappedBufferReader : IComputeBufferReader
+		{
+			readonly IpcBuffer _buffer;
+
+			public WrappedBufferReader(IpcBuffer buffer) => _buffer = buffer;
+
+			public void Dispose() => _buffer.Dispose();
+
+			public bool IsComplete => _buffer.Reader.IsComplete;
+
+			public void Advance(int size) => _buffer.Reader.Advance(size);
+
+			public ReadOnlyMemory<byte> GetMemory() => _buffer.Reader.GetMemory();
+
+			public ValueTask WaitAsync(int currentLength, CancellationToken cancellationToken) => _buffer.Reader.WaitAsync(currentLength, cancellationToken);
+		}
+
+		/// <summary>
+		/// Creates a buffer which can be used to receive data from the remote
+		/// </summary>
+		/// <param name="channel">Channel to write to</param>
+		/// <param name="id">Identifier for the new buffer</param>
+		/// <param name="capacity">Capacity of the buffer</param>
+		/// <param name="cancellationToken">Cancellation token for the operation</param>
+		/// <returns>New compute buffer</returns>
+		public static async Task<IComputeBufferReader> CreateReceiveBufferAsync(this IComputeChannel channel, int id, long capacity, CancellationToken cancellationToken)
+		{
+			IpcBuffer buffer = await CreateBufferAsync(channel, id, capacity, false, cancellationToken);
+			return new WrappedBufferReader(buffer);
+		}
+
+		static async Task<IpcBuffer> CreateBufferAsync(this IComputeChannel channel, int id, long capacity, bool send, CancellationToken cancellationToken)
+		{
+			using (IComputeMessageBuilder builder = channel.CreateMessage(ComputeMessageType.CreateBufferRequest))
+			{
+				builder.WriteInt32(id);
+				builder.WriteInt64(capacity);
+				builder.WriteBoolean(send);
+				builder.Send();
+			}
+
+			using IComputeMessage response = await channel.ReceiveAsync(cancellationToken);
+			if (response.Type != ComputeMessageType.CreateBufferResponse)
+			{
+				throw new NotSupportedException();
+			}
+
+			string handle = response.ReadString();
+			return IpcBuffer.OpenExisting(handle);
+		}
+
+		/// <summary>
+		/// Parse a <see cref="CreateBufferRequest"/> message
+		/// </summary>
+		public static CreateBufferRequest ParseCreateBufferRequest(this IComputeMessage message)
+		{
+			int channelId = message.ReadInt32();
+			long capacity = message.ReadInt64();
+			bool send = message.ReadBoolean();
+			return new CreateBufferRequest(channelId, capacity, send);
+		}
+
+		#endregion
+
+		#region Test Messages
+
 		/// <summary>
 		/// Send a message to request that a byte string be xor'ed with a particular value
 		/// </summary>
@@ -150,5 +679,7 @@ namespace EpicGames.Horde.Compute
 			ReadOnlyMemory<byte> data = message.Data;
 			return new XorRequestMessage(data[0..^1], data.Span[^1]);
 		}
+
+		#endregion
 	}
 }
diff --git a/Engine/Source/Programs/Shared/EpicGames.Horde/Compute/IComputeSocket.cs b/Engine/Source/Programs/Shared/EpicGames.Horde/Compute/IComputeSocket.cs
new file mode 100644
--- /dev/null
+++ b/Engine/Source/Programs/Shared/EpicGames.Horde/Compute/IComputeSocket.cs
@@ -0,0 +1,100 @@
+// Copyright Epic Games, Inc. All Rights Reserved.
+
+using System;
+using System.Threading;
+using System.Threading.Tasks;
+using EpicGames.Horde.Compute.Buffers;
+
+namespace EpicGames.Horde.Compute
+{
+	/// <summary>
+	/// Socket for sending and reciving messages
+	/// </summary>
+	public interface IComputeSocket : IAsyncDisposable
+	{
+		/// <summary>
+		/// Attaches buffers to a channel.
+		/// </summary>
+		/// <param name="channelId">Identifier of the channel</param>
+		/// <param name="sendBufferReader">The reader to attach</param>
+		/// <param name="receiveBufferWriter">The writer to attach</param>
+		/// <returns>Identifier for the message channel</returns>
+		void AttachBuffers(int channelId, IComputeBufferReader? sendBufferReader, IComputeBufferWriter? receiveBufferWriter);
+
+		/// <summary>
+		/// Creates a channel to receive data on the given channel id
+		/// </summary>
+		/// <param name="channelId">Identifier of the channel</param>
+		/// <param name="sendBuffer">Buffer for sending messages</param>
+		/// <param name="receiveBuffer">Buffer for receiving messages</param>
+		IComputeChannel AttachMessageChannel(int channelId, IComputeBuffer sendBuffer, IComputeBuffer receiveBuffer);
+	}
+
+	/// <summary>
+	/// Extension methods for sockets
+	/// </summary>
+	public static class ComputeSocketExtensions
+	{
+		class HeapBufferChannel : IComputeChannel
+		{
+			readonly IComputeChannel _inner;
+			readonly HeapBuffer _receiveBuffer;
+			readonly HeapBuffer _sendBuffer;
+
+			public HeapBufferChannel(IComputeChannel inner, HeapBuffer receiveBuffer, HeapBuffer sendBuffer)
+			{
+				_inner = inner;
+				_sendBuffer = sendBuffer;
+				_receiveBuffer = receiveBuffer;
+			}
+
+			/// <inheritdoc/>
+			public async ValueTask DisposeAsync()
+			{
+				await _inner.DisposeAsync();
+
+				_sendBuffer.Dispose();
+				_receiveBuffer.Dispose();
+			}
+
+			/// <inheritdoc/>
+			public ValueTask<IComputeMessage> ReceiveAsync(CancellationToken cancellationToken) => _inner.ReceiveAsync(cancellationToken);
+
+			/// <inheritdoc/>
+			public IComputeMessageBuilder CreateMessage(ComputeMessageType type, int sizeHint = 0) => _inner.CreateMessage(type, sizeHint);
+		}
+
+		/// <summary>
+		/// Creates a channel to receive data on the given channel id
+		/// </summary>
+		/// <param name="socket">Socket to attach a channel on</param>
+		/// <param name="channelId">Identifier of the channel</param>
+		public static IComputeChannel AttachMessageChannel(this IComputeSocket socket, int channelId)
+		{
+			const int MaxMessageSize = 64 * 1024;
+			const int BufferSize = MaxMessageSize * 3;
+
+			HeapBuffer receiveBuffer = new HeapBuffer(BufferSize);
+			HeapBuffer sendBuffer = new HeapBuffer(BufferSize);
+
+			IComputeChannel inner = socket.AttachMessageChannel(channelId, sendBuffer, receiveBuffer);
+			return new HeapBufferChannel(inner, receiveBuffer, sendBuffer);
+		}
+
+		/// <summary>
+		/// Attaches a send buffer to a socket. Data will be read from this buffer and replicated a receive buffer attached with the same id on the remote.
+		/// </summary>
+		/// <param name="socket">Socket to attach to</param>
+		/// <param name="channelId">Identifier for the buffer</param>
+		/// <param name="reader">Source to read from</param>
+		public static void AttachSendBuffer(this IComputeSocket socket, int channelId, IComputeBufferReader reader) => socket.AttachBuffers(channelId, reader, null);
+
+		/// <summary>
+		/// Attaches a receive buffer to a socket. Data will be read into this buffer from the other end of the lease.
+		/// </summary>
+		/// <param name="socket">Socket to attach to</param>
+		/// <param name="channelId">Identifier for the buffer</param>
+		/// <param name="writer">The buffer to attach</param>
+		public static void AttachReceiveBuffer(this IComputeSocket socket, int channelId, IComputeBufferWriter writer) => socket.AttachBuffers(channelId, null, writer);
+	}
+}
diff --git a/Engine/Source/Programs/Shared/EpicGames.Horde/Compute/LoopbackComputeClient.cs b/Engine/Source/Programs/Shared/EpicGames.Horde/Compute/LoopbackComputeClient.cs
deleted file mode 100644
diff --git a/Engine/Source/Programs/Shared/EpicGames.Horde/Compute/ServerComputeClient.cs b/Engine/Source/Programs/Shared/EpicGames.Horde/Compute/ServerComputeClient.cs
deleted file mode 100644
diff --git a/Engine/Source/Programs/Shared/EpicGames.Horde/Compute/SharedMemory.cs b/Engine/Source/Programs/Shared/EpicGames.Horde/Compute/SharedMemory.cs
deleted file mode 100644
diff --git a/Engine/Source/Programs/Shared/EpicGames.Horde/Storage/TreeReader.cs b/Engine/Source/Programs/Shared/EpicGames.Horde/Storage/TreeReader.cs
--- a/Engine/Source/Programs/Shared/EpicGames.Horde/Storage/TreeReader.cs
+++ b/Engine/Source/Programs/Shared/EpicGames.Horde/Storage/TreeReader.cs
@@ -64,8 +64,6 @@ namespace EpicGames.Horde.Storage
 			typeof(Logs.LogNode),
 			typeof(Logs.LogChunkNode),
 			typeof(Logs.LogIndexNode),
-			typeof(Compute.Cpp.CppComputeNode),
-			typeof(Compute.Cpp.CppComputeOutputNode)
 		};
 
 		/// <summary>
