Horde: Fix build errors.

diff --git a/Engine/Source/Programs/Horde/Horde.Agent/Leases/Handlers/ComputeHandler.cs b/Engine/Source/Programs/Horde/Horde.Agent/Leases/Handlers/ComputeHandler.cs
--- a/Engine/Source/Programs/Horde/Horde.Agent/Leases/Handlers/ComputeHandler.cs
+++ b/Engine/Source/Programs/Horde/Horde.Agent/Leases/Handlers/ComputeHandler.cs
@@ -61,7 +61,7 @@ namespace Horde.Agent.Leases.Handlers
 					return LeaseResult.Success;
 				}
 
-				await using (ComputeSocket socket = new ComputeSocket(new TcpTransport(tcpClient.Client), true, _loggerFactory))
+				await using (ComputeSocket socket = new ComputeSocket(new TcpTransport(tcpClient.Client), _loggerFactory))
 				{
 					await RunAsync(socket, cancellationToken);
 					await socket.CloseAsync(cancellationToken);
diff --git a/Engine/Source/Programs/Shared/EpicGames.Horde.Tests/BufferTests.cs b/Engine/Source/Programs/Shared/EpicGames.Horde.Tests/BufferTests.cs
--- a/Engine/Source/Programs/Shared/EpicGames.Horde.Tests/BufferTests.cs
+++ b/Engine/Source/Programs/Shared/EpicGames.Horde.Tests/BufferTests.cs
@@ -17,9 +17,9 @@ namespace EpicGames.Horde.Tests
 	public class BufferTests
 	{
 		[TestMethod]
-		public async Task TestLocalBuffer()
+		public async Task TestHeapBuffer()
 		{
-			using RentedLocalBuffer buffer = new RentedLocalBuffer(8000);
+			using HeapBuffer buffer = new HeapBuffer(8000);
 			await TestProducerConsumerAsync(buffer.Writer, buffer.Reader);
 		}
 
diff --git a/Engine/Source/Programs/Shared/EpicGames.Horde/Compute/Buffers/IpcBuffer.cs b/Engine/Source/Programs/Shared/EpicGames.Horde/Compute/Buffers/IpcBuffer.cs
--- a/Engine/Source/Programs/Shared/EpicGames.Horde/Compute/Buffers/IpcBuffer.cs
+++ b/Engine/Source/Programs/Shared/EpicGames.Horde/Compute/Buffers/IpcBuffer.cs
@@ -8,7 +8,6 @@ using System.Threading;
 using System.Threading.Tasks;
 using EpicGames.Core;
 using Microsoft.Win32.SafeHandles;
-using System.Runtime.Versioning;
 using System.IO.MemoryMappedFiles;
 using System.Reflection;
 using System.IO;
@@ -18,7 +17,7 @@ namespace EpicGames.Horde.Compute.Buffers
 	/// <summary>
 	/// Implementation of <see cref="IComputeBuffer"/> suitable for cross-process communication
 	/// </summary>
-	public sealed class IpcBuffer : ComputeBuffer, IDisposable
+	public sealed class IpcBuffer : ComputeBuffer
 	{
 		[StructLayout(LayoutKind.Sequential)]
 		class SECURITY_ATTRIBUTES
diff --git a/Engine/Source/Programs/Shared/EpicGames.Horde/Compute/Clients/LoopbackComputeClient.cs b/Engine/Source/Programs/Shared/EpicGames.Horde/Compute/Clients/LoopbackComputeClient.cs
--- a/Engine/Source/Programs/Shared/EpicGames.Horde/Compute/Clients/LoopbackComputeClient.cs
+++ b/Engine/Source/Programs/Shared/EpicGames.Horde/Compute/Clients/LoopbackComputeClient.cs
@@ -58,7 +58,7 @@ namespace EpicGames.Horde.Compute.Clients
 		{
 			using Socket tcpSocket = await listener.AcceptAsync(cancellationToken);
 
-			await using (ComputeSocket socket = new ComputeSocket(new TcpTransport(tcpSocket), true, loggerFactory))
+			await using (ComputeSocket socket = new ComputeSocket(new TcpTransport(tcpSocket), loggerFactory))
 			{
 				await serverFunc(socket, cancellationToken);
 				await socket.CloseAsync(cancellationToken);
@@ -68,7 +68,7 @@ namespace EpicGames.Horde.Compute.Clients
 		/// <inheritdoc/>
 		public async Task<TResult> ExecuteAsync<TResult>(ClusterId clusterId, Requirements? requirements, Func<IComputeLease, CancellationToken, Task<TResult>> handler, CancellationToken cancellationToken)
 		{
-			await using ComputeSocket socket = new ComputeSocket(new TcpTransport(_socket), false, _loggerFactory);
+			await using ComputeSocket socket = new ComputeSocket(new TcpTransport(_socket), _loggerFactory);
 			await using ComputeLease lease = new ComputeLease(new Dictionary<string, int>(), socket);
 			TResult result = await handler(lease, cancellationToken);
 			await socket.CloseAsync(cancellationToken);
diff --git a/Engine/Source/Programs/Shared/EpicGames.Horde/Compute/Clients/ServerComputeClient.cs b/Engine/Source/Programs/Shared/EpicGames.Horde/Compute/Clients/ServerComputeClient.cs
--- a/Engine/Source/Programs/Shared/EpicGames.Horde/Compute/Clients/ServerComputeClient.cs
+++ b/Engine/Source/Programs/Shared/EpicGames.Horde/Compute/Clients/ServerComputeClient.cs
@@ -133,7 +133,7 @@ namespace EpicGames.Horde.Compute.Clients
 			// Pass the rest of the call over to the handler
 			byte[] key = StringUtils.ParseHexString(responseMessage.Key);
 
-			await using ComputeLease lease = new ComputeLease(responseMessage.AssignedResources, new ComputeSocket(new TcpTransport(socket), false, _loggerFactory));
+			await using ComputeLease lease = new ComputeLease(responseMessage.AssignedResources, new ComputeSocket(new TcpTransport(socket), _loggerFactory));
 			return await handler(lease, cancellationToken);
 		}
 	}
diff --git a/Engine/Source/Programs/Shared/EpicGames.Horde/Compute/ComputeSocket.cs b/Engine/Source/Programs/Shared/EpicGames.Horde/Compute/ComputeSocket.cs
--- a/Engine/Source/Programs/Shared/EpicGames.Horde/Compute/ComputeSocket.cs
+++ b/Engine/Source/Programs/Shared/EpicGames.Horde/Compute/ComputeSocket.cs
@@ -20,9 +20,7 @@ namespace EpicGames.Horde.Compute
 		readonly object _lockObject = new object();
 		
 		bool _complete;
-		int _nextChannelId;
 
-		readonly Queue<int> _freeChannelIds = new Queue<int>();
 		readonly IComputeTransport _transport;
 		readonly ILoggerFactory _loggerFactory;
 		readonly ILogger _logger;
@@ -36,20 +34,16 @@ namespace EpicGames.Horde.Compute
 		readonly Dictionary<int, IComputeBufferReader> _sendBuffers = new Dictionary<int, IComputeBufferReader>();
 		readonly Dictionary<int, Task> _sendTasks = new Dictionary<int, Task>();
 
-		int Endpoint => (_nextChannelId & 1);
-
 		/// <summary>
 		/// Constructor
 		/// </summary>
 		/// <param name="transport">Transport to communicate with the remote</param>
-		/// <param name="isRemote">Whether this socket is the remote end of the connection</param>
 		/// <param name="loggerFactory">Logger for trace output</param>
-		public ComputeSocket(IComputeTransport transport, bool isRemote, ILoggerFactory loggerFactory)
+		public ComputeSocket(IComputeTransport transport, ILoggerFactory loggerFactory)
 		{
 			_transport = transport;
 			_loggerFactory = loggerFactory;
 			_logger = loggerFactory.CreateLogger<ComputeSocket>();
-			_nextChannelId = isRemote ? 1 : 2;
 			_readTask = RunReaderAsync(transport, _cancellationSource.Token);
 		}
 
@@ -68,13 +62,13 @@ namespace EpicGames.Horde.Compute
 			// a message to be sent back to our read task allowing it to shut down gracefully.
 			byte[] header = new byte[8];
 			await _transport.WriteAsync(header, cancellationToken);
-			_logger.LogTrace("{Endpoint}: Sent shutdown packet", Endpoint);
+			_logger.LogTrace("Sent shutdown packet");
 
 			// Wait for the reader to stop
 			await Task.WhenAny(_readTask, cancellationTask);
 			cancellationToken.ThrowIfCancellationRequested();
 
-			_logger.LogTrace("{Endpoint}: Closed", Endpoint);
+			_logger.LogTrace("Closed");
 		}
 
 		/// <inheritdoc/>
@@ -96,7 +90,7 @@ namespace EpicGames.Horde.Compute
 
 		async Task RunReaderAsync(IComputeTransport transport, CancellationToken cancellationToken)
 		{
-			_logger.LogTrace("{Endpoint}: Started reader", Endpoint);
+			_logger.LogTrace("Started socket reader");
 
 			byte[] header = new byte[8];
 			try
@@ -111,26 +105,26 @@ namespace EpicGames.Horde.Compute
 					int bytesRead = await transport.ReadAsync(header, cancellationToken);
 					if (bytesRead == 0)
 					{
-						_logger.LogTrace("{Endpoint}: End of socket", Endpoint);
+						_logger.LogTrace("End of socket");
 						break;
 					}
 
 					// Parse the target buffer and packet size
 					int id = BinaryPrimitives.ReadInt32LittleEndian(header);
 					int size = BinaryPrimitives.ReadInt32LittleEndian(header.AsSpan(4));
-					_logger.LogTrace("{Endpoint}: Read {ChannelId} -> {Size} bytes", Endpoint, id, size);
+					_logger.LogTrace("Read {ChannelId} -> {Size} bytes", id, size);
 
 					// If the size if negative, we're closing the entire connection
 					if (size == 0)
 					{
-						_logger.LogTrace("{Endpoint}: Received shutdown packet", Endpoint);
+						_logger.LogTrace("Received shutdown packet");
 						break;
 					}
 
 					// Dispatch it to the correct place
 					if (size < 0)
 					{
-						_logger.LogTrace("{Endpoint}: Detaching buffer {Id}", Endpoint, id);
+						_logger.LogTrace("Detaching buffer {Id}", id);
 						DetachReceiveBuffer(cachedWriters, id);
 					}
 					else
@@ -155,7 +149,7 @@ namespace EpicGames.Horde.Compute
 				}
 			}
 
-			_logger.LogTrace("{Endpoint}: Closing reader", Endpoint);
+			_logger.LogTrace("Closing reader");
 		}
 
 		async Task<IComputeBufferWriter> GetReceiveBufferAsync(Dictionary<int, IComputeBufferWriter> cachedWriters, int id)
@@ -243,7 +237,7 @@ namespace EpicGames.Horde.Compute
 					}
 					else if (reader.IsComplete)
 					{
-						_logger.LogTrace("{Endpoint}: Sending complete packet for channel {ChannelId}", Endpoint, id);
+						_logger.LogTrace("Sending complete packet for channel {ChannelId}", id);
 						BinaryPrimitives.WriteInt32LittleEndian(header.AsSpan(4), -1);
 						await _transport.WriteAsync(header, cancellationToken);
 						break;
