Horde: Fix warning.

diff --git a/Engine/Source/Programs/Shared/EpicGames.Horde/Compute/Clients/ServerComputeClient.cs b/Engine/Source/Programs/Shared/EpicGames.Horde/Compute/Clients/ServerComputeClient.cs
--- a/Engine/Source/Programs/Shared/EpicGames.Horde/Compute/Clients/ServerComputeClient.cs
+++ b/Engine/Source/Programs/Shared/EpicGames.Horde/Compute/Clients/ServerComputeClient.cs
@@ -133,7 +133,8 @@ namespace EpicGames.Horde.Compute.Clients
 			// Pass the rest of the call over to the handler
 			byte[] key = StringUtils.ParseHexString(responseMessage.Key);
 
-			await using ComputeLease lease = new ComputeLease(responseMessage.AssignedResources, new ComputeSocket(new TcpTransport(socket), _loggerFactory));
+			await using ComputeSocket computeSocket = new ComputeSocket(new TcpTransport(socket), _loggerFactory);
+			await using ComputeLease lease = new ComputeLease(responseMessage.AssignedResources, computeSocket);
 			return await handler(lease, cancellationToken);
 		}
 	}
