Horde: Return agent properties back to the initiator as part of a compute lease.

diff --git a/Engine/Source/Programs/Horde/Horde.Server/Compute/ComputeController.cs b/Engine/Source/Programs/Horde/Horde.Server/Compute/ComputeController.cs
--- a/Engine/Source/Programs/Horde/Horde.Server/Compute/ComputeController.cs
+++ b/Engine/Source/Programs/Horde/Horde.Server/Compute/ComputeController.cs
@@ -54,6 +54,11 @@ namespace Horde.Server.Compute
 		/// Resources assigned to this machine
 		/// </summary>
 		public Dictionary<string, int> AssignedResources { get; set; } = new Dictionary<string, int>();
+
+		/// <summary>
+		/// Properties of the agent assigned to do the work
+		/// </summary>
+		public IReadOnlyList<string> Properties { get; set; } = new List<string>();
 	}
 
 	/// <summary>
@@ -110,6 +115,7 @@ namespace Horde.Server.Compute
 			response.Port = computeResource.Port;
 			response.Nonce = StringUtils.FormatHexString(computeResource.Task.Nonce.Span);
 			response.Key = StringUtils.FormatHexString(computeResource.Task.Key.Span);
+			response.Properties = computeResource.Properties;
 
 			foreach (KeyValuePair<string, int> pair in computeResource.Task.Resources)
 			{
diff --git a/Engine/Source/Programs/Horde/Horde.Server/Compute/ComputeTaskSource.cs b/Engine/Source/Programs/Horde/Horde.Server/Compute/ComputeTaskSource.cs
--- a/Engine/Source/Programs/Horde/Horde.Server/Compute/ComputeTaskSource.cs
+++ b/Engine/Source/Programs/Horde/Horde.Server/Compute/ComputeTaskSource.cs
@@ -44,14 +44,20 @@ namespace Horde.Server.Compute
 		/// </summary>
 		public ComputeTask Task { get; }
 
+		/// <summary>
+		/// Properties of the assigned agent
+		/// </summary>
+		public IReadOnlyList<string> Properties { get; }
+
 		/// <summary>
 		/// Constructor
 		/// </summary>
-		public ComputeResource(IPAddress ip, int port, ComputeTask task)
+		public ComputeResource(IPAddress ip, int port, ComputeTask task, IReadOnlyList<string> properties)
 		{
 			Ip = ip;
 			Port = port;
 			Task = task;
+			Properties = properties;
 		}
 	}
 
@@ -121,7 +127,7 @@ namespace Horde.Server.Compute
 							if (node.Value.Lease.TrySetResult(lease))
 							{
 								Waiter waiter = node.Value;
-								return new ComputeResource(waiter.Ip, waiter.Port, computeTask);
+								return new ComputeResource(waiter.Ip, waiter.Port, computeTask, waiter.Agent.Properties);
 							}
 						}
 					}
diff --git a/Engine/Source/Programs/Shared/EpicGames.Horde/Compute/Clients/LoopbackComputeClient.cs b/Engine/Source/Programs/Shared/EpicGames.Horde/Compute/Clients/LoopbackComputeClient.cs
--- a/Engine/Source/Programs/Shared/EpicGames.Horde/Compute/Clients/LoopbackComputeClient.cs
+++ b/Engine/Source/Programs/Shared/EpicGames.Horde/Compute/Clients/LoopbackComputeClient.cs
@@ -69,7 +69,7 @@ namespace EpicGames.Horde.Compute.Clients
 		public async Task<TResult> ExecuteAsync<TResult>(ClusterId clusterId, Requirements? requirements, Func<IComputeLease, CancellationToken, Task<TResult>> handler, CancellationToken cancellationToken)
 		{
 			await using ComputeSocket socket = new ComputeSocket(new TcpTransport(_socket), _loggerFactory);
-			await using ComputeLease lease = new ComputeLease(new Dictionary<string, int>(), socket);
+			await using ComputeLease lease = new ComputeLease(new List<string>(), new Dictionary<string, int>(), socket);
 			TResult result = await handler(lease, cancellationToken);
 			await socket.CloseAsync(cancellationToken);
 			return result;
diff --git a/Engine/Source/Programs/Shared/EpicGames.Horde/Compute/Clients/ServerComputeClient.cs b/Engine/Source/Programs/Shared/EpicGames.Horde/Compute/Clients/ServerComputeClient.cs
--- a/Engine/Source/Programs/Shared/EpicGames.Horde/Compute/Clients/ServerComputeClient.cs
+++ b/Engine/Source/Programs/Shared/EpicGames.Horde/Compute/Clients/ServerComputeClient.cs
@@ -36,6 +36,7 @@ namespace EpicGames.Horde.Compute.Clients
 			public string Nonce { get; set; } = String.Empty;
 			public string Key { get; set; } = String.Empty;
 			public Dictionary<string, int> AssignedResources { get; set; } = new Dictionary<string, int>();
+			public List<string> Properties { get; set; } = new List<string>();
 		}
 
 		readonly HttpClient? _defaultHttpClient;
@@ -134,7 +135,7 @@ namespace EpicGames.Horde.Compute.Clients
 			byte[] key = StringUtils.ParseHexString(responseMessage.Key);
 
 			await using ComputeSocket computeSocket = new ComputeSocket(new TcpTransport(socket), _loggerFactory);
-			await using ComputeLease lease = new ComputeLease(responseMessage.AssignedResources, computeSocket);
+			await using ComputeLease lease = new ComputeLease(responseMessage.Properties, responseMessage.AssignedResources, computeSocket);
 			return await handler(lease, cancellationToken);
 		}
 	}
diff --git a/Engine/Source/Programs/Shared/EpicGames.Horde/Compute/ComputeLease.cs b/Engine/Source/Programs/Shared/EpicGames.Horde/Compute/ComputeLease.cs
--- a/Engine/Source/Programs/Shared/EpicGames.Horde/Compute/ComputeLease.cs
+++ b/Engine/Source/Programs/Shared/EpicGames.Horde/Compute/ComputeLease.cs
@@ -11,6 +11,9 @@ namespace EpicGames.Horde.Compute
 	/// </summary>
 	public sealed class ComputeLease : IComputeLease, IAsyncDisposable
 	{
+		/// <inheritdoc/>
+		public IReadOnlyList<string> Properties { get; }
+
 		/// <inheritdoc/>
 		public IReadOnlyDictionary<string, int> AssignedResources { get; }
 
@@ -20,10 +23,12 @@ namespace EpicGames.Horde.Compute
 		/// <summary>
 		/// Constructor
 		/// </summary>
+		/// <param name="properties">Properties of the assigned agent</param>
 		/// <param name="assignedResources">Resources assigned to this lease</param>
 		/// <param name="socket">Socket for communication</param>
-		public ComputeLease(IReadOnlyDictionary<string, int> assignedResources, IComputeSocket socket)
+		public ComputeLease(IReadOnlyList<string> properties, IReadOnlyDictionary<string, int> assignedResources, IComputeSocket socket)
 		{
+			Properties = properties;
 			AssignedResources = assignedResources;
 			Socket = socket;
 		}
diff --git a/Engine/Source/Programs/Shared/EpicGames.Horde/Compute/IComputeLease.cs b/Engine/Source/Programs/Shared/EpicGames.Horde/Compute/IComputeLease.cs
--- a/Engine/Source/Programs/Shared/EpicGames.Horde/Compute/IComputeLease.cs
+++ b/Engine/Source/Programs/Shared/EpicGames.Horde/Compute/IComputeLease.cs
@@ -10,6 +10,11 @@ namespace EpicGames.Horde.Compute
 	/// </summary>
 	public interface IComputeLease : IAsyncDisposable
 	{
+		/// <summary>
+		/// Properties of the remote machine
+		/// </summary>
+		IReadOnlyList<string> Properties { get; }
+
 		/// <summary>
 		/// Resources assigned to this lease
 		/// </summary>
