EpicGames.Core: Allow adding arbitrary processes to ManagedProcessGroup objects (eg. processes created with regular C# Process constructs).

diff --git a/Engine/Source/Programs/Shared/EpicGames.Core/ManagedProcess.cs b/Engine/Source/Programs/Shared/EpicGames.Core/ManagedProcess.cs
--- a/Engine/Source/Programs/Shared/EpicGames.Core/ManagedProcess.cs
+++ b/Engine/Source/Programs/Shared/EpicGames.Core/ManagedProcess.cs
@@ -103,6 +103,15 @@ namespace EpicGames.Core
 
 		[DllImport("kernel32.dll", SetLastError = true)]
 		static extern bool QueryInformationJobObject(SafeFileHandle hJob, int JobObjectInformationClass, ref JOBOBJECT_BASIC_ACCOUNTING_INFORMATION lpJobObjectInformation, int cbJobObjectInformationLength);
+
+		[DllImport("kernel32.dll", SetLastError = true)]
+		static extern int AssignProcessToJobObject(SafeFileHandle? hJob, IntPtr hProcess);
+
+		[DllImport("kernel32.dll")]
+		static extern bool IsProcessInJob(IntPtr hProcess, IntPtr hJob, out bool result);
+
+		[DllImport("kernel32.dll", SetLastError = true)]
+		static extern IntPtr GetCurrentProcess();
 #pragma warning restore IDE1006 // Naming Styles
 #pragma warning restore IDE0049 // Naming Styles
 
@@ -154,6 +163,33 @@ namespace EpicGames.Core
 			}
 		}
 
+		/// <summary>
+		/// Adds a process to this group
+		/// </summary>
+		/// <param name="process">Process to add</param>
+		public void AddProcess(Process process) => AddProcess(process.Handle);
+
+		/// <summary>
+		/// Adds a process to this group
+		/// </summary>
+		/// <param name="processHandle">Handle to the process to add</param>
+		public void AddProcess(IntPtr processHandle)
+		{
+			if(SupportsJobObjects && AssignProcessToJobObject(JobHandle, processHandle) == 0)
+			{
+				// Support for nested job objects was only added in Windows 8; prior to that, assigning processes to job objects would fail. Figure out if we're already in a job, and ignore the error if we are.
+				int originalError = Marshal.GetLastWin32Error();
+
+				bool bProcessInJob;
+				IsProcessInJob(GetCurrentProcess(), IntPtr.Zero, out bProcessInJob);
+
+				if (!bProcessInJob)
+				{
+					throw new Win32ExceptionWithCode(originalError, "Unable to assign process to job object");
+				}
+			}
+		}
+
 		/// <summary>
 		/// Returns the total CPU time usage for this job.
 		/// </summary>
@@ -240,9 +276,6 @@ namespace EpicGames.Core
 			public SafeHandle? hStdError;
 		}
 
-		[DllImport("kernel32.dll", SetLastError = true)]
-		static extern int AssignProcessToJobObject(SafeFileHandle? hJob, IntPtr hProcess);
-
 		[DllImport("kernel32.dll", SetLastError = true)]
 		static extern int CloseHandle(IntPtr hObject);
 
@@ -296,9 +329,6 @@ namespace EpicGames.Core
 		[DllImport("kernel32.dll", SetLastError = true)]
 		static extern int GetExitCodeProcess(SafeFileHandle hProcess, out int lpExitCode);
 
-		[DllImport("kernel32.dll")]
-		static extern bool IsProcessInJob(IntPtr hProcess, IntPtr hJob, out bool result);
-
 		[DllImport("kernel32.dll", SetLastError = true)]
 		static extern int TerminateProcess(SafeHandleZeroOrMinusOneIsInvalid hProcess, uint uExitCode);
 
@@ -727,26 +757,11 @@ namespace EpicGames.Core
 					}
 
 					// Add it to our job object
-					if (group != null && AssignProcessToJobObject(group.JobHandle, processInfo.hProcess) == 0)
-					{
-						// Support for nested job objects was only added in Windows 8; prior to that, assigning processes to job objects would fail. Figure out if we're already in a job, and ignore the error if we are.
-						int originalError = Marshal.GetLastWin32Error();
-
-						bool bProcessInJob;
-						IsProcessInJob(GetCurrentProcess(), IntPtr.Zero, out bProcessInJob);
-
-						if (!bProcessInJob)
-						{
-							throw new Win32ExceptionWithCode(originalError, "Unable to assign process to job object");
-						}
-					}
+					group?.AddProcess(processInfo.hProcess);
 
 					// Create a JobObject for each spawned process to do CPU usage accounting of spawned process and all its children
 					_accountingProcessGroup = new ManagedProcessGroup(bKillOnJobClose: false);
-					if (AssignProcessToJobObject(_accountingProcessGroup.JobHandle, processInfo.hProcess) == 0) 
-					{
-						throw new Win32Exception();
-					}
+					_accountingProcessGroup.AddProcess(processInfo.hProcess);
 
 					// On systems with more than one processor group (more than one CPU socket, or more than 64 cores), it is possible
 					// that processes launched from here may be scheduled by the operating system in a way that impedes overall throughput.
