Minor performance improvement in C# UHT file reading by limiting the buffer size to reduce GC stress.

diff --git a/Engine/Source/Programs/Shared/EpicGames.UHT/Utils/UhtFileManager.cs b/Engine/Source/Programs/Shared/EpicGames.UHT/Utils/UhtFileManager.cs
--- a/Engine/Source/Programs/Shared/EpicGames.UHT/Utils/UhtFileManager.cs
+++ b/Engine/Source/Programs/Shared/EpicGames.UHT/Utils/UhtFileManager.cs
@@ -104,7 +104,7 @@ namespace EpicGames.UHT.Utils
 
 			try
 			{
-				using FileStream fs = new(filePath, FileMode.Open, FileAccess.Read, FileShare.Read, 4 * 1024, FileOptions.SequentialScan);
+				using FileStream fs = new(filePath, FileMode.Open, FileAccess.Read, FileShare.Read, 1, FileOptions.SequentialScan);
 				using UhtBorrowByteBuffer byteBuffer = new((int)fs.Length);
 				Span<byte> byteSpan = byteBuffer.Buffer.Memory.Span;
 				int readLength = fs.Read(byteSpan);
@@ -136,7 +136,7 @@ namespace EpicGames.UHT.Utils
 
 			try
 			{
-				using FileStream fs = new(filePath, FileMode.Open, FileAccess.Read, FileShare.Read, 4 * 1024, FileOptions.SequentialScan);
+				using FileStream fs = new(filePath, FileMode.Open, FileAccess.Read, FileShare.Read, 1, FileOptions.SequentialScan);
 				using UhtRentedPoolBuffer<byte> byteBuffer = new((int)fs.Length);
 				Span<byte> byteSpan = byteBuffer.Buffer.Memory.Span;
 				int readLength = fs.Read(byteSpan);
@@ -208,7 +208,7 @@ namespace EpicGames.UHT.Utils
 
 			try
 			{
-				using FileStream fs = new(filePath, FileMode.Open, FileAccess.Read, FileShare.Read, 4 * 1024, FileOptions.SequentialScan);
+				using FileStream fs = new(filePath, FileMode.Open, FileAccess.Read, FileShare.Read, 1, FileOptions.SequentialScan);
 				using UhtRentedPoolBuffer<byte> byteBuffer = new((int)fs.Length);
 				Span<byte> byteSpan = byteBuffer.Buffer.Memory.Span;
 				int readLength = fs.Read(byteSpan);
