Horde: Various remote execution improvements.

* Add a sample app which uploads an application, streams values to it, and prints the output it echoes back to us.
* Add an in-memory transport which uses a C# pipeline to buffer data.
* Add a client implementation that launches a local Horde agent and connects directly to it.
* Add an IPC message channel to allow attaching new shared memory buffers from remote applications.
* Simplify the API for establishing connections.

diff --git a/Engine/Source/Programs/Horde/Horde.Agent/Commands/Compute/ComputeCommand.cs b/Engine/Source/Programs/Horde/Horde.Agent/Commands/Compute/ComputeCommand.cs
--- a/Engine/Source/Programs/Horde/Horde.Agent/Commands/Compute/ComputeCommand.cs
+++ b/Engine/Source/Programs/Horde/Horde.Agent/Commands/Compute/ComputeCommand.cs
@@ -10,7 +10,6 @@ using System.Net.Sockets;
 using System.Reflection;
 using System.Threading;
 using System.Threading.Tasks;
-using Amazon.EC2.Model;
 using EpicGames.Core;
 using EpicGames.Horde.Common;
 using EpicGames.Horde.Compute;
@@ -20,7 +19,6 @@ using Horde.Agent.Leases.Handlers;
 using Microsoft.Extensions.DependencyInjection;
 using Microsoft.Extensions.Logging;
 using Microsoft.Extensions.Options;
-using Polly;
 
 namespace Horde.Agent.Commands.Compute
 {
@@ -52,7 +50,7 @@ namespace Horde.Agent.Commands.Compute
 		/// <inheritdoc/>
 		public override async Task<int> ExecuteAsync(ILogger logger)
 		{
-			await using IComputeClient client = CreateClient();
+			await using IComputeClient client = CreateClient(_serviceProvider.GetRequiredService<ILogger<IComputeClient>>());
 
 			Requirements? requirements = null;
 			if (Requirements != null)
@@ -73,20 +71,19 @@ namespace Horde.Agent.Commands.Compute
 		/// <returns></returns>
 		protected abstract Task<bool> HandleRequestAsync(IComputeLease lease, CancellationToken cancellationToken);
 
-		IComputeClient CreateClient()
+		IComputeClient CreateClient(ILogger logger)
 		{
-			ILoggerFactory loggerFactory = _serviceProvider.GetRequiredService<ILoggerFactory>();
 			if (Local)
 			{
-				return new LocalComputeClient((socket, ctx) => ComputeWorkerCommand.RunWorkerAsync(socket, _serviceProvider, ctx), 2000, loggerFactory);
+				return new LocalComputeClient((socket, ctx) => ComputeWorkerCommand.RunWorkerAsync(socket, _serviceProvider, ctx), 2000, logger);
 			}
 			else if (Loopback)
 			{
-				return new LoopbackComputeClient(2000, loggerFactory);
+				return new LoopbackComputeClient(2000, logger);
 			}
 			else
 			{
-				return new ServerComputeClient(CreateHttpClient, loggerFactory);
+				return new ServerComputeClient(CreateHttpClient, logger);
 			}
 		}
 
@@ -103,20 +100,19 @@ namespace Horde.Agent.Commands.Compute
 		class LoopbackComputeClient : IComputeClient
 		{
 			readonly int _port;
-			readonly ILoggerFactory _loggerFactory;
+			readonly ILogger _logger;
 
-			public LoopbackComputeClient(int port, ILoggerFactory loggerFactory)
+			public LoopbackComputeClient(int port, ILogger logger)
 			{
 				_port = port;
-				_loggerFactory = loggerFactory;
+				_logger = logger;
 			}
 
 			public ValueTask DisposeAsync() => new ValueTask();
 
 			public async Task<TResult> ExecuteAsync<TResult>(ClusterId clusterId, Requirements? requirements, Func<IComputeLease, CancellationToken, Task<TResult>> handler, CancellationToken cancellationToken)
 			{
-				ILogger logger = _loggerFactory.CreateLogger<LoopbackComputeClient>();
-				logger.LogInformation("** INITIATOR **");
+				_logger.LogInformation("** CLIENT **");
 
 				TResult result;
 				using (Socket listener = new Socket(SocketType.Stream, ProtocolType.IP))
@@ -158,7 +154,7 @@ namespace Horde.Agent.Commands.Compute
 				TResult result;
 				using (Socket tcpSocket = await listener.AcceptAsync(cancellationToken))
 				{
-					await using (ComputeSocket socket = new ComputeSocket(new TcpTransport(tcpSocket), _loggerFactory))
+					await using (ClientComputeSocket socket = new ClientComputeSocket(new TcpTransport(tcpSocket), _logger))
 					{
 						await using ComputeLease lease = new ComputeLease(new List<string>(), new Dictionary<string, int>(), socket);
 						result = await handler(lease, cancellationToken);
@@ -190,12 +186,10 @@ namespace Horde.Agent.Commands.Compute
 		{
 			logger.LogInformation("** WORKER **");
 
-			ILoggerFactory loggerFactory = _serviceProvider.GetRequiredService<ILoggerFactory>();
-
 			using Socket tcpSocket = new Socket(SocketType.Stream, ProtocolType.IP);
 			await tcpSocket.ConnectAsync(IPAddress.Loopback, Port);
 
-			await using (ComputeSocket socket = new ComputeSocket(new TcpTransport(tcpSocket), loggerFactory))
+			await using (ClientComputeSocket socket = new ClientComputeSocket(new TcpTransport(tcpSocket), logger))
 			{
 				logger.LogInformation("Running worker...");
 				await RunWorkerAsync(socket, _serviceProvider, CancellationToken.None);
diff --git a/Engine/Source/Programs/Horde/Horde.Agent/Commands/Compute/CppComputeCommand.cs b/Engine/Source/Programs/Horde/Horde.Agent/Commands/Compute/CppComputeCommand.cs
--- a/Engine/Source/Programs/Horde/Horde.Agent/Commands/Compute/CppComputeCommand.cs
+++ b/Engine/Source/Programs/Horde/Horde.Agent/Commands/Compute/CppComputeCommand.cs
@@ -56,14 +56,14 @@ namespace Horde.Agent.Commands.Compute
 			NodeLocator sandbox = await CreateSandboxAsync(TaskFile, storage, cancellationToken);
 
 			// Open a socket and upload the sandbox
-			await using (IComputeChannel channel = lease.Socket.AttachMessageChannel(0))
+			await using (IComputeMessageChannel channel = lease.Socket.CreateMessageChannel(0, 30 * 1024 * 1024, _logger))
 			{
 				await channel.UploadFilesAsync("", sandbox, storage, cancellationToken);
 
 				// Execute the program
 				using (LogParser parser = new LogParser(_logger, new List<string>()))
 				{
-					await channel.ExecuteProcessAsync(jsonComputeTask.Executable, jsonComputeTask.Arguments, jsonComputeTask.WorkingDir, jsonComputeTask.EnvVars, parser.WriteData, cancellationToken);
+					await channel.ExecuteAsync(jsonComputeTask.Executable, jsonComputeTask.Arguments, jsonComputeTask.WorkingDir, jsonComputeTask.EnvVars, parser.WriteData, cancellationToken);
 				}
 			}
 
diff --git a/Engine/Source/Programs/Horde/Horde.Agent/Commands/Compute/DdcComputeCommand.cs b/Engine/Source/Programs/Horde/Horde.Agent/Commands/Compute/DdcComputeCommand.cs
--- a/Engine/Source/Programs/Horde/Horde.Agent/Commands/Compute/DdcComputeCommand.cs
+++ b/Engine/Source/Programs/Horde/Horde.Agent/Commands/Compute/DdcComputeCommand.cs
@@ -31,10 +31,10 @@ namespace Horde.Agent.Commands.Compute
 		protected override async Task<bool> HandleRequestAsync(IComputeLease lease, CancellationToken cancellationToken)
 		{
 			IComputeSocket socket = lease.Socket;
-			await using (IComputeChannel channel = socket.AttachMessageChannel(0))
+			await using (IComputeMessageChannel channel = socket.CreateMessageChannel(0, _logger))
 			{
 				_logger.LogInformation("Sending XOR request");
-				channel.XorRequest(new byte[] { 1, 2, 3, 4, 5 }, (byte)123);
+				channel.SendXorRequest(new byte[] { 1, 2, 3, 4, 5 }, (byte)123);
 
 				_logger.LogInformation("Waiting for response...");
 				IComputeMessage response = await channel.ReceiveAsync(cancellationToken);
diff --git a/Engine/Source/Programs/Horde/Horde.Agent/Commands/Compute/XorComputeCommand.cs b/Engine/Source/Programs/Horde/Horde.Agent/Commands/Compute/XorComputeCommand.cs
--- a/Engine/Source/Programs/Horde/Horde.Agent/Commands/Compute/XorComputeCommand.cs
+++ b/Engine/Source/Programs/Horde/Horde.Agent/Commands/Compute/XorComputeCommand.cs
@@ -31,15 +31,14 @@ namespace Horde.Agent.Commands.Compute
 		protected override async Task<bool> HandleRequestAsync(IComputeLease lease, CancellationToken cancellationToken)
 		{
 			IComputeSocket socket = lease.Socket;
-			await using (IComputeChannel mainChannel = socket.AttachMessageChannel(0))
+			await using (IComputeMessageChannel mainChannel = socket.CreateMessageChannel(0, _logger))
 			{
 				_logger.LogInformation("Forking compute channel...");
 
-				await using IComputeChannel channel = socket.AttachMessageChannel(1);
-				mainChannel.Fork(1);
+				IComputeMessageChannel channel = mainChannel;
 
 				_logger.LogInformation("Sending XOR request");
-				channel.XorRequest(new byte[] { 1, 2, 3, 4, 5 }, (byte)123);
+				channel.SendXorRequest(new byte[] { 1, 2, 3, 4, 5 }, (byte)123);
 
 				_logger.LogInformation("Waiting for response...");
 				IComputeMessage response = await channel.ReceiveAsync(cancellationToken);
diff --git a/Engine/Source/Programs/Horde/Horde.Agent/Leases/Handlers/ComputeHandler.cs b/Engine/Source/Programs/Horde/Horde.Agent/Leases/Handlers/ComputeHandler.cs
--- a/Engine/Source/Programs/Horde/Horde.Agent/Leases/Handlers/ComputeHandler.cs
+++ b/Engine/Source/Programs/Horde/Horde.Agent/Leases/Handlers/ComputeHandler.cs
@@ -29,19 +29,17 @@ namespace Horde.Agent.Leases.Handlers
 		readonly ComputeListenerService _listenerService;
 		readonly IMemoryCache _memoryCache;
 		readonly DirectoryReference _sandboxDir;
-		readonly ILoggerFactory _loggerFactory;
 		readonly ILogger _logger;
 
 		/// <summary>
 		/// Constructor
 		/// </summary>
-		public ComputeHandler(ComputeListenerService listenerService, IMemoryCache memoryCache, ILoggerFactory loggerFactory)
+		public ComputeHandler(ComputeListenerService listenerService, IMemoryCache memoryCache, ILogger<ComputeHandler> logger)
 		{
 			_listenerService = listenerService;
 			_memoryCache = memoryCache;
 			_sandboxDir = DirectoryReference.Combine(Program.DataDir, "Sandbox");
-			_loggerFactory = loggerFactory;
-			_logger = loggerFactory.CreateLogger<ComputeHandler>();
+			_logger = logger;
 		}
 
 		/// <inheritdoc/>
@@ -61,7 +59,7 @@ namespace Horde.Agent.Leases.Handlers
 					return LeaseResult.Success;
 				}
 
-				await using (ComputeSocket socket = new ComputeSocket(new TcpTransport(tcpClient.Client), _loggerFactory))
+				await using (ClientComputeSocket socket = new ClientComputeSocket(new TcpTransport(tcpClient.Client), _logger))
 				{
 					await RunAsync(socket, cancellationToken);
 					await socket.CloseAsync(cancellationToken);
@@ -86,7 +84,7 @@ namespace Horde.Agent.Leases.Handlers
 
 		async Task RunAsync(IComputeSocket socket, int channelId, CancellationToken cancellationToken)
 		{
-			await using (IComputeChannel channel = socket.AttachMessageChannel(channelId))
+			await using (IComputeMessageChannel channel = socket.CreateMessageChannel(channelId, 30 * 1024 * 1024, _logger))
 			{
 				List<Task> childTasks = new List<Task>();
 				for (; ; )
@@ -117,7 +115,7 @@ namespace Horde.Agent.Leases.Handlers
 								DeleteFiles(deleteFiles.Filter);
 							}
 							break;
-						case ComputeMessageType.ExecuteProcess:
+						case ComputeMessageType.Execute:
 							{
 								ExecuteProcessMessage executeProcess = message.ParseExecuteProcessMessage();
 								await ExecuteProcessAsync(socket, channel, executeProcess.Executable, executeProcess.Arguments, executeProcess.WorkingDir, executeProcess.EnvVars, cancellationToken);  
@@ -136,13 +134,11 @@ namespace Horde.Agent.Leases.Handlers
 			}
 		}
 
-		static void RunXor(IComputeChannel channel, ReadOnlyMemory<byte> source, byte value)
+		static void RunXor(IComputeMessageChannel channel, ReadOnlyMemory<byte> source, byte value)
 		{
-			using (IComputeMessageBuilder writer = channel.CreateMessage(ComputeMessageType.XorResponse, source.Length))
-			{
-				XorData(source.Span, writer.GetSpanAndAdvance(source.Length), value);
-				writer.Send();
-			}
+			using IComputeMessageBuilder response = channel.CreateMessage(ComputeMessageType.XorResponse, source.Length);
+			XorData(source.Span, response.GetSpanAndAdvance(source.Length), value);
+			response.Send();
 		}
 
 		static void XorData(ReadOnlySpan<byte> source, Span<byte> target, byte value)
@@ -153,7 +149,7 @@ namespace Horde.Agent.Leases.Handlers
 			}
 		}
 
-		async Task WriteFilesAsync(IComputeChannel channel, string path, NodeLocator locator, CancellationToken cancellationToken)
+		async Task WriteFilesAsync(IComputeMessageChannel channel, string path, NodeLocator locator, CancellationToken cancellationToken)
 		{
 			using ComputeStorageClient store = new ComputeStorageClient(channel);
 			TreeReader reader = new TreeReader(store, _memoryCache, _logger);
@@ -168,9 +164,9 @@ namespace Horde.Agent.Leases.Handlers
 
 			await directoryNode.CopyToDirectoryAsync(reader, outputDir.ToDirectoryInfo(), _logger, cancellationToken);
 
-			using (IComputeMessageBuilder builder = channel.CreateMessage(ComputeMessageType.WriteFilesResponse))
+			using (IComputeMessageBuilder message = channel.CreateMessage(ComputeMessageType.WriteFilesResponse))
 			{
-				builder.Send();
+				message.Send();
 			}
 		}
 
@@ -185,7 +181,7 @@ namespace Horde.Agent.Leases.Handlers
 			}
 		}
 
-		async Task ExecuteProcessAsync(IComputeSocket socket, IComputeChannel channel, string executable, IReadOnlyList<string> arguments, string? workingDir, IReadOnlyDictionary<string, string?>? envVars, CancellationToken cancellationToken)
+		async Task ExecuteProcessAsync(IComputeSocket socket, IComputeMessageChannel channel, string executable, IReadOnlyList<string> arguments, string? workingDir, IReadOnlyDictionary<string, string?>? envVars, CancellationToken cancellationToken)
 		{
 			if (RuntimeInformation.IsOSPlatform(OSPlatform.Windows))
 			{
@@ -197,10 +193,9 @@ namespace Horde.Agent.Leases.Handlers
 			}
 		}
 
-		async Task ExecuteProcessWindowsAsync(IComputeSocket socket, IComputeChannel channel, string executable, IReadOnlyList<string> arguments, string? workingDir, IReadOnlyDictionary<string, string?>? envVars, CancellationToken cancellationToken)
+		async Task ExecuteProcessWindowsAsync(IComputeSocket socket, IComputeMessageChannel channel, string executable, IReadOnlyList<string> arguments, string? workingDir, IReadOnlyDictionary<string, string?>? envVars, CancellationToken cancellationToken)
 		{
-			using IpcBuffer sendBuffer = IpcBuffer.CreateNew(16384); // Messages that we send to child process
-			using IpcBuffer receiveBuffer = IpcBuffer.CreateNew(16384); // Messages that child process sends to us
+			await using IpcComputeMessageChannel ipcChannel = new IpcComputeMessageChannel(4096, _logger);
 
 			Dictionary<string, string?> newEnvVars = new Dictionary<string, string?>();
 			if (envVars != null)
@@ -211,40 +206,40 @@ namespace Horde.Agent.Leases.Handlers
 				}
 			}
 
-			newEnvVars[ComputeEnvVarNames.DefaultSendBuffer] = receiveBuffer.GetHandle(); // Note: opposite ends that we're sending/receiving to
-			newEnvVars[ComputeEnvVarNames.DefaultRecvBuffer] = sendBuffer.GetHandle();
+			newEnvVars[ComputeSocket.WorkerIpcEnvVar] = ipcChannel.GetStringHandle();
 
-			await using ComputeChannel controlChannel = new ComputeChannel(receiveBuffer.Reader, sendBuffer.Writer, _loggerFactory.CreateLogger<ComputeChannel>());
+			Task processTask = ExecuteProcessInternalAsync(channel, executable, arguments, workingDir, newEnvVars, cancellationToken);
+			processTask = processTask.ContinueWith(x => ipcChannel.ForceComplete(), cancellationToken, TaskContinuationOptions.None, TaskScheduler.Default);
 
-			Task processTask = ExecuteProcessInternalAsync(channel, executable, arguments, workingDir, envVars, cancellationToken);
-			processTask = processTask.ContinueWith(x => receiveBuffer.Writer.MarkComplete(), cancellationToken, TaskContinuationOptions.None, TaskScheduler.Default);
-
-			List<IpcBuffer> buffers = new List<IpcBuffer>();
+			List<IDisposable> buffers = new List<IDisposable>();
 			try
 			{
 				for (; ; )
 				{
-					using IComputeMessage message = await controlChannel.ReceiveAsync(cancellationToken);
+					using IComputeMessage message = await ipcChannel.ReceiveAsync(cancellationToken);
 					switch (message.Type)
 					{
 						case ComputeMessageType.None:
 							await processTask;
 							return;
-						case ComputeMessageType.CreateBufferRequest:
+						case ComputeMessageType.AttachRecvBuffer:
 							{
-								CreateBufferRequest createBuffer = message.ParseCreateBufferRequest();
-
-								IpcBuffer buffer = IpcBuffer.CreateNew(createBuffer.Capacity);
-								buffers.Add(buffer);
-
-								if (createBuffer.Send)
-								{
-									socket.AttachSendBuffer(createBuffer.ChannelId, buffer.Reader);
-								}
-								else
-								{
-									socket.AttachReceiveBuffer(createBuffer.ChannelId, buffer.Writer);
-								}
+								AttachRecvBufferRequest attachRecvBuffer = IpcComputeMessageChannel.ParseAttachRecvBuffer(message);
+#pragma warning disable CA2000 // Dispose objects before losing scope
+								(IComputeBufferReader reader, IComputeBufferWriter writer) = SharedMemoryBuffer.OpenIpcHandle(attachRecvBuffer.Handle).ToShared();
+								reader.Dispose();
+								socket.AttachRecvBuffer(attachRecvBuffer.ChannelId, writer);
+#pragma warning restore CA2000 // Dispose objects before losing scope
+							}
+							break;
+						case ComputeMessageType.AttachSendBuffer:
+							{
+								AttachSendBufferRequest attachSendBuffer = IpcComputeMessageChannel.ParseAttachSendBuffer(message);
+#pragma warning disable CA2000 // Dispose objects before losing scope
+								(IComputeBufferReader reader, IComputeBufferWriter writer) = SharedMemoryBuffer.OpenIpcHandle(attachSendBuffer.Handle).ToShared();
+								writer.Dispose();
+								socket.AttachSendBuffer(attachSendBuffer.ChannelId, reader);
+#pragma warning restore CA2000 // Dispose objects before losing scope
 							}
 							break;
 						default:
@@ -254,14 +249,14 @@ namespace Horde.Agent.Leases.Handlers
 			}
 			finally
 			{
-				foreach (IpcBuffer buffer in buffers)
+				foreach (SharedMemoryBuffer buffer in buffers)
 				{
 					buffer.Dispose();
 				}
 			}
 		}
 
-		async Task ExecuteProcessInternalAsync(IComputeChannel channel, string executable, IReadOnlyList<string> arguments, string? workingDir, IReadOnlyDictionary<string, string?>? envVars, CancellationToken cancellationToken)
+		async Task ExecuteProcessInternalAsync(IComputeMessageChannel channel, string executable, IReadOnlyList<string> arguments, string? workingDir, IReadOnlyDictionary<string, string?>? envVars, CancellationToken cancellationToken)
 		{
 			string resolvedExecutable = FileReference.Combine(_sandboxDir, executable).FullName;
 			string resolvedCommandLine = CommandLineArguments.Join(arguments);
@@ -293,10 +288,10 @@ namespace Horde.Agent.Leases.Handlers
 						int length = await process.ReadAsync(buffer, 0, buffer.Length, cancellationToken);
 						if (length == 0)
 						{
-							channel.SendExecuteProcessResponse(process.ExitCode);
+							channel.SendExecuteResult(process.ExitCode);
 							return;
 						}
-						channel.SendProcessOutput(buffer.AsMemory(0, length));
+						channel.SendExecuteOutput(buffer.AsMemory(0, length));
 					}
 				}
 			}
diff --git a/Engine/Source/Programs/Horde/Samples/RemoteClient/RemoteClient.csproj b/Engine/Source/Programs/Horde/Samples/RemoteClient/RemoteClient.csproj
new file mode 100644
--- /dev/null
+++ b/Engine/Source/Programs/Horde/Samples/RemoteClient/RemoteClient.csproj
@@ -0,0 +1,14 @@
+<Project Sdk="Microsoft.NET.Sdk">
+
+  <PropertyGroup>
+    <OutputType>Exe</OutputType>
+    <TargetFramework>net6.0</TargetFramework>
+    <ImplicitUsings>enable</ImplicitUsings>
+    <Nullable>enable</Nullable>
+  </PropertyGroup>
+
+  <ItemGroup>
+    <ProjectReference Include="..\..\..\Shared\EpicGames.Horde\EpicGames.Horde.csproj" />
+  </ItemGroup>
+
+</Project>
diff --git a/Engine/Source/Programs/Horde/Samples/RemoteClient/RemoteClientApp.cs b/Engine/Source/Programs/Horde/Samples/RemoteClient/RemoteClientApp.cs
new file mode 100644
--- /dev/null
+++ b/Engine/Source/Programs/Horde/Samples/RemoteClient/RemoteClientApp.cs
@@ -0,0 +1,79 @@
+// Copyright Epic Games, Inc. All Rights Reserved.
+
+using System.Reflection;
+using EpicGames.Core;
+using EpicGames.Horde.Compute;
+using EpicGames.Horde.Compute.Clients;
+using EpicGames.Horde.Storage;
+using EpicGames.Horde.Storage.Backends;
+using EpicGames.Horde.Storage.Nodes;
+using Microsoft.Extensions.Logging;
+
+namespace RemoteClient
+{
+	class ClientApp
+	{
+		static FileReference CurrentAssembly { get; } = new FileReference(Assembly.GetExecutingAssembly().Location);
+		static DirectoryReference AssemblyDir { get; } = CurrentAssembly.Directory;
+		static DirectoryReference ProjectDir { get; } = DirectoryReference.Combine(AssemblyDir, "../../..");
+		static string ConfigPath { get; } = AssemblyDir.MakeRelativeTo(ProjectDir);
+
+		static FileReference HordeAgentFile { get; } = FileReference.Combine(ProjectDir, "../../Horde.Agent", ConfigPath, "HordeAgent.dll");
+		static FileReference RemoteServerFile { get; } = FileReference.Combine(ProjectDir, "../RemoteServer", ConfigPath, "RemoteServer.exe");
+
+		static async Task Main()
+		{
+			ILogger logger = Log.Logger;
+
+			await using IComputeClient client = CreateClient(true, logger);
+			await client.ExecuteAsync(new ClusterId("default"), null, (lease, ctx) => ExecuteAsync(lease, logger, ctx), CancellationToken.None);
+		}
+
+		static IComputeClient CreateClient(bool loopback, ILogger logger)
+		{
+			if (loopback)
+			{
+				return new AgentComputeClient(HordeAgentFile.FullName, 2000, logger);
+			}
+			else
+			{
+				return new ServerComputeClient(new Uri("https://localhost:5001"), null, logger);
+			}
+		}
+
+		static async Task<int> ExecuteAsync(IComputeLease lease, ILogger logger, CancellationToken cancellationToken)
+		{
+			IComputeSocket socket = lease.Socket;
+			await using (IComputeMessageChannel channel = socket.CreateMessageChannel(0, 30 * 1024 * 1024, logger))
+			{
+				// Upload the sandbox
+				MemoryStorageClient storage = new MemoryStorageClient();
+				using (TreeWriter treeWriter = new TreeWriter(storage))
+				{
+					DirectoryNode sandbox = new DirectoryNode();
+					await sandbox.CopyFromDirectoryAsync(RemoteServerFile.Directory.ToDirectoryInfo(), new ChunkingOptions(), treeWriter, null, cancellationToken);
+					NodeHandle handle = await treeWriter.FlushAsync(sandbox, cancellationToken);
+					await channel.UploadFilesAsync("", handle.Locator, storage, cancellationToken);
+				}
+
+				// Run the task remotely in the background and echo the output to the console
+				Task childProcessTask = channel.ExecuteAsync(RemoteServerFile.GetFileName(), new List<string>(), null, null, (string x) => logger.LogInformation("ChildProcess: {Message}", x), cancellationToken);
+
+				// Generate data into a buffer attached to channel 1. The remote server will echo them back to us as it receives them.
+				IComputeBufferWriter writer = socket.AttachSendBuffer(1, 1024 * 1024);
+				for (int idx = 0; idx < 100; idx++)
+				{
+					logger.LogInformation("Writing value: {Value}", idx);
+					writer.GetMemory().Span[0] = (byte)idx;
+					writer.Advance(1);
+					await Task.Delay(1000, cancellationToken);
+				}
+				writer.MarkComplete();
+
+				// Wait for the child process to finish
+				await childProcessTask;
+			}
+			return 0;
+		}
+	}
+}
\ No newline at end of file
diff --git a/Engine/Source/Programs/Horde/Samples/RemoteServer/RemoteServer.csproj b/Engine/Source/Programs/Horde/Samples/RemoteServer/RemoteServer.csproj
new file mode 100644
--- /dev/null
+++ b/Engine/Source/Programs/Horde/Samples/RemoteServer/RemoteServer.csproj
@@ -0,0 +1,14 @@
+<Project Sdk="Microsoft.NET.Sdk">
+
+  <PropertyGroup>
+    <OutputType>Exe</OutputType>
+    <TargetFramework>net6.0</TargetFramework>
+    <ImplicitUsings>enable</ImplicitUsings>
+    <Nullable>enable</Nullable>
+  </PropertyGroup>
+
+  <ItemGroup>
+    <ProjectReference Include="..\..\..\Shared\EpicGames.Horde\EpicGames.Horde.csproj" />
+  </ItemGroup>
+
+</Project>
diff --git a/Engine/Source/Programs/Horde/Samples/RemoteServer/RemoteServerApp.cs b/Engine/Source/Programs/Horde/Samples/RemoteServer/RemoteServerApp.cs
new file mode 100644
--- /dev/null
+++ b/Engine/Source/Programs/Horde/Samples/RemoteServer/RemoteServerApp.cs
@@ -0,0 +1,30 @@
+// Copyright Epic Games, Inc. All Rights Reserved.
+
+using EpicGames.Horde.Compute;
+using Microsoft.Extensions.Logging.Abstractions;
+
+namespace RemoteServer
+{
+	class ServerApp
+	{
+		static async Task Main()
+		{
+			await using IComputeSocket socket = ComputeSocket.ConnectAsWorker(NullLogger.Instance);
+			Console.WriteLine("Connected to client");
+
+			using IComputeBufferReader reader = socket.AttachRecvBuffer(1, 1024 * 1024);
+			while (!reader.IsComplete)
+			{
+				ReadOnlyMemory<byte> memory = reader.GetMemory();
+				if (memory.Length == 0)
+				{
+					await reader.WaitForDataAsync(0, CancellationToken.None);
+					continue;
+				}
+
+				Console.WriteLine("Received data: {0}", memory.Span[0]);
+				reader.Advance(1);
+			}
+		}
+	}
+}
\ No newline at end of file
diff --git a/Engine/Source/Programs/Shared/EpicGames.Horde.Tests/BufferTests.cs b/Engine/Source/Programs/Shared/EpicGames.Horde.Tests/BufferTests.cs
--- a/Engine/Source/Programs/Shared/EpicGames.Horde.Tests/BufferTests.cs
+++ b/Engine/Source/Programs/Shared/EpicGames.Horde.Tests/BufferTests.cs
@@ -2,6 +2,8 @@
 
 using System;
 using System.Diagnostics;
+using System.IO.Pipelines;
+using System.IO.Pipes;
 using System.Linq;
 using System.Runtime.InteropServices;
 using System.Security.Cryptography;
@@ -9,6 +11,9 @@ using System.Threading;
 using System.Threading.Tasks;
 using EpicGames.Horde.Compute;
 using EpicGames.Horde.Compute.Buffers;
+using EpicGames.Horde.Compute.Transports;
+using Microsoft.Extensions.Logging.Abstractions;
+using Microsoft.VisualStudio.TestPlatform.CrossPlatEngine;
 using Microsoft.VisualStudio.TestTools.UnitTesting;
 
 namespace EpicGames.Horde.Tests
@@ -17,33 +22,38 @@ namespace EpicGames.Horde.Tests
 	public class BufferTests
 	{
 		[TestMethod]
-		public async Task TestHeapBuffer()
+		public async Task TestPooledBuffer()
 		{
-			using HeapBuffer buffer = new HeapBuffer(8000);
-			await TestProducerConsumerAsync(buffer.Writer, buffer.Reader);
+			await TestProducerConsumerAsync(length => new PooledBuffer(length));
 		}
 
 		[TestMethod]
-		public async Task TestIpcBuffer()
+		public async Task TestSharedMemoryBuffer()
 		{
 			if (RuntimeInformation.IsOSPlatform(OSPlatform.Windows))
 			{
-				using IpcBuffer buffer1 = IpcBuffer.CreateNew(16384);
-				string info = buffer1.DuplicateIntoProcess(Process.GetCurrentProcess().Handle);
-				using IpcBuffer buffer2 = IpcBuffer.OpenExisting(info);
-				await TestProducerConsumerAsync(buffer1.Writer, buffer2.Reader);
+				await TestProducerConsumerAsync(length => new SharedMemoryBuffer(length));
 			}
 		}
 
-		static async Task TestProducerConsumerAsync(IComputeBufferWriter writer, IComputeBufferReader reader)
+		static async Task TestProducerConsumerAsync(Func<int, IComputeBuffer> createBuffer)
 		{
 			const int Length = 8000;
 
+			Pipe sourceToTargetPipe = new Pipe();
+			Pipe targetToSourcePipe = new Pipe();
+
+			await using IComputeSocket sourceSocket = new ClientComputeSocket(new PipeTransport(targetToSourcePipe.Reader, sourceToTargetPipe.Writer), NullLogger.Instance);
+			await using IComputeSocket targetSocket = new ClientComputeSocket(new PipeTransport(sourceToTargetPipe.Reader, targetToSourcePipe.Writer), NullLogger.Instance);
+
+			using IComputeBufferWriter sourceWriter = sourceSocket.AttachSendBuffer(0, createBuffer(Length));
+			using IComputeBufferReader targetReader = targetSocket.AttachRecvBuffer(0, createBuffer(Length));
+
 			byte[] input = RandomNumberGenerator.GetBytes(Length);
-			Task producerTask = RunProducerAsync(writer, input);
+			Task producerTask = RunProducerAsync(sourceWriter, input);
 
 			byte[] output = new byte[Length];
-			await RunConsumerAsync(reader, output);
+			await RunConsumerAsync(targetReader, output);
 
 			await producerTask;
 			Assert.IsTrue(input.SequenceEqual(output));
@@ -72,7 +82,7 @@ namespace EpicGames.Horde.Tests
 				int length = Math.Min(memory.Length, 7);
 				memory.Slice(0, length).CopyTo(output.Slice(offset));
 				reader.Advance(length);
-				await reader.WaitAsync(memory.Length - length, CancellationToken.None);
+				await reader.WaitForDataAsync(memory.Length - length, CancellationToken.None);
 				offset += length;
 			}
 		}
diff --git a/Engine/Source/Programs/Shared/EpicGames.Horde.Tests/TransportTests.cs b/Engine/Source/Programs/Shared/EpicGames.Horde.Tests/TransportTests.cs
--- a/Engine/Source/Programs/Shared/EpicGames.Horde.Tests/TransportTests.cs
+++ b/Engine/Source/Programs/Shared/EpicGames.Horde.Tests/TransportTests.cs
@@ -74,7 +74,7 @@ namespace EpicGames.Horde.Tests
 
 				for (int offset = 0; offset < output.Length;)
 				{
-					offset += await aesTransport.ReadAsync(output.AsMemory(offset), CancellationToken.None);
+					offset += await aesTransport.ReadPartialAsync(output.AsMemory(offset), CancellationToken.None);
 				}
 			}
 
@@ -89,7 +89,7 @@ namespace EpicGames.Horde.Tests
 		{
 			for (int offset = 0; offset < buffer.Length;)
 			{
-				int size = await transport.ReadAsync(buffer.Slice(offset), CancellationToken.None);
+				int size = await transport.ReadPartialAsync(buffer.Slice(offset), CancellationToken.None);
 				Assert.IsTrue(size > 0);
 				offset += size;
 			}
diff --git a/Engine/Source/Programs/Shared/EpicGames.Horde/Compute/Buffers/HeapBuffer.cs b/Engine/Source/Programs/Shared/EpicGames.Horde/Compute/Buffers/HeapBuffer.cs
deleted file mode 100644
diff --git a/Engine/Source/Programs/Shared/EpicGames.Horde/Compute/Buffers/IpcBuffer.cs b/Engine/Source/Programs/Shared/EpicGames.Horde/Compute/Buffers/IpcBuffer.cs
deleted file mode 100644
diff --git a/Engine/Source/Programs/Shared/EpicGames.Horde/Compute/Buffers/PooledBuffer.cs b/Engine/Source/Programs/Shared/EpicGames.Horde/Compute/Buffers/PooledBuffer.cs
new file mode 100644
--- /dev/null
+++ b/Engine/Source/Programs/Shared/EpicGames.Horde/Compute/Buffers/PooledBuffer.cs
@@ -0,0 +1,162 @@
+// Copyright Epic Games, Inc. All Rights Reserved.
+
+using System;
+using System.Buffers;
+using System.Threading;
+using System.Threading.Tasks;
+
+namespace EpicGames.Horde.Compute.Buffers
+{
+	/// <summary>
+	/// In-process buffer used to store compute messages
+	/// </summary>
+	class PooledBuffer : ComputeBufferBase
+	{
+		readonly IMemoryOwner<byte> _memoryOwner;
+		readonly Memory<byte> _memory;
+
+		int _readPosition;
+		int _writePosition;
+		bool _finishedWriting;
+		TaskCompletionSource? _writtenTcs;
+		TaskCompletionSource? _flushedTcs;
+
+		/// <inheritdoc/>
+		public override long Length => _memory.Length;
+
+		/// <summary>
+		/// Creates a local buffer with the given capacity
+		/// </summary>
+		/// <param name="capacity">Capacity of the buffer</param>
+		public PooledBuffer(int capacity)
+		{
+			_memoryOwner = MemoryPool<byte>.Shared.Rent(capacity);
+			_memory = _memoryOwner.Memory.Slice(0, capacity);
+		}
+
+		/// <inheritdoc/>
+		protected override void Dispose(bool disposing)
+		{
+			base.Dispose(disposing);
+
+			if (disposing)
+			{
+				_memoryOwner.Dispose();
+			}
+		}
+
+		/// <inheritdoc/>
+		public override Memory<byte> GetMemory(long offset, int length)
+		{
+			if (offset > Int32.MaxValue)
+			{
+				throw new ArgumentException("Offset is out of range", nameof(offset));
+			}
+			return _memory.Slice((int)offset, length);
+		}
+
+		#region Reader
+
+		/// <inheritdoc/>
+		public override bool FinishedReading() => _finishedWriting && _readPosition == _writePosition;
+
+		/// <inheritdoc/>
+		public override void AdvanceReadPosition(int size)
+		{
+			_readPosition += size;
+
+			if (_readPosition == _writePosition)
+			{
+				_flushedTcs?.TrySetResult();
+			}
+		}
+
+		/// <inheritdoc/>
+		public override ReadOnlyMemory<byte> GetReadMemory() => _memory.Slice(_readPosition, _writePosition - _readPosition);
+
+		/// <inheritdoc/>
+		public override async ValueTask WaitForDataAsync(int currentLength, CancellationToken cancellationToken)
+		{
+			int initialWritePosition = _readPosition + currentLength;
+
+			// Early out if we're already past this
+			if (_finishedWriting || _writePosition > initialWritePosition)
+			{
+				return;
+			}
+
+			TaskCompletionSource tcs = new TaskCompletionSource();
+			using (IDisposable registration = cancellationToken.Register(x => tcs.SetCanceled((CancellationToken)x!), cancellationToken))
+			{
+				try
+				{
+					_writtenTcs = tcs;
+					for (; ; )
+					{
+						if (_finishedWriting || _writePosition > initialWritePosition)
+						{
+							break;
+						}
+						else
+						{
+							await _writtenTcs.Task;
+						}
+					}
+				}
+				finally
+				{
+					_writtenTcs = null;
+				}
+			}
+		}
+
+		#endregion
+
+		#region Writer
+
+		/// <inheritdoc/>
+		public override void FinishWriting()
+		{
+			_finishedWriting = true;
+			_writtenTcs?.TrySetResult();
+		}
+
+		/// <inheritdoc/>
+		public override void AdvanceWritePosition(int size)
+		{
+			if (_finishedWriting)
+			{
+				throw new InvalidOperationException("Cannot update write position after marking as complete");
+			}
+
+			_writePosition += size;
+			_writtenTcs?.TrySetResult();
+		}
+
+		/// <inheritdoc/>
+		public override Memory<byte> GetWriteMemory() => _memory.Slice(_writePosition);
+
+		/// <inheritdoc/>
+		public override async ValueTask FlushWritesAsync(CancellationToken cancellationToken)
+		{
+			if (_readPosition < _writePosition)
+			{
+				TaskCompletionSource tcs = new TaskCompletionSource();
+				using (_ = cancellationToken.Register(x => tcs.SetCanceled((CancellationToken)x!), cancellationToken, false))
+				{
+					TaskCompletionSource? originalTcs = Interlocked.CompareExchange(ref _flushedTcs, tcs, null);
+					if (originalTcs != null)
+					{
+						_ = originalTcs.Task.ContinueWith(x => tcs.TrySetResult(), cancellationToken, TaskContinuationOptions.None, TaskScheduler.Default);
+					}
+				}
+				if (_readPosition < _writePosition)
+				{
+					await tcs.Task;
+				}
+			}
+		}
+
+		#endregion
+	}
+}
diff --git a/Engine/Source/Programs/Shared/EpicGames.Horde/Compute/Buffers/SharedMemoryBuffer.cs b/Engine/Source/Programs/Shared/EpicGames.Horde/Compute/Buffers/SharedMemoryBuffer.cs
new file mode 100644
--- /dev/null
+++ b/Engine/Source/Programs/Shared/EpicGames.Horde/Compute/Buffers/SharedMemoryBuffer.cs
@@ -0,0 +1,289 @@
+// Copyright Epic Games, Inc. All Rights Reserved.
+
+using System;
+using EpicGames.Core;
+using Microsoft.Win32.SafeHandles;
+using System.IO.MemoryMappedFiles;
+using System.Reflection;
+using System.IO;
+using System.Buffers.Binary;
+using System.Threading.Tasks;
+using System.Threading;
+using System.Linq;
+using System.Globalization;
+using System.ComponentModel;
+
+namespace EpicGames.Horde.Compute.Buffers
+{
+	/// <summary>
+	/// Implementation of <see cref="IComputeBuffer"/> suitable for cross-process communication
+	/// </summary>
+	public sealed class SharedMemoryBuffer : ComputeBufferBase
+	{
+		const int HeaderLength = 32;
+
+		long ReadPosition
+		{
+			get => BinaryPrimitives.ReadInt64LittleEndian(_header.Span.Slice(8));
+			set => BinaryPrimitives.WriteInt64LittleEndian(_header.Span.Slice(8), value);
+		}
+
+		long WritePosition
+		{
+			get => BinaryPrimitives.ReadInt64LittleEndian(_header.Span.Slice(16));
+			set => BinaryPrimitives.WriteInt64LittleEndian(_header.Span.Slice(16), value);
+		}
+
+		bool FinishedWriting
+		{
+			get => _header.Span[24] != 0;
+			set => _header.Span[24] = value ? (byte)1 : (byte)0;
+		}
+
+		const int ChunkOffsetPow2 = 10;
+
+		readonly MemoryMappedFile _memoryMappedFile;
+		readonly MemoryMappedViewAccessor _memoryMappedViewAccessor;
+		readonly MemoryMappedView _memoryMappedView;
+		readonly long _length;
+		readonly Memory<byte> _header;
+		readonly Memory<byte> _memory; // TODO: DEPRECATE
+		readonly Memory<byte>[] _chunks;
+
+		readonly Native.EventHandle _writtenEvent;
+		readonly Native.EventHandle _flushedEvent;
+
+		/// <inheritdoc/>
+		public override long Length => _length;
+
+		/// <summary>
+		/// Creates a new shared memory buffer with the given capacity
+		/// </summary>
+		/// <param name="capacity">Capacity of the buffer</param>
+		public SharedMemoryBuffer(long capacity)
+			: this(MemoryMappedFile.CreateNew(null, capacity, MemoryMappedFileAccess.ReadWrite, MemoryMappedFileOptions.None, HandleInheritability.Inheritable))
+		{
+		}
+
+		/// <summary>
+		/// 
+		/// </summary>
+		/// <param name="buffer"></param>
+		public SharedMemoryBuffer(MemoryMappedFile buffer)
+			: this(buffer, new Native.EventHandle(HandleInheritability.Inheritable), new Native.EventHandle(HandleInheritability.Inheritable))
+		{
+		}
+
+		/// <summary>
+		/// Creates a shared memory buffer from a memory mapped file
+		/// </summary>
+		internal SharedMemoryBuffer(MemoryMappedFile memoryMappedFile, Native.EventHandle writtenEvent, Native.EventHandle flushedEvent)
+		{
+			_memoryMappedFile = memoryMappedFile;
+			_memoryMappedViewAccessor = memoryMappedFile.CreateViewAccessor();
+			_memoryMappedView = new MemoryMappedView(_memoryMappedViewAccessor);
+
+			_writtenEvent = writtenEvent;
+			_flushedEvent = flushedEvent;
+
+			_header = _memoryMappedView.GetMemory(0, HeaderLength);
+			_length = _memoryMappedViewAccessor.Capacity;
+
+			_memory = _memoryMappedView.GetMemory(HeaderLength, (int)(_length - HeaderLength)); // TODO: REMOVE
+
+			int chunkCount = (int)((_length + ((1 << ChunkOffsetPow2) - 1)) >> ChunkOffsetPow2);
+			_chunks = new Memory<byte>[chunkCount];
+
+			for (int chunkIdx = 0; chunkIdx < chunkCount; chunkIdx++)
+			{
+				long offset = (chunkIdx << ChunkOffsetPow2) + HeaderLength;
+				int length = (int)Math.Min(Length - offset, Int32.MaxValue);
+				_chunks[chunkIdx] = _memoryMappedView.GetMemory(offset, length);
+			}
+		}
+
+		/// <summary>
+		/// Opens an IpcBuffer from a string passed in from another process
+		/// </summary>
+		/// <param name="handle">Descriptor for the buffer to open</param>
+		public static MemoryMappedFile OpenMemoryMappedFileFromHandle(IntPtr handle)
+		{
+			MethodInfo? setHandleInfo = typeof(SafeMemoryMappedFileHandle).GetMethod("SetHandle", BindingFlags.Instance | BindingFlags.NonPublic, new[] { typeof(IntPtr) });
+			if (setHandleInfo == null)
+			{
+				throw new InvalidOperationException("Cannot find SetHandle method for SafeMemoryMappedFileHandle");
+			}
+
+			ConstructorInfo? constructorInfo = typeof(MemoryMappedFile).GetConstructor(BindingFlags.Instance | BindingFlags.NonPublic, null, new[] { typeof(SafeMemoryMappedFileHandle) }, null);
+			if (constructorInfo == null)
+			{
+				throw new InvalidOperationException("Cannot find private constructor for memory mapped file");
+			}
+
+			SafeMemoryMappedFileHandle safeHandle = new SafeMemoryMappedFileHandle();
+			setHandleInfo.Invoke(safeHandle, new object[] { handle });
+
+			return (MemoryMappedFile)constructorInfo.Invoke(new object[] { safeHandle });
+		}
+
+		/// <summary>
+		/// Opens a <see cref="SharedMemoryBuffer"/> from handles returned by <see cref="GetIpcHandle()"/>
+		/// </summary>
+		public static SharedMemoryBuffer OpenIpcHandle(string handle)
+		{
+			IntPtr[] handles = handle.Split('.').Select(x => new IntPtr((long)UInt64.Parse(x, NumberStyles.None, null))).ToArray();
+			if (handles.Length != 3)
+			{
+				throw new ArgumentException($"Malformed ipc handle string: {handle}", nameof(handle));
+			}
+
+			MemoryMappedFile memoryMappedFile = OpenMemoryMappedFileFromHandle(handles[0]);
+			Native.EventHandle writtenEvent = new Native.EventHandle(handles[1], true);
+			Native.EventHandle flushedEvent = new Native.EventHandle(handles[2], true);
+
+			return new SharedMemoryBuffer(memoryMappedFile, writtenEvent, flushedEvent);
+		}
+
+		/// <summary>
+		/// Gets a string that can be used to open the same buffer in another process
+		/// </summary>
+		public string GetIpcHandle()
+		{
+			IntPtr memoryHandle = _memoryMappedFile.SafeMemoryMappedFileHandle.DangerousGetHandle();
+			IntPtr writtenEventHandle = _writtenEvent.SafeWaitHandle.DangerousGetHandle();
+			IntPtr flushedEventHandle = _flushedEvent.SafeWaitHandle.DangerousGetHandle();
+			return GetIpcHandle(memoryHandle, writtenEventHandle, flushedEventHandle);
+		}
+
+		/// <summary>
+		/// Duplicates handles for this buffer into the given process, and returns a string that the target process can use to open it.
+		/// </summary>
+		/// <param name="targetProcessHandle">Handle to the target process</param>
+		/// <returns>String that the target process can pass into a call to <see cref="OpenIpcHandle(String)"/> to open the buffer.</returns>
+		public string GetIpcHandle(IntPtr targetProcessHandle)
+		{
+			IntPtr sourceProcessHandle = Native.GetCurrentProcess();
+
+			IntPtr targetMemoryHandle;
+			if (!Native.DuplicateHandle(sourceProcessHandle, _memoryMappedFile.SafeMemoryMappedFileHandle.DangerousGetHandle(), targetProcessHandle, out targetMemoryHandle, 0, false, Native.DUPLICATE_SAME_ACCESS))
+			{
+				throw new Win32Exception();
+			}
+
+			IntPtr targetWrittenEventHandle;
+			if (!Native.DuplicateHandle(sourceProcessHandle, _writtenEvent.SafeWaitHandle.DangerousGetHandle(), targetProcessHandle, out targetWrittenEventHandle, 0, false, Native.DUPLICATE_SAME_ACCESS))
+			{
+				throw new Win32Exception();
+			}
+
+			IntPtr targetFlushedEventHandle;
+			if(!Native.DuplicateHandle(sourceProcessHandle, _flushedEvent.SafeWaitHandle.DangerousGetHandle(), targetProcessHandle, out targetFlushedEventHandle, 0, false, Native.DUPLICATE_SAME_ACCESS))
+			{
+				throw new Win32Exception();
+			}
+
+			return GetIpcHandle(targetMemoryHandle, targetWrittenEventHandle, targetFlushedEventHandle);
+		}
+
+		static string GetIpcHandle(IntPtr memoryHandle, IntPtr writtenEventHandle, IntPtr flushedEventHandle)
+		{
+			return $"{(ulong)memoryHandle.ToInt64()}.{(ulong)writtenEventHandle.ToInt64()}.{(ulong)flushedEventHandle.ToInt64()}";
+		}
+
+		/// <inheritdoc/>
+		protected override void Dispose(bool disposing)
+		{
+			base.Dispose(disposing);
+
+			if (disposing)
+			{
+				_flushedEvent.Dispose();
+				_writtenEvent.Dispose();
+				_memoryMappedView.Dispose();
+				_memoryMappedViewAccessor.Dispose();
+				_memoryMappedFile.Dispose();
+			}
+		}
+
+		/// <inheritdoc/>
+		public override Memory<byte> GetMemory(long offset, int length)
+		{
+			int chunkIdx = (int)(offset >> ChunkOffsetPow2);
+			long baseOffset = offset - (chunkIdx << ChunkOffsetPow2);
+			return _chunks[chunkIdx].Slice((int)(offset - baseOffset), length);
+		}
+
+		#region Reader
+
+		/// <inheritdoc/>
+		public override bool FinishedReading() => FinishedWriting && ReadPosition == WritePosition;
+
+		/// <inheritdoc/>
+		public override void AdvanceReadPosition(int size)
+		{
+			ReadPosition += size;
+
+			if (ReadPosition == WritePosition)
+			{
+				_flushedEvent.Set();
+			}
+		}
+
+		/// <inheritdoc/>
+		public override ReadOnlyMemory<byte> GetReadMemory() => _memory.Slice((int)ReadPosition, (int)(WritePosition - ReadPosition));
+
+		/// <inheritdoc/>
+		public override async ValueTask WaitForDataAsync(int currentLength, CancellationToken cancellationToken)
+		{
+			long initialWritePosition = ReadPosition + currentLength;
+			for (; ; )
+			{
+				if (FinishedWriting || WritePosition > initialWritePosition)
+				{
+					return;
+				}
+				else
+				{
+					await _writtenEvent.WaitOneAsync(cancellationToken);
+				}
+			}
+		}
+
+		#endregion
+
+		#region Writer
+
+		/// <inheritdoc/>
+		public override void FinishWriting()
+		{
+			FinishedWriting = true;
+			_writtenEvent.Set();
+		}
+
+		/// <inheritdoc/>
+		public override void AdvanceWritePosition(int size)
+		{
+			if (FinishedWriting)
+			{
+				throw new InvalidOperationException("Cannot update write position after marking as complete");
+			}
+
+			WritePosition += size;
+			_writtenEvent.Set();
+		}
+
+		/// <inheritdoc/>
+		public override Memory<byte> GetWriteMemory() => _memory.Slice((int)WritePosition);
+
+		/// <inheritdoc/>
+		public override async ValueTask FlushWritesAsync(CancellationToken cancellationToken)
+		{
+			while (ReadPosition < WritePosition)
+			{
+				await _flushedEvent.WaitOneAsync(cancellationToken);
+			}
+		}
+
+		#endregion
+	}
+}
diff --git a/Engine/Source/Programs/Shared/EpicGames.Horde/Compute/Clients/AgentComputeClient.cs b/Engine/Source/Programs/Shared/EpicGames.Horde/Compute/Clients/AgentComputeClient.cs
new file mode 100644
--- /dev/null
+++ b/Engine/Source/Programs/Shared/EpicGames.Horde/Compute/Clients/AgentComputeClient.cs
@@ -0,0 +1,113 @@
+// Copyright Epic Games, Inc. All Rights Reserved.
+
+using System;
+using System.Collections.Generic;
+using System.Diagnostics;
+using System.Net;
+using System.Net.Sockets;
+using System.Runtime.InteropServices;
+using System.Threading;
+using System.Threading.Tasks;
+using EpicGames.Core;
+using EpicGames.Horde.Compute.Transports;
+using Microsoft.Extensions.Logging;
+
+namespace EpicGames.Horde.Compute.Clients
+{
+	/// <summary>
+	/// Runs a local Horde Agent process to process compute requests without communicating with a server
+	/// </summary>
+	public sealed class AgentComputeClient : IComputeClient
+	{
+		readonly string _hordeAgentAssembly;
+		readonly int _port;
+		readonly ILogger _logger;
+
+		/// <summary>
+		/// Constructor
+		/// </summary>
+		/// <param name="hordeAgentFile">Path to the Horde Agent assembly</param>
+		/// <param name="port">Loopback port to connect on</param>
+		/// <param name="logger">Factory for logger instances</param>
+		public AgentComputeClient(string hordeAgentFile, int port, ILogger logger)
+		{
+			_hordeAgentAssembly = hordeAgentFile;
+			_port = port;
+			_logger = logger;
+		}
+
+		/// <inheritdoc/>
+		public ValueTask DisposeAsync() => new ValueTask();
+
+		/// <inheritdoc/>
+		public async Task<TResult> ExecuteAsync<TResult>(ClusterId clusterId, Requirements? requirements, Func<IComputeLease, CancellationToken, Task<TResult>> handler, CancellationToken cancellationToken)
+		{
+			_logger.LogInformation("** CLIENT **");
+			_logger.LogInformation("Launching {Path} to handle remote", _hordeAgentAssembly);
+
+			TResult result;
+			using (Socket listener = new Socket(SocketType.Stream, ProtocolType.IP))
+			{
+				listener.Bind(new IPEndPoint(IPAddress.Loopback, _port));
+				listener.Listen();
+
+				using CancellationTokenSource cancellationSource = CancellationTokenSource.CreateLinkedTokenSource(cancellationToken);
+
+				Task<TResult> listenTask = ListenAsync<TResult>(listener, handler, cancellationSource.Token);
+				using (ManagedProcessGroup group = new ManagedProcessGroup())
+				{
+					if (RuntimeInformation.IsOSPlatform(OSPlatform.Windows))
+					{
+						List<string> arguments = new List<string>();
+						arguments.Add(_hordeAgentAssembly);
+						arguments.Add("computeworker");
+						arguments.Add($"-port={_port}");
+
+						using ManagedProcess process = new ManagedProcess(group, "dotnet", CommandLineArguments.Join(arguments), null, null, ProcessPriorityClass.Normal);
+
+						string? line;
+						while ((line = await process.ReadLineAsync(cancellationToken)) != null)
+						{
+							_logger.LogInformation("Output: {Line}", line);
+						}
+
+						await process.WaitForExitAsync(cancellationToken);
+					}
+					else
+					{
+						using Process process = new Process();
+						process.StartInfo.FileName = "dotnet";
+						process.StartInfo.ArgumentList.Add(_hordeAgentAssembly);
+						process.StartInfo.ArgumentList.Add("computeworker");
+						process.StartInfo.ArgumentList.Add($"-port={_port}");
+						process.StartInfo.UseShellExecute = true;
+						process.Start();
+
+						await process.WaitForExitAsync(cancellationToken);
+					}
+				}
+
+				cancellationSource.Cancel();
+				result = await listenTask;
+			}
+
+			_logger.LogInformation("Client terminated.");
+			return result;
+		}
+
+		async Task<TResult> ListenAsync<TResult>(Socket listener, Func<IComputeLease, CancellationToken, Task<TResult>> handler, CancellationToken cancellationToken)
+		{
+			TResult result;
+			using (Socket tcpSocket = await listener.AcceptAsync(cancellationToken))
+			{
+				await using (ClientComputeSocket socket = new ClientComputeSocket(new TcpTransport(tcpSocket), _logger))
+				{
+					await using ComputeLease lease = new ComputeLease(new List<string>(), new Dictionary<string, int>(), socket);
+					result = await handler(lease, cancellationToken);
+					await socket.CloseAsync(cancellationToken);
+				}
+			}
+			return result;
+		}
+	}
+}
diff --git a/Engine/Source/Programs/Shared/EpicGames.Horde/Compute/Clients/LocalComputeClient.cs b/Engine/Source/Programs/Shared/EpicGames.Horde/Compute/Clients/LocalComputeClient.cs
--- a/Engine/Source/Programs/Shared/EpicGames.Horde/Compute/Clients/LocalComputeClient.cs
+++ b/Engine/Source/Programs/Shared/EpicGames.Horde/Compute/Clients/LocalComputeClient.cs
@@ -13,30 +13,30 @@ using Microsoft.Extensions.Logging;
 namespace EpicGames.Horde.Compute.Clients
 {
 	/// <summary>
-	/// Implementation of <see cref="IComputeChannel"/> which marshals data over a loopback connection to a method running on a background task in the same process.
+	/// Implementation of <see cref="IComputeClient"/> which marshals data over a loopback connection to a method running on a background task in the same process.
 	/// </summary>
 	public sealed class LocalComputeClient : IComputeClient
 	{
 		readonly BackgroundTask _listenerTask;
 		readonly Socket _listener;
 		readonly Socket _socket;
-		readonly ILoggerFactory _loggerFactory;
+		readonly ILogger _logger;
 
 		/// <summary>
 		/// Constructor
 		/// </summary>
 		/// <param name="serverFunc">Callback to process the server side of the connection</param>
 		/// <param name="port">Port to connect on</param>
-		/// <param name="loggerFactory">Logger for diagnostic output</param>
-		public LocalComputeClient(Func<IComputeSocket, CancellationToken, Task> serverFunc, int port, ILoggerFactory loggerFactory)
+		/// <param name="logger">Logger for diagnostic output</param>
+		public LocalComputeClient(Func<IComputeSocket, CancellationToken, Task> serverFunc, int port, ILogger logger)
 		{
-			_loggerFactory = loggerFactory;
+			_logger = logger;
 
 			_listener = new Socket(SocketType.Stream, ProtocolType.IP);
 			_listener.Bind(new IPEndPoint(IPAddress.Loopback, port));
 			_listener.Listen();
 
-			_listenerTask = BackgroundTask.StartNew(ctx => RunListenerAsync(_listener, serverFunc, loggerFactory, ctx));
+			_listenerTask = BackgroundTask.StartNew(ctx => RunListenerAsync(_listener, serverFunc, logger, ctx));
 
 			_socket = new Socket(SocketType.Stream, ProtocolType.IP);
 			_socket.Connect(IPAddress.Loopback, port);
@@ -53,11 +53,11 @@ namespace EpicGames.Horde.Compute.Clients
 		/// <summary>
 		/// Sets up the loopback listener and calls the server method
 		/// </summary>
-		static async Task RunListenerAsync(Socket listener, Func<IComputeSocket, CancellationToken, Task> serverFunc, ILoggerFactory loggerFactory, CancellationToken cancellationToken)
+		static async Task RunListenerAsync(Socket listener, Func<IComputeSocket, CancellationToken, Task> serverFunc, ILogger logger, CancellationToken cancellationToken)
 		{
 			using Socket tcpSocket = await listener.AcceptAsync(cancellationToken);
 
-			await using (ComputeSocket socket = new ComputeSocket(new TcpTransport(tcpSocket), loggerFactory))
+			await using (ClientComputeSocket socket = new ClientComputeSocket(new TcpTransport(tcpSocket), logger))
 			{
 				await serverFunc(socket, cancellationToken);
 				await socket.CloseAsync(cancellationToken);
@@ -67,7 +67,7 @@ namespace EpicGames.Horde.Compute.Clients
 		/// <inheritdoc/>
 		public async Task<TResult> ExecuteAsync<TResult>(ClusterId clusterId, Requirements? requirements, Func<IComputeLease, CancellationToken, Task<TResult>> handler, CancellationToken cancellationToken)
 		{
-			await using ComputeSocket socket = new ComputeSocket(new TcpTransport(_socket), _loggerFactory);
+			await using ClientComputeSocket socket = new ClientComputeSocket(new TcpTransport(_socket), _logger);
 			await using ComputeLease lease = new ComputeLease(new List<string>(), new Dictionary<string, int>(), socket);
 			TResult result = await handler(lease, cancellationToken);
 			await socket.CloseAsync(cancellationToken);
diff --git a/Engine/Source/Programs/Shared/EpicGames.Horde/Compute/Clients/ServerComputeClient.cs b/Engine/Source/Programs/Shared/EpicGames.Horde/Compute/Clients/ServerComputeClient.cs
--- a/Engine/Source/Programs/Shared/EpicGames.Horde/Compute/Clients/ServerComputeClient.cs
+++ b/Engine/Source/Programs/Shared/EpicGames.Horde/Compute/Clients/ServerComputeClient.cs
@@ -7,6 +7,7 @@ using System;
 using System.Collections.Generic;
 using System.Net;
 using System.Net.Http;
+using System.Net.Http.Headers;
 using System.Net.Http.Json;
 using System.Net.Sockets;
 using System.Threading;
@@ -42,15 +43,15 @@ namespace EpicGames.Horde.Compute.Clients
 		readonly HttpClient? _defaultHttpClient;
 		readonly Func<HttpClient> _createHttpClient;
 		readonly CancellationTokenSource _cancellationSource = new CancellationTokenSource();
-		readonly ILoggerFactory _loggerFactory;
 		readonly ILogger _logger;
 
 		/// <summary>
 		/// Constructor
 		/// </summary>
 		/// <param name="serverUri">Uri of the server to connect to</param>
-		/// <param name="loggerFactory">Logger for diagnostic messages</param>
-		public ServerComputeClient(Uri serverUri, ILoggerFactory loggerFactory)
+		/// <param name="authHeader">Authentication header</param>
+		/// <param name="logger">Logger for diagnostic messages</param>
+		public ServerComputeClient(Uri serverUri, AuthenticationHeaderValue? authHeader, ILogger logger)
 		{
 #pragma warning disable CA2000 // Dispose objects before losing scope
 			// This is disposed via HttpClient
@@ -59,23 +60,22 @@ namespace EpicGames.Horde.Compute.Clients
 
 			_defaultHttpClient = new HttpClient(handler, true);
 			_defaultHttpClient.BaseAddress = serverUri;
+			_defaultHttpClient.DefaultRequestHeaders.Authorization = authHeader;
 #pragma warning restore CA2000 // Dispose objects before losing scope
 
 			_createHttpClient = GetDefaultHttpClient;
-			_loggerFactory = loggerFactory;
-			_logger = loggerFactory.CreateLogger<ServerComputeClient>();
+			_logger = logger;
 		}
 
 		/// <summary>
 		/// Constructor
 		/// </summary>
 		/// <param name="createHttpClient">Creates an HTTP client with the correct base address for the server</param>
-		/// <param name="loggerFactory">Logger for diagnostic messages</param>
-		public ServerComputeClient(Func<HttpClient> createHttpClient, ILoggerFactory loggerFactory)
+		/// <param name="logger">Logger for diagnostic messages</param>
+		public ServerComputeClient(Func<HttpClient> createHttpClient, ILogger logger)
 		{
 			_createHttpClient = createHttpClient;
-			_loggerFactory = loggerFactory;
-			_logger = loggerFactory.CreateLogger<ServerComputeClient>();
+			_logger = logger;
 		}
 
 		/// <summary>
@@ -134,7 +134,7 @@ namespace EpicGames.Horde.Compute.Clients
 			// Pass the rest of the call over to the handler
 			byte[] key = StringUtils.ParseHexString(responseMessage.Key);
 
-			await using ComputeSocket computeSocket = new ComputeSocket(new TcpTransport(socket), _loggerFactory);
+			await using ClientComputeSocket computeSocket = new ClientComputeSocket(new TcpTransport(socket), _logger);
 			await using ComputeLease lease = new ComputeLease(responseMessage.Properties, responseMessage.AssignedResources, computeSocket);
 			return await handler(lease, cancellationToken);
 		}
diff --git a/Engine/Source/Programs/Shared/EpicGames.Horde/Compute/ComputeBuffer.cs b/Engine/Source/Programs/Shared/EpicGames.Horde/Compute/ComputeBuffer.cs
deleted file mode 100644
diff --git a/Engine/Source/Programs/Shared/EpicGames.Horde/Compute/ComputeBufferBase.cs b/Engine/Source/Programs/Shared/EpicGames.Horde/Compute/ComputeBufferBase.cs
new file mode 100644
--- /dev/null
+++ b/Engine/Source/Programs/Shared/EpicGames.Horde/Compute/ComputeBufferBase.cs
@@ -0,0 +1,155 @@
+// Copyright Epic Games, Inc. All Rights Reserved.
+
+using System;
+using System.Threading;
+using System.Threading.Tasks;
+
+namespace EpicGames.Horde.Compute
+{
+	/// <summary>
+	/// Base implementation for <see cref="IComputeBuffer"/>, which offers reader and writer objects that redirect to the owner
+	/// </summary>
+	public abstract class ComputeBufferBase : IComputeBuffer
+	{
+		class ReaderImpl : IComputeBufferReader
+		{
+			readonly ComputeBufferBase _buffer;
+
+			/// <inheritdoc/>
+			public IComputeBuffer Buffer => _buffer;
+
+			public ReaderImpl(ComputeBufferBase outer) => _buffer = outer;
+
+			/// <inheritdoc/>
+			public void Dispose()
+			{
+				Dispose(true);
+				GC.SuppressFinalize(this);
+			}
+
+			/// <summary>
+			/// 
+			/// </summary>
+			/// <param name="disposing"></param>
+			protected virtual void Dispose(bool disposing)
+			{
+			}
+
+			/// <inheritdoc/>
+			public bool IsComplete => _buffer.FinishedReading();
+
+			/// <inheritdoc/>
+			public void Advance(int size) => _buffer.AdvanceReadPosition(size);
+
+			/// <inheritdoc/>
+			public ReadOnlyMemory<byte> GetMemory() => _buffer.GetReadMemory();
+
+			/// <inheritdoc/>
+			public ValueTask WaitForDataAsync(int currentLength, CancellationToken cancellationToken) => _buffer.WaitForDataAsync(currentLength, cancellationToken);
+		}
+
+		class WriterImpl : IComputeBufferWriter
+		{
+			readonly ComputeBufferBase _buffer;
+
+			/// <inheritdoc/>
+			public IComputeBuffer Buffer => _buffer;
+
+			public WriterImpl(ComputeBufferBase buffer) => _buffer = buffer;
+
+			/// <inheritdoc/>
+			public void Dispose()
+			{
+				Dispose(true);
+				GC.SuppressFinalize(this);
+			}
+
+			/// <summary>
+			/// 
+			/// </summary>
+			/// <param name="disposing"></param>
+			protected virtual void Dispose(bool disposing)
+			{
+				_buffer.FinishWriting();
+			}
+
+			/// <inheritdoc/>
+			public void MarkComplete() => _buffer.FinishWriting();
+
+			/// <inheritdoc/>
+			public void Advance(int size) => _buffer.AdvanceWritePosition(size);
+
+			/// <inheritdoc/>
+			public Memory<byte> GetMemory() => _buffer.GetWriteMemory();
+
+			/// <inheritdoc/>
+			public ValueTask FlushAsync(CancellationToken cancellationToken) => _buffer.FlushWritesAsync(cancellationToken);
+		}
+
+		/// <inheritdoc/>
+		public void Dispose()
+		{
+			Dispose(true);
+			GC.SuppressFinalize(this);
+		}
+
+		/// <inheritdoc/>
+		public abstract long Length { get; }
+
+		/// <inheritdoc/>
+		public IComputeBufferReader Reader { get; }
+
+		/// <inheritdoc/>
+		public IComputeBufferWriter Writer { get; }
+
+		/// <summary>
+		/// Default constructor
+		/// </summary>
+		protected ComputeBufferBase()
+		{
+			Reader = new ReaderImpl(this);
+			Writer = new WriterImpl(this);
+		}
+
+		/// <summary>
+		/// Standard dispose pattern
+		/// </summary>
+		protected virtual void Dispose(bool disposing)
+		{
+		}
+
+		/// <inheritdoc/>
+		public abstract Memory<byte> GetMemory(long offset, int length);
+
+		#region Reader
+
+		/// <inheritdoc cref="IComputeBufferReader.IsComplete"/>
+		public abstract bool FinishedReading();
+
+		/// <inheritdoc cref="IComputeBufferReader.Advance(Int32)"/>
+		public abstract void AdvanceReadPosition(int size);
+
+		/// <inheritdoc cref="IComputeBufferReader.GetMemory"/>
+		public abstract ReadOnlyMemory<byte> GetReadMemory();
+
+		/// <inheritdoc cref="IComputeBufferReader.WaitForDataAsync(Int32, CancellationToken)"/>
+		public abstract ValueTask WaitForDataAsync(int currentLength, CancellationToken cancellationToken);
+
+		#endregion
+		#region Writer
+
+		/// <inheritdoc cref="IComputeBufferWriter.MarkComplete"/>
+		public abstract void FinishWriting();
+
+		/// <inheritdoc cref="IComputeBufferWriter.Advance(Int32)"/>
+		public abstract void AdvanceWritePosition(int size);
+
+		/// <inheritdoc cref="IComputeBufferWriter.GetMemory()"/>
+		public abstract Memory<byte> GetWriteMemory();
+
+		/// <inheritdoc cref="IComputeBufferWriter.FlushAsync(CancellationToken)"/>
+		public abstract ValueTask FlushWritesAsync(CancellationToken cancellationToken);
+
+		#endregion
+	}
+}
diff --git a/Engine/Source/Programs/Shared/EpicGames.Horde/Compute/ComputeChannel.cs b/Engine/Source/Programs/Shared/EpicGames.Horde/Compute/ComputeChannel.cs
deleted file mode 100644
diff --git a/Engine/Source/Programs/Shared/EpicGames.Horde/Compute/ComputeMessageChannel.cs b/Engine/Source/Programs/Shared/EpicGames.Horde/Compute/ComputeMessageChannel.cs
new file mode 100644
--- /dev/null
+++ b/Engine/Source/Programs/Shared/EpicGames.Horde/Compute/ComputeMessageChannel.cs
@@ -0,0 +1,204 @@
+// Copyright Epic Games, Inc. All Rights Reserved.
+
+using System;
+using System.Buffers;
+using System.Buffers.Binary;
+using System.Text;
+using System.Threading;
+using System.Threading.Tasks;
+using Microsoft.Extensions.Logging;
+
+namespace EpicGames.Horde.Compute
+{
+	/// <summary>
+	/// Implementation of a compute channel
+	/// </summary>
+	public class ComputeMessageChannel : IComputeMessageChannel, IAsyncDisposable
+	{
+		// Length of a message header. Consists of a 1 byte type field, followed by 4 byte length field.
+		const int HeaderLength = 5;
+
+		/// <summary>
+		/// Standard implementation of a message
+		/// </summary>
+		sealed class Message : IComputeMessage
+		{
+			/// <inheritdoc/>
+			public ComputeMessageType Type { get; }
+
+			/// <inheritdoc/>
+			public ReadOnlyMemory<byte> Data { get; }
+
+			readonly IMemoryOwner<byte> _memoryOwner;
+			int _position;
+
+			public Message(ComputeMessageType type, ReadOnlyMemory<byte> data)
+			{
+				_memoryOwner = MemoryPool<byte>.Shared.Rent(data.Length);
+				data.CopyTo(_memoryOwner.Memory);
+
+				Type = type;
+				Data = _memoryOwner.Memory.Slice(0, data.Length);
+			}
+
+			/// <inheritdoc/>
+			public void Dispose()
+			{
+				_memoryOwner.Dispose();
+			}
+
+			/// <inheritdoc/>
+			public ReadOnlyMemory<byte> GetMemory(int minSize = 1) => Data.Slice(_position);
+
+			/// <inheritdoc/>
+			public void Advance(int length) => _position += length;
+		}
+
+		/// <summary>
+		/// Allows creating new messages in rented memory
+		/// </summary>
+		class MessageBuilder : IComputeMessageBuilder
+		{
+			readonly ComputeMessageChannel _channel;
+			readonly IComputeBufferWriter _sendBufferWriter;
+			readonly ComputeMessageType _type;
+			int _length;
+
+			/// <inheritdoc/>
+			public int Length => _length;
+
+			public MessageBuilder(ComputeMessageChannel channel, IComputeBufferWriter sendBufferWriter, ComputeMessageType type)
+			{
+				_channel = channel;
+				_sendBufferWriter = sendBufferWriter;
+				_type = type;
+				_length = 0;
+			}
+
+			public void Dispose()
+			{
+				if (_channel._currentBuilder == this)
+				{
+					_channel._currentBuilder = null;
+				}
+			}
+
+			public void Send()
+			{
+				Span<byte> header = _sendBufferWriter.GetMemory().Span;
+				header[0] = (byte)_type;
+				BinaryPrimitives.WriteInt32LittleEndian(header.Slice(1, 4), _length);
+
+				if (_channel._logger.IsEnabled(LogLevel.Trace))
+				{
+					_channel.LogMessageInfo("SEND", _type, _sendBufferWriter.GetMemory().Slice(HeaderLength, _length).Span);
+				}
+
+				_sendBufferWriter.Advance(HeaderLength + _length);
+				_length = 0;
+			}
+
+			/// <inheritdoc/>
+			public void Advance(int count) => _length += count;
+
+			/// <inheritdoc/>
+			public Memory<byte> GetMemory(int sizeHint = 0) => _sendBufferWriter.GetMemory().Slice(HeaderLength + _length);
+
+			/// <inheritdoc/>
+			public Span<byte> GetSpan(int sizeHint = 0) => GetMemory(sizeHint).Span;
+		}
+
+		readonly IComputeBufferReader _receiveBufferReader;
+		readonly IComputeBufferWriter _sendBufferWriter;
+
+		// Can lock chunked memory writer to acuqire pointer
+		readonly ILogger _logger;
+
+#pragma warning disable CA2213
+		MessageBuilder? _currentBuilder;
+#pragma warning restore CA2213
+
+		/// <summary>
+		/// Constructor
+		/// </summary>
+		/// <param name="receiveBufferReader">Reader for incoming messages</param>
+		/// <param name="sendBufferWriter">Writer for outgoing messages</param>
+		/// <param name="logger">Logger for diagnostic output</param>
+		public ComputeMessageChannel(IComputeBufferReader receiveBufferReader, IComputeBufferWriter sendBufferWriter, ILogger logger)
+		{
+			_receiveBufferReader = receiveBufferReader;
+			_sendBufferWriter = sendBufferWriter;
+			_logger = logger;
+		}
+
+		/// <inheritdoc/>
+		public async ValueTask DisposeAsync()
+		{
+			await DisposeAsync(true);
+			GC.SuppressFinalize(this);
+		}
+
+		/// <summary>
+		/// Overridable dispose method
+		/// </summary>
+		protected virtual async ValueTask DisposeAsync(bool disposing)
+		{
+			_currentBuilder?.Dispose();
+			_sendBufferWriter.MarkComplete();
+			await _sendBufferWriter.FlushAsync(CancellationToken.None);
+			_sendBufferWriter.Dispose();
+		}
+
+		/// <inheritdoc/>
+		public async ValueTask<IComputeMessage> ReceiveAsync(CancellationToken cancellationToken)
+		{
+			while (!_receiveBufferReader.IsComplete)
+			{
+				ReadOnlyMemory<byte> memory = _receiveBufferReader.GetMemory();
+				if (memory.Length >= HeaderLength)
+				{
+					int length = BinaryPrimitives.ReadInt32LittleEndian(memory.Span.Slice(1, 4));
+					if (memory.Length >= HeaderLength + length)
+					{
+						ComputeMessageType type = (ComputeMessageType)memory.Span[0];
+						Message message = new Message(type, memory.Slice(HeaderLength, length));
+						if (_logger.IsEnabled(LogLevel.Trace))
+						{
+							LogMessageInfo("RECV", message.Type, message.Data.Span);
+						}
+						_receiveBufferReader.Advance(HeaderLength + length);
+						return message;
+					}
+				}
+				await _receiveBufferReader.WaitForDataAsync(memory.Length, cancellationToken);
+			}
+			return new Message(ComputeMessageType.None, ReadOnlyMemory<byte>.Empty);
+		}
+
+		void LogMessageInfo(string verb, ComputeMessageType type, ReadOnlySpan<byte> data)
+		{
+			StringBuilder bytes = new StringBuilder();
+			for (int offset = 0; offset < 16 && offset < data.Length; offset++)
+			{
+				bytes.Append($" {data[offset]:X2}");
+			}
+			if (data.Length > 16)
+			{
+				bytes.Append("..");
+			}
+			_logger.LogTrace("{Verb} {Type,-22} [{Length,10:n0}] = {Bytes}", verb, type, data.Length, bytes.ToString());
+		}
+
+		/// <inheritdoc/>
+		public IComputeMessageBuilder CreateMessage(ComputeMessageType type, int maxSize)
+		{
+			if (_currentBuilder != null)
+			{
+				throw new InvalidOperationException("Only one writer can be active at a time. Dispose of the previous writer first.");
+			}
+
+			_currentBuilder = new MessageBuilder(this, _sendBufferWriter, type);
+			return _currentBuilder;
+		}
+	}
+}
diff --git a/Engine/Source/Programs/Shared/EpicGames.Horde/Compute/ComputeSocket.cs b/Engine/Source/Programs/Shared/EpicGames.Horde/Compute/ComputeSocket.cs
deleted file mode 100644
diff --git a/Engine/Source/Programs/Shared/EpicGames.Horde/Compute/ComputeSocketBase.cs b/Engine/Source/Programs/Shared/EpicGames.Horde/Compute/ComputeSocketBase.cs
new file mode 100644
--- /dev/null
+++ b/Engine/Source/Programs/Shared/EpicGames.Horde/Compute/ComputeSocketBase.cs
@@ -0,0 +1,322 @@
+// Copyright Epic Games, Inc. All Rights Reserved.
+
+using System;
+using System.Buffers;
+using System.Buffers.Binary;
+using System.Collections.Generic;
+using System.Diagnostics;
+using System.Threading;
+using System.Threading.Tasks;
+using EpicGames.Core;
+using Microsoft.Extensions.Logging;
+
+namespace EpicGames.Horde.Compute
+{
+	/// <summary>
+	/// Manages a set of readers and writers to buffers across a transport layer
+	/// </summary>
+	public abstract class ComputeSocketBase : IComputeSocket, IAsyncDisposable
+	{
+		readonly object _lockObject = new object();
+
+		bool _complete;
+
+		readonly IComputeTransport _transport;
+		readonly ILogger _logger;
+		readonly Task _readTask;
+		CancellationTokenSource _cancellationSource = new CancellationTokenSource();
+
+		readonly Dictionary<int, IComputeBufferWriter> _receiveBuffers = new Dictionary<int, IComputeBufferWriter>();
+		readonly AsyncEvent _recvBufferWritersChangedEvent = new AsyncEvent();
+
+		readonly SemaphoreSlim _sendSemaphore = new SemaphoreSlim(1, 1);
+		readonly Dictionary<int, IComputeBufferReader> _sendBuffers = new Dictionary<int, IComputeBufferReader>();
+		readonly Dictionary<int, Task> _sendTasks = new Dictionary<int, Task>();
+
+		/// <summary>
+		/// Constructor
+		/// </summary>
+		/// <param name="transport">Transport to communicate with the remote</param>
+		/// <param name="logger">Logger for trace output</param>
+		protected ComputeSocketBase(IComputeTransport transport, ILogger logger)
+		{
+			_transport = transport;
+			_logger = logger;
+			_readTask = RunReaderAsync(transport, _cancellationSource.Token);
+		}
+
+		/// <summary>
+		/// Attempt to gracefully close the current connection and shutdown both ends of the transport
+		/// </summary>
+		public async ValueTask CloseAsync(CancellationToken cancellationToken)
+		{
+			Task cancellationTask = Task.Delay(-1, cancellationToken);
+
+			// Wait for all the individual send tasks to complete
+			await Task.WhenAll(_sendTasks.Values);
+			cancellationToken.ThrowIfCancellationRequested();
+
+			// Send a final message indicating that the lease is done. This will allow the senders on the remote end to terminate, and trigger
+			// a message to be sent back to our read task allowing it to shut down gracefully.
+			byte[] header = new byte[8];
+			await _transport.WriteAsync(header, cancellationToken);
+			_logger.LogTrace("Sent shutdown packet");
+
+			// Wait for the reader to stop
+			await Task.WhenAny(_readTask, cancellationTask);
+			cancellationToken.ThrowIfCancellationRequested();
+
+			_logger.LogTrace("Closed");
+		}
+
+		/// <inheritdoc/>
+		public virtual async ValueTask DisposeAsync()
+		{
+			if (_cancellationSource != null)
+			{
+				_cancellationSource.Cancel();
+
+				await Task.WhenAll(_sendTasks.Values);
+				await _readTask;
+
+				_cancellationSource.Dispose();
+				_cancellationSource = null!;
+			}
+
+			_sendSemaphore.Dispose();
+			GC.SuppressFinalize(this);
+		}
+
+		async Task RunReaderAsync(IComputeTransport transport, CancellationToken cancellationToken)
+		{
+			_logger.LogTrace("Started socket reader");
+
+			byte[] header = new byte[8];
+			try
+			{
+				// Maintain a local cache of buffers to be able to query for them without having to acquire a global lock
+				Dictionary<int, IComputeBufferWriter> cachedWriters = new Dictionary<int, IComputeBufferWriter>();
+
+				Memory<byte> last = Memory<byte>.Empty;
+
+				// Process messages from the remote
+				for (; ; )
+				{
+					// Read the next packet header
+					if (!await transport.ReadOptionalAsync(header, cancellationToken))
+					{
+						_logger.LogTrace("End of socket");
+						break;
+					}
+
+					// Parse the target buffer and packet size
+					int id = BinaryPrimitives.ReadInt32LittleEndian(header);
+					int size = BinaryPrimitives.ReadInt32LittleEndian(header.AsSpan(4));
+
+					// If the size if negative, we're closing the entire connection
+					if (size == 0)
+					{
+						_logger.LogTrace("Received shutdown packet");
+						break;
+					}
+
+					// Dispatch it to the correct place
+					if (size < 0)
+					{
+						_logger.LogTrace("Detaching buffer {Id}", id);
+						DetachRecvBufferWriter(cachedWriters, id);
+					}
+					else
+					{
+						IComputeBufferWriter writer = await GetReceiveBufferAsync(cachedWriters, id);
+						await transport.ReadFullAsync(writer.GetMemory().Slice(0, size), cancellationToken);
+						writer.Advance(size);
+					}
+				}
+			}
+			catch (OperationCanceledException)
+			{
+			}
+
+			// Mark all buffers as complete
+			lock (_lockObject)
+			{
+				_complete = true;
+				foreach (IComputeBufferWriter writer in _receiveBuffers.Values)
+				{
+					writer.MarkComplete();
+				}
+			}
+
+			_logger.LogTrace("Closing reader");
+		}
+
+		async Task<IComputeBufferWriter> GetReceiveBufferAsync(Dictionary<int, IComputeBufferWriter> cachedWriters, int id)
+		{
+			IComputeBufferWriter? writer;
+			if (cachedWriters.TryGetValue(id, out writer))
+			{
+				return writer;
+			}
+
+			int delayTime = 1;
+			Task? delayTask = null;
+			Stopwatch? timer = null;
+
+			for (; ; )
+			{
+				Task waitTask;
+				lock (_lockObject)
+				{
+					if (_receiveBuffers.TryGetValue(id, out writer))
+					{
+						cachedWriters.Add(id, writer);
+						return writer;
+					}
+					waitTask = _recvBufferWritersChangedEvent.Task;
+				}
+
+				timer ??= Stopwatch.StartNew();
+				delayTask ??= Task.Delay(TimeSpan.FromSeconds(delayTime));
+
+				while (!waitTask.IsCompleted)
+				{
+					if (await Task.WhenAny(waitTask, delayTask) == delayTask)
+					{
+						_logger.LogWarning("Socket stalled for {Time}s waiting for a buffer to be attached to channel {ChannelId}", (int)timer.Elapsed.TotalSeconds, id);
+						delayTime = Math.Min(15, delayTime + 2);
+						delayTask = Task.Delay(TimeSpan.FromSeconds(delayTime));
+					}
+				}
+
+				await waitTask;
+			}
+		}
+
+		class SendSegment : ReadOnlySequenceSegment<byte>
+		{
+			public void Set(ReadOnlyMemory<byte> memory, ReadOnlySequenceSegment<byte>? next, long runningIndex)
+			{
+				Memory = memory;
+				Next = next;
+				RunningIndex = runningIndex;
+			}
+		}
+
+		async Task SendFromBufferAsync(int id, IComputeBufferReader reader, CancellationToken cancellationToken)
+		{
+			byte[] header = new byte[8];
+
+			BinaryPrimitives.WriteInt32LittleEndian(header, id);
+
+			SendSegment headerSegment = new SendSegment();
+			SendSegment bodySegment = new SendSegment();
+			headerSegment.Set(header, bodySegment, 0);
+
+			for (; ; )
+			{
+				// Wait for something to read
+				await reader.WaitForDataAsync(0, cancellationToken);
+
+				// Acquire the semaphore
+				await _sendSemaphore.WaitAsync(cancellationToken);
+				try
+				{
+					ReadOnlyMemory<byte> memory = reader.GetMemory();
+					if (memory.Length > 0)
+					{
+						BinaryPrimitives.WriteInt32LittleEndian(header.AsSpan(4), memory.Length);
+
+						bodySegment.Set(memory, null, header.Length);
+						ReadOnlySequence<byte> sequence = new ReadOnlySequence<byte>(headerSegment, 0, bodySegment, memory.Length);
+						await _transport.WriteAsync(sequence, cancellationToken);
+
+						reader.Advance(memory.Length);
+					}
+					else if (reader.IsComplete)
+					{
+						_logger.LogTrace("Sending complete packet for channel {ChannelId}", id);
+						BinaryPrimitives.WriteInt32LittleEndian(header.AsSpan(4), -1);
+						await _transport.WriteAsync(header, cancellationToken);
+						break;
+					}
+				}
+				finally
+				{
+					_sendSemaphore.Release();
+				}
+			}
+
+			lock (_lockObject)
+			{
+				_sendBuffers.Remove(id);
+				_sendTasks.Remove(id);
+			}
+		}
+
+		/// <inheritdoc/>
+		public abstract IComputeBuffer CreateBuffer(long capacity);
+
+		/// <summary>
+		/// Registers the write end of an attached receive buffer
+		/// </summary>
+		/// <param name="channelId">Channel id</param>
+		/// <param name="recvBufferWriter">Writer for the receive buffer</param>
+		public void AttachRecvBuffer(int channelId, IComputeBufferWriter recvBufferWriter)
+		{
+			bool complete;
+			lock (_lockObject)
+			{
+				complete = _complete;
+				if (!complete)
+				{
+					_receiveBuffers.Add(channelId, recvBufferWriter);
+				}
+			}
+
+			if (recvBufferWriter != null)
+			{
+				if (complete)
+				{
+					recvBufferWriter.MarkComplete();
+				}
+				else
+				{
+					_recvBufferWritersChangedEvent.Set();
+				}
+			}
+		}
+
+		void DetachRecvBufferWriter(Dictionary<int, IComputeBufferWriter> cachedWriters, int id)
+		{
+			cachedWriters.Remove(id);
+
+			IComputeBufferWriter? writer;
+			lock (_lockObject)
+			{
+#pragma warning disable CA2000 // Dispose objects before losing scope
+				_receiveBuffers.Remove(id, out writer);
+#pragma warning restore CA2000 // Dispose objects before losing scope
+			}
+			if (writer != null)
+			{
+				writer.MarkComplete();
+				writer.Dispose();
+			}
+		}
+
+		/// <summary>
+		/// Registers the read end of an attached send buffer
+		/// </summary>
+		/// <param name="channelId">Channel id</param>
+		/// <param name="sendBufferReader">Writer for the receive buffer</param>
+		public void AttachSendBuffer(int channelId, IComputeBufferReader sendBufferReader)
+		{
+			lock (_lockObject)
+			{
+				_sendBuffers.Add(channelId, sendBufferReader);
+				_sendTasks.Add(channelId, Task.Run(() => SendFromBufferAsync(channelId, sendBufferReader, _cancellationSource.Token), _cancellationSource.Token));
+			}
+		}
+	}
+}
diff --git a/Engine/Source/Programs/Shared/EpicGames.Horde/Compute/ComputeStorageClient.cs b/Engine/Source/Programs/Shared/EpicGames.Horde/Compute/ComputeStorageClient.cs
--- a/Engine/Source/Programs/Shared/EpicGames.Horde/Compute/ComputeStorageClient.cs
+++ b/Engine/Source/Programs/Shared/EpicGames.Horde/Compute/ComputeStorageClient.cs
@@ -15,14 +15,14 @@ namespace EpicGames.Horde.Compute
 	/// </summary>
 	public sealed class ComputeStorageClient : StorageClientBase, IDisposable
 	{
-		readonly IComputeChannel _channel;
+		readonly IComputeMessageChannel _channel;
 		readonly SemaphoreSlim _semaphore;
 
 		/// <summary>
 		/// Constructor
 		/// </summary>
 		/// <param name="channel"></param>
-		public ComputeStorageClient(IComputeChannel channel)
+		public ComputeStorageClient(IComputeMessageChannel channel)
 		{
 			_channel = channel;
 			_semaphore = new SemaphoreSlim(1);
diff --git a/Engine/Source/Programs/Shared/EpicGames.Horde/Compute/IComputeBuffer.cs b/Engine/Source/Programs/Shared/EpicGames.Horde/Compute/IComputeBuffer.cs
--- a/Engine/Source/Programs/Shared/EpicGames.Horde/Compute/IComputeBuffer.cs
+++ b/Engine/Source/Programs/Shared/EpicGames.Horde/Compute/IComputeBuffer.cs
@@ -1,6 +1,7 @@
 // Copyright Epic Games, Inc. All Rights Reserved.
 
 using System;
+using System.IO;
 using System.Threading;
 using System.Threading.Tasks;
 
@@ -12,21 +13,36 @@ namespace EpicGames.Horde.Compute
 	public interface IComputeBuffer : IDisposable
 	{
 		/// <summary>
-		/// Reader from the buffer
+		/// Total length of this buffer
+		/// </summary>
+		long Length { get; }
+
+		/// <summary>
+		/// Reader for this buffer
 		/// </summary>
 		IComputeBufferReader Reader { get; }
 
 		/// <summary>
-		/// Writer to the buffer
+		/// Writer for this buffer
 		/// </summary>
 		IComputeBufferWriter Writer { get; }
+
+		/// <summary>
+		/// Access the underlying memory for the buffer
+		/// </summary>
+		Memory<byte> GetMemory(long offset, int length);
 	}
 
 	/// <summary>
 	/// Read interface for a compute buffer
 	/// </summary>
-	public interface IComputeBufferReader
+	public interface IComputeBufferReader : IDisposable
 	{
+		/// <summary>
+		/// Accessor for the buffer this is reading from
+		/// </summary>
+		IComputeBuffer Buffer { get; }
+
 		/// <summary>
 		/// Whether this buffer is complete (no more data will be added)
 		/// </summary>
@@ -49,14 +65,19 @@ namespace EpicGames.Horde.Compute
 		/// </summary>
 		/// <param name="currentLength">Current length of the buffer</param>
 		/// <param name="cancellationToken">Cancellation token for the operation</param>
-		ValueTask WaitAsync(int currentLength, CancellationToken cancellationToken);
+		ValueTask WaitForDataAsync(int currentLength, CancellationToken cancellationToken);
 	}
 
 	/// <summary>
 	/// Buffer that can receive data from a remote machine.
 	/// </summary>
-	public interface IComputeBufferWriter
+	public interface IComputeBufferWriter : IDisposable
 	{
+		/// <summary>
+		/// Accessor for the buffer we're writing to
+		/// </summary>
+		IComputeBuffer Buffer { get; }
+
 		/// <summary>
 		/// Mark the output to this buffer as complete
 		/// </summary>
@@ -81,18 +102,111 @@ namespace EpicGames.Horde.Compute
 	}
 
 	/// <summary>
-	/// Default environment variable names for compute channels
+	/// Extension methods for <see cref="IComputeBuffer"/>
 	/// </summary>
-	public static class ComputeEnvVarNames
+	public static class ComputeBufferExtensions
 	{
+		sealed class RefCountedBuffer
+		{
+			int _refFlags = 1 | 2;
+
+			public IComputeBuffer Buffer { get; }
+
+			public RefCountedBuffer(IComputeBuffer buffer) => Buffer = buffer;
+
+			public void Release(int flag)
+			{
+				for (; ; )
+				{
+					int initialRefFlags = _refFlags;
+					if (Interlocked.CompareExchange(ref _refFlags, initialRefFlags & ~flag, initialRefFlags) != initialRefFlags)
+					{
+						continue;
+					}
+					if (initialRefFlags == flag)
+					{
+						Buffer.Dispose();
+					}
+					break;
+				}
+			}
+		}
+
+		class RefCountedReader : IComputeBufferReader
+		{
+			public readonly RefCountedBuffer _buffer;
+			public readonly IComputeBufferReader _reader;
+
+			public IComputeBuffer Buffer => _reader.Buffer;
+
+			public RefCountedReader(RefCountedBuffer refCountedBuffer)
+			{
+				_buffer = refCountedBuffer;
+				_reader = refCountedBuffer.Buffer.Reader;
+			}
+
+			public void Dispose() => _buffer.Release(1);
+			public bool IsComplete => _reader.IsComplete;
+			public void Advance(int size) => _reader.Advance(size);
+			public ReadOnlyMemory<byte> GetMemory() => _reader.GetMemory();
+			public ValueTask WaitForDataAsync(int currentLength, CancellationToken cancellationToken) => _reader.WaitForDataAsync(currentLength, cancellationToken);
+		}
+
+		class RefCountedWriter : IComputeBufferWriter
+		{
+			public readonly RefCountedBuffer _buffer;
+			public readonly IComputeBufferWriter _writer;
+
+			public IComputeBuffer Buffer => _writer.Buffer;
+
+			public RefCountedWriter(RefCountedBuffer refCountedBuffer)
+			{
+				_buffer = refCountedBuffer;
+				_writer = refCountedBuffer.Buffer.Writer;
+			}
+
+			public void Dispose() => _buffer.Release(2);
+			public void Advance(int size) => _writer.Advance(size);
+			public ValueTask FlushAsync(CancellationToken cancellationToken) => _writer.FlushAsync(cancellationToken);
+			public Memory<byte> GetMemory() => _writer.GetMemory();
+			public void MarkComplete() => _writer.MarkComplete();
+		}
+
 		/// <summary>
-		/// Default name for a channel that can be used to send messages to the initiator
+		/// Converts a compute buffer to be disposable through its reader/writer 
 		/// </summary>
-		public const string DefaultSendBuffer = "UE_HORDE_SEND_BUFFER";
+		/// <param name="buffer"></param>
+		/// <returns></returns>
+		public static (IComputeBufferReader, IComputeBufferWriter) ToShared(this IComputeBuffer buffer)
+		{
+			RefCountedBuffer refCountedBuffer = new RefCountedBuffer(buffer);
+			return (new RefCountedReader(refCountedBuffer), new RefCountedWriter(refCountedBuffer));
+		}
 
 		/// <summary>
-		/// Default name for a channel that can be used to receive messages from the initiator
+		/// Waits until there is a block of the certain size in the buffer, and returns it
 		/// </summary>
-		public const string DefaultRecvBuffer = "UE_HORDE_RECV_BUFFER";
+		/// <param name="reader">Instance to read from</param>
+		/// <param name="minLength">Minimum length of the memory</param>
+		/// <param name="cancellationToken"></param>
+		/// <returns></returns>
+		public static async ValueTask<ReadOnlyMemory<byte>> GetMemoryAsync(this IComputeBufferReader reader, int minLength, CancellationToken cancellationToken)
+		{
+			for (; ; )
+			{
+				bool complete = reader.IsComplete;
+
+				ReadOnlyMemory<byte> memory = reader.GetMemory();
+				if (memory.Length >= minLength)
+				{
+					return memory;
+				}
+				else if (complete)
+				{
+					throw new EndOfStreamException();
+				}
+				await reader.WaitForDataAsync(memory.Length, cancellationToken);
+			}
+		}
 	}
 }
diff --git a/Engine/Source/Programs/Shared/EpicGames.Horde/Compute/IComputeChannel.cs b/Engine/Source/Programs/Shared/EpicGames.Horde/Compute/IComputeChannel.cs
deleted file mode 100644
diff --git a/Engine/Source/Programs/Shared/EpicGames.Horde/Compute/IComputeMessage.cs b/Engine/Source/Programs/Shared/EpicGames.Horde/Compute/IComputeMessage.cs
--- a/Engine/Source/Programs/Shared/EpicGames.Horde/Compute/IComputeMessage.cs
+++ b/Engine/Source/Programs/Shared/EpicGames.Horde/Compute/IComputeMessage.cs
@@ -4,6 +4,7 @@ using System;
 using System.Buffers.Binary;
 using System.Collections.Generic;
 using System.IO;
+using System.Text;
 using System.Threading;
 using System.Threading.Tasks;
 using EpicGames.Core;
@@ -55,25 +56,20 @@ namespace EpicGames.Horde.Compute
 		/// </summary>
 		DeleteFiles = 0x12,
 
-		/// <summary>
-		/// Delete a directory in the remote (Initiator -> Remote)
-		/// </summary>
-		DeleteDirectory = 0x13,
-
 		/// <summary>
 		/// Execute a process in a sandbox (Initiator -> Remote)
 		/// </summary>
-		ExecuteProcess = 0x16,
+		Execute = 0x16,
 
 		/// <summary>
 		/// Returns the process exit code (Remote -> Initiator)
 		/// </summary>
-		ExecuteProcessResponse = 0x17,
+		ExecuteResult = 0x17,
 
 		/// <summary>
 		/// Returns output from the child process to the caller (Remote -> Initiator)
 		/// </summary>
-		ProcessOutput = 0x18,
+		ExecuteOutput = 0x18,
 
 		#endregion
 
@@ -94,14 +90,14 @@ namespace EpicGames.Horde.Compute
 		#region Buffers
 
 		/// <summary>
-		/// Creates a new receive buffer for messages
+		/// Registers a send buffer created by a child process
 		/// </summary>
-		CreateBufferRequest = 0x30,
+		AttachSendBuffer = 0x30,
 
 		/// <summary>
-		/// Receives information about the new receive buffer
+		/// Registers a receive buffer created by a child process
 		/// </summary>
-		CreateBufferResponse = 0x31,
+		AttachRecvBuffer = 0x31,
 
 		#endregion
 
@@ -180,14 +176,14 @@ namespace EpicGames.Horde.Compute
 	public record struct ReadBlobMessage(BlobLocator Locator, int Offset, int Length);
 
 	/// <summary>
-	/// Request creation of a new buffer, and attaching it to the lease
+	/// Request a receive buffer in a worker process
 	/// </summary>
-	public record struct CreateBufferRequest(int ChannelId, long Capacity, bool Send);
+	public record struct AttachSendBufferRequest(int ChannelId, string Handle);
 
 	/// <summary>
-	/// Response to create a new buffer
+	/// Request a receive buffer in a worker process
 	/// </summary>
-	public record struct CreateBufferResponse(string Handle);
+	public record struct AttachRecvBufferRequest(int ChannelId, string Handle);
 
 	/// <summary>
 	/// Message for running an XOR command
@@ -204,19 +200,17 @@ namespace EpicGames.Horde.Compute
 		/// <summary>
 		/// Sends an exception response to the remote
 		/// </summary>
-		public static void Exception(this IComputeChannel channel, Exception ex) => Exception(channel, ex.Message, ex.StackTrace);
+		public static void SendException(this IComputeMessageChannel channel, Exception ex) => SendException(channel, ex.Message, ex.StackTrace);
 
 		/// <summary>
 		/// Sends an exception response to the remote
 		/// </summary>
-		public static void Exception(this IComputeChannel channel, string description, string? trace)
+		public static void SendException(this IComputeMessageChannel channel, string description, string? trace)
 		{
-			using (IComputeMessageBuilder builder = channel.CreateMessage(ComputeMessageType.Exception))
-			{
-				builder.WriteString(description);
-				builder.WriteOptionalString(trace);
-				builder.Send();
-			}
+			using IComputeMessageBuilder message = channel.CreateMessage(ComputeMessageType.Exception);
+			message.WriteString(description);
+			message.WriteOptionalString(trace);
+			message.Send();
 		}
 
 		/// <summary>
@@ -228,9 +222,9 @@ namespace EpicGames.Horde.Compute
 			string description = message.ReadString();
 			return new ExceptionMessage(msg, description);
 		}
-
+/*
 		/// <inheritdoc/>
-		public static void Fork(this IComputeChannel channel, int newChannelId)
+		public static void Fork(this IComputeMessageChannel channel, int newChannelId)
 		{
 			using (IComputeMessageBuilder builder = channel.CreateMessage(ComputeMessageType.Fork))
 			{
@@ -238,7 +232,7 @@ namespace EpicGames.Horde.Compute
 				builder.Send();
 			}
 		}
-
+*/
 		/// <summary>
 		/// Parses a fork message from the given compute message
 		/// </summary>
@@ -252,19 +246,17 @@ namespace EpicGames.Horde.Compute
 		/// Sends untyped user data to the remote
 		/// </summary>
 		/// <param name="channel">Current channel</param>
-		/// <param name="span">Data to send</param>
-		public static void SendUserData(this IComputeChannel channel, ReadOnlySpan<byte> span)
+		/// <param name="memory">Data to send</param>
+		public static void SendUserData(this IComputeMessageChannel channel, ReadOnlyMemory<byte> memory)
 		{
-			using (IComputeMessageBuilder builder = channel.CreateMessage(ComputeMessageType.UserData))
-			{
-				builder.WriteFixedLengthBytes(span);
-				builder.Send();
-			}
+			using IComputeMessageBuilder message = channel.CreateMessage(ComputeMessageType.UserData, memory.Length);
+			message.WriteFixedLengthBytes(memory.Span);
+			message.Send();
 		}
 
 		#region Process
 
-		static async Task<IComputeMessage> RunStorageServer(this IComputeChannel channel, IStorageClient storage, CancellationToken cancellationToken)
+		static async Task<IComputeMessage> RunStorageServer(this IComputeMessageChannel channel, IStorageClient storage, CancellationToken cancellationToken)
 		{
 			for (; ; )
 			{
@@ -294,13 +286,13 @@ namespace EpicGames.Horde.Compute
 		/// <param name="locator">Location of a <see cref="DirectoryNode"/> describing contents of the sandbox</param>
 		/// <param name="storage">Storage for the sandbox data</param>
 		/// <param name="cancellationToken">Cancellation token for the operation</param>
-		public static async Task UploadFilesAsync(this IComputeChannel channel, string path, NodeLocator locator, IStorageClient storage, CancellationToken cancellationToken)
+		public static async Task UploadFilesAsync(this IComputeMessageChannel channel, string path, NodeLocator locator, IStorageClient storage, CancellationToken cancellationToken)
 		{
-			using (IComputeMessageBuilder writer = channel.CreateMessage(ComputeMessageType.WriteFiles))
+			using (IComputeMessageBuilder request = channel.CreateMessage(ComputeMessageType.WriteFiles))
 			{
-				writer.WriteString(path);
-				writer.WriteNodeLocator(locator);
-				writer.Send();
+				request.WriteString(path);
+				request.WriteNodeLocator(locator);
+				request.Send();
 			}
 
 			using IComputeMessage response = await RunStorageServer(channel, storage, cancellationToken);
@@ -325,13 +317,11 @@ namespace EpicGames.Horde.Compute
 		/// </summary>
 		/// <param name="channel">Current channel</param>
 		/// <param name="paths">Paths of files or directories to delete</param>
-		public static void DeleteFiles(this IComputeChannel channel, IReadOnlyList<string> paths)
+		public static void DeleteFiles(this IComputeMessageChannel channel, IReadOnlyList<string> paths)
 		{
-			using (IComputeMessageBuilder writer = channel.CreateMessage(ComputeMessageType.DeleteDirectory))
-			{
-				writer.WriteList(paths, MemoryWriterExtensions.WriteString);
-				writer.Send();
-			}
+			IComputeMessageBuilder request = channel.CreateMessage(ComputeMessageType.DeleteFiles);
+			request.WriteList(paths, MemoryWriterExtensions.WriteString);
+			request.Send();
 		}
 
 		/// <summary>
@@ -343,6 +333,53 @@ namespace EpicGames.Horde.Compute
 			return new DeleteFilesMessage(files);
 		}
 
+		class StringOutputWriter
+		{
+			byte[] _buffer;
+			int _length;
+			readonly Action<string> _handler;
+
+			public StringOutputWriter(Action<string> handler)
+			{
+				_buffer = new byte[10];
+				_handler = handler;
+			}
+
+			public void WriteData(ReadOnlyMemory<byte> data)
+			{
+				int requiredLength = _length + data.Length;
+				if (requiredLength > _buffer.Length)
+				{
+					Array.Resize(ref _buffer, requiredLength);
+				}
+
+				data.CopyTo(_buffer.AsMemory(_length));
+				_length += data.Length;
+
+				ReadOnlySpan<byte> output = _buffer.AsSpan(0, _length);
+				for (; ; )
+				{
+					int newLineIdx = output.IndexOf((byte)'\n');
+					if (newLineIdx == -1)
+					{
+						break;
+					}
+
+					int length = newLineIdx;
+					if (length > 0 && output[length - 1] == (byte)'\r')
+					{
+						length--;
+					}
+
+					_handler(Encoding.UTF8.GetString(output.Slice(0, length)));
+					output = output.Slice(newLineIdx + 1);
+				}
+
+				_length = output.Length;
+				output.CopyTo(_buffer);
+			}
+		}
+
 		/// <summary>
 		/// Executes a remote process
 		/// </summary>
@@ -353,15 +390,31 @@ namespace EpicGames.Horde.Compute
 		/// <param name="outputHandler">Output callback for stdout</param>
 		/// <param name="envVars">Environment variables for the child process</param>
 		/// <param name="cancellationToken">Cancellation token for the operation</param>
-		public static async Task<int> ExecuteProcessAsync(this IComputeChannel channel, string executable, IReadOnlyList<string> arguments, string? workingDir, IReadOnlyDictionary<string, string?>? envVars, Action<ReadOnlyMemory<byte>> outputHandler, CancellationToken cancellationToken)
+		public static Task<int> ExecuteAsync(this IComputeMessageChannel channel, string executable, IReadOnlyList<string> arguments, string? workingDir, IReadOnlyDictionary<string, string?>? envVars, Action<string> outputHandler, CancellationToken cancellationToken)
 		{
-			using (IComputeMessageBuilder writer = channel.CreateMessage(ComputeMessageType.ExecuteProcess))
+			StringOutputWriter writer = new StringOutputWriter(outputHandler);
+			return ExecuteAsync(channel, executable, arguments, workingDir, envVars, writer.WriteData, cancellationToken);
+		}
+
+		/// <summary>
+		/// Executes a remote process
+		/// </summary>
+		/// <param name="channel">Current channel</param>
+		/// <param name="executable">Executable to run, relative to the sandbox root</param>
+		/// <param name="arguments">Arguments for the child process</param>
+		/// <param name="workingDir">Working directory for the process</param>
+		/// <param name="outputHandler">Output callback for stdout</param>
+		/// <param name="envVars">Environment variables for the child process</param>
+		/// <param name="cancellationToken">Cancellation token for the operation</param>
+		public static async Task<int> ExecuteAsync(this IComputeMessageChannel channel, string executable, IReadOnlyList<string> arguments, string? workingDir, IReadOnlyDictionary<string, string?>? envVars, Action<ReadOnlyMemory<byte>> outputHandler, CancellationToken cancellationToken)
+		{
+			using (IComputeMessageBuilder request = channel.CreateMessage(ComputeMessageType.Execute))
 			{
-				writer.WriteString(executable);
-				writer.WriteList(arguments, MemoryWriterExtensions.WriteString);
-				writer.WriteOptionalString(workingDir);
-				writer.WriteDictionary(envVars ?? new Dictionary<string, string?>(), MemoryWriterExtensions.WriteString, MemoryWriterExtensions.WriteOptionalString);
-				writer.Send();
+				request.WriteString(executable);
+				request.WriteList(arguments, MemoryWriterExtensions.WriteString);
+				request.WriteOptionalString(workingDir);
+				request.WriteDictionary(envVars ?? new Dictionary<string, string?>(), MemoryWriterExtensions.WriteString, MemoryWriterExtensions.WriteOptionalString);
+				request.Send();
 			}
 
 			for (; ; )
@@ -372,10 +425,10 @@ namespace EpicGames.Horde.Compute
 					case ComputeMessageType.Exception:
 						ExceptionMessage exception = message.ParseExceptionMessage();
 						throw new ComputeRemoteException(exception);
-					case ComputeMessageType.ProcessOutput:
+					case ComputeMessageType.ExecuteOutput:
 						outputHandler(message.Data);
 						break;
-					case ComputeMessageType.ExecuteProcessResponse:
+					case ComputeMessageType.ExecuteResult:
 						ExecuteProcessResponseMessage executeProcessResponse = message.ParseExecuteProcessResponse();
 						return executeProcessResponse.ExitCode;
 					default:
@@ -399,13 +452,11 @@ namespace EpicGames.Horde.Compute
 		/// <summary>
 		/// Sends output from a child process
 		/// </summary>
-		public static void SendProcessOutput(this IComputeChannel channel, ReadOnlyMemory<byte> data)
+		public static void SendExecuteOutput(this IComputeMessageChannel channel, ReadOnlyMemory<byte> data)
 		{
-			using (IComputeMessageBuilder builder = channel.CreateMessage(ComputeMessageType.ProcessOutput))
-			{
-				builder.WriteFixedLengthBytes(data.Span);
-				builder.Send();
-			}
+			using IComputeMessageBuilder message = channel.CreateMessage(ComputeMessageType.ExecuteOutput);
+			message.WriteFixedLengthBytes(data.Span);
+			message.Send();
 		}
 
 		/// <summary>
@@ -413,13 +464,11 @@ namespace EpicGames.Horde.Compute
 		/// </summary>
 		/// <param name="channel"></param>
 		/// <param name="exitCode">Exit code from the process</param>
-		public static void SendExecuteProcessResponse(this IComputeChannel channel, int exitCode)
+		public static void SendExecuteResult(this IComputeMessageChannel channel, int exitCode)
 		{
-			using (IComputeMessageBuilder builder = channel.CreateMessage(ComputeMessageType.ExecuteProcessResponse))
-			{
-				builder.WriteInt32(exitCode);
-				builder.Send();
-			}
+			using IComputeMessageBuilder builder = channel.CreateMessage(ComputeMessageType.ExecuteResult);
+			builder.WriteInt32(exitCode);
+			builder.Send();
 		}
 
 		/// <summary>
@@ -481,14 +530,14 @@ namespace EpicGames.Horde.Compute
 		/// <param name="length">Length of data to return</param>
 		/// <param name="cancellationToken">Cancellation token for the operation</param>
 		/// <returns>Stream containing the blob data</returns>
-		public static async Task<Stream> ReadBlobAsync(this IComputeChannel channel, BlobLocator locator, int offset, int length, CancellationToken cancellationToken)
+		public static async Task<Stream> ReadBlobAsync(this IComputeMessageChannel channel, BlobLocator locator, int offset, int length, CancellationToken cancellationToken)
 		{
-			using (IComputeMessageBuilder writer = channel.CreateMessage(ComputeMessageType.ReadBlob))
+			using (IComputeMessageBuilder request = channel.CreateMessage(ComputeMessageType.ReadBlob))
 			{
-				writer.WriteBlobLocator(locator);
-				writer.WriteUnsignedVarInt(offset);
-				writer.WriteUnsignedVarInt(length);
-				writer.Send();
+				request.WriteBlobLocator(locator);
+				request.WriteUnsignedVarInt(offset);
+				request.WriteUnsignedVarInt(length);
+				request.Send();
 			}
 
 			IComputeMessage response = await channel.ReceiveAsync(cancellationToken);
@@ -508,7 +557,7 @@ namespace EpicGames.Horde.Compute
 		/// <param name="message">The read request</param>
 		/// <param name="storage">Storage client to retrieve the blob from</param>
 		/// <param name="cancellationToken">Cancellation token for the operation</param>
-		public static Task SendBlobDataAsync(this IComputeChannel channel, ReadBlobMessage message, IStorageClient storage, CancellationToken cancellationToken)
+		public static Task SendBlobDataAsync(this IComputeMessageChannel channel, ReadBlobMessage message, IStorageClient storage, CancellationToken cancellationToken)
 		{
 			return SendBlobDataAsync(channel, message.Locator, message.Offset, message.Length, storage, cancellationToken);
 		}
@@ -522,7 +571,7 @@ namespace EpicGames.Horde.Compute
 		/// <param name="length">Length of the data</param>
 		/// <param name="storage">Storage client to retrieve the blob from</param>
 		/// <param name="cancellationToken">Cancellation token for the operation</param>
-		public static async Task SendBlobDataAsync(this IComputeChannel channel, BlobLocator locator, int offset, int length, IStorageClient storage, CancellationToken cancellationToken)
+		public static async Task SendBlobDataAsync(this IComputeMessageChannel channel, BlobLocator locator, int offset, int length, IStorageClient storage, CancellationToken cancellationToken)
 		{
 			byte[] data;
 			if (offset == 0 && length == 0)
@@ -544,114 +593,11 @@ namespace EpicGames.Horde.Compute
 				}
 			}
 
-			using (IComputeMessageBuilder writer = channel.CreateMessage(ComputeMessageType.ReadBlobResponse, length + 128))
-			{
-				writer.WriteFixedLengthBytes(data);
-				writer.Send();
-			}
-		}
-
-		#endregion
-
-		#region Child process
-
-		/// <summary>
-		/// Wrapper around the lifetime of an <see cref="IpcBuffer"/> which only exposes the writer interface
-		/// </summary>
-		class WrappedBufferWriter : IComputeBufferWriter
-		{
-			readonly IpcBuffer _buffer;
-
-			public WrappedBufferWriter(IpcBuffer buffer) => _buffer = buffer;
-
-			public void Dispose() => _buffer.Dispose();
-
-			public void Advance(int size) => _buffer.Writer.Advance(size);
-
-			public ValueTask FlushAsync(CancellationToken cancellationToken) => _buffer.Writer.FlushAsync(cancellationToken);
-
-			public Memory<byte> GetMemory() => _buffer.Writer.GetMemory();
-
-			public void MarkComplete() => _buffer.Writer.MarkComplete();
-		}
-
-		/// <summary>
-		/// Creates a buffer which can be used to send data to the remote
-		/// </summary>
-		/// <param name="channel">Channel to write to</param>
-		/// <param name="id">Identifier for the new buffer</param>
-		/// <param name="capacity">Capacity of the buffer</param>
-		/// <param name="cancellationToken">Cancellation token for the operation</param>
-		/// <returns>New compute buffer</returns>
-		public static async Task<IComputeBufferWriter> CreateSendBufferAsync(this IComputeChannel channel, int id, long capacity, CancellationToken cancellationToken)
-		{
-			IpcBuffer buffer = await CreateBufferAsync(channel, id, capacity, true, cancellationToken);
-			return new WrappedBufferWriter(buffer);
-		}
-
-		/// <summary>
-		/// Wrapper around the lifetime of an <see cref="IpcBuffer"/> which only exposes the reader interface
-		/// </summary>
-		class WrappedBufferReader : IComputeBufferReader
-		{
-			readonly IpcBuffer _buffer;
-
-			public WrappedBufferReader(IpcBuffer buffer) => _buffer = buffer;
-
-			public void Dispose() => _buffer.Dispose();
-
-			public bool IsComplete => _buffer.Reader.IsComplete;
-
-			public void Advance(int size) => _buffer.Reader.Advance(size);
-
-			public ReadOnlyMemory<byte> GetMemory() => _buffer.Reader.GetMemory();
-
-			public ValueTask WaitAsync(int currentLength, CancellationToken cancellationToken) => _buffer.Reader.WaitAsync(currentLength, cancellationToken);
-		}
-
-		/// <summary>
-		/// Creates a buffer which can be used to receive data from the remote
-		/// </summary>
-		/// <param name="channel">Channel to write to</param>
-		/// <param name="id">Identifier for the new buffer</param>
-		/// <param name="capacity">Capacity of the buffer</param>
-		/// <param name="cancellationToken">Cancellation token for the operation</param>
-		/// <returns>New compute buffer</returns>
-		public static async Task<IComputeBufferReader> CreateReceiveBufferAsync(this IComputeChannel channel, int id, long capacity, CancellationToken cancellationToken)
-		{
-			IpcBuffer buffer = await CreateBufferAsync(channel, id, capacity, false, cancellationToken);
-			return new WrappedBufferReader(buffer);
-		}
-
-		static async Task<IpcBuffer> CreateBufferAsync(this IComputeChannel channel, int id, long capacity, bool send, CancellationToken cancellationToken)
-		{
-			using (IComputeMessageBuilder builder = channel.CreateMessage(ComputeMessageType.CreateBufferRequest))
-			{
-				builder.WriteInt32(id);
-				builder.WriteInt64(capacity);
-				builder.WriteBoolean(send);
-				builder.Send();
-			}
-
-			using IComputeMessage response = await channel.ReceiveAsync(cancellationToken);
-			if (response.Type != ComputeMessageType.CreateBufferResponse)
+			using (IComputeMessageBuilder response = channel.CreateMessage(ComputeMessageType.ReadBlobResponse, data.Length + 128))
 			{
-				throw new NotSupportedException();
+				response.WriteFixedLengthBytes(data);
+				response.Send();
 			}
-
-			string handle = response.ReadString();
-			return IpcBuffer.OpenExisting(handle);
-		}
-
-		/// <summary>
-		/// Parse a <see cref="CreateBufferRequest"/> message
-		/// </summary>
-		public static CreateBufferRequest ParseCreateBufferRequest(this IComputeMessage message)
-		{
-			int channelId = message.ReadInt32();
-			long capacity = message.ReadInt64();
-			bool send = message.ReadBoolean();
-			return new CreateBufferRequest(channelId, capacity, send);
 		}
 
 		#endregion
@@ -661,14 +607,12 @@ namespace EpicGames.Horde.Compute
 		/// <summary>
 		/// Send a message to request that a byte string be xor'ed with a particular value
 		/// </summary>
-		public static void XorRequest(this IComputeChannel channel, ReadOnlyMemory<byte> data, byte value)
+		public static void SendXorRequest(this IComputeMessageChannel channel, ReadOnlyMemory<byte> data, byte value)
 		{
-			using (IComputeMessageBuilder builder = channel.CreateMessage(ComputeMessageType.XorRequest))
-			{
-				builder.WriteFixedLengthBytes(data.Span);
-				builder.WriteUInt8(value);
-				builder.Send();
-			}
+			using IComputeMessageBuilder message = channel.CreateMessage(ComputeMessageType.XorRequest);
+			message.WriteFixedLengthBytes(data.Span);
+			message.WriteUInt8(value);
+			message.Send();
 		}
 
 		/// <summary>
diff --git a/Engine/Source/Programs/Shared/EpicGames.Horde/Compute/IComputeMessageChannel.cs b/Engine/Source/Programs/Shared/EpicGames.Horde/Compute/IComputeMessageChannel.cs
new file mode 100644
--- /dev/null
+++ b/Engine/Source/Programs/Shared/EpicGames.Horde/Compute/IComputeMessageChannel.cs
@@ -0,0 +1,61 @@
+// Copyright Epic Games, Inc. All Rights Reserved.
+
+using System;
+using System.Threading;
+using System.Threading.Tasks;
+using EpicGames.Core;
+
+namespace EpicGames.Horde.Compute
+{
+	/// <summary>
+	/// Full-duplex channel for sending and receiving messages
+	/// </summary>
+	public interface IComputeMessageChannel : IAsyncDisposable
+	{
+		/// <summary>
+		/// Reads a message from the channel
+		/// </summary>
+		/// <param name="cancellationToken">Cancellation token for the operation</param>
+		/// <returns>Data for a message that was read. Must be disposed.</returns>
+		ValueTask<IComputeMessage> ReceiveAsync(CancellationToken cancellationToken = default);
+
+		/// <summary>
+		/// Creates a new builder for a message
+		/// </summary>
+		/// <param name="type">Type of the message</param>
+		/// <param name="maxSize">Maximum size of the message that will be written</param>
+		/// <returns>New builder for messages</returns>
+		IComputeMessageBuilder CreateMessage(ComputeMessageType type, int maxSize = 1024);
+	}
+
+	/// <summary>
+	/// Writer for compute messages
+	/// </summary>
+	public interface IComputeMessageBuilder : IMemoryWriter, IDisposable
+	{
+		/// <summary>
+		/// Sends the current message
+		/// </summary>
+		void Send();
+	}
+
+	/// <summary>
+	/// Extension methods to allow creating channels from leases
+	/// </summary>
+	public static class ComputeChannelExtensions
+	{
+		/// <summary>
+		/// Forwards an existing message across a channel
+		/// </summary>
+		/// <param name="channel">Channel to send on</param>
+		/// <param name="message">The message to be sent</param>
+		public static void Send(this IComputeMessageChannel channel, IComputeMessage message)
+		{
+			using (IComputeMessageBuilder builder = channel.CreateMessage(message.Type, message.Data.Length))
+			{
+				builder.WriteFixedLengthBytes(message.Data.Span);
+				builder.Send();
+			}
+		}
+	}
+}
diff --git a/Engine/Source/Programs/Shared/EpicGames.Horde/Compute/IComputeSocket.cs b/Engine/Source/Programs/Shared/EpicGames.Horde/Compute/IComputeSocket.cs
--- a/Engine/Source/Programs/Shared/EpicGames.Horde/Compute/IComputeSocket.cs
+++ b/Engine/Source/Programs/Shared/EpicGames.Horde/Compute/IComputeSocket.cs
@@ -1,9 +1,9 @@
 // Copyright Epic Games, Inc. All Rights Reserved.
 
 using System;
-using System.Threading;
-using System.Threading.Tasks;
-using EpicGames.Horde.Compute.Buffers;
+using System.Globalization;
+using EpicGames.Horde.Compute.Sockets;
+using Microsoft.Extensions.Logging;
 
 namespace EpicGames.Horde.Compute
 {
@@ -13,88 +13,137 @@ namespace EpicGames.Horde.Compute
 	public interface IComputeSocket : IAsyncDisposable
 	{
 		/// <summary>
-		/// Attaches buffers to a channel.
+		/// Creates a new buffer of the given capacity, suitable for this socket type
 		/// </summary>
-		/// <param name="channelId">Identifier of the channel</param>
-		/// <param name="sendBufferReader">The reader to attach</param>
-		/// <param name="receiveBufferWriter">The writer to attach</param>
-		/// <returns>Identifier for the message channel</returns>
-		void AttachBuffers(int channelId, IComputeBufferReader? sendBufferReader, IComputeBufferWriter? receiveBufferWriter);
+		/// <param name="capacity">Capacity of the buffer</param>
+		IComputeBuffer CreateBuffer(long capacity);
 
 		/// <summary>
-		/// Creates a channel to receive data on the given channel id
+		/// Attaches a buffer to receive data.
 		/// </summary>
-		/// <param name="channelId">Identifier of the channel</param>
-		/// <param name="sendBuffer">Buffer for sending messages</param>
-		/// <param name="receiveBuffer">Buffer for receiving messages</param>
-		IComputeChannel AttachMessageChannel(int channelId, IComputeBuffer sendBuffer, IComputeBuffer receiveBuffer);
+		/// <param name="channelId">Channel to receive data on</param>
+		/// <param name="writer">Writer for the received data</param>
+		void AttachRecvBuffer(int channelId, IComputeBufferWriter writer);
+
+		/// <summary>
+		/// Attaches a buffer to receive data.
+		/// </summary>
+		/// <param name="channelId">Channel to receive data on</param>
+		/// <param name="reader">Reader for incoming data</param>
+		void AttachSendBuffer(int channelId, IComputeBufferReader reader);
 	}
 
 	/// <summary>
-	/// Extension methods for sockets
+	/// Utility methods for creating sockets
 	/// </summary>
-	public static class ComputeSocketExtensions
+	public static class ComputeSocket
 	{
-		class HeapBufferChannel : IComputeChannel
-		{
-			readonly IComputeChannel _inner;
-			readonly HeapBuffer _receiveBuffer;
-			readonly HeapBuffer _sendBuffer;
+		/// <summary>
+		/// Environment variable for channel info for communicating with the worker process
+		/// </summary>
+		public const string WorkerIpcEnvVar = "UE_HORDE_COMPUTE_IPC";
 
-			public HeapBufferChannel(IComputeChannel inner, HeapBuffer receiveBuffer, HeapBuffer sendBuffer)
+		/// <summary>
+		/// Creates a socket for a worker
+		/// </summary>
+		/// <param name="logger">Logger for diagnostic messages</param>
+		public static IComputeSocket ConnectAsWorker(ILogger logger)
+		{
+			string? handle = Environment.GetEnvironmentVariable(WorkerIpcEnvVar);
+			if (handle == null)
 			{
-				_inner = inner;
-				_sendBuffer = sendBuffer;
-				_receiveBuffer = receiveBuffer;
+				throw new InvalidOperationException($"Environment variable {WorkerIpcEnvVar} is not defined; cannot connect as worker.");
 			}
 
-			/// <inheritdoc/>
-			public async ValueTask DisposeAsync()
+			try
 			{
-				await _inner.DisposeAsync();
-
-				_sendBuffer.Dispose();
-				_receiveBuffer.Dispose();
+				IpcComputeMessageChannel channel = IpcComputeMessageChannel.FromStringHandle(handle, logger);
+				return new WorkerComputeSocket(channel, logger);
+			}
+			catch(Exception ex)
+			{
+				throw new Exception($"While connecting using '{handle}'", ex);
 			}
-
-			/// <inheritdoc/>
-			public ValueTask<IComputeMessage> ReceiveAsync(CancellationToken cancellationToken) => _inner.ReceiveAsync(cancellationToken);
-
-			/// <inheritdoc/>
-			public IComputeMessageBuilder CreateMessage(ComputeMessageType type, int sizeHint = 0) => _inner.CreateMessage(type, sizeHint);
 		}
+	}
+
+	/// <summary>
+	/// Extension methods for <see cref="IComputeMessageChannel"/>
+	/// </summary>
+	public static class ComputeSocketExtensions
+	{
+		/// <summary>
+		/// Attaches a buffer to receive data.
+		/// </summary>
+		/// <param name="socket">Socket to attach to</param>
+		/// <param name="channelId">Channel to receive data on</param>
+		/// <param name="capacity">Capacity for the buffer</param>
+		public static IComputeBufferReader AttachRecvBuffer(this IComputeSocket socket, int channelId, long capacity) => AttachRecvBuffer(socket, channelId, socket.CreateBuffer(capacity));
 
 		/// <summary>
-		/// Creates a channel to receive data on the given channel id
+		/// Attaches a buffer to receive data.
 		/// </summary>
-		/// <param name="socket">Socket to attach a channel on</param>
-		/// <param name="channelId">Identifier of the channel</param>
-		public static IComputeChannel AttachMessageChannel(this IComputeSocket socket, int channelId)
+		/// <param name="socket">Socket to attach to</param>
+		/// <param name="channelId">Channel to receive data on</param>
+		/// <param name="buffer">Buffer to attach to the socket</param>
+		public static IComputeBufferReader AttachRecvBuffer(this IComputeSocket socket, int channelId, IComputeBuffer buffer)
 		{
-			const int MaxMessageSize = 64 * 1024;
-			const int BufferSize = MaxMessageSize * 3;
-
-			HeapBuffer receiveBuffer = new HeapBuffer(BufferSize);
-			HeapBuffer sendBuffer = new HeapBuffer(BufferSize);
-
-			IComputeChannel inner = socket.AttachMessageChannel(channelId, sendBuffer, receiveBuffer);
-			return new HeapBufferChannel(inner, receiveBuffer, sendBuffer);
+			(IComputeBufferReader reader, IComputeBufferWriter writer) = buffer.ToShared();
+			socket.AttachRecvBuffer(channelId, writer);
+			return reader;
 		}
 
 		/// <summary>
-		/// Attaches a send buffer to a socket. Data will be read from this buffer and replicated a receive buffer attached with the same id on the remote.
+		/// Attaches a buffer from which to send data.
 		/// </summary>
 		/// <param name="socket">Socket to attach to</param>
-		/// <param name="channelId">Identifier for the buffer</param>
-		/// <param name="reader">Source to read from</param>
-		public static void AttachSendBuffer(this IComputeSocket socket, int channelId, IComputeBufferReader reader) => socket.AttachBuffers(channelId, reader, null);
+		/// <param name="channelId">Channel to receive data on</param>
+		/// <param name="capacity">Capacity for the buffer</param>
+		public static IComputeBufferWriter AttachSendBuffer(this IComputeSocket socket, int channelId, long capacity) => AttachSendBuffer(socket, channelId, socket.CreateBuffer(capacity));
 
 		/// <summary>
-		/// Attaches a receive buffer to a socket. Data will be read into this buffer from the other end of the lease.
+		/// Attaches a buffer from which to send data.
 		/// </summary>
 		/// <param name="socket">Socket to attach to</param>
-		/// <param name="channelId">Identifier for the buffer</param>
-		/// <param name="writer">The buffer to attach</param>
-		public static void AttachReceiveBuffer(this IComputeSocket socket, int channelId, IComputeBufferWriter writer) => socket.AttachBuffers(channelId, null, writer);
+		/// <param name="channelId">Channel to receive data on</param>
+		/// <param name="buffer">Buffer for the data to send</param>
+		public static IComputeBufferWriter AttachSendBuffer(this IComputeSocket socket, int channelId, IComputeBuffer buffer)
+		{
+			(IComputeBufferReader reader, IComputeBufferWriter writer) = buffer.ToShared();
+			socket.AttachSendBuffer(channelId, reader);
+			return writer;
+		}
+
+		/// <summary>
+		/// Creates a message channel with the given identifier
+		/// </summary>
+		/// <param name="socket">Socket to create a channel for</param>
+		/// <param name="channelId">Identifier for the channel</param>
+		/// <param name="logger">Logger for the channel</param>
+		public static IComputeMessageChannel CreateMessageChannel(this IComputeSocket socket, int channelId, ILogger logger) => socket.CreateMessageChannel(channelId, 65536, logger);
+
+		/// <summary>
+		/// Creates a message channel with the given identifier
+		/// </summary>
+		/// <param name="socket">Socket to create a channel for</param>
+		/// <param name="channelId">Identifier for the channel</param>
+		/// <param name="bufferSize">Size of the send and receive buffer</param>
+		/// <param name="logger">Logger for the channel</param>
+		public static IComputeMessageChannel CreateMessageChannel(this IComputeSocket socket, int channelId, long bufferSize, ILogger logger) => CreateMessageChannel(socket, channelId, bufferSize, bufferSize, logger);
+
+		/// <summary>
+		/// Creates a message channel with the given identifier
+		/// </summary>
+		/// <param name="socket">Socket to create a channel for</param>
+		/// <param name="channelId">Identifier for the channel</param>
+		/// <param name="sendBufferSize">Size of the send buffer</param>
+		/// <param name="recvBufferSize">Size of the recieve buffer</param>
+		/// <param name="logger">Logger for the channel</param>
+		public static IComputeMessageChannel CreateMessageChannel(this IComputeSocket socket, int channelId, long sendBufferSize, long recvBufferSize, ILogger logger)
+		{
+			IComputeBufferReader reader = socket.AttachRecvBuffer(channelId, recvBufferSize);
+			IComputeBufferWriter writer = socket.AttachSendBuffer(channelId, sendBufferSize);
+			return new ComputeMessageChannel(reader, writer, logger);
+		}
 	}
 }
diff --git a/Engine/Source/Programs/Shared/EpicGames.Horde/Compute/IComputeTransport.cs b/Engine/Source/Programs/Shared/EpicGames.Horde/Compute/IComputeTransport.cs
--- a/Engine/Source/Programs/Shared/EpicGames.Horde/Compute/IComputeTransport.cs
+++ b/Engine/Source/Programs/Shared/EpicGames.Horde/Compute/IComputeTransport.cs
@@ -2,6 +2,7 @@
 
 using System;
 using System.Buffers;
+using System.IO;
 using System.Threading;
 using System.Threading.Tasks;
 
@@ -12,19 +13,24 @@ namespace EpicGames.Horde.Compute
 	/// </summary>
 	public interface IComputeTransport
 	{
+		/// <summary>
+		/// Position in the stream; used for debugging
+		/// </summary>
+		long Position { get; }
+
 		/// <summary>
 		/// Reads data from the underlying transport into an output buffer
 		/// </summary>
 		/// <param name="buffer">Buffer to read into</param>
 		/// <param name="cancellationToken">Cancellation token for the operation</param>
-		public ValueTask<int> ReadAsync(Memory<byte> buffer, CancellationToken cancellationToken);
+		ValueTask<int> ReadPartialAsync(Memory<byte> buffer, CancellationToken cancellationToken);
 
 		/// <summary>
 		/// Writes data to the underlying transport
 		/// </summary>
 		/// <param name="buffer">Buffer to be written</param>
 		/// <param name="cancellationToken">Cancellation token for the operation</param>
-		public ValueTask WriteAsync(ReadOnlySequence<byte> buffer, CancellationToken cancellationToken);
+		ValueTask WriteAsync(ReadOnlySequence<byte> buffer, CancellationToken cancellationToken);
 	}
 
 	/// <summary>
@@ -32,6 +38,46 @@ namespace EpicGames.Horde.Compute
 	/// </summary>
 	public static class ComputeTransportExtensions
 	{
+		/// <summary>
+		/// Fill the given buffer with data
+		/// </summary>
+		/// <param name="transport">Transport object</param>
+		/// <param name="buffer">Buffer to read into</param>
+		/// <param name="cancellationToken">Cancellation token for the operation</param>
+		public static async ValueTask ReadFullAsync(this IComputeTransport transport, Memory<byte> buffer, CancellationToken cancellationToken)
+		{
+			int read = 0;
+			while (read < buffer.Length)
+			{
+				int partialRead = await transport.ReadPartialAsync(buffer.Slice(read, buffer.Length - read), cancellationToken);
+				if (partialRead == 0)
+				{
+					throw new EndOfStreamException();
+				}
+				read += partialRead;
+			}
+		}
+
+		/// <summary>
+		/// Fill the given buffer with data
+		/// </summary>
+		/// <param name="transport">Transport object</param>
+		/// <param name="buffer">Buffer to read into</param>
+		/// <param name="cancellationToken">Cancellation token for the operation</param>
+		public static async ValueTask<bool> ReadOptionalAsync(this IComputeTransport transport, Memory<byte> buffer, CancellationToken cancellationToken)
+		{
+			int read = await transport.ReadPartialAsync(buffer, cancellationToken);
+			if (read == 0)
+			{
+				return false;
+			}
+			if (read < buffer.Length)
+			{
+				await transport.ReadFullAsync(buffer.Slice(read), cancellationToken);
+			}
+			return true;
+		}
+
 		/// <summary>
 		/// Writes data to the underlying transport
 		/// </summary>
diff --git a/Engine/Source/Programs/Shared/EpicGames.Horde/Compute/IpcComputeMessageChannel.cs b/Engine/Source/Programs/Shared/EpicGames.Horde/Compute/IpcComputeMessageChannel.cs
new file mode 100644
--- /dev/null
+++ b/Engine/Source/Programs/Shared/EpicGames.Horde/Compute/IpcComputeMessageChannel.cs
@@ -0,0 +1,136 @@
+// Copyright Epic Games, Inc. All Rights Reserved.
+
+using System;
+using System.Globalization;
+using System.Threading.Tasks;
+using EpicGames.Core;
+using EpicGames.Horde.Compute.Buffers;
+using Microsoft.Extensions.Logging;
+
+namespace EpicGames.Horde.Compute
+{
+	/// <summary>
+	/// Message channel for communicating with child processes without using a socket
+	/// </summary>
+	public class IpcComputeMessageChannel : ComputeMessageChannel
+	{
+		readonly IntPtr _parentProcessHandle;
+		readonly SharedMemoryBuffer _recvBuffer;
+		readonly SharedMemoryBuffer _sendBuffer;
+
+		static IntPtr s_currentProcessHandle;
+
+		/// <summary>
+		/// 
+		/// </summary>
+		/// <param name="capacity"></param>
+		/// <param name="logger"></param>
+		public IpcComputeMessageChannel(long capacity, ILogger logger)
+			: this(IntPtr.Zero, new SharedMemoryBuffer(capacity), new SharedMemoryBuffer(capacity), logger)
+		{
+		}
+
+		/// <summary>
+		/// Constructor
+		/// </summary>
+		internal IpcComputeMessageChannel(IntPtr parentProcessHandle, SharedMemoryBuffer recvBuffer, SharedMemoryBuffer sendBuffer, ILogger logger)
+			: base(recvBuffer.Reader, sendBuffer.Writer, logger)
+		{
+			_parentProcessHandle = parentProcessHandle;
+			_recvBuffer = recvBuffer;
+			_sendBuffer = sendBuffer;
+		}
+
+		/// <inheritdoc/>
+		protected override ValueTask DisposeAsync(bool disposing)
+		{
+			_recvBuffer.Dispose();
+			_sendBuffer.Dispose();
+			return new ValueTask();
+		}
+
+		/// <summary>
+		/// Creates an IPC channel from a string handle
+		/// </summary>
+		/// <param name="handle">Handle to parse</param>
+		/// <param name="logger">Logger for the new channel</param>
+		public static IpcComputeMessageChannel FromStringHandle(string handle, ILogger logger)
+		{
+			string[] handles = handle.Split('/');
+			if (handles.Length != 3)
+			{
+				throw new ArgumentException($"Invalid handle for IPC channel: {handle}", nameof(handle));
+			}
+
+			IntPtr parentProcessHandle = new IntPtr((long)UInt64.Parse(handles[0], NumberStyles.None, null));
+			SharedMemoryBuffer recvBuffer = SharedMemoryBuffer.OpenIpcHandle(handles[1]);
+			SharedMemoryBuffer sendBuffer = SharedMemoryBuffer.OpenIpcHandle(handles[2]);
+
+			return new IpcComputeMessageChannel(parentProcessHandle, sendBuffer, recvBuffer, logger); // Note: Swapping send/recv buffer here
+		}
+
+		/// <summary>
+		/// Gets a string handle that can be used to open this channel in another process
+		/// </summary>
+		public string GetStringHandle()
+		{
+			if (s_currentProcessHandle == IntPtr.Zero)
+			{
+				IntPtr currentProcessHandle = Native.GetCurrentProcess();
+				Native.DuplicateHandle(currentProcessHandle, currentProcessHandle, currentProcessHandle, out s_currentProcessHandle, 0, true, Native.DUPLICATE_SAME_ACCESS);
+			}
+			return $"{(ulong)s_currentProcessHandle.ToInt64()}/{_recvBuffer.GetIpcHandle()}/{_sendBuffer.GetIpcHandle()}";
+		}
+
+		/// <summary>
+		/// Force the stream closed
+		/// </summary>
+		public void ForceComplete() => _recvBuffer.FinishWriting();
+
+		/// <summary>
+		/// Registers a receive buffer with the agent process
+		/// </summary>
+		public void AttachRecvBuffer(int channelId, SharedMemoryBuffer buffer)
+		{
+			string handle = buffer.GetIpcHandle(_parentProcessHandle);
+
+			using IComputeMessageBuilder message = CreateMessage(ComputeMessageType.AttachRecvBuffer, handle.Length + 20);
+			message.WriteInt32(channelId);
+			message.WriteString(handle);
+			message.Send();
+		}
+
+		/// <summary>
+		/// Parse a <see cref="AttachRecvBufferRequest"/> message
+		/// </summary>
+		public static AttachRecvBufferRequest ParseAttachRecvBuffer(IComputeMessage message)
+		{
+			int channelId = message.ReadInt32();
+			string handle = message.ReadString();
+			return new AttachRecvBufferRequest(channelId, handle);
+		}
+
+		/// <summary>
+		/// Registers a send buffer with the agent process
+		/// </summary>
+		public void AttachSendBuffer(int channelId, SharedMemoryBuffer buffer)
+		{
+			string handle = buffer.GetIpcHandle(_parentProcessHandle);
+
+			using IComputeMessageBuilder message = CreateMessage(ComputeMessageType.AttachSendBuffer, handle.Length + 20);
+			message.WriteInt32(channelId);
+			message.WriteString(handle);
+			message.Send();
+		}
+
+		/// <summary>
+		/// Parse a <see cref="AttachSendBufferRequest"/> message
+		/// </summary>
+		public static AttachSendBufferRequest ParseAttachSendBuffer(IComputeMessage message)
+		{
+			int channelId = message.ReadInt32();
+			string handle = message.ReadString();
+			return new AttachSendBufferRequest(channelId, handle);
+		}
+	}
+}
diff --git a/Engine/Source/Programs/Shared/EpicGames.Horde/Compute/Sockets/ClientComputeSocket.cs b/Engine/Source/Programs/Shared/EpicGames.Horde/Compute/Sockets/ClientComputeSocket.cs
new file mode 100644
--- /dev/null
+++ b/Engine/Source/Programs/Shared/EpicGames.Horde/Compute/Sockets/ClientComputeSocket.cs
@@ -0,0 +1,37 @@
+// Copyright Epic Games, Inc. All Rights Reserved.
+
+using System;
+using EpicGames.Horde.Compute.Buffers;
+using Microsoft.Extensions.Logging;
+
+namespace EpicGames.Horde.Compute
+{
+	/// <summary>
+	/// Manages a set of readers and writers to buffers across a transport layer
+	/// </summary>
+	public sealed class ClientComputeSocket : ComputeSocketBase, IAsyncDisposable
+	{
+		/// <summary>
+		/// Constructor
+		/// </summary>
+		/// <param name="transport">Transport to communicate with the remote</param>
+		/// <param name="logger">Logger for trace output</param>
+		public ClientComputeSocket(IComputeTransport transport, ILogger logger)
+			: base(transport, logger)
+		{
+		}
+
+		/// <inheritdoc/>
+		public override IComputeBuffer CreateBuffer(long capacity)
+		{
+			if (capacity > Int32.MaxValue)
+			{
+				return new SharedMemoryBuffer(capacity);
+			}
+			else
+			{
+				return new PooledBuffer((int)capacity);
+			}
+		}
+	}
+}
diff --git a/Engine/Source/Programs/Shared/EpicGames.Horde/Compute/Sockets/WorkerComputeSocket.cs b/Engine/Source/Programs/Shared/EpicGames.Horde/Compute/Sockets/WorkerComputeSocket.cs
new file mode 100644
--- /dev/null
+++ b/Engine/Source/Programs/Shared/EpicGames.Horde/Compute/Sockets/WorkerComputeSocket.cs
@@ -0,0 +1,48 @@
+// Copyright Epic Games, Inc. All Rights Reserved.
+
+using System;
+using System.Threading.Tasks;
+using EpicGames.Horde.Compute.Buffers;
+using Microsoft.Extensions.Logging;
+
+namespace EpicGames.Horde.Compute.Sockets
+{
+	sealed class WorkerComputeSocket : IComputeSocket
+	{
+		readonly IpcComputeMessageChannel _ipcMessageChannel;
+		readonly ILogger _logger;
+
+		/// <summary>
+		/// Constructor
+		/// </summary>
+		/// <param name="ipcMessageChannel">Channel for IPC messages</param>
+		/// <param name="logger">Logger for trace output</param>
+		public WorkerComputeSocket(IpcComputeMessageChannel ipcMessageChannel, ILogger logger)
+		{
+			_ipcMessageChannel = ipcMessageChannel;
+			_logger = logger;
+		}
+
+		/// <inheritdoc/>
+		public ValueTask DisposeAsync() => _ipcMessageChannel.DisposeAsync();
+
+		/// <inheritdoc/>
+		public IComputeBuffer CreateBuffer(long capacity) => new SharedMemoryBuffer(capacity);
+
+		/// <inheritdoc/>
+		public void AttachRecvBuffer(int channelId, IComputeBufferWriter writer)
+		{
+			SharedMemoryBuffer buffer = (SharedMemoryBuffer)writer.Buffer;
+			_ipcMessageChannel.AttachRecvBuffer(channelId, buffer);
+			_logger.LogTrace("Attached recv buffer to channel {ChannelId}", channelId);
+		}
+
+		/// <inheritdoc/>
+		public void AttachSendBuffer(int channelId, IComputeBufferReader reader)
+		{
+			SharedMemoryBuffer buffer = (SharedMemoryBuffer)reader.Buffer;
+			_ipcMessageChannel.AttachSendBuffer(channelId, buffer);
+			_logger.LogTrace("Attached send buffer to channel {ChannelId}", channelId);
+		}
+	}
+}
diff --git a/Engine/Source/Programs/Shared/EpicGames.Horde/Compute/Transports/AesTransport.cs b/Engine/Source/Programs/Shared/EpicGames.Horde/Compute/Transports/AesTransport.cs
--- a/Engine/Source/Programs/Shared/EpicGames.Horde/Compute/Transports/AesTransport.cs
+++ b/Engine/Source/Programs/Shared/EpicGames.Horde/Compute/Transports/AesTransport.cs
@@ -46,6 +46,9 @@ namespace EpicGames.Horde.Compute.Transports
 		const int WritePacketSize = 64 * 1024;
 		readonly BackgroundTask _backgroundReadTask;
 
+		/// <inheritdoc/>
+		public long Position { get; private set; }
+
 		/// <summary>
 		/// Constructor
 		/// </summary>
@@ -97,7 +100,7 @@ namespace EpicGames.Horde.Compute.Transports
 			for (; ; )
 			{
 				// Read more data into the buffer
-				int read = await _inner.ReadAsync(memory.Slice(size), cancellationToken);
+				int read = await _inner.ReadPartialAsync(memory.Slice(size), cancellationToken);
 				if (read == 0)
 				{
 					break;
@@ -140,23 +143,15 @@ namespace EpicGames.Horde.Compute.Transports
 		}
 
 		/// <inheritdoc/>
-		public async ValueTask<int> ReadAsync(Memory<byte> buffer, CancellationToken cancellationToken)
+		public async ValueTask<int> ReadPartialAsync(Memory<byte> buffer, CancellationToken cancellationToken)
 		{
 			int sizeRead = 0;
-			while (sizeRead < buffer.Length)
+			while (sizeRead == 0)
 			{
 				// Try to get more data into the read buffer
 				if (_readBuffer.Length == 0)
 				{
-					ReadResult result;
-					if (sizeRead == 0)
-					{
-						result = await _readPipe.Reader.ReadAsync(cancellationToken);
-					}
-					else if (!_readPipe.Reader.TryRead(out result) || result.IsCompleted)
-					{
-						return sizeRead;
-					}
+					ReadResult result = await _readPipe.Reader.ReadAsync(cancellationToken);
 					_readBuffer = result.Buffer;
 				}
 
@@ -166,6 +161,7 @@ namespace EpicGames.Horde.Compute.Transports
 				buffer = buffer.Slice(copySize);
 				_readBuffer = _readBuffer.Slice(copySize);
 				sizeRead += copySize;
+				Position += copySize;
 			}
 			return sizeRead;
 		}
@@ -186,6 +182,8 @@ namespace EpicGames.Horde.Compute.Transports
 				Memory<byte> writeBuffer = _writeBufferOwner.Memory.Slice(_writeBufferOffset, HeaderLength + plainText.Length + FooterLength);
 				Encrypt(plainText.Span, writeBuffer.Span);
 
+				Position += plainText.Length;
+
 				await _lastWriteTask;
 				_lastWriteTask = _inner.WriteAsync(writeBuffer, cancellationToken).AsTask();
 
diff --git a/Engine/Source/Programs/Shared/EpicGames.Horde/Compute/Transports/PipeTransport.cs b/Engine/Source/Programs/Shared/EpicGames.Horde/Compute/Transports/PipeTransport.cs
new file mode 100644
--- /dev/null
+++ b/Engine/Source/Programs/Shared/EpicGames.Horde/Compute/Transports/PipeTransport.cs
@@ -0,0 +1,90 @@
+// Copyright Epic Games, Inc. All Rights Reserved.
+
+using System;
+using System.Buffers;
+using System.IO.Pipelines;
+using System.Threading;
+using System.Threading.Tasks;
+
+namespace EpicGames.Horde.Compute.Transports
+{
+	/// <summary>
+	/// Implementation of <see cref="IComputeTransport"/> for communicating over a <see cref="Pipe"/>. 
+	/// (Note: this uses a .NET in-process pipe, not an IPC pipe).
+	/// </summary>
+	public class PipeTransport : IComputeTransport
+	{
+		readonly PipeReader _reader;
+		readonly PipeWriter _writer;
+
+		ReadOnlySequence<byte> _readBuffer = ReadOnlySequence<byte>.Empty;
+		SequencePosition? _readEnd;
+
+		/// <inheritdoc/>
+		public long Position { get; private set; }
+
+		/// <summary>
+		/// Constructor
+		/// </summary>
+		/// <param name="reader">Reader for the pipe</param>
+		/// <param name="writer">Writer for the pipe</param>
+		public PipeTransport(PipeReader reader, PipeWriter writer)
+		{
+			_reader = reader;
+			_writer = writer;
+		}
+
+		/// <inheritdoc/>
+		public async ValueTask<int> ReadPartialAsync(Memory<byte> buffer, CancellationToken cancellationToken)
+		{
+			int sizeRead = 0;
+			while (sizeRead < buffer.Length)
+			{
+				// Try to get more data into the read buffer
+				if (_readBuffer.Length == 0)
+				{
+					if (_readEnd != null)
+					{
+						_reader.AdvanceTo(_readEnd.Value);
+						_readEnd = null;
+					}
+
+					ReadResult result;
+					if (sizeRead == 0)
+					{
+						result = await _reader.ReadAsync(cancellationToken);
+					}
+					else if (!_reader.TryRead(out result) || result.IsCompleted)
+					{
+						return sizeRead;
+					}
+
+					_readBuffer = result.Buffer;
+					_readEnd = _readBuffer.End;
+				}
+
+				// Copy as much of the next chunk as we can
+				int copySize = Math.Min(buffer.Length, _readBuffer.First.Length);
+				if (copySize > 0)
+				{
+					_readBuffer.First.Slice(0, copySize).CopyTo(buffer);
+					buffer = buffer.Slice(copySize);
+					_readBuffer = _readBuffer.Slice(copySize);
+					sizeRead += copySize;
+				}
+			}
+			Position += sizeRead;
+			return sizeRead;
+		}
+
+		/// <inheritdoc/>
+		public async ValueTask WriteAsync(ReadOnlySequence<byte> buffer, CancellationToken cancellationToken)
+		{
+			foreach (ReadOnlyMemory<byte> memory in buffer)
+			{
+				await _writer.WriteAsync(memory, cancellationToken);
+				Position += memory.Length;
+			}
+		}
+	}
+}
diff --git a/Engine/Source/Programs/Shared/EpicGames.Horde/Compute/Transports/StreamTransport.cs b/Engine/Source/Programs/Shared/EpicGames.Horde/Compute/Transports/StreamTransport.cs
--- a/Engine/Source/Programs/Shared/EpicGames.Horde/Compute/Transports/StreamTransport.cs
+++ b/Engine/Source/Programs/Shared/EpicGames.Horde/Compute/Transports/StreamTransport.cs
@@ -15,6 +15,9 @@ namespace EpicGames.Horde.Compute.Transports
 	{
 		readonly Stream _stream;
 
+		/// <inheritdoc/>
+		public long Position { get; private set; }
+
 		/// <summary>
 		/// Constructor
 		/// </summary>
@@ -22,7 +25,12 @@ namespace EpicGames.Horde.Compute.Transports
 		public StreamTransport(Stream stream) => _stream = stream;
 
 		/// <inheritdoc/>
-		public async ValueTask<int> ReadAsync(Memory<byte> buffer, CancellationToken cancellationToken) => await _stream.ReadAsync(buffer, cancellationToken);
+		public async ValueTask<int> ReadPartialAsync(Memory<byte> buffer, CancellationToken cancellationToken)
+		{
+			int length = await _stream.ReadAsync(buffer, cancellationToken);
+			Position += length;
+			return length;
+		}
 
 		/// <inheritdoc/>
 		public async ValueTask WriteAsync(ReadOnlySequence<byte> buffer, CancellationToken cancellationToken)
@@ -30,6 +38,7 @@ namespace EpicGames.Horde.Compute.Transports
 			foreach (ReadOnlyMemory<byte> memory in buffer)
 			{
 				await _stream.WriteAsync(memory, cancellationToken);
+				Position += memory.Length;
 			}
 		}
 	}
diff --git a/Engine/Source/Programs/Shared/EpicGames.Horde/Compute/Transports/TcpTransport.cs b/Engine/Source/Programs/Shared/EpicGames.Horde/Compute/Transports/TcpTransport.cs
--- a/Engine/Source/Programs/Shared/EpicGames.Horde/Compute/Transports/TcpTransport.cs
+++ b/Engine/Source/Programs/Shared/EpicGames.Horde/Compute/Transports/TcpTransport.cs
@@ -16,6 +16,9 @@ namespace EpicGames.Horde.Compute.Transports
 	{
 		readonly Socket _socket;
 
+		/// <inheritdoc/>
+		public long Position { get; private set; }
+
 		/// <summary>
 		/// Constructor
 		/// </summary>
@@ -23,7 +26,12 @@ namespace EpicGames.Horde.Compute.Transports
 		public TcpTransport(Socket socket) => _socket = socket;
 
 		/// <inheritdoc/>
-		public ValueTask<int> ReadAsync(Memory<byte> buffer, CancellationToken cancellationToken) => _socket.ReceiveAsync(buffer, SocketFlags.None, cancellationToken);
+		public async ValueTask<int> ReadPartialAsync(Memory<byte> buffer, CancellationToken cancellationToken)
+		{
+			int read = await _socket.ReceiveAsync(buffer, SocketFlags.None, cancellationToken);
+			Position += read;
+			return read;
+		}
 
 		/// <inheritdoc/>
 		public async ValueTask WriteAsync(ReadOnlySequence<byte> buffer, CancellationToken cancellationToken)
@@ -31,6 +39,7 @@ namespace EpicGames.Horde.Compute.Transports
 			foreach (ReadOnlyMemory<byte> memory in buffer)
 			{
 				await _socket.SendMessageAsync(memory, SocketFlags.None, cancellationToken);
+				Position += memory.Length;
 			}
 		}
 	}
