Horde: Missing file.

diff --git a/Engine/Source/Programs/Shared/EpicGames.Horde/Compute/Native.cs b/Engine/Source/Programs/Shared/EpicGames.Horde/Compute/Native.cs
new file mode 100644
--- /dev/null
+++ b/Engine/Source/Programs/Shared/EpicGames.Horde/Compute/Native.cs
@@ -0,0 +1,54 @@
+// Copyright Epic Games, Inc. All Rights Reserved.
+
+using System;
+using System.IO;
+using System.Runtime.InteropServices;
+using System.Threading;
+using Microsoft.Win32.SafeHandles;
+
+namespace EpicGames.Horde.Compute
+{
+	internal class Native
+	{
+		[StructLayout(LayoutKind.Sequential)]
+		public class SECURITY_ATTRIBUTES
+		{
+			public int nLength;
+			public IntPtr lpSecurityDescriptor;
+			public int bInheritHandle;
+		}
+
+		[DllImport("kernel32.dll")]
+		public static extern SafeWaitHandle CreateEvent(SECURITY_ATTRIBUTES lpEventAttributes, bool bManualReset, bool bInitialState, string? lpName);
+
+		[DllImport("kernel32.dll")]
+		public static extern bool SetEvent(SafeWaitHandle hEvent);
+
+		public const uint DUPLICATE_SAME_ACCESS = 2;
+
+		[DllImport("kernel32.dll", SetLastError = true)]
+		public static extern IntPtr GetCurrentProcess();
+
+		[DllImport("kernel32.dll", SetLastError = true)]
+		[return: MarshalAs(UnmanagedType.Bool)]
+		public static extern bool DuplicateHandle(IntPtr hSourceProcessHandle, IntPtr hSourceHandle, IntPtr hTargetProcessHandle, out IntPtr lpTargetHandle, uint dwDesiredAccess, [MarshalAs(UnmanagedType.Bool)] bool bInheritHandle, uint dwOptions);
+
+		public class EventHandle : WaitHandle
+		{
+			public EventHandle(HandleInheritability handleInheritability)
+			{
+				Native.SECURITY_ATTRIBUTES securityAttributes = new Native.SECURITY_ATTRIBUTES();
+				securityAttributes.nLength = Marshal.SizeOf(securityAttributes);
+				securityAttributes.bInheritHandle = (handleInheritability == HandleInheritability.Inheritable) ? 1 : 0;
+				SafeWaitHandle = CreateEvent(securityAttributes, false, false, null);
+			}
+
+			public EventHandle(IntPtr handle, bool ownsHandle)
+			{
+				SafeWaitHandle = new SafeWaitHandle(handle, ownsHandle);
+			}
+
+			public void Set() => SetEvent(SafeWaitHandle);
+		}
+	}
+}
