Horde: Return leases directly from compute clients, rather than invoking a callback with the allocated lease.

diff --git a/Engine/Source/Programs/Horde/Horde.Agent/Commands/Compute/ComputeCommand.cs b/Engine/Source/Programs/Horde/Horde.Agent/Commands/Compute/ComputeCommand.cs
--- a/Engine/Source/Programs/Horde/Horde.Agent/Commands/Compute/ComputeCommand.cs
+++ b/Engine/Source/Programs/Horde/Horde.Agent/Commands/Compute/ComputeCommand.cs
@@ -58,8 +58,13 @@ namespace Horde.Agent.Commands.Compute
 				requirements = new Requirements(Condition.Parse(Requirements));
 			}
 
-			bool result = await client.ExecuteAsync(new ClusterId(ClusterId), requirements, HandleRequestAsync, CancellationToken.None);
+			await using IComputeLease? lease = await client.TryAssignWorkerAsync(new ClusterId(ClusterId), requirements, CancellationToken.None);
+			if (lease == null)
+			{
+				throw new Exception("Unable to create lease");
+			}
 
+			bool result = await HandleRequestAsync(lease, CancellationToken.None);
 			return result ? 0 : 1;
 		}
 
@@ -79,7 +84,7 @@ namespace Horde.Agent.Commands.Compute
 			}
 			else if (Loopback)
 			{
-				return new LoopbackComputeClient(2000, logger);
+				return new AgentComputeClient(Assembly.GetExecutingAssembly().Location, 2000, logger);
 			}
 			else
 			{
@@ -96,74 +101,6 @@ namespace Horde.Agent.Commands.Compute
 			client.DefaultRequestHeaders.Authorization = new AuthenticationHeaderValue("Bearer", profile.Token);
 			return client;
 		}
-
-		class LoopbackComputeClient : IComputeClient
-		{
-			readonly int _port;
-			readonly ILogger _logger;
-
-			public LoopbackComputeClient(int port, ILogger logger)
-			{
-				_port = port;
-				_logger = logger;
-			}
-
-			public ValueTask DisposeAsync() => new ValueTask();
-
-			public async Task<TResult> ExecuteAsync<TResult>(ClusterId clusterId, Requirements? requirements, Func<IComputeLease, CancellationToken, Task<TResult>> handler, CancellationToken cancellationToken)
-			{
-				_logger.LogInformation("** CLIENT **");
-
-				TResult result;
-				using (Socket listener = new Socket(SocketType.Stream, ProtocolType.IP))
-				{
-					listener.Bind(new IPEndPoint(IPAddress.Loopback, _port));
-					listener.Listen();
-
-					using (CancellationTokenSource cancellationSource = CancellationTokenSource.CreateLinkedTokenSource(cancellationToken))
-					{
-						Task<TResult> listenTask = ListenAsync<TResult>(listener, handler, cancellationSource.Token);
-						try
-						{
-							using ManagedProcessGroup group = new ManagedProcessGroup();
-
-							using Process process = new Process();
-							process.StartInfo.FileName = "dotnet";
-							process.StartInfo.ArgumentList.Add(Assembly.GetExecutingAssembly().Location);
-							process.StartInfo.ArgumentList.Add("computeworker");
-							process.StartInfo.ArgumentList.Add($"-port={_port}");
-							process.StartInfo.UseShellExecute = true;
-							process.Start();
-
-							group.AddProcess(process.Handle);
-
-							await process.WaitForExitAsync(cancellationToken);
-						}
-						finally
-						{
-							cancellationSource.Cancel();
-						}
-						result = await listenTask;
-					}
-				}
-				return result;
-			}
-
-			async Task<TResult> ListenAsync<TResult>(Socket listener, Func<IComputeLease, CancellationToken, Task<TResult>> handler, CancellationToken cancellationToken)
-			{
-				TResult result;
-				using (Socket tcpSocket = await listener.AcceptAsync(cancellationToken))
-				{
-					await using (ClientComputeSocket socket = new ClientComputeSocket(new TcpTransport(tcpSocket), _logger))
-					{
-						await using ComputeLease lease = new ComputeLease(new List<string>(), new Dictionary<string, int>(), socket);
-						result = await handler(lease, cancellationToken);
-						await socket.CloseAsync(cancellationToken);
-					}
-				}
-				return result;
-			}
-		}
 	}
 
 	/// <summary>
diff --git a/Engine/Source/Programs/Horde/Samples/RemoteClient/RemoteClientApp.cs b/Engine/Source/Programs/Horde/Samples/RemoteClient/RemoteClientApp.cs
--- a/Engine/Source/Programs/Horde/Samples/RemoteClient/RemoteClientApp.cs
+++ b/Engine/Source/Programs/Horde/Samples/RemoteClient/RemoteClientApp.cs
@@ -24,43 +24,40 @@ namespace RemoteClient
 		static async Task Main()
 		{
 			ILogger logger = Log.Logger;
+			CancellationToken cancellationToken = CancellationToken.None;
 
+			// Create the client to handle our requests
 			await using IComputeClient client = CreateClient(true, logger);
-			await client.ExecuteAsync(new ClusterId("default"), null, (lease, ctx) => ExecuteAsync(lease, logger, ctx), CancellationToken.None);
-		}
 
-		static IComputeClient CreateClient(bool loopback, ILogger logger)
-		{
-			if (loopback)
-			{
-				return new AgentComputeClient(HordeAgentFile.FullName, 2000, logger);
-			}
-			else
+			// Allocate a worker
+			await using IComputeLease? lease = await client.TryAssignWorkerAsync(new ClusterId("default"), null, cancellationToken);
+			if (lease == null)
 			{
-				return new ServerComputeClient(new Uri("https://localhost:5001"), null, logger);
+				logger.LogInformation("Unable to connect to remote");
+				return;
 			}
-		}
 
-		static async Task<int> ExecuteAsync(IComputeLease lease, ILogger logger, CancellationToken cancellationToken)
-		{
-			IComputeSocket socket = lease.Socket;
-			await using (IComputeMessageChannel channel = socket.CreateMessageChannel(0, 30 * 1024 * 1024, logger))
+			// Create a message channel on channel id 0. The Horde Agent always listens on this channel for requests.
+			const int ControlChannelId = 0;
+			await using IComputeMessageChannel channel = lease.Socket.CreateMessageChannel(ControlChannelId, 30 * 1024 * 1024, logger);
+
+			// Upload the sandbox
+			MemoryStorageClient storage = new MemoryStorageClient();
+			using (TreeWriter treeWriter = new TreeWriter(storage))
 			{
-				// Upload the sandbox
-				MemoryStorageClient storage = new MemoryStorageClient();
-				using (TreeWriter treeWriter = new TreeWriter(storage))
-				{
-					DirectoryNode sandbox = new DirectoryNode();
-					await sandbox.CopyFromDirectoryAsync(RemoteServerFile.Directory.ToDirectoryInfo(), new ChunkingOptions(), treeWriter, null, cancellationToken);
-					NodeHandle handle = await treeWriter.FlushAsync(sandbox, cancellationToken);
-					await channel.UploadFilesAsync("", handle.Locator, storage, cancellationToken);
-				}
+				DirectoryNode sandbox = new DirectoryNode();
+				await sandbox.CopyFromDirectoryAsync(RemoteServerFile.Directory.ToDirectoryInfo(), new ChunkingOptions(), treeWriter, null, cancellationToken);
+				NodeHandle handle = await treeWriter.FlushAsync(sandbox, cancellationToken);
+				await channel.UploadFilesAsync("", handle.Locator, storage, cancellationToken);
+			}
 
-				// Run the task remotely in the background and echo the output to the console
-				Task childProcessTask = channel.ExecuteAsync(RemoteServerFile.GetFileName(), new List<string>(), null, null, (string x) => logger.LogInformation("ChildProcess: {Message}", x), cancellationToken);
+			// Run the task remotely in the background and echo the output to the console
+			Task childProcessTask = channel.ExecuteAsync(RemoteServerFile.GetFileName(), new List<string>(), null, null, (string x) => logger.LogInformation("ChildProcess: {Message}", x), cancellationToken);
 
-				// Generate data into a buffer attached to channel 1. The remote server will echo them back to us as it receives them.
-				IComputeBufferWriter writer = socket.AttachSendBuffer(1, 1024 * 1024);
+			// Generate data into a buffer attached to channel 1. The remote server will echo them back to us as it receives them, then exit when the channel is complete/closed.
+			const int DataChannelId = 1;
+			using (IComputeBufferWriter writer = lease.Socket.AttachSendBuffer(DataChannelId, 1024 * 1024))
+			{
 				for (int idx = 0; idx < 100; idx++)
 				{
 					logger.LogInformation("Writing value: {Value}", idx);
@@ -68,12 +65,22 @@ namespace RemoteClient
 					writer.Advance(1);
 					await Task.Delay(1000, cancellationToken);
 				}
-				writer.MarkComplete();
+			}
 
-				// Wait for the child process to finish
-				await childProcessTask;
+			// Wait for the child process to finish
+			await childProcessTask;
+		}
+
+		static IComputeClient CreateClient(bool loopback, ILogger logger)
+		{
+			if (loopback)
+			{
+				return new AgentComputeClient(HordeAgentFile.FullName, 2000, logger);
+			}
+			else
+			{
+				return new ServerComputeClient(new Uri("https://localhost:5001"), null, logger);
 			}
-			return 0;
 		}
 	}
 }
\ No newline at end of file
diff --git a/Engine/Source/Programs/Shared/EpicGames.Horde/Compute/Clients/AgentComputeClient.cs b/Engine/Source/Programs/Shared/EpicGames.Horde/Compute/Clients/AgentComputeClient.cs
--- a/Engine/Source/Programs/Shared/EpicGames.Horde/Compute/Clients/AgentComputeClient.cs
+++ b/Engine/Source/Programs/Shared/EpicGames.Horde/Compute/Clients/AgentComputeClient.cs
@@ -5,6 +5,7 @@ using System.Collections.Generic;
 using System.Diagnostics;
 using System.Net;
 using System.Net.Sockets;
+using System.Runtime.CompilerServices;
 using System.Runtime.InteropServices;
 using System.Threading;
 using System.Threading.Tasks;
@@ -19,6 +20,29 @@ namespace EpicGames.Horde.Compute.Clients
 	/// </summary>
 	public sealed class AgentComputeClient : IComputeClient
 	{
+		class LeaseImpl : IComputeLease
+		{
+			readonly IAsyncEnumerator<IComputeSocket> _source;
+
+			/// <inheritdoc/>
+			public IReadOnlyList<string> Properties { get; } = new List<string>();
+
+			/// <inheritdoc/>
+			public IReadOnlyDictionary<string, int> AssignedResources => new Dictionary<string, int>();
+
+			/// <inheritdoc/>
+			public IComputeSocket Socket => _source.Current;
+
+			public LeaseImpl(IAsyncEnumerator<IComputeSocket> source) => _source = source;
+
+			/// <inheritdoc/>
+			public async ValueTask DisposeAsync()
+			{
+				await _source.MoveNextAsync();
+				await _source.DisposeAsync();
+			}
+		}
+
 		readonly string _hordeAgentAssembly;
 		readonly int _port;
 		readonly ILogger _logger;
@@ -26,12 +50,12 @@ namespace EpicGames.Horde.Compute.Clients
 		/// <summary>
 		/// Constructor
 		/// </summary>
-		/// <param name="hordeAgentFile">Path to the Horde Agent assembly</param>
+		/// <param name="hordeAgentAssembly">Path to the Horde Agent assembly</param>
 		/// <param name="port">Loopback port to connect on</param>
 		/// <param name="logger">Factory for logger instances</param>
-		public AgentComputeClient(string hordeAgentFile, int port, ILogger logger)
+		public AgentComputeClient(string hordeAgentAssembly, int port, ILogger logger)
 		{
-			_hordeAgentAssembly = hordeAgentFile;
+			_hordeAgentAssembly = hordeAgentAssembly;
 			_port = port;
 			_logger = logger;
 		}
@@ -40,74 +64,69 @@ namespace EpicGames.Horde.Compute.Clients
 		public ValueTask DisposeAsync() => new ValueTask();
 
 		/// <inheritdoc/>
-		public async Task<TResult> ExecuteAsync<TResult>(ClusterId clusterId, Requirements? requirements, Func<IComputeLease, CancellationToken, Task<TResult>> handler, CancellationToken cancellationToken)
+		public async Task<IComputeLease?> TryAssignWorkerAsync(ClusterId clusterId, Requirements? requirements, CancellationToken cancellationToken)
 		{
 			_logger.LogInformation("** CLIENT **");
 			_logger.LogInformation("Launching {Path} to handle remote", _hordeAgentAssembly);
 
-			TResult result;
-			using (Socket listener = new Socket(SocketType.Stream, ProtocolType.IP))
+			// The connection logic is an async enumerator that returns the socket, then shuts down.
+			IAsyncEnumerator<IComputeSocket> source = ConnectAsync(cancellationToken).GetAsyncEnumerator(cancellationToken);
+			if (!await source.MoveNextAsync())
 			{
-				listener.Bind(new IPEndPoint(IPAddress.Loopback, _port));
-				listener.Listen();
+				await source.DisposeAsync();
+				return null;
+			}
 
-				using CancellationTokenSource cancellationSource = CancellationTokenSource.CreateLinkedTokenSource(cancellationToken);
+			return new LeaseImpl(source);
+		}
 
-				Task<TResult> listenTask = ListenAsync<TResult>(listener, handler, cancellationSource.Token);
-				using (ManagedProcessGroup group = new ManagedProcessGroup())
-				{
-					if (RuntimeInformation.IsOSPlatform(OSPlatform.Windows))
-					{
-						List<string> arguments = new List<string>();
-						arguments.Add(_hordeAgentAssembly);
-						arguments.Add("computeworker");
-						arguments.Add($"-port={_port}");
-
-						using ManagedProcess process = new ManagedProcess(group, "dotnet", CommandLineArguments.Join(arguments), null, null, ProcessPriorityClass.Normal);
-
-						string? line;
-						while ((line = await process.ReadLineAsync(cancellationToken)) != null)
-						{
-							_logger.LogInformation("Output: {Line}", line);
-						}
-
-						await process.WaitForExitAsync(cancellationToken);
-					}
-					else
-					{
-						using Process process = new Process();
-						process.StartInfo.FileName = "dotnet";
-						process.StartInfo.ArgumentList.Add(_hordeAgentAssembly);
-						process.StartInfo.ArgumentList.Add("computeworker");
-						process.StartInfo.ArgumentList.Add($"-port={_port}");
-						process.StartInfo.UseShellExecute = true;
-						process.Start();
-
-						await process.WaitForExitAsync(cancellationToken);
-					}
-				}
+		async IAsyncEnumerable<IComputeSocket> ConnectAsync([EnumeratorCancellation] CancellationToken cancellationToken)
+		{
+			using Socket listener = new Socket(SocketType.Stream, ProtocolType.IP);
+			listener.Bind(new IPEndPoint(IPAddress.Loopback, _port));
+			listener.Listen();
 
-				cancellationSource.Cancel();
-				result = await listenTask;
-			}
+			using BackgroundTask agentTask = BackgroundTask.StartNew(ctx => RunAgentAsync(_hordeAgentAssembly, _port, _logger, ctx));
+			using Socket tcpSocket = await listener.AcceptAsync(cancellationToken);
+
+			await using ClientComputeSocket socket = new ClientComputeSocket(new TcpTransport(tcpSocket), _logger);
+			yield return socket;
 
-			_logger.LogInformation("Client terminated.");
-			return result;
+			await socket.CloseAsync(cancellationToken);
 		}
 
-		async Task<TResult> ListenAsync<TResult>(Socket listener, Func<IComputeLease, CancellationToken, Task<TResult>> handler, CancellationToken cancellationToken)
+		static async Task RunAgentAsync(string hordeAgentAssembly, int port, ILogger logger, CancellationToken cancellationToken)
 		{
-			TResult result;
-			using (Socket tcpSocket = await listener.AcceptAsync(cancellationToken))
+			using ManagedProcessGroup group = new ManagedProcessGroup();
+			if (RuntimeInformation.IsOSPlatform(OSPlatform.Windows))
 			{
-				await using (ClientComputeSocket socket = new ClientComputeSocket(new TcpTransport(tcpSocket), _logger))
+				List<string> arguments = new List<string>();
+				arguments.Add(hordeAgentAssembly);
+				arguments.Add("computeworker");
+				arguments.Add($"-port={port}");
+
+				using ManagedProcess process = new ManagedProcess(group, "dotnet", CommandLineArguments.Join(arguments), null, null, ProcessPriorityClass.Normal);
+
+				string? line;
+				while ((line = await process.ReadLineAsync(cancellationToken)) != null)
 				{
-					await using ComputeLease lease = new ComputeLease(new List<string>(), new Dictionary<string, int>(), socket);
-					result = await handler(lease, cancellationToken);
-					await socket.CloseAsync(cancellationToken);
+					logger.LogInformation("Output: {Line}", line);
 				}
+
+				await process.WaitForExitAsync(cancellationToken);
+			}
+			else
+			{
+				using Process process = new Process();
+				process.StartInfo.FileName = "dotnet";
+				process.StartInfo.ArgumentList.Add(hordeAgentAssembly);
+				process.StartInfo.ArgumentList.Add("computeworker");
+				process.StartInfo.ArgumentList.Add($"-port={port}");
+				process.StartInfo.UseShellExecute = true;
+				process.Start();
+
+				await process.WaitForExitAsync(cancellationToken);
 			}
-			return result;
 		}
 	}
 }
diff --git a/Engine/Source/Programs/Shared/EpicGames.Horde/Compute/Clients/LocalComputeClient.cs b/Engine/Source/Programs/Shared/EpicGames.Horde/Compute/Clients/LocalComputeClient.cs
--- a/Engine/Source/Programs/Shared/EpicGames.Horde/Compute/Clients/LocalComputeClient.cs
+++ b/Engine/Source/Programs/Shared/EpicGames.Horde/Compute/Clients/LocalComputeClient.cs
@@ -17,6 +17,18 @@ namespace EpicGames.Horde.Compute.Clients
 	/// </summary>
 	public sealed class LocalComputeClient : IComputeClient
 	{
+		class LeaseImpl : IComputeLease
+		{
+			public IReadOnlyList<string> Properties { get; } = new List<string>();
+			public IReadOnlyDictionary<string, int> AssignedResources => new Dictionary<string, int>();
+			public IComputeSocket Socket => _socket;
+
+			readonly ClientComputeSocket _socket;
+
+			public LeaseImpl(ClientComputeSocket socket) => _socket = socket;
+			public ValueTask DisposeAsync() => _socket.DisposeAsync();
+		}
+
 		readonly BackgroundTask _listenerTask;
 		readonly Socket _listener;
 		readonly Socket _socket;
@@ -65,13 +77,12 @@ namespace EpicGames.Horde.Compute.Clients
 		}
 
 		/// <inheritdoc/>
-		public async Task<TResult> ExecuteAsync<TResult>(ClusterId clusterId, Requirements? requirements, Func<IComputeLease, CancellationToken, Task<TResult>> handler, CancellationToken cancellationToken)
+		public Task<IComputeLease?> TryAssignWorkerAsync(ClusterId clusterId, Requirements? requirements, CancellationToken cancellationToken)
 		{
-			await using ClientComputeSocket socket = new ClientComputeSocket(new TcpTransport(_socket), _logger);
-			await using ComputeLease lease = new ComputeLease(new List<string>(), new Dictionary<string, int>(), socket);
-			TResult result = await handler(lease, cancellationToken);
-			await socket.CloseAsync(cancellationToken);
-			return result;
+#pragma warning disable CA2000 // Dispose objects before losing scope
+			ClientComputeSocket socket = new ClientComputeSocket(new TcpTransport(_socket), _logger);
+			return Task.FromResult<IComputeLease?>(new LeaseImpl(socket));
+#pragma warning restore CA2000 // Dispose objects before losing scope
 		}
 	}
 }
diff --git a/Engine/Source/Programs/Shared/EpicGames.Horde/Compute/Clients/ServerComputeClient.cs b/Engine/Source/Programs/Shared/EpicGames.Horde/Compute/Clients/ServerComputeClient.cs
--- a/Engine/Source/Programs/Shared/EpicGames.Horde/Compute/Clients/ServerComputeClient.cs
+++ b/Engine/Source/Programs/Shared/EpicGames.Horde/Compute/Clients/ServerComputeClient.cs
@@ -10,6 +10,7 @@ using System.Net.Http;
 using System.Net.Http.Headers;
 using System.Net.Http.Json;
 using System.Net.Sockets;
+using System.Runtime.CompilerServices;
 using System.Threading;
 using System.Threading.Tasks;
 
@@ -40,6 +41,29 @@ namespace EpicGames.Horde.Compute.Clients
 			public List<string> Properties { get; set; } = new List<string>();
 		}
 
+		record class LeaseInfo(IReadOnlyList<string> Properties, IReadOnlyDictionary<string, int> AssignedResources, IComputeSocket Socket);
+
+		class LeaseImpl : IComputeLease
+		{
+			readonly IAsyncEnumerator<LeaseInfo> _source;
+
+			public IReadOnlyList<string> Properties => _source.Current.Properties;
+			public IReadOnlyDictionary<string, int> AssignedResources => _source.Current.AssignedResources;
+			public IComputeSocket Socket => _source.Current.Socket;
+
+			public LeaseImpl(IAsyncEnumerator<LeaseInfo> source)
+			{
+				_source = source;
+			}
+
+			/// <inheritdoc/>
+			public async ValueTask DisposeAsync()
+			{
+				await _source.MoveNextAsync();
+				await _source.DisposeAsync();
+			}
+		}
+
 		readonly HttpClient? _defaultHttpClient;
 		readonly Func<HttpClient> _createHttpClient;
 		readonly CancellationTokenSource _cancellationSource = new CancellationTokenSource();
@@ -99,7 +123,19 @@ namespace EpicGames.Horde.Compute.Clients
 		}
 
 		/// <inheritdoc/>
-		public async Task<TResult> ExecuteAsync<TResult>(ClusterId clusterId, Requirements? requirements, Func<IComputeLease, CancellationToken, Task<TResult>> handler, CancellationToken cancellationToken)
+		public async Task<IComputeLease?> TryAssignWorkerAsync(ClusterId clusterId, Requirements? requirements, CancellationToken cancellationToken)
+		{
+			IAsyncEnumerator<LeaseInfo> source = ConnectAsync(clusterId, requirements, cancellationToken).GetAsyncEnumerator(cancellationToken);
+			if (!await source.MoveNextAsync(cancellationToken))
+			{
+				await source.DisposeAsync();
+				return null;
+			}
+			return new LeaseImpl(source);
+		}
+
+		/// <inheritdoc/>
+		async IAsyncEnumerable<LeaseInfo> ConnectAsync(ClusterId clusterId, Requirements? requirements, [EnumeratorCancellation] CancellationToken cancellationToken)
 		{
 			_logger.LogInformation("Requesting compute resource");
 
@@ -135,8 +171,7 @@ namespace EpicGames.Horde.Compute.Clients
 			byte[] key = StringUtils.ParseHexString(responseMessage.Key);
 
 			await using ClientComputeSocket computeSocket = new ClientComputeSocket(new TcpTransport(socket), _logger);
-			await using ComputeLease lease = new ComputeLease(responseMessage.Properties, responseMessage.AssignedResources, computeSocket);
-			return await handler(lease, cancellationToken);
+			yield return new LeaseInfo(responseMessage.Properties, responseMessage.AssignedResources, computeSocket);
 		}
 	}
 }
diff --git a/Engine/Source/Programs/Shared/EpicGames.Horde/Compute/ComputeLease.cs b/Engine/Source/Programs/Shared/EpicGames.Horde/Compute/ComputeLease.cs
deleted file mode 100644
diff --git a/Engine/Source/Programs/Shared/EpicGames.Horde/Compute/IComputeClient.cs b/Engine/Source/Programs/Shared/EpicGames.Horde/Compute/IComputeClient.cs
--- a/Engine/Source/Programs/Shared/EpicGames.Horde/Compute/IComputeClient.cs
+++ b/Engine/Source/Programs/Shared/EpicGames.Horde/Compute/IComputeClient.cs
@@ -16,8 +16,7 @@ namespace EpicGames.Horde.Compute
 		/// </summary>
 		/// <param name="clusterId">Cluster to execute the request</param>
 		/// <param name="requirements">Requirements for the agent</param>
-		/// <param name="handler">Handler for the lease</param>
 		/// <param name="cancellationToken">Cancellation token for the operation</param>
-		public Task<TResult> ExecuteAsync<TResult>(ClusterId clusterId, Requirements? requirements, Func<IComputeLease, CancellationToken, Task<TResult>> handler, CancellationToken cancellationToken);
+		public Task<IComputeLease?> TryAssignWorkerAsync(ClusterId clusterId, Requirements? requirements, CancellationToken cancellationToken);
 	}
 }
