Horde: Build fix.

diff --git a/Engine/Source/Programs/Shared/EpicGames.Horde/Compute/Clients/ServerComputeClient.cs b/Engine/Source/Programs/Shared/EpicGames.Horde/Compute/Clients/ServerComputeClient.cs
--- a/Engine/Source/Programs/Shared/EpicGames.Horde/Compute/Clients/ServerComputeClient.cs
+++ b/Engine/Source/Programs/Shared/EpicGames.Horde/Compute/Clients/ServerComputeClient.cs
@@ -126,7 +126,7 @@ namespace EpicGames.Horde.Compute.Clients
 		public async Task<IComputeLease?> TryAssignWorkerAsync(ClusterId clusterId, Requirements? requirements, CancellationToken cancellationToken)
 		{
 			IAsyncEnumerator<LeaseInfo> source = ConnectAsync(clusterId, requirements, cancellationToken).GetAsyncEnumerator(cancellationToken);
-			if (!await source.MoveNextAsync(cancellationToken))
+			if (!await source.MoveNextAsync())
 			{
 				await source.DisposeAsync();
 				return null;
