Horde: Add an exception when there is insufficient space in the compute buffer to write a message of the requested size.

diff --git a/Engine/Source/Programs/Shared/EpicGames.Horde/Compute/ComputeExceptions.cs b/Engine/Source/Programs/Shared/EpicGames.Horde/Compute/ComputeExceptions.cs
--- a/Engine/Source/Programs/Shared/EpicGames.Horde/Compute/ComputeExceptions.cs
+++ b/Engine/Source/Programs/Shared/EpicGames.Horde/Compute/ComputeExceptions.cs
@@ -18,6 +18,21 @@ namespace EpicGames.Horde.Compute
 		}
 	}
 
+	/// <summary>
+	/// Exception thrown for internal reasons
+	/// </summary>
+	public sealed class ComputeInternalException : ComputeException
+	{
+		/// <summary>
+		/// Constructor
+		/// </summary>
+		/// <param name="message"></param>
+		public ComputeInternalException(string message)
+			: base(message)
+		{
+		}
+	}
+
 	/// <summary>
 	/// Exception thrown on a remote machine
 	/// </summary>
diff --git a/Engine/Source/Programs/Shared/EpicGames.Horde/Compute/ComputeMessageChannel.cs b/Engine/Source/Programs/Shared/EpicGames.Horde/Compute/ComputeMessageChannel.cs
--- a/Engine/Source/Programs/Shared/EpicGames.Horde/Compute/ComputeMessageChannel.cs
+++ b/Engine/Source/Programs/Shared/EpicGames.Horde/Compute/ComputeMessageChannel.cs
@@ -197,6 +197,12 @@ namespace EpicGames.Horde.Compute
 				throw new InvalidOperationException("Only one writer can be active at a time. Dispose of the previous writer first.");
 			}
 
+			Memory<byte> memory = _sendBufferWriter.GetMemory();
+			if (memory.Length < maxSize)
+			{
+				throw new ComputeInternalException($"Insufficient space in buffer to write message of {maxSize} bytes; only {memory.Length} remaining");
+			}
+
 			_currentBuilder = new MessageBuilder(this, _sendBufferWriter, type);
 			return _currentBuilder;
 		}
