UGS: Fix parsing of datetime strings with a period in the timezone name (eg. "W. Europe Standard Time").

diff --git a/Engine/Source/Programs/Shared/EpicGames.Perforce/PerforceReflection.cs b/Engine/Source/Programs/Shared/EpicGames.Perforce/PerforceReflection.cs
--- a/Engine/Source/Programs/Shared/EpicGames.Perforce/PerforceReflection.cs
+++ b/Engine/Source/Programs/Shared/EpicGames.Perforce/PerforceReflection.cs
@@ -644,7 +644,7 @@ namespace EpicGames.Perforce
 		static object ParseStringAsNullableDateTimeOffset(Utf8String str)
 		{
 			string text = str.ToString();
-			return DateTimeOffset.Parse(Regex.Replace(text, "[a-zA-Z ]*$", "")); // Strip timezone name (eg. "EST")
+			return DateTimeOffset.Parse(Regex.Replace(text, "[a-zA-Z. ]*$", "")); // Strip timezone name (eg. "EST")
 		}
 	}
 }
