Horde: Allow compute buffers to be reset back to empty once the reader catches up to the writer.

diff --git a/Engine/Source/Programs/Shared/EpicGames.Core/MemoryMappedView.cs b/Engine/Source/Programs/Shared/EpicGames.Core/MemoryMappedView.cs
--- a/Engine/Source/Programs/Shared/EpicGames.Core/MemoryMappedView.cs
+++ b/Engine/Source/Programs/Shared/EpicGames.Core/MemoryMappedView.cs
@@ -48,6 +48,11 @@ namespace EpicGames.Core
 			memoryMappedViewAccessor.SafeMemoryMappedViewHandle.AcquirePointer(ref _data);
 		}
 
+		/// <summary>
+		/// Gets a pointer to the data
+		/// </summary>
+		public byte* GetPointer() => _data;
+
 		/// <summary>
 		/// Gets a memory object for the given range
 		/// </summary>
diff --git a/Engine/Source/Programs/Shared/EpicGames.Horde/Compute/Buffers/PooledBuffer.cs b/Engine/Source/Programs/Shared/EpicGames.Horde/Compute/Buffers/PooledBuffer.cs
--- a/Engine/Source/Programs/Shared/EpicGames.Horde/Compute/Buffers/PooledBuffer.cs
+++ b/Engine/Source/Programs/Shared/EpicGames.Horde/Compute/Buffers/PooledBuffer.cs
@@ -16,6 +16,7 @@ namespace EpicGames.Horde.Compute.Buffers
 		readonly Memory<byte> _memory;
 
 		int _readPosition;
+		int _readLength;
 		int _writePosition;
 		bool _finishedWriting;
 		TaskCompletionSource? _writtenTcs;
@@ -24,6 +25,9 @@ namespace EpicGames.Horde.Compute.Buffers
 		/// <inheritdoc/>
 		public override long Length => _memory.Length;
 
+		// Accessor for the value of read length. 
+		int GetInterlockedReadLength() => Interlocked.CompareExchange(ref _readLength, 0, 0);
+
 		/// <summary>
 		/// Creates a local buffer with the given capacity
 		/// </summary>
@@ -58,29 +62,26 @@ namespace EpicGames.Horde.Compute.Buffers
 		#region Reader
 
 		/// <inheritdoc/>
-		public override bool FinishedReading() => _finishedWriting && _readPosition == _writePosition;
+		public override bool FinishedReading() => _finishedWriting && GetInterlockedReadLength() == 0;
 
 		/// <inheritdoc/>
 		public override void AdvanceReadPosition(int size)
 		{
 			_readPosition += size;
-
-			if (_readPosition == _writePosition)
+			if (Interlocked.Add(ref _readLength, -size) == 0)
 			{
 				_flushedTcs?.TrySetResult();
 			}
 		}
 
 		/// <inheritdoc/>
-		public override ReadOnlyMemory<byte> GetReadMemory() => _memory.Slice(_readPosition, _writePosition - _readPosition);
+		public override ReadOnlyMemory<byte> GetReadMemory() => _memory.Slice(_readPosition, GetInterlockedReadLength());
 
 		/// <inheritdoc/>
 		public override async ValueTask WaitForDataAsync(int currentLength, CancellationToken cancellationToken)
 		{
-			int initialWritePosition = _readPosition + currentLength;
-
 			// Early out if we're already past this
-			if (_finishedWriting || _writePosition > initialWritePosition)
+			if (_finishedWriting || GetInterlockedReadLength() != currentLength)
 			{
 				return;
 			}
@@ -93,7 +94,7 @@ namespace EpicGames.Horde.Compute.Buffers
 					_writtenTcs = tcs;
 					for (; ; )
 					{
-						if (_finishedWriting || _writePosition > initialWritePosition)
+						if (_finishedWriting || GetInterlockedReadLength() != currentLength)
 						{
 							break;
 						}
@@ -130,16 +131,26 @@ namespace EpicGames.Horde.Compute.Buffers
 			}
 
 			_writePosition += size;
+			Interlocked.Add(ref _readLength, size);
+
 			_writtenTcs?.TrySetResult();
 		}
 
 		/// <inheritdoc/>
-		public override Memory<byte> GetWriteMemory() => _memory.Slice(_writePosition);
+		public override Memory<byte> GetWriteMemory()
+		{
+			if (GetInterlockedReadLength() == 0)
+			{
+				_readPosition = 0;
+				_writePosition = 0;
+			}
+			return _memory.Slice(_writePosition);
+		}
 
 		/// <inheritdoc/>
 		public override async ValueTask FlushWritesAsync(CancellationToken cancellationToken)
 		{
-			if (_readPosition < _writePosition)
+			if (GetInterlockedReadLength() != 0)
 			{
 				TaskCompletionSource tcs = new TaskCompletionSource();
 				using (_ = cancellationToken.Register(x => tcs.SetCanceled((CancellationToken)x!), cancellationToken, false))
@@ -150,7 +161,7 @@ namespace EpicGames.Horde.Compute.Buffers
 						_ = originalTcs.Task.ContinueWith(x => tcs.TrySetResult(), cancellationToken, TaskContinuationOptions.None, TaskScheduler.Default);
 					}
 				}
-				if (_readPosition < _writePosition)
+				if (GetInterlockedReadLength() != 0)
 				{
 					await tcs.Task;
 				}
diff --git a/Engine/Source/Programs/Shared/EpicGames.Horde/Compute/Buffers/SharedMemoryBuffer.cs b/Engine/Source/Programs/Shared/EpicGames.Horde/Compute/Buffers/SharedMemoryBuffer.cs
--- a/Engine/Source/Programs/Shared/EpicGames.Horde/Compute/Buffers/SharedMemoryBuffer.cs
+++ b/Engine/Source/Programs/Shared/EpicGames.Horde/Compute/Buffers/SharedMemoryBuffer.cs
@@ -22,22 +22,55 @@ namespace EpicGames.Horde.Compute.Buffers
 	{
 		const int HeaderLength = 32;
 
-		long ReadPosition
+		sealed unsafe class Header
 		{
-			get => BinaryPrimitives.ReadInt64LittleEndian(_header.Span.Slice(8));
-			set => BinaryPrimitives.WriteInt64LittleEndian(_header.Span.Slice(8), value);
-		}
+			readonly byte* _finishedWritingPtr;
+			readonly long* _readPositionPtr;
+			readonly long* _readLengthPtr;
+			readonly long* _writePositionPtr;
 
-		long WritePosition
-		{
-			get => BinaryPrimitives.ReadInt64LittleEndian(_header.Span.Slice(16));
-			set => BinaryPrimitives.WriteInt64LittleEndian(_header.Span.Slice(16), value);
-		}
+			public Header(MemoryMappedView view)
+			{
+				byte* basePtr = view.GetPointer();
+				_finishedWritingPtr = basePtr;
 
-		bool FinishedWriting
-		{
-			get => _header.Span[24] != 0;
-			set => _header.Span[24] = value ? (byte)1 : (byte)0;
+				long* nextPtr = (long*)basePtr + 1;
+				_readPositionPtr = nextPtr++;
+				_readLengthPtr = nextPtr++;
+				_writePositionPtr = nextPtr++;
+			}
+
+			public bool FinishedWriting => *_finishedWritingPtr != 0;
+
+			public long ReadPosition => Interlocked.Read(ref *_readPositionPtr);
+			public long ReadLength => Interlocked.Read(ref *_readLengthPtr);
+			public long WritePosition => Interlocked.Read(ref *_writePositionPtr);
+
+			public void TryReset()
+			{
+				if (ReadLength == 0)
+				{
+					Interlocked.Exchange(ref *_writePositionPtr, 0);
+					Interlocked.Exchange(ref *_readPositionPtr, 0);
+				}
+			}
+
+			public void FinishWriting()
+			{
+				*_finishedWritingPtr = 1;
+			}
+
+			public long AdvanceReadPosition(long value)
+			{
+				Interlocked.Add(ref *_readPositionPtr, value);
+				return Interlocked.Add(ref *_readLengthPtr, -value);
+			}
+
+			public long AdvanceWritePosition(long value)
+			{
+				Interlocked.Add(ref *_writePositionPtr, value);
+				return Interlocked.Add(ref *_readLengthPtr, value);
+			}
 		}
 
 		const int ChunkOffsetPow2 = 10;
@@ -46,7 +79,7 @@ namespace EpicGames.Horde.Compute.Buffers
 		readonly MemoryMappedViewAccessor _memoryMappedViewAccessor;
 		readonly MemoryMappedView _memoryMappedView;
 		readonly long _length;
-		readonly Memory<byte> _header;
+		readonly Header _header;
 		readonly Memory<byte> _memory; // TODO: DEPRECATE
 		readonly Memory<byte>[] _chunks;
 
@@ -86,7 +119,7 @@ namespace EpicGames.Horde.Compute.Buffers
 			_writtenEvent = writtenEvent;
 			_flushedEvent = flushedEvent;
 
-			_header = _memoryMappedView.GetMemory(0, HeaderLength);
+			_header = new Header(_memoryMappedView);
 			_length = _memoryMappedViewAccessor.Capacity;
 
 			_memory = _memoryMappedView.GetMemory(HeaderLength, (int)(_length - HeaderLength)); // TODO: REMOVE
@@ -216,29 +249,29 @@ namespace EpicGames.Horde.Compute.Buffers
 		#region Reader
 
 		/// <inheritdoc/>
-		public override bool FinishedReading() => FinishedWriting && ReadPosition == WritePosition;
+		public override bool FinishedReading() => _header.FinishedWriting && _header.ReadLength == 0;
 
 		/// <inheritdoc/>
 		public override void AdvanceReadPosition(int size)
 		{
-			ReadPosition += size;
-
-			if (ReadPosition == WritePosition)
+			if (_header.AdvanceReadPosition(size) == 0)
 			{
 				_flushedEvent.Set();
 			}
 		}
 
 		/// <inheritdoc/>
-		public override ReadOnlyMemory<byte> GetReadMemory() => _memory.Slice((int)ReadPosition, (int)(WritePosition - ReadPosition));
+		public override ReadOnlyMemory<byte> GetReadMemory()
+		{
+			return _memory.Slice((int)_header.ReadPosition, (int)_header.ReadLength);
+		}
 
 		/// <inheritdoc/>
 		public override async ValueTask WaitForDataAsync(int currentLength, CancellationToken cancellationToken)
 		{
-			long initialWritePosition = ReadPosition + currentLength;
 			for (; ; )
 			{
-				if (FinishedWriting || WritePosition > initialWritePosition)
+				if (_header.FinishedWriting || _header.ReadLength > currentLength)
 				{
 					return;
 				}
@@ -256,32 +289,37 @@ namespace EpicGames.Horde.Compute.Buffers
 		/// <inheritdoc/>
 		public override void FinishWriting()
 		{
-			FinishedWriting = true;
+			_header.FinishWriting();
 			_writtenEvent.Set();
 		}
 
 		/// <inheritdoc/>
 		public override void AdvanceWritePosition(int size)
 		{
-			if (FinishedWriting)
+			if (_header.FinishedWriting)
 			{
 				throw new InvalidOperationException("Cannot update write position after marking as complete");
 			}
 
-			WritePosition += size;
+			_header.AdvanceWritePosition(size);
 			_writtenEvent.Set();
 		}
 
 		/// <inheritdoc/>
-		public override Memory<byte> GetWriteMemory() => _memory.Slice((int)WritePosition);
+		public override Memory<byte> GetWriteMemory()
+		{
+			_header.TryReset();
+			return _memory.Slice((int)_header.WritePosition);
+		}
 
 		/// <inheritdoc/>
 		public override async ValueTask FlushWritesAsync(CancellationToken cancellationToken)
 		{
-			while (ReadPosition < WritePosition)
+			while (_header.ReadLength > 0)
 			{
 				await _flushedEvent.WaitOneAsync(cancellationToken);
 			}
+			_header.TryReset();
 		}
 
 		#endregion
