Horde: Allow message buffers to automatically flush if there is insufficient space remaining.

* CreateMessage is now CreateMessageAsync, and blocks on a flush if it needs more space.
* Decreased default buffer sizes for compute handler.

diff --git a/Engine/Source/Programs/Horde/Horde.Agent/Commands/Compute/CppComputeCommand.cs b/Engine/Source/Programs/Horde/Horde.Agent/Commands/Compute/CppComputeCommand.cs
--- a/Engine/Source/Programs/Horde/Horde.Agent/Commands/Compute/CppComputeCommand.cs
+++ b/Engine/Source/Programs/Horde/Horde.Agent/Commands/Compute/CppComputeCommand.cs
@@ -56,15 +56,10 @@ namespace Horde.Agent.Commands.Compute
 			NodeLocator sandbox = await CreateSandboxAsync(TaskFile, storage, cancellationToken);
 
 			// Open a socket and upload the sandbox
-			await using (IComputeMessageChannel channel = lease.Socket.CreateMessageChannel(0, 30 * 1024 * 1024, _logger))
+			await using (IComputeMessageChannel channel = await lease.Socket.CreateMessageChannelAsync(0, 4 * 1024 * 1024, _logger, cancellationToken))
 			{
 				await channel.UploadFilesAsync("", sandbox, storage, cancellationToken);
-
-				// Execute the program
-				using (LogParser parser = new LogParser(_logger, new List<string>()))
-				{
-					await channel.ExecuteAsync(jsonComputeTask.Executable, jsonComputeTask.Arguments, jsonComputeTask.WorkingDir, jsonComputeTask.EnvVars, parser.WriteData, cancellationToken);
-				}
+				await channel.ExecuteAsync(jsonComputeTask.Executable, jsonComputeTask.Arguments, jsonComputeTask.WorkingDir, jsonComputeTask.EnvVars, (string x) => _logger.LogInformation("Child Process: {Text}", x), cancellationToken);
 			}
 
 			return true;
diff --git a/Engine/Source/Programs/Horde/Horde.Agent/Commands/Compute/DdcComputeCommand.cs b/Engine/Source/Programs/Horde/Horde.Agent/Commands/Compute/DdcComputeCommand.cs
--- a/Engine/Source/Programs/Horde/Horde.Agent/Commands/Compute/DdcComputeCommand.cs
+++ b/Engine/Source/Programs/Horde/Horde.Agent/Commands/Compute/DdcComputeCommand.cs
@@ -31,10 +31,10 @@ namespace Horde.Agent.Commands.Compute
 		protected override async Task<bool> HandleRequestAsync(IComputeLease lease, CancellationToken cancellationToken)
 		{
 			IComputeSocket socket = lease.Socket;
-			await using (IComputeMessageChannel channel = socket.CreateMessageChannel(0, _logger))
+			await using (IComputeMessageChannel channel = await socket.CreateMessageChannelAsync(0, _logger, cancellationToken))
 			{
 				_logger.LogInformation("Sending XOR request");
-				channel.SendXorRequest(new byte[] { 1, 2, 3, 4, 5 }, (byte)123);
+				await channel.SendXorRequestAsync(new byte[] { 1, 2, 3, 4, 5 }, (byte)123, cancellationToken);
 
 				_logger.LogInformation("Waiting for response...");
 				IComputeMessage response = await channel.ReceiveAsync(cancellationToken);
diff --git a/Engine/Source/Programs/Horde/Horde.Agent/Commands/Compute/XorComputeCommand.cs b/Engine/Source/Programs/Horde/Horde.Agent/Commands/Compute/XorComputeCommand.cs
--- a/Engine/Source/Programs/Horde/Horde.Agent/Commands/Compute/XorComputeCommand.cs
+++ b/Engine/Source/Programs/Horde/Horde.Agent/Commands/Compute/XorComputeCommand.cs
@@ -31,14 +31,14 @@ namespace Horde.Agent.Commands.Compute
 		protected override async Task<bool> HandleRequestAsync(IComputeLease lease, CancellationToken cancellationToken)
 		{
 			IComputeSocket socket = lease.Socket;
-			await using (IComputeMessageChannel mainChannel = socket.CreateMessageChannel(0, _logger))
+			await using (IComputeMessageChannel mainChannel = await socket.CreateMessageChannelAsync(0, _logger, cancellationToken))
 			{
 				_logger.LogInformation("Forking compute channel...");
 
 				IComputeMessageChannel channel = mainChannel;
 
 				_logger.LogInformation("Sending XOR request");
-				channel.SendXorRequest(new byte[] { 1, 2, 3, 4, 5 }, (byte)123);
+				await channel.SendXorRequestAsync(new byte[] { 1, 2, 3, 4, 5 }, (byte)123, cancellationToken);
 
 				_logger.LogInformation("Waiting for response...");
 				IComputeMessage response = await channel.ReceiveAsync(cancellationToken);
diff --git a/Engine/Source/Programs/Horde/Horde.Agent/Leases/Handlers/ComputeHandler.cs b/Engine/Source/Programs/Horde/Horde.Agent/Leases/Handlers/ComputeHandler.cs
--- a/Engine/Source/Programs/Horde/Horde.Agent/Leases/Handlers/ComputeHandler.cs
+++ b/Engine/Source/Programs/Horde/Horde.Agent/Leases/Handlers/ComputeHandler.cs
@@ -84,7 +84,7 @@ namespace Horde.Agent.Leases.Handlers
 
 		async Task RunAsync(IComputeSocket socket, int channelId, CancellationToken cancellationToken)
 		{
-			await using (IComputeMessageChannel channel = socket.CreateMessageChannel(channelId, 30 * 1024 * 1024, _logger))
+			await using (IComputeMessageChannel channel = await socket.CreateMessageChannelAsync(channelId, 4 * 1024 * 1024, _logger, cancellationToken))
 			{
 				List<Task> childTasks = new List<Task>();
 				for (; ; )
@@ -124,7 +124,7 @@ namespace Horde.Agent.Leases.Handlers
 						case ComputeMessageType.XorRequest:
 							{
 								XorRequestMessage xorRequest = message.AsXorRequest();
-								RunXor(channel, xorRequest.Data, xorRequest.Value);
+								await RunXor(channel, xorRequest.Data, xorRequest.Value, cancellationToken);
 							}
 							break;
 						default:
@@ -134,9 +134,9 @@ namespace Horde.Agent.Leases.Handlers
 			}
 		}
 
-		static void RunXor(IComputeMessageChannel channel, ReadOnlyMemory<byte> source, byte value)
+		static async ValueTask RunXor(IComputeMessageChannel channel, ReadOnlyMemory<byte> source, byte value, CancellationToken cancellationToken)
 		{
-			using IComputeMessageBuilder response = channel.CreateMessage(ComputeMessageType.XorResponse, source.Length);
+			using IComputeMessageBuilder response = await channel.CreateMessageAsync(ComputeMessageType.XorResponse, source.Length, cancellationToken);
 			XorData(source.Span, response.GetSpanAndAdvance(source.Length), value);
 			response.Send();
 		}
@@ -164,7 +164,7 @@ namespace Horde.Agent.Leases.Handlers
 
 			await directoryNode.CopyToDirectoryAsync(reader, outputDir.ToDirectoryInfo(), _logger, cancellationToken);
 
-			using (IComputeMessageBuilder message = channel.CreateMessage(ComputeMessageType.WriteFilesResponse))
+			using (IComputeMessageBuilder message = await channel.CreateMessageAsync(ComputeMessageType.WriteFilesResponse, cancellationToken))
 			{
 				message.Send();
 			}
@@ -228,7 +228,7 @@ namespace Horde.Agent.Leases.Handlers
 #pragma warning disable CA2000 // Dispose objects before losing scope
 								(IComputeBufferReader reader, IComputeBufferWriter writer) = SharedMemoryBuffer.OpenIpcHandle(attachRecvBuffer.Handle).ToShared();
 								reader.Dispose();
-								socket.AttachRecvBuffer(attachRecvBuffer.ChannelId, writer);
+								await socket.AttachRecvBufferAsync(attachRecvBuffer.ChannelId, writer, cancellationToken);
 #pragma warning restore CA2000 // Dispose objects before losing scope
 							}
 							break;
@@ -238,7 +238,7 @@ namespace Horde.Agent.Leases.Handlers
 #pragma warning disable CA2000 // Dispose objects before losing scope
 								(IComputeBufferReader reader, IComputeBufferWriter writer) = SharedMemoryBuffer.OpenIpcHandle(attachSendBuffer.Handle).ToShared();
 								writer.Dispose();
-								socket.AttachSendBuffer(attachSendBuffer.ChannelId, reader);
+								await socket.AttachSendBufferAsync(attachSendBuffer.ChannelId, reader, cancellationToken);
 #pragma warning restore CA2000 // Dispose objects before losing scope
 							}
 							break;
@@ -288,10 +288,10 @@ namespace Horde.Agent.Leases.Handlers
 						int length = await process.ReadAsync(buffer, 0, buffer.Length, cancellationToken);
 						if (length == 0)
 						{
-							channel.SendExecuteResult(process.ExitCode);
+							await channel.SendExecuteResultAsync(process.ExitCode, cancellationToken);
 							return;
 						}
-						channel.SendExecuteOutput(buffer.AsMemory(0, length));
+						await channel.SendExecuteOutputAsync(buffer.AsMemory(0, length), cancellationToken);
 					}
 				}
 			}
diff --git a/Engine/Source/Programs/Horde/Samples/RemoteClient/RemoteClientApp.cs b/Engine/Source/Programs/Horde/Samples/RemoteClient/RemoteClientApp.cs
--- a/Engine/Source/Programs/Horde/Samples/RemoteClient/RemoteClientApp.cs
+++ b/Engine/Source/Programs/Horde/Samples/RemoteClient/RemoteClientApp.cs
@@ -41,7 +41,7 @@ namespace RemoteClient
 
 			// Create a message channel on channel id 0. The Horde Agent always listens on this channel for requests.
 			const int ControlChannelId = 0;
-			await using IComputeMessageChannel channel = lease.Socket.CreateMessageChannel(ControlChannelId, 30 * 1024 * 1024, logger);
+			await using IComputeMessageChannel channel = await lease.Socket.CreateMessageChannelAsync(ControlChannelId, 4 * 1024 * 1024, logger, cancellationToken);
 
 			// Upload the sandbox
 			MemoryStorageClient storage = new MemoryStorageClient();
@@ -58,7 +58,7 @@ namespace RemoteClient
 
 			// Generate data into a buffer attached to channel 1. The remote server will echo them back to us as it receives them, then exit when the channel is complete/closed.
 			const int DataChannelId = 1;
-			using (IComputeBufferWriter writer = lease.Socket.AttachSendBuffer(DataChannelId, 1024 * 1024))
+			using (IComputeBufferWriter writer = await lease.Socket.AttachSendBufferAsync(DataChannelId, 1024 * 1024, cancellationToken))
 			{
 				for (int idx = 0; idx < 100; idx++)
 				{
diff --git a/Engine/Source/Programs/Horde/Samples/RemoteServer/RemoteServerApp.cs b/Engine/Source/Programs/Horde/Samples/RemoteServer/RemoteServerApp.cs
--- a/Engine/Source/Programs/Horde/Samples/RemoteServer/RemoteServerApp.cs
+++ b/Engine/Source/Programs/Horde/Samples/RemoteServer/RemoteServerApp.cs
@@ -9,16 +9,18 @@ namespace RemoteServer
 	{
 		static async Task Main()
 		{
+			CancellationToken cancellationToken = CancellationToken.None;
+
 			await using IComputeSocket socket = ComputeSocket.ConnectAsWorker(NullLogger.Instance);
 			Console.WriteLine("Connected to client");
 
-			using IComputeBufferReader reader = socket.AttachRecvBuffer(1, 1024 * 1024);
+			using IComputeBufferReader reader = await socket.AttachRecvBufferAsync(1, 1024 * 1024, cancellationToken);
 			while (!reader.IsComplete)
 			{
 				ReadOnlyMemory<byte> memory = reader.GetMemory();
 				if (memory.Length == 0)
 				{
-					await reader.WaitForDataAsync(0, CancellationToken.None);
+					await reader.WaitForDataAsync(0, cancellationToken);
 					continue;
 				}
 
diff --git a/Engine/Source/Programs/Shared/EpicGames.Horde.Tests/BufferTests.cs b/Engine/Source/Programs/Shared/EpicGames.Horde.Tests/BufferTests.cs
--- a/Engine/Source/Programs/Shared/EpicGames.Horde.Tests/BufferTests.cs
+++ b/Engine/Source/Programs/Shared/EpicGames.Horde.Tests/BufferTests.cs
@@ -24,7 +24,7 @@ namespace EpicGames.Horde.Tests
 		[TestMethod]
 		public async Task TestPooledBuffer()
 		{
-			await TestProducerConsumerAsync(length => new PooledBuffer(length));
+			await TestProducerConsumerAsync(length => new PooledBuffer(length), CancellationToken.None);
 		}
 
 		[TestMethod]
@@ -32,11 +32,11 @@ namespace EpicGames.Horde.Tests
 		{
 			if (RuntimeInformation.IsOSPlatform(OSPlatform.Windows))
 			{
-				await TestProducerConsumerAsync(length => new SharedMemoryBuffer(length));
+				await TestProducerConsumerAsync(length => new SharedMemoryBuffer(length), CancellationToken.None);
 			}
 		}
 
-		static async Task TestProducerConsumerAsync(Func<int, IComputeBuffer> createBuffer)
+		static async Task TestProducerConsumerAsync(Func<int, IComputeBuffer> createBuffer, CancellationToken cancellationToken)
 		{
 			const int Length = 8000;
 
@@ -46,8 +46,8 @@ namespace EpicGames.Horde.Tests
 			await using IComputeSocket sourceSocket = new ClientComputeSocket(new PipeTransport(targetToSourcePipe.Reader, sourceToTargetPipe.Writer), NullLogger.Instance);
 			await using IComputeSocket targetSocket = new ClientComputeSocket(new PipeTransport(sourceToTargetPipe.Reader, targetToSourcePipe.Writer), NullLogger.Instance);
 
-			using IComputeBufferWriter sourceWriter = sourceSocket.AttachSendBuffer(0, createBuffer(Length));
-			using IComputeBufferReader targetReader = targetSocket.AttachRecvBuffer(0, createBuffer(Length));
+			using IComputeBufferWriter sourceWriter = await sourceSocket.AttachSendBufferAsync(0, createBuffer(Length), cancellationToken);
+			using IComputeBufferReader targetReader = await targetSocket.AttachRecvBufferAsync(0, createBuffer(Length), cancellationToken);
 
 			byte[] input = RandomNumberGenerator.GetBytes(Length);
 			Task producerTask = RunProducerAsync(sourceWriter, input);
diff --git a/Engine/Source/Programs/Shared/EpicGames.Horde/Compute/Buffers/PooledBuffer.cs b/Engine/Source/Programs/Shared/EpicGames.Horde/Compute/Buffers/PooledBuffer.cs
--- a/Engine/Source/Programs/Shared/EpicGames.Horde/Compute/Buffers/PooledBuffer.cs
+++ b/Engine/Source/Programs/Shared/EpicGames.Horde/Compute/Buffers/PooledBuffer.cs
@@ -15,12 +15,20 @@ namespace EpicGames.Horde.Compute.Buffers
 		readonly IMemoryOwner<byte> _memoryOwner;
 		readonly Memory<byte> _memory;
 
+		enum State
+		{
+			Normal,
+			ReaderWaitingForData, // Once new data is available, set _readerTcs
+			WriterWaitingToFlush, // Once flush has been complete, set _writerTcs
+			Finished
+		}
+
 		int _readPosition;
 		int _readLength;
 		int _writePosition;
-		bool _finishedWriting;
-		TaskCompletionSource? _writtenTcs;
-		TaskCompletionSource? _flushedTcs;
+		int _state;
+		TaskCompletionSource _readerTcs = new TaskCompletionSource(); // Used when reader is waiting for writer. Reset by the reader.
+		TaskCompletionSource _writerTcs = new TaskCompletionSource(); // Used when writer is waiting for reader. Reset by the writer.
 
 		/// <inheritdoc/>
 		public override long Length => _memory.Length;
@@ -62,16 +70,13 @@ namespace EpicGames.Horde.Compute.Buffers
 		#region Reader
 
 		/// <inheritdoc/>
-		public override bool FinishedReading() => _finishedWriting && GetInterlockedReadLength() == 0;
+		public override bool FinishedReading() => (State)_state == State.Finished && GetInterlockedReadLength() == 0;
 
 		/// <inheritdoc/>
 		public override void AdvanceReadPosition(int size)
 		{
 			_readPosition += size;
-			if (Interlocked.Add(ref _readLength, -size) == 0)
-			{
-				_flushedTcs?.TrySetResult();
-			}
+			Interlocked.Add(ref _readLength, -size);
 		}
 
 		/// <inheritdoc/>
@@ -80,35 +85,45 @@ namespace EpicGames.Horde.Compute.Buffers
 		/// <inheritdoc/>
 		public override async ValueTask WaitForDataAsync(int currentLength, CancellationToken cancellationToken)
 		{
-			// Early out if we're already past this
-			if (_finishedWriting || GetInterlockedReadLength() != currentLength)
-			{
-				return;
-			}
-
-			TaskCompletionSource tcs = new TaskCompletionSource();
-			using (IDisposable registration = cancellationToken.Register(x => tcs.SetCanceled((CancellationToken)x!), cancellationToken))
+			try
 			{
-				try
+				while (GetInterlockedReadLength() == currentLength)
 				{
-					_writtenTcs = tcs;
-					for (; ; )
+					int state = _state;
+					if (state == (int)State.Finished)
 					{
-						if (_finishedWriting || GetInterlockedReadLength() != currentLength)
-						{
-							break;
-						}
-						else
+						break;
+					}
+					else if (state == (int)State.Normal)
+					{
+						TaskCompletionSource readerTcs = _readerTcs;
+						if (readerTcs.Task.IsCompleted)
 						{
-							await _writtenTcs.Task;
+							Interlocked.CompareExchange(ref _readerTcs, new TaskCompletionSource(), readerTcs);
 						}
+						Interlocked.CompareExchange(ref _state, (int)State.ReaderWaitingForData, state);
 					}
-				}
-				finally
-				{
-					_writtenTcs = null;
+					else if (state == (int)State.ReaderWaitingForData)
+					{
+						await _readerTcs.Task.WaitAsync(cancellationToken);
+					}
+					else if (state == (int)State.WriterWaitingToFlush)
+					{
+						Compact();
+						Interlocked.CompareExchange(ref _state, (int)State.Normal, state);
+						_writerTcs.TrySetResult();
+					}
+					else
+					{
+						throw new NotImplementedException();
+					}
+					cancellationToken.ThrowIfCancellationRequested();
 				}
 			}
+			finally
+			{
+				Interlocked.CompareExchange(ref _state, (int)State.Normal, (int)State.ReaderWaitingForData);
+			}
 		}
 
 		#endregion
@@ -118,14 +133,21 @@ namespace EpicGames.Horde.Compute.Buffers
 		/// <inheritdoc/>
 		public override void FinishWriting()
 		{
-			_finishedWriting = true;
-			_writtenTcs?.TrySetResult();
+			for (; ; )
+			{
+				int state = _state;
+				if (Interlocked.CompareExchange(ref _state, (int)State.Finished, state) == state)
+				{
+					_readerTcs.TrySetResult();
+					break;
+				}
+			}
 		}
 
 		/// <inheritdoc/>
 		public override void AdvanceWritePosition(int size)
 		{
-			if (_finishedWriting)
+			if (_state == (int)State.Finished)
 			{
 				throw new InvalidOperationException("Cannot update write position after marking as complete");
 			}
@@ -133,39 +155,61 @@ namespace EpicGames.Horde.Compute.Buffers
 			_writePosition += size;
 			Interlocked.Add(ref _readLength, size);
 
-			_writtenTcs?.TrySetResult();
+			if (Interlocked.CompareExchange(ref _state, (int)State.Normal, (int)State.ReaderWaitingForData) == (int)State.ReaderWaitingForData)
+			{
+				_readerTcs.TrySetResult();
+			}
 		}
 
 		/// <inheritdoc/>
 		public override Memory<byte> GetWriteMemory()
 		{
-			if (GetInterlockedReadLength() == 0)
-			{
-				_readPosition = 0;
-				_writePosition = 0;
-			}
 			return _memory.Slice(_writePosition);
 		}
 
 		/// <inheritdoc/>
-		public override async ValueTask FlushWritesAsync(CancellationToken cancellationToken)
+		public override async ValueTask FlushAsync(CancellationToken cancellationToken)
 		{
-			if (GetInterlockedReadLength() != 0)
+			try
 			{
-				TaskCompletionSource tcs = new TaskCompletionSource();
-				using (_ = cancellationToken.Register(x => tcs.SetCanceled((CancellationToken)x!), cancellationToken, false))
+				while (_readPosition > 0)
 				{
-					TaskCompletionSource? originalTcs = Interlocked.CompareExchange(ref _flushedTcs, tcs, null);
-					if (originalTcs != null)
+					int state = _state;
+					if (state == (int)State.Finished)
 					{
-						_ = originalTcs.Task.ContinueWith(x => tcs.TrySetResult(), cancellationToken, TaskContinuationOptions.None, TaskScheduler.Default);
+						break;
+					}
+					else if (state == (int)State.Normal)
+					{
+						TaskCompletionSource writerTcs = _writerTcs;
+						if (writerTcs.Task.IsCompleted)
+						{
+							Interlocked.CompareExchange(ref _writerTcs, new TaskCompletionSource(), writerTcs);
+						}
+						Interlocked.CompareExchange(ref _state, (int)State.WriterWaitingToFlush, state);
+					}
+					else if (state == (int)State.ReaderWaitingForData)
+					{
+						Compact();
+						break;
+					}
+					else if (state == (int)State.WriterWaitingToFlush)
+					{
+						await _writerTcs.Task.WaitAsync(cancellationToken);
 					}
-				}
-				if (GetInterlockedReadLength() != 0)
-				{
-					await tcs.Task;
 				}
 			}
+			finally
+			{
+				Interlocked.CompareExchange(ref _state, (int)State.Normal, (int)State.WriterWaitingToFlush);
+			}
+		}
+
+		void Compact()
+		{
+			_memory.Slice(_readPosition, _readLength).CopyTo(_memory);
+			_readPosition = 0;
+			_writePosition = _readLength;
 		}
 
 		#endregion
diff --git a/Engine/Source/Programs/Shared/EpicGames.Horde/Compute/Buffers/SharedMemoryBuffer.cs b/Engine/Source/Programs/Shared/EpicGames.Horde/Compute/Buffers/SharedMemoryBuffer.cs
--- a/Engine/Source/Programs/Shared/EpicGames.Horde/Compute/Buffers/SharedMemoryBuffer.cs
+++ b/Engine/Source/Programs/Shared/EpicGames.Horde/Compute/Buffers/SharedMemoryBuffer.cs
@@ -6,7 +6,6 @@ using Microsoft.Win32.SafeHandles;
 using System.IO.MemoryMappedFiles;
 using System.Reflection;
 using System.IO;
-using System.Buffers.Binary;
 using System.Threading.Tasks;
 using System.Threading;
 using System.Linq;
@@ -22,9 +21,17 @@ namespace EpicGames.Horde.Compute.Buffers
 	{
 		const int HeaderLength = 32;
 
+		enum State
+		{
+			Normal,
+			ReaderWaitingForData, // Once new data is available, set _readerTcs
+			WriterWaitingToFlush, // Once flush has been complete, set _writerTcs
+			Finished
+		}
+
 		sealed unsafe class Header
 		{
-			readonly byte* _finishedWritingPtr;
+			readonly int* _statePtr;
 			readonly long* _readPositionPtr;
 			readonly long* _readLengthPtr;
 			readonly long* _writePositionPtr;
@@ -32,7 +39,7 @@ namespace EpicGames.Horde.Compute.Buffers
 			public Header(MemoryMappedView view)
 			{
 				byte* basePtr = view.GetPointer();
-				_finishedWritingPtr = basePtr;
+				_statePtr = (int*)basePtr;
 
 				long* nextPtr = (long*)basePtr + 1;
 				_readPositionPtr = nextPtr++;
@@ -40,24 +47,22 @@ namespace EpicGames.Horde.Compute.Buffers
 				_writePositionPtr = nextPtr++;
 			}
 
-			public bool FinishedWriting => *_finishedWritingPtr != 0;
-
+			public State State => (State)Interlocked.CompareExchange(ref *_statePtr, 0, 0);
 			public long ReadPosition => Interlocked.Read(ref *_readPositionPtr);
 			public long ReadLength => Interlocked.Read(ref *_readLengthPtr);
 			public long WritePosition => Interlocked.Read(ref *_writePositionPtr);
 
-			public void TryReset()
+			public void Compact()
 			{
-				if (ReadLength == 0)
-				{
-					Interlocked.Exchange(ref *_writePositionPtr, 0);
-					Interlocked.Exchange(ref *_readPositionPtr, 0);
-				}
+				Interlocked.Exchange(ref *_readPositionPtr, 0);
+				Interlocked.Exchange(ref *_writePositionPtr, ReadLength);
 			}
 
+			public bool UpdateState(State oldState, State newState) => Interlocked.CompareExchange(ref *_statePtr, (int)newState, (int)oldState) == (int)oldState;
+
 			public void FinishWriting()
 			{
-				*_finishedWritingPtr = 1;
+				*_statePtr = (int)State.Finished;
 			}
 
 			public long AdvanceReadPosition(long value)
@@ -83,8 +88,8 @@ namespace EpicGames.Horde.Compute.Buffers
 		readonly Memory<byte> _memory; // TODO: DEPRECATE
 		readonly Memory<byte>[] _chunks;
 
-		readonly Native.EventHandle _writtenEvent;
-		readonly Native.EventHandle _flushedEvent;
+		readonly Native.EventHandle _readerEvent;
+		readonly Native.EventHandle _writerEvent;
 
 		/// <inheritdoc/>
 		public override long Length => _length;
@@ -116,8 +121,8 @@ namespace EpicGames.Horde.Compute.Buffers
 			_memoryMappedViewAccessor = memoryMappedFile.CreateViewAccessor();
 			_memoryMappedView = new MemoryMappedView(_memoryMappedViewAccessor);
 
-			_writtenEvent = writtenEvent;
-			_flushedEvent = flushedEvent;
+			_readerEvent = writtenEvent;
+			_writerEvent = flushedEvent;
 
 			_header = new Header(_memoryMappedView);
 			_length = _memoryMappedViewAccessor.Capacity;
@@ -183,8 +188,8 @@ namespace EpicGames.Horde.Compute.Buffers
 		public string GetIpcHandle()
 		{
 			IntPtr memoryHandle = _memoryMappedFile.SafeMemoryMappedFileHandle.DangerousGetHandle();
-			IntPtr writtenEventHandle = _writtenEvent.SafeWaitHandle.DangerousGetHandle();
-			IntPtr flushedEventHandle = _flushedEvent.SafeWaitHandle.DangerousGetHandle();
+			IntPtr writtenEventHandle = _readerEvent.SafeWaitHandle.DangerousGetHandle();
+			IntPtr flushedEventHandle = _writerEvent.SafeWaitHandle.DangerousGetHandle();
 			return GetIpcHandle(memoryHandle, writtenEventHandle, flushedEventHandle);
 		}
 
@@ -204,13 +209,13 @@ namespace EpicGames.Horde.Compute.Buffers
 			}
 
 			IntPtr targetWrittenEventHandle;
-			if (!Native.DuplicateHandle(sourceProcessHandle, _writtenEvent.SafeWaitHandle.DangerousGetHandle(), targetProcessHandle, out targetWrittenEventHandle, 0, false, Native.DUPLICATE_SAME_ACCESS))
+			if (!Native.DuplicateHandle(sourceProcessHandle, _readerEvent.SafeWaitHandle.DangerousGetHandle(), targetProcessHandle, out targetWrittenEventHandle, 0, false, Native.DUPLICATE_SAME_ACCESS))
 			{
 				throw new Win32Exception();
 			}
 
 			IntPtr targetFlushedEventHandle;
-			if(!Native.DuplicateHandle(sourceProcessHandle, _flushedEvent.SafeWaitHandle.DangerousGetHandle(), targetProcessHandle, out targetFlushedEventHandle, 0, false, Native.DUPLICATE_SAME_ACCESS))
+			if(!Native.DuplicateHandle(sourceProcessHandle, _writerEvent.SafeWaitHandle.DangerousGetHandle(), targetProcessHandle, out targetFlushedEventHandle, 0, false, Native.DUPLICATE_SAME_ACCESS))
 			{
 				throw new Win32Exception();
 			}
@@ -230,8 +235,8 @@ namespace EpicGames.Horde.Compute.Buffers
 
 			if (disposing)
 			{
-				_flushedEvent.Dispose();
-				_writtenEvent.Dispose();
+				_writerEvent.Dispose();
+				_readerEvent.Dispose();
 				_memoryMappedView.Dispose();
 				_memoryMappedViewAccessor.Dispose();
 				_memoryMappedFile.Dispose();
@@ -249,14 +254,14 @@ namespace EpicGames.Horde.Compute.Buffers
 		#region Reader
 
 		/// <inheritdoc/>
-		public override bool FinishedReading() => _header.FinishedWriting && _header.ReadLength == 0;
+		public override bool FinishedReading() => _header.State == State.Finished && _header.ReadLength == 0;
 
 		/// <inheritdoc/>
 		public override void AdvanceReadPosition(int size)
 		{
 			if (_header.AdvanceReadPosition(size) == 0)
 			{
-				_flushedEvent.Set();
+				_writerEvent.Set();
 			}
 		}
 
@@ -269,17 +274,41 @@ namespace EpicGames.Horde.Compute.Buffers
 		/// <inheritdoc/>
 		public override async ValueTask WaitForDataAsync(int currentLength, CancellationToken cancellationToken)
 		{
-			for (; ; )
+			try
 			{
-				if (_header.FinishedWriting || _header.ReadLength > currentLength)
+				while (_header.ReadLength == currentLength)
 				{
-					return;
-				}
-				else
-				{
-					await _writtenEvent.WaitOneAsync(cancellationToken);
+					State state = _header.State;
+					if (state == State.Finished)
+					{
+						break;
+					}
+					else if (state == State.Normal)
+					{
+						_readerEvent.Reset();
+						_header.UpdateState(state, State.ReaderWaitingForData);
+					}
+					else if (state == State.ReaderWaitingForData)
+					{
+						await _readerEvent.WaitOneAsync(cancellationToken);
+					}
+					else if (state == State.WriterWaitingToFlush)
+					{
+						_header.Compact();
+						_header.UpdateState(state, State.Normal);
+						_writerEvent.Set();
+					}
+					else
+					{
+						throw new NotImplementedException();
+					}
+					cancellationToken.ThrowIfCancellationRequested();
 				}
 			}
+			finally
+			{
+				_header.UpdateState(State.ReaderWaitingForData, State.Normal);
+			}
 		}
 
 		#endregion
@@ -290,36 +319,60 @@ namespace EpicGames.Horde.Compute.Buffers
 		public override void FinishWriting()
 		{
 			_header.FinishWriting();
-			_writtenEvent.Set();
+			_readerEvent.Set();
 		}
 
 		/// <inheritdoc/>
 		public override void AdvanceWritePosition(int size)
 		{
-			if (_header.FinishedWriting)
+			if (_header.State == State.Finished)
 			{
 				throw new InvalidOperationException("Cannot update write position after marking as complete");
 			}
 
 			_header.AdvanceWritePosition(size);
-			_writtenEvent.Set();
+
+			if (_header.UpdateState(State.ReaderWaitingForData, State.Normal))
+			{
+				_readerEvent.Set();
+			}
 		}
 
 		/// <inheritdoc/>
-		public override Memory<byte> GetWriteMemory()
-		{
-			_header.TryReset();
-			return _memory.Slice((int)_header.WritePosition);
-		}
+		public override Memory<byte> GetWriteMemory() => _memory.Slice((int)_header.WritePosition);
 
 		/// <inheritdoc/>
-		public override async ValueTask FlushWritesAsync(CancellationToken cancellationToken)
+		public override async ValueTask FlushAsync(CancellationToken cancellationToken)
 		{
-			while (_header.ReadLength > 0)
+			try
+			{
+				while (_header.ReadPosition > 0)
+				{
+					State state = _header.State;
+					if (state == State.Finished)
+					{
+						break;
+					}
+					else if (state == State.Normal)
+					{
+						_writerEvent.Reset();
+						_header.UpdateState(state, State.WriterWaitingToFlush);
+					}
+					else if (state == State.ReaderWaitingForData)
+					{
+						_header.Compact();
+						break;
+					}
+					else if (state == State.WriterWaitingToFlush)
+					{
+						await _writerEvent.WaitOneAsync(cancellationToken);
+					}
+				}
+			}
+			finally
 			{
-				await _flushedEvent.WaitOneAsync(cancellationToken);
+				_header.UpdateState(State.WriterWaitingToFlush, State.Normal);
 			}
-			_header.TryReset();
 		}
 
 		#endregion
diff --git a/Engine/Source/Programs/Shared/EpicGames.Horde/Compute/ComputeBufferBase.cs b/Engine/Source/Programs/Shared/EpicGames.Horde/Compute/ComputeBufferBase.cs
--- a/Engine/Source/Programs/Shared/EpicGames.Horde/Compute/ComputeBufferBase.cs
+++ b/Engine/Source/Programs/Shared/EpicGames.Horde/Compute/ComputeBufferBase.cs
@@ -83,7 +83,7 @@ namespace EpicGames.Horde.Compute
 			public Memory<byte> GetMemory() => _buffer.GetWriteMemory();
 
 			/// <inheritdoc/>
-			public ValueTask FlushAsync(CancellationToken cancellationToken) => _buffer.FlushWritesAsync(cancellationToken);
+			public ValueTask FlushAsync(CancellationToken cancellationToken) => _buffer.FlushAsync(cancellationToken);
 		}
 
 		/// <inheritdoc/>
@@ -148,7 +148,7 @@ namespace EpicGames.Horde.Compute
 		public abstract Memory<byte> GetWriteMemory();
 
 		/// <inheritdoc cref="IComputeBufferWriter.FlushAsync(CancellationToken)"/>
-		public abstract ValueTask FlushWritesAsync(CancellationToken cancellationToken);
+		public abstract ValueTask FlushAsync(CancellationToken cancellationToken);
 
 		#endregion
 	}
diff --git a/Engine/Source/Programs/Shared/EpicGames.Horde/Compute/ComputeMessageChannel.cs b/Engine/Source/Programs/Shared/EpicGames.Horde/Compute/ComputeMessageChannel.cs
--- a/Engine/Source/Programs/Shared/EpicGames.Horde/Compute/ComputeMessageChannel.cs
+++ b/Engine/Source/Programs/Shared/EpicGames.Horde/Compute/ComputeMessageChannel.cs
@@ -169,6 +169,10 @@ namespace EpicGames.Horde.Compute
 						_receiveBufferReader.Advance(HeaderLength + length);
 						return message;
 					}
+					if (_receiveBufferReader.Buffer.Length < HeaderLength + length)
+					{
+						throw new Exception($"Receive buffer is too small for message; size is {_receiveBufferReader.Buffer.Length}, need {HeaderLength + length}");
+					}
 				}
 				await _receiveBufferReader.WaitForDataAsync(memory.Length, cancellationToken);
 			}
@@ -190,21 +194,32 @@ namespace EpicGames.Horde.Compute
 		}
 
 		/// <inheritdoc/>
-		public IComputeMessageBuilder CreateMessage(ComputeMessageType type, int maxSize)
+		public async ValueTask<IComputeMessageBuilder> CreateMessageAsync(ComputeMessageType type, int maxSize, CancellationToken cancellationToken)
 		{
 			if (_currentBuilder != null)
 			{
 				throw new InvalidOperationException("Only one writer can be active at a time. Dispose of the previous writer first.");
 			}
 
-			Memory<byte> memory = _sendBufferWriter.GetMemory();
-			if (memory.Length < maxSize)
+			for(; ;)
 			{
-				throw new ComputeInternalException($"Insufficient space in buffer to write message of {maxSize} bytes; only {memory.Length} remaining");
-			}
+				Memory<byte> memory = _sendBufferWriter.GetMemory();
+				if (memory.Length >= maxSize)
+				{
+					_currentBuilder = new MessageBuilder(this, _sendBufferWriter, type);
+					return _currentBuilder;
+				}
 
-			_currentBuilder = new MessageBuilder(this, _sendBufferWriter, type);
-			return _currentBuilder;
+				long totalLength = _sendBufferWriter.Buffer.Length;
+				if (totalLength < maxSize)
+				{
+					throw new ComputeInternalException($"Insufficient space in buffer to write message of {maxSize} bytes; buffer only has {totalLength} remaining");
+				}
+				else
+				{
+					await _sendBufferWriter.FlushAsync(cancellationToken);
+				}
+			}
 		}
 	}
 }
diff --git a/Engine/Source/Programs/Shared/EpicGames.Horde/Compute/ComputeSocketBase.cs b/Engine/Source/Programs/Shared/EpicGames.Horde/Compute/ComputeSocketBase.cs
--- a/Engine/Source/Programs/Shared/EpicGames.Horde/Compute/ComputeSocketBase.cs
+++ b/Engine/Source/Programs/Shared/EpicGames.Horde/Compute/ComputeSocketBase.cs
@@ -129,8 +129,23 @@ namespace EpicGames.Horde.Compute
 					else
 					{
 						IComputeBufferWriter writer = await GetReceiveBufferAsync(cachedWriters, id);
-						await transport.ReadFullAsync(writer.GetMemory().Slice(0, size), cancellationToken);
-						writer.Advance(size);
+						for (int offset = 0; offset < size;)
+						{
+							Memory<byte> memory = writer.GetMemory();
+							if (memory.Length == 0)
+							{
+								_logger.LogTrace("No space in buffer {Id}, flushing", id);
+								await writer.FlushAsync(cancellationToken);
+							}
+							else
+							{
+								int maxRead = Math.Min(size - offset, memory.Length);
+								int read = await transport.ReadPartialAsync(memory.Slice(0, maxRead), cancellationToken);
+								writer.Advance(read);
+
+								offset += read;
+							}
+						}
 					}
 				}
 			}
@@ -257,12 +272,8 @@ namespace EpicGames.Horde.Compute
 		/// <inheritdoc/>
 		public abstract IComputeBuffer CreateBuffer(long capacity);
 
-		/// <summary>
-		/// Registers the write end of an attached receive buffer
-		/// </summary>
-		/// <param name="channelId">Channel id</param>
-		/// <param name="recvBufferWriter">Writer for the receive buffer</param>
-		public void AttachRecvBuffer(int channelId, IComputeBufferWriter recvBufferWriter)
+		/// <inheritdoc/>
+		public ValueTask AttachRecvBufferAsync(int channelId, IComputeBufferWriter recvBufferWriter, CancellationToken cancellationToken)
 		{
 			bool complete;
 			lock (_lockObject)
@@ -285,6 +296,7 @@ namespace EpicGames.Horde.Compute
 					_recvBufferWritersChangedEvent.Set();
 				}
 			}
+			return new ValueTask();
 		}
 
 		void DetachRecvBufferWriter(Dictionary<int, IComputeBufferWriter> cachedWriters, int id)
@@ -310,13 +322,15 @@ namespace EpicGames.Horde.Compute
 		/// </summary>
 		/// <param name="channelId">Channel id</param>
 		/// <param name="sendBufferReader">Writer for the receive buffer</param>
-		public void AttachSendBuffer(int channelId, IComputeBufferReader sendBufferReader)
+		/// <param name="cancellationToken">Cancellation token for the operation</param>
+		public ValueTask AttachSendBufferAsync(int channelId, IComputeBufferReader sendBufferReader, CancellationToken cancellationToken)
 		{
 			lock (_lockObject)
 			{
 				_sendBuffers.Add(channelId, sendBufferReader);
 				_sendTasks.Add(channelId, Task.Run(() => SendFromBufferAsync(channelId, sendBufferReader, _cancellationSource.Token), _cancellationSource.Token));
 			}
+			return new ValueTask();
 		}
 	}
 }
diff --git a/Engine/Source/Programs/Shared/EpicGames.Horde/Compute/IComputeMessage.cs b/Engine/Source/Programs/Shared/EpicGames.Horde/Compute/IComputeMessage.cs
--- a/Engine/Source/Programs/Shared/EpicGames.Horde/Compute/IComputeMessage.cs
+++ b/Engine/Source/Programs/Shared/EpicGames.Horde/Compute/IComputeMessage.cs
@@ -200,14 +200,14 @@ namespace EpicGames.Horde.Compute
 		/// <summary>
 		/// Sends an exception response to the remote
 		/// </summary>
-		public static void SendException(this IComputeMessageChannel channel, Exception ex) => SendException(channel, ex.Message, ex.StackTrace);
+		public static ValueTask SendExceptionAsync(this IComputeMessageChannel channel, Exception ex, CancellationToken cancellationToken) => SendExceptionAsync(channel, ex.Message, ex.StackTrace, cancellationToken);
 
 		/// <summary>
 		/// Sends an exception response to the remote
 		/// </summary>
-		public static void SendException(this IComputeMessageChannel channel, string description, string? trace)
+		public static async ValueTask SendExceptionAsync(this IComputeMessageChannel channel, string description, string? trace, CancellationToken cancellationToken)
 		{
-			using IComputeMessageBuilder message = channel.CreateMessage(ComputeMessageType.Exception);
+			using IComputeMessageBuilder message = await channel.CreateMessageAsync(ComputeMessageType.Exception, cancellationToken);
 			message.WriteString(description);
 			message.WriteOptionalString(trace);
 			message.Send();
@@ -245,11 +245,9 @@ namespace EpicGames.Horde.Compute
 		/// <summary>
 		/// Sends untyped user data to the remote
 		/// </summary>
-		/// <param name="channel">Current channel</param>
-		/// <param name="memory">Data to send</param>
-		public static void SendUserData(this IComputeMessageChannel channel, ReadOnlyMemory<byte> memory)
+		public static async ValueTask SendUserDataAsync(this IComputeMessageChannel channel, ReadOnlyMemory<byte> memory, CancellationToken cancellationToken)
 		{
-			using IComputeMessageBuilder message = channel.CreateMessage(ComputeMessageType.UserData, memory.Length);
+			using IComputeMessageBuilder message = await channel.CreateMessageAsync(ComputeMessageType.UserData, memory.Length, cancellationToken);
 			message.WriteFixedLengthBytes(memory.Span);
 			message.Send();
 		}
@@ -288,7 +286,7 @@ namespace EpicGames.Horde.Compute
 		/// <param name="cancellationToken">Cancellation token for the operation</param>
 		public static async Task UploadFilesAsync(this IComputeMessageChannel channel, string path, NodeLocator locator, IStorageClient storage, CancellationToken cancellationToken)
 		{
-			using (IComputeMessageBuilder request = channel.CreateMessage(ComputeMessageType.WriteFiles))
+			using (IComputeMessageBuilder request = await channel.CreateMessageAsync(ComputeMessageType.WriteFiles, cancellationToken))
 			{
 				request.WriteString(path);
 				request.WriteNodeLocator(locator);
@@ -317,9 +315,10 @@ namespace EpicGames.Horde.Compute
 		/// </summary>
 		/// <param name="channel">Current channel</param>
 		/// <param name="paths">Paths of files or directories to delete</param>
-		public static void DeleteFiles(this IComputeMessageChannel channel, IReadOnlyList<string> paths)
+		/// <param name="cancellationToken">Cancellation token for the operation</param>
+		public static async ValueTask DeleteFilesAsync(this IComputeMessageChannel channel, IReadOnlyList<string> paths, CancellationToken cancellationToken)
 		{
-			IComputeMessageBuilder request = channel.CreateMessage(ComputeMessageType.DeleteFiles);
+			IComputeMessageBuilder request = await channel.CreateMessageAsync(ComputeMessageType.DeleteFiles, cancellationToken);
 			request.WriteList(paths, MemoryWriterExtensions.WriteString);
 			request.Send();
 		}
@@ -408,7 +407,7 @@ namespace EpicGames.Horde.Compute
 		/// <param name="cancellationToken">Cancellation token for the operation</param>
 		public static async Task<int> ExecuteAsync(this IComputeMessageChannel channel, string executable, IReadOnlyList<string> arguments, string? workingDir, IReadOnlyDictionary<string, string?>? envVars, Action<ReadOnlyMemory<byte>> outputHandler, CancellationToken cancellationToken)
 		{
-			using (IComputeMessageBuilder request = channel.CreateMessage(ComputeMessageType.Execute))
+			using (IComputeMessageBuilder request = await channel.CreateMessageAsync(ComputeMessageType.Execute, cancellationToken))
 			{
 				request.WriteString(executable);
 				request.WriteList(arguments, MemoryWriterExtensions.WriteString);
@@ -452,9 +451,9 @@ namespace EpicGames.Horde.Compute
 		/// <summary>
 		/// Sends output from a child process
 		/// </summary>
-		public static void SendExecuteOutput(this IComputeMessageChannel channel, ReadOnlyMemory<byte> data)
+		public static async ValueTask SendExecuteOutputAsync(this IComputeMessageChannel channel, ReadOnlyMemory<byte> data, CancellationToken cancellationToken)
 		{
-			using IComputeMessageBuilder message = channel.CreateMessage(ComputeMessageType.ExecuteOutput);
+			using IComputeMessageBuilder message = await channel.CreateMessageAsync(ComputeMessageType.ExecuteOutput, cancellationToken);
 			message.WriteFixedLengthBytes(data.Span);
 			message.Send();
 		}
@@ -464,9 +463,10 @@ namespace EpicGames.Horde.Compute
 		/// </summary>
 		/// <param name="channel"></param>
 		/// <param name="exitCode">Exit code from the process</param>
-		public static void SendExecuteResult(this IComputeMessageChannel channel, int exitCode)
+		/// <param name="cancellationToken">Cancellation token for the operation</param>
+		public static async ValueTask SendExecuteResultAsync(this IComputeMessageChannel channel, int exitCode, CancellationToken cancellationToken)
 		{
-			using IComputeMessageBuilder builder = channel.CreateMessage(ComputeMessageType.ExecuteResult);
+			using IComputeMessageBuilder builder = await channel.CreateMessageAsync(ComputeMessageType.ExecuteResult, cancellationToken);
 			builder.WriteInt32(exitCode);
 			builder.Send();
 		}
@@ -532,7 +532,7 @@ namespace EpicGames.Horde.Compute
 		/// <returns>Stream containing the blob data</returns>
 		public static async Task<Stream> ReadBlobAsync(this IComputeMessageChannel channel, BlobLocator locator, int offset, int length, CancellationToken cancellationToken)
 		{
-			using (IComputeMessageBuilder request = channel.CreateMessage(ComputeMessageType.ReadBlob))
+			using (IComputeMessageBuilder request = await channel.CreateMessageAsync(ComputeMessageType.ReadBlob, cancellationToken))
 			{
 				request.WriteBlobLocator(locator);
 				request.WriteUnsignedVarInt(offset);
@@ -593,7 +593,7 @@ namespace EpicGames.Horde.Compute
 				}
 			}
 
-			using (IComputeMessageBuilder response = channel.CreateMessage(ComputeMessageType.ReadBlobResponse, data.Length + 128))
+			using (IComputeMessageBuilder response = await channel.CreateMessageAsync(ComputeMessageType.ReadBlobResponse, data.Length + 128, cancellationToken))
 			{
 				response.WriteFixedLengthBytes(data);
 				response.Send();
@@ -607,9 +607,9 @@ namespace EpicGames.Horde.Compute
 		/// <summary>
 		/// Send a message to request that a byte string be xor'ed with a particular value
 		/// </summary>
-		public static void SendXorRequest(this IComputeMessageChannel channel, ReadOnlyMemory<byte> data, byte value)
+		public static async ValueTask SendXorRequestAsync(this IComputeMessageChannel channel, ReadOnlyMemory<byte> data, byte value, CancellationToken cancellationToken)
 		{
-			using IComputeMessageBuilder message = channel.CreateMessage(ComputeMessageType.XorRequest);
+			using IComputeMessageBuilder message = await channel.CreateMessageAsync(ComputeMessageType.XorRequest, cancellationToken);
 			message.WriteFixedLengthBytes(data.Span);
 			message.WriteUInt8(value);
 			message.Send();
diff --git a/Engine/Source/Programs/Shared/EpicGames.Horde/Compute/IComputeMessageChannel.cs b/Engine/Source/Programs/Shared/EpicGames.Horde/Compute/IComputeMessageChannel.cs
--- a/Engine/Source/Programs/Shared/EpicGames.Horde/Compute/IComputeMessageChannel.cs
+++ b/Engine/Source/Programs/Shared/EpicGames.Horde/Compute/IComputeMessageChannel.cs
@@ -24,8 +24,9 @@ namespace EpicGames.Horde.Compute
 		/// </summary>
 		/// <param name="type">Type of the message</param>
 		/// <param name="maxSize">Maximum size of the message that will be written</param>
+		/// <param name="cancellationToken">Cancellation token for the operation</param>
 		/// <returns>New builder for messages</returns>
-		IComputeMessageBuilder CreateMessage(ComputeMessageType type, int maxSize = 1024);
+		ValueTask<IComputeMessageBuilder> CreateMessageAsync(ComputeMessageType type, int maxSize, CancellationToken cancellationToken);
 	}
 
 	/// <summary>
@@ -44,14 +45,27 @@ namespace EpicGames.Horde.Compute
 	/// </summary>
 	public static class ComputeChannelExtensions
 	{
+		/// <summary>
+		/// Creates a new builder for a message
+		/// </summary>
+		/// <param name="channel">Channel to send on</param>
+		/// <param name="type">Type of the message</param>
+		/// <param name="cancellationToken">Cancellation token for the operation</param>
+		/// <returns>New builder for messages</returns>
+		public static ValueTask<IComputeMessageBuilder> CreateMessageAsync(this IComputeMessageChannel channel, ComputeMessageType type, CancellationToken cancellationToken)
+		{
+			return channel.CreateMessageAsync(type, 1024, cancellationToken);
+		}
+
 		/// <summary>
 		/// Forwards an existing message across a channel
 		/// </summary>
 		/// <param name="channel">Channel to send on</param>
 		/// <param name="message">The message to be sent</param>
-		public static void Send(this IComputeMessageChannel channel, IComputeMessage message)
+		/// <param name="cancellationToken">Cancellation token for the operation</param>
+		public static async ValueTask SendAsync(this IComputeMessageChannel channel, IComputeMessage message, CancellationToken cancellationToken)
 		{
-			using (IComputeMessageBuilder builder = channel.CreateMessage(message.Type, message.Data.Length))
+			using (IComputeMessageBuilder builder = await channel.CreateMessageAsync(message.Type, message.Data.Length, cancellationToken))
 			{
 				builder.WriteFixedLengthBytes(message.Data.Span);
 				builder.Send();
diff --git a/Engine/Source/Programs/Shared/EpicGames.Horde/Compute/IComputeSocket.cs b/Engine/Source/Programs/Shared/EpicGames.Horde/Compute/IComputeSocket.cs
--- a/Engine/Source/Programs/Shared/EpicGames.Horde/Compute/IComputeSocket.cs
+++ b/Engine/Source/Programs/Shared/EpicGames.Horde/Compute/IComputeSocket.cs
@@ -2,6 +2,8 @@
 
 using System;
 using System.Globalization;
+using System.Threading;
+using System.Threading.Tasks;
 using EpicGames.Horde.Compute.Sockets;
 using Microsoft.Extensions.Logging;
 
@@ -23,14 +25,16 @@ namespace EpicGames.Horde.Compute
 		/// </summary>
 		/// <param name="channelId">Channel to receive data on</param>
 		/// <param name="writer">Writer for the received data</param>
-		void AttachRecvBuffer(int channelId, IComputeBufferWriter writer);
+		/// <param name="cancellationToken">Cancellation token for the operation</param>
+		ValueTask AttachRecvBufferAsync(int channelId, IComputeBufferWriter writer, CancellationToken cancellationToken);
 
 		/// <summary>
 		/// Attaches a buffer to receive data.
 		/// </summary>
 		/// <param name="channelId">Channel to receive data on</param>
 		/// <param name="reader">Reader for incoming data</param>
-		void AttachSendBuffer(int channelId, IComputeBufferReader reader);
+		/// <param name="cancellationToken">Cancellation token for the operation</param>
+		ValueTask AttachSendBufferAsync(int channelId, IComputeBufferReader reader, CancellationToken cancellationToken);
 	}
 
 	/// <summary>
@@ -78,7 +82,11 @@ namespace EpicGames.Horde.Compute
 		/// <param name="socket">Socket to attach to</param>
 		/// <param name="channelId">Channel to receive data on</param>
 		/// <param name="capacity">Capacity for the buffer</param>
-		public static IComputeBufferReader AttachRecvBuffer(this IComputeSocket socket, int channelId, long capacity) => AttachRecvBuffer(socket, channelId, socket.CreateBuffer(capacity));
+		/// <param name="cancellationToken">Cancellation token for the operation</param>
+		public static ValueTask<IComputeBufferReader> AttachRecvBufferAsync(this IComputeSocket socket, int channelId, long capacity, CancellationToken cancellationToken)
+		{
+			return AttachRecvBufferAsync(socket, channelId, socket.CreateBuffer(capacity), cancellationToken);
+		}
 
 		/// <summary>
 		/// Attaches a buffer to receive data.
@@ -86,10 +94,11 @@ namespace EpicGames.Horde.Compute
 		/// <param name="socket">Socket to attach to</param>
 		/// <param name="channelId">Channel to receive data on</param>
 		/// <param name="buffer">Buffer to attach to the socket</param>
-		public static IComputeBufferReader AttachRecvBuffer(this IComputeSocket socket, int channelId, IComputeBuffer buffer)
+		/// <param name="cancellationToken">Cancellation token for the operation</param>
+		public static async ValueTask<IComputeBufferReader> AttachRecvBufferAsync(this IComputeSocket socket, int channelId, IComputeBuffer buffer, CancellationToken cancellationToken)
 		{
 			(IComputeBufferReader reader, IComputeBufferWriter writer) = buffer.ToShared();
-			socket.AttachRecvBuffer(channelId, writer);
+			await socket.AttachRecvBufferAsync(channelId, writer, cancellationToken);
 			return reader;
 		}
 
@@ -99,7 +108,9 @@ namespace EpicGames.Horde.Compute
 		/// <param name="socket">Socket to attach to</param>
 		/// <param name="channelId">Channel to receive data on</param>
 		/// <param name="capacity">Capacity for the buffer</param>
-		public static IComputeBufferWriter AttachSendBuffer(this IComputeSocket socket, int channelId, long capacity) => AttachSendBuffer(socket, channelId, socket.CreateBuffer(capacity));
+		/// <param name="cancellationToken">Cancellation token for the operation</param>
+		public static ValueTask<IComputeBufferWriter> AttachSendBufferAsync(this IComputeSocket socket, int channelId, long capacity, CancellationToken cancellationToken) 
+			=> AttachSendBufferAsync(socket, channelId, socket.CreateBuffer(capacity), cancellationToken);
 
 		/// <summary>
 		/// Attaches a buffer from which to send data.
@@ -107,10 +118,11 @@ namespace EpicGames.Horde.Compute
 		/// <param name="socket">Socket to attach to</param>
 		/// <param name="channelId">Channel to receive data on</param>
 		/// <param name="buffer">Buffer for the data to send</param>
-		public static IComputeBufferWriter AttachSendBuffer(this IComputeSocket socket, int channelId, IComputeBuffer buffer)
+		/// <param name="cancellationToken">Cancellation token for the operation</param>
+		public static async ValueTask<IComputeBufferWriter> AttachSendBufferAsync(this IComputeSocket socket, int channelId, IComputeBuffer buffer, CancellationToken cancellationToken)
 		{
 			(IComputeBufferReader reader, IComputeBufferWriter writer) = buffer.ToShared();
-			socket.AttachSendBuffer(channelId, reader);
+			await socket.AttachSendBufferAsync(channelId, reader, cancellationToken);
 			return writer;
 		}
 
@@ -120,7 +132,9 @@ namespace EpicGames.Horde.Compute
 		/// <param name="socket">Socket to create a channel for</param>
 		/// <param name="channelId">Identifier for the channel</param>
 		/// <param name="logger">Logger for the channel</param>
-		public static IComputeMessageChannel CreateMessageChannel(this IComputeSocket socket, int channelId, ILogger logger) => socket.CreateMessageChannel(channelId, 65536, logger);
+		/// <param name="cancellationToken">Cancellation token for the operation</param>
+		public static ValueTask<IComputeMessageChannel> CreateMessageChannelAsync(this IComputeSocket socket, int channelId, ILogger logger, CancellationToken cancellationToken) 
+			=> socket.CreateMessageChannelAsync(channelId, 65536, logger, cancellationToken);
 
 		/// <summary>
 		/// Creates a message channel with the given identifier
@@ -129,7 +143,9 @@ namespace EpicGames.Horde.Compute
 		/// <param name="channelId">Identifier for the channel</param>
 		/// <param name="bufferSize">Size of the send and receive buffer</param>
 		/// <param name="logger">Logger for the channel</param>
-		public static IComputeMessageChannel CreateMessageChannel(this IComputeSocket socket, int channelId, long bufferSize, ILogger logger) => CreateMessageChannel(socket, channelId, bufferSize, bufferSize, logger);
+		/// <param name="cancellationToken">Cancellation token for the operation</param>
+		public static ValueTask<IComputeMessageChannel> CreateMessageChannelAsync(this IComputeSocket socket, int channelId, long bufferSize, ILogger logger, CancellationToken cancellationToken)
+			=> socket.CreateMessageChannelAsync(channelId, bufferSize, bufferSize, logger, cancellationToken);
 
 		/// <summary>
 		/// Creates a message channel with the given identifier
@@ -139,10 +155,11 @@ namespace EpicGames.Horde.Compute
 		/// <param name="sendBufferSize">Size of the send buffer</param>
 		/// <param name="recvBufferSize">Size of the recieve buffer</param>
 		/// <param name="logger">Logger for the channel</param>
-		public static IComputeMessageChannel CreateMessageChannel(this IComputeSocket socket, int channelId, long sendBufferSize, long recvBufferSize, ILogger logger)
+		/// <param name="cancellationToken">Cancellation token for the operation</param>
+		public static async ValueTask<IComputeMessageChannel> CreateMessageChannelAsync(this IComputeSocket socket, int channelId, long sendBufferSize, long recvBufferSize, ILogger logger, CancellationToken cancellationToken)
 		{
-			IComputeBufferReader reader = socket.AttachRecvBuffer(channelId, recvBufferSize);
-			IComputeBufferWriter writer = socket.AttachSendBuffer(channelId, sendBufferSize);
+			IComputeBufferReader reader = await socket.AttachRecvBufferAsync(channelId, recvBufferSize, cancellationToken);
+			IComputeBufferWriter writer = await socket.AttachSendBufferAsync(channelId, sendBufferSize, cancellationToken);
 			return new ComputeMessageChannel(reader, writer, logger);
 		}
 	}
diff --git a/Engine/Source/Programs/Shared/EpicGames.Horde/Compute/IpcComputeMessageChannel.cs b/Engine/Source/Programs/Shared/EpicGames.Horde/Compute/IpcComputeMessageChannel.cs
--- a/Engine/Source/Programs/Shared/EpicGames.Horde/Compute/IpcComputeMessageChannel.cs
+++ b/Engine/Source/Programs/Shared/EpicGames.Horde/Compute/IpcComputeMessageChannel.cs
@@ -2,6 +2,7 @@
 
 using System;
 using System.Globalization;
+using System.Threading;
 using System.Threading.Tasks;
 using EpicGames.Core;
 using EpicGames.Horde.Compute.Buffers;
@@ -90,11 +91,11 @@ namespace EpicGames.Horde.Compute
 		/// <summary>
 		/// Registers a receive buffer with the agent process
 		/// </summary>
-		public void AttachRecvBuffer(int channelId, SharedMemoryBuffer buffer)
+		public async ValueTask AttachRecvBufferAsync(int channelId, SharedMemoryBuffer buffer, CancellationToken cancellationToken)
 		{
 			string handle = buffer.GetIpcHandle(_parentProcessHandle);
 
-			using IComputeMessageBuilder message = CreateMessage(ComputeMessageType.AttachRecvBuffer, handle.Length + 20);
+			using IComputeMessageBuilder message = await CreateMessageAsync(ComputeMessageType.AttachRecvBuffer, handle.Length + 20, cancellationToken);
 			message.WriteInt32(channelId);
 			message.WriteString(handle);
 			message.Send();
@@ -113,11 +114,11 @@ namespace EpicGames.Horde.Compute
 		/// <summary>
 		/// Registers a send buffer with the agent process
 		/// </summary>
-		public void AttachSendBuffer(int channelId, SharedMemoryBuffer buffer)
+		public async ValueTask AttachSendBufferAsync(int channelId, SharedMemoryBuffer buffer, CancellationToken cancellationToken)
 		{
 			string handle = buffer.GetIpcHandle(_parentProcessHandle);
 
-			using IComputeMessageBuilder message = CreateMessage(ComputeMessageType.AttachSendBuffer, handle.Length + 20);
+			using IComputeMessageBuilder message = await CreateMessageAsync(ComputeMessageType.AttachSendBuffer, handle.Length + 20, cancellationToken);
 			message.WriteInt32(channelId);
 			message.WriteString(handle);
 			message.Send();
diff --git a/Engine/Source/Programs/Shared/EpicGames.Horde/Compute/Sockets/WorkerComputeSocket.cs b/Engine/Source/Programs/Shared/EpicGames.Horde/Compute/Sockets/WorkerComputeSocket.cs
--- a/Engine/Source/Programs/Shared/EpicGames.Horde/Compute/Sockets/WorkerComputeSocket.cs
+++ b/Engine/Source/Programs/Shared/EpicGames.Horde/Compute/Sockets/WorkerComputeSocket.cs
@@ -1,6 +1,7 @@
 // Copyright Epic Games, Inc. All Rights Reserved.
 
 using System;
+using System.Threading;
 using System.Threading.Tasks;
 using EpicGames.Horde.Compute.Buffers;
 using Microsoft.Extensions.Logging;
@@ -30,18 +31,18 @@ namespace EpicGames.Horde.Compute.Sockets
 		public IComputeBuffer CreateBuffer(long capacity) => new SharedMemoryBuffer(capacity);
 
 		/// <inheritdoc/>
-		public void AttachRecvBuffer(int channelId, IComputeBufferWriter writer)
+		public async ValueTask AttachRecvBufferAsync(int channelId, IComputeBufferWriter writer, CancellationToken cancellationToken)
 		{
 			SharedMemoryBuffer buffer = (SharedMemoryBuffer)writer.Buffer;
-			_ipcMessageChannel.AttachRecvBuffer(channelId, buffer);
+			await _ipcMessageChannel.AttachRecvBufferAsync(channelId, buffer, cancellationToken);
 			_logger.LogTrace("Attached recv buffer to channel {ChannelId}", channelId);
 		}
 
 		/// <inheritdoc/>
-		public void AttachSendBuffer(int channelId, IComputeBufferReader reader)
+		public async ValueTask AttachSendBufferAsync(int channelId, IComputeBufferReader reader, CancellationToken cancellationToken)
 		{
 			SharedMemoryBuffer buffer = (SharedMemoryBuffer)reader.Buffer;
-			_ipcMessageChannel.AttachSendBuffer(channelId, buffer);
+			await _ipcMessageChannel.AttachSendBufferAsync(channelId, buffer, cancellationToken);
 			_logger.LogTrace("Attached send buffer to channel {ChannelId}", channelId);
 		}
 	}
