Horde: Missing file.

diff --git a/Engine/Source/Programs/Shared/EpicGames.Horde/Compute/Native.cs b/Engine/Source/Programs/Shared/EpicGames.Horde/Compute/Native.cs
--- a/Engine/Source/Programs/Shared/EpicGames.Horde/Compute/Native.cs
+++ b/Engine/Source/Programs/Shared/EpicGames.Horde/Compute/Native.cs
@@ -21,7 +21,10 @@ namespace EpicGames.Horde.Compute
 		[DllImport("kernel32.dll")]
 		public static extern SafeWaitHandle CreateEvent(SECURITY_ATTRIBUTES lpEventAttributes, bool bManualReset, bool bInitialState, string? lpName);
 
-		[DllImport("kernel32.dll")]
+		[DllImport("kernel32.dll", SetLastError = true)]
+		public static extern bool ResetEvent(SafeWaitHandle hEvent);
+
+		[DllImport("kernel32.dll", SetLastError = true)]
 		public static extern bool SetEvent(SafeWaitHandle hEvent);
 
 		public const uint DUPLICATE_SAME_ACCESS = 2;
@@ -40,7 +43,7 @@ namespace EpicGames.Horde.Compute
 				Native.SECURITY_ATTRIBUTES securityAttributes = new Native.SECURITY_ATTRIBUTES();
 				securityAttributes.nLength = Marshal.SizeOf(securityAttributes);
 				securityAttributes.bInheritHandle = (handleInheritability == HandleInheritability.Inheritable) ? 1 : 0;
-				SafeWaitHandle = CreateEvent(securityAttributes, false, false, null);
+				SafeWaitHandle = CreateEvent(securityAttributes, true, false, null);
 			}
 
 			public EventHandle(IntPtr handle, bool ownsHandle)
@@ -49,6 +52,8 @@ namespace EpicGames.Horde.Compute
 			}
 
 			public void Set() => SetEvent(SafeWaitHandle);
+
+			public void Reset() => ResetEvent(SafeWaitHandle);
 		}
 	}
 }
