Horde: Move the worker message loop into EpicGames.Horde, making it easier to test locally using the LocalComputeClient.

diff --git a/Engine/Source/Programs/Horde/Horde.Agent/Commands/Compute/ComputeCommand.cs b/Engine/Source/Programs/Horde/Horde.Agent/Commands/Compute/ComputeCommand.cs
--- a/Engine/Source/Programs/Horde/Horde.Agent/Commands/Compute/ComputeCommand.cs
+++ b/Engine/Source/Programs/Horde/Horde.Agent/Commands/Compute/ComputeCommand.cs
@@ -1,8 +1,6 @@
 // Copyright Epic Games, Inc. All Rights Reserved.
 
 using System;
-using System.Collections.Generic;
-using System.Diagnostics;
 using System.Net;
 using System.Net.Http;
 using System.Net.Http.Headers;
@@ -15,7 +13,7 @@ using EpicGames.Horde.Common;
 using EpicGames.Horde.Compute;
 using EpicGames.Horde.Compute.Clients;
 using EpicGames.Horde.Compute.Transports;
-using Horde.Agent.Leases.Handlers;
+using Microsoft.Extensions.Caching.Memory;
 using Microsoft.Extensions.DependencyInjection;
 using Microsoft.Extensions.Logging;
 using Microsoft.Extensions.Options;
@@ -36,6 +34,9 @@ namespace Horde.Agent.Commands.Compute
 		[CommandLine("-Loopback")]
 		public bool Loopback { get; set; }
 
+		[CommandLine("-Sandbox=")]
+		public DirectoryReference SandboxDir { get; set; } = DirectoryReference.Combine(Program.DataDir, "Sandbox");
+
 		readonly IServiceProvider _serviceProvider;
 		readonly IHttpClientFactory _httpClientFactory;
 		readonly IOptions<AgentSettings> _settings;
@@ -80,7 +81,7 @@ namespace Horde.Agent.Commands.Compute
 		{
 			if (Local)
 			{
-				return new LocalComputeClient((socket, ctx) => ComputeWorkerCommand.RunWorkerAsync(socket, _serviceProvider, ctx), 2000, logger);
+				return new LocalComputeClient(2000, SandboxDir, logger);
 			}
 			else if (Loopback)
 			{
@@ -109,14 +110,14 @@ namespace Horde.Agent.Commands.Compute
 	[Command("computeworker", "Runs the agent as a local compute host, accepting incoming connections on the loopback adapter with a given port")]
 	class ComputeWorkerCommand : Command
 	{
-		readonly IServiceProvider _serviceProvider;
+		readonly IMemoryCache _memoryCache;
 
 		[CommandLine("-Port=")]
 		int Port { get; set; } = 2000;
 
-		public ComputeWorkerCommand(IServiceProvider serviceProvider)
+		public ComputeWorkerCommand(IMemoryCache memoryCache)
 		{
-			_serviceProvider = serviceProvider;
+			_memoryCache = memoryCache;
 		}
 
 		public override async Task<int> ExecuteAsync(ILogger logger)
@@ -129,7 +130,7 @@ namespace Horde.Agent.Commands.Compute
 			await using (ClientComputeSocket socket = new ClientComputeSocket(new TcpTransport(tcpSocket), logger))
 			{
 				logger.LogInformation("Running worker...");
-				await RunWorkerAsync(socket, _serviceProvider, CancellationToken.None);
+				await RunWorkerAsync(socket, _memoryCache, logger, CancellationToken.None);
 				logger.LogInformation("Worker complete");
 				await socket.CloseAsync(CancellationToken.None);
 			}
@@ -138,10 +139,12 @@ namespace Horde.Agent.Commands.Compute
 			return 0;
 		}
 
-		public static async Task RunWorkerAsync(IComputeSocket socket, IServiceProvider serviceProvider, CancellationToken cancellationToken)
+		public static async Task RunWorkerAsync(IComputeSocket socket, IMemoryCache memoryCache, ILogger logger, CancellationToken cancellationToken)
 		{
-			ComputeHandler handler = ActivatorUtilities.CreateInstance<ComputeHandler>(serviceProvider);
-			await handler.RunAsync(socket, cancellationToken);
+			DirectoryReference sandboxDir = DirectoryReference.Combine(Program.DataDir, "Sandbox");
+
+			ComputeWorker worker = new ComputeWorker(sandboxDir, memoryCache, logger);
+			await worker.RunAsync(socket, cancellationToken);
 		}
 	}
 }
diff --git a/Engine/Source/Programs/Horde/Horde.Agent/Leases/Handlers/ComputeHandler.cs b/Engine/Source/Programs/Horde/Horde.Agent/Leases/Handlers/ComputeHandler.cs
--- a/Engine/Source/Programs/Horde/Horde.Agent/Leases/Handlers/ComputeHandler.cs
+++ b/Engine/Source/Programs/Horde/Horde.Agent/Leases/Handlers/ComputeHandler.cs
@@ -1,19 +1,12 @@
 // Copyright Epic Games, Inc. All Rights Reserved.
 
 using System;
-using System.Collections.Generic;
-using System.Diagnostics;
 using System.Net.Sockets;
-using System.Runtime.CompilerServices;
-using System.Runtime.InteropServices;
 using System.Threading;
 using System.Threading.Tasks;
 using EpicGames.Core;
 using EpicGames.Horde.Compute;
-using EpicGames.Horde.Compute.Buffers;
 using EpicGames.Horde.Compute.Transports;
-using EpicGames.Horde.Storage;
-using EpicGames.Horde.Storage.Nodes;
 using Horde.Agent.Services;
 using HordeCommon.Rpc.Tasks;
 using Microsoft.Extensions.Caching.Memory;
@@ -61,7 +54,8 @@ namespace Horde.Agent.Leases.Handlers
 
 				await using (ClientComputeSocket socket = new ClientComputeSocket(new TcpTransport(tcpClient.Client), _logger))
 				{
-					await RunAsync(socket, cancellationToken);
+					ComputeWorker worker = new ComputeWorker(_sandboxDir, _memoryCache, _logger);
+					await worker.RunAsync(socket, cancellationToken);
 					await socket.CloseAsync(cancellationToken);
 					return LeaseResult.Success;
 				}
@@ -76,226 +70,6 @@ namespace Horde.Agent.Leases.Handlers
 				tcpClient?.Dispose();
 			}
 		}
-
-		public async Task RunAsync(IComputeSocket socket, CancellationToken cancellationToken)
-		{
-			await RunAsync(socket, 0, cancellationToken);
-		}
-
-		async Task RunAsync(IComputeSocket socket, int channelId, CancellationToken cancellationToken)
-		{
-			await using (IComputeMessageChannel channel = await socket.CreateMessageChannelAsync(channelId, 4 * 1024 * 1024, _logger, cancellationToken))
-			{
-				List<Task> childTasks = new List<Task>();
-				for (; ; )
-				{
-					using IComputeMessage message = await channel.ReceiveAsync(cancellationToken);
-					_logger.LogTrace("Compute Channel {ChannelId}: {MessageType}", channelId, message.Type);
-
-					switch (message.Type)
-					{
-						case ComputeMessageType.None:
-							await Task.WhenAll(childTasks);
-							return;
-						case ComputeMessageType.Fork:
-							{
-								ForkMessage fork = message.ParseForkMessage();
-								childTasks.Add(RunAsync(socket, fork.ChannelId, cancellationToken));
-							}
-							break;
-						case ComputeMessageType.WriteFiles:
-							{
-								UploadFilesMessage writeFiles = message.ParseUploadFilesMessage();
-								await WriteFilesAsync(channel, writeFiles.Name, writeFiles.Locator, cancellationToken);
-							}
-							break;
-						case ComputeMessageType.DeleteFiles:
-							{
-								DeleteFilesMessage deleteFiles = message.ParseDeleteFilesMessage();
-								DeleteFiles(deleteFiles.Filter);
-							}
-							break;
-						case ComputeMessageType.Execute:
-							{
-								ExecuteProcessMessage executeProcess = message.ParseExecuteProcessMessage();
-								await ExecuteProcessAsync(socket, channel, executeProcess.Executable, executeProcess.Arguments, executeProcess.WorkingDir, executeProcess.EnvVars, cancellationToken);  
-							}
-							break;
-						case ComputeMessageType.XorRequest:
-							{
-								XorRequestMessage xorRequest = message.AsXorRequest();
-								await RunXor(channel, xorRequest.Data, xorRequest.Value, cancellationToken);
-							}
-							break;
-						default:
-							throw new ComputeInvalidMessageException(message);
-					}
-				}
-			}
-		}
-
-		static async ValueTask RunXor(IComputeMessageChannel channel, ReadOnlyMemory<byte> source, byte value, CancellationToken cancellationToken)
-		{
-			using IComputeMessageBuilder response = await channel.CreateMessageAsync(ComputeMessageType.XorResponse, source.Length, cancellationToken);
-			XorData(source.Span, response.GetSpanAndAdvance(source.Length), value);
-			response.Send();
-		}
-
-		static void XorData(ReadOnlySpan<byte> source, Span<byte> target, byte value)
-		{
-			for (int idx = 0; idx < source.Length; idx++)
-			{
-				target[idx] = (byte)(source[idx] ^ value);
-			}
-		}
-
-		async Task WriteFilesAsync(IComputeMessageChannel channel, string path, NodeLocator locator, CancellationToken cancellationToken)
-		{
-			using ComputeStorageClient store = new ComputeStorageClient(channel);
-			TreeReader reader = new TreeReader(store, _memoryCache, _logger);
-
-			DirectoryNode directoryNode = await reader.ReadNodeAsync<DirectoryNode>(locator, cancellationToken);
-
-			DirectoryReference outputDir = DirectoryReference.Combine(_sandboxDir, path);
-			if (!outputDir.IsUnderDirectory(_sandboxDir))
-			{
-				throw new InvalidOperationException("Cannot write files outside sandbox");
-			}
-
-			await directoryNode.CopyToDirectoryAsync(reader, outputDir.ToDirectoryInfo(), _logger, cancellationToken);
-
-			using (IComputeMessageBuilder message = await channel.CreateMessageAsync(ComputeMessageType.WriteFilesResponse, cancellationToken))
-			{
-				message.Send();
-			}
-		}
-
-		void DeleteFiles(IReadOnlyList<string> deleteFiles)
-		{
-			FileFilter filter = new FileFilter(deleteFiles);
-
-			List<FileReference> files = filter.ApplyToDirectory(_sandboxDir, false);
-			foreach (FileReference file in files)
-			{
-				FileUtils.ForceDeleteFile(file);
-			}
-		}
-
-		async Task ExecuteProcessAsync(IComputeSocket socket, IComputeMessageChannel channel, string executable, IReadOnlyList<string> arguments, string? workingDir, IReadOnlyDictionary<string, string?>? envVars, CancellationToken cancellationToken)
-		{
-			if (RuntimeInformation.IsOSPlatform(OSPlatform.Windows))
-			{
-				await ExecuteProcessWindowsAsync(socket, channel, executable, arguments, workingDir, envVars, cancellationToken);
-			}
-			else
-			{
-				await ExecuteProcessInternalAsync(channel, executable, arguments, workingDir, envVars, cancellationToken);
-			}
-		}
-
-		async Task ExecuteProcessWindowsAsync(IComputeSocket socket, IComputeMessageChannel channel, string executable, IReadOnlyList<string> arguments, string? workingDir, IReadOnlyDictionary<string, string?>? envVars, CancellationToken cancellationToken)
-		{
-			await using IpcComputeMessageChannel ipcChannel = new IpcComputeMessageChannel(4096, _logger);
-
-			Dictionary<string, string?> newEnvVars = new Dictionary<string, string?>();
-			if (envVars != null)
-			{
-				foreach ((string name, string? value) in envVars)
-				{
-					newEnvVars.Add(name, value);
-				}
-			}
-
-			newEnvVars[ComputeSocket.WorkerIpcEnvVar] = ipcChannel.GetStringHandle();
-
-			Task processTask = ExecuteProcessInternalAsync(channel, executable, arguments, workingDir, newEnvVars, cancellationToken);
-			processTask = processTask.ContinueWith(x => ipcChannel.ForceComplete(), cancellationToken, TaskContinuationOptions.None, TaskScheduler.Default);
-
-			List<IDisposable> buffers = new List<IDisposable>();
-			try
-			{
-				for (; ; )
-				{
-					using IComputeMessage message = await ipcChannel.ReceiveAsync(cancellationToken);
-					switch (message.Type)
-					{
-						case ComputeMessageType.None:
-							await processTask;
-							return;
-						case ComputeMessageType.AttachRecvBuffer:
-							{
-								AttachRecvBufferRequest attachRecvBuffer = IpcComputeMessageChannel.ParseAttachRecvBuffer(message);
-#pragma warning disable CA2000 // Dispose objects before losing scope
-								(IComputeBufferReader reader, IComputeBufferWriter writer) = SharedMemoryBuffer.OpenIpcHandle(attachRecvBuffer.Handle).ToShared();
-								reader.Dispose();
-								await socket.AttachRecvBufferAsync(attachRecvBuffer.ChannelId, writer, cancellationToken);
-#pragma warning restore CA2000 // Dispose objects before losing scope
-							}
-							break;
-						case ComputeMessageType.AttachSendBuffer:
-							{
-								AttachSendBufferRequest attachSendBuffer = IpcComputeMessageChannel.ParseAttachSendBuffer(message);
-#pragma warning disable CA2000 // Dispose objects before losing scope
-								(IComputeBufferReader reader, IComputeBufferWriter writer) = SharedMemoryBuffer.OpenIpcHandle(attachSendBuffer.Handle).ToShared();
-								writer.Dispose();
-								await socket.AttachSendBufferAsync(attachSendBuffer.ChannelId, reader, cancellationToken);
-#pragma warning restore CA2000 // Dispose objects before losing scope
-							}
-							break;
-						default:
-							throw new ComputeInvalidMessageException(message);
-					}
-				}
-			}
-			finally
-			{
-				foreach (SharedMemoryBuffer buffer in buffers)
-				{
-					buffer.Dispose();
-				}
-			}
-		}
-
-		async Task ExecuteProcessInternalAsync(IComputeMessageChannel channel, string executable, IReadOnlyList<string> arguments, string? workingDir, IReadOnlyDictionary<string, string?>? envVars, CancellationToken cancellationToken)
-		{
-			string resolvedExecutable = FileReference.Combine(_sandboxDir, executable).FullName;
-			string resolvedCommandLine = CommandLineArguments.Join(arguments);
-			string resolvedWorkingDir = DirectoryReference.Combine(_sandboxDir, workingDir ?? String.Empty).FullName;
-
-			Dictionary<string, string> resolvedEnvVars = ManagedProcess.GetCurrentEnvVars();
-			if (envVars != null)
-			{
-				foreach ((string key, string? value) in envVars)
-				{
-					if (value == null)
-					{
-						resolvedEnvVars.Remove(key);
-					}
-					else
-					{
-						resolvedEnvVars[key] = value;
-					}
-				}
-			}
-
-			using (ManagedProcessGroup group = new ManagedProcessGroup())
-			{
-				using (ManagedProcess process = new ManagedProcess(group, resolvedExecutable, resolvedCommandLine, resolvedWorkingDir, resolvedEnvVars, null, ProcessPriorityClass.Normal))
-				{
-					byte[] buffer = new byte[1024];
-					for (; ; )
-					{
-						int length = await process.ReadAsync(buffer, 0, buffer.Length, cancellationToken);
-						if (length == 0)
-						{
-							await channel.SendExecuteResultAsync(process.ExitCode, cancellationToken);
-							return;
-						}
-						await channel.SendExecuteOutputAsync(buffer.AsMemory(0, length), cancellationToken);
-					}
-				}
-			}
-		}
 	}
 }
 
diff --git a/Engine/Source/Programs/Horde/Samples/RemoteClient/RemoteClientApp.cs b/Engine/Source/Programs/Horde/Samples/RemoteClient/RemoteClientApp.cs
--- a/Engine/Source/Programs/Horde/Samples/RemoteClient/RemoteClientApp.cs
+++ b/Engine/Source/Programs/Horde/Samples/RemoteClient/RemoteClientApp.cs
@@ -20,7 +20,6 @@ namespace RemoteClient
 		static DirectoryReference ProjectDir { get; } = DirectoryReference.Combine(AssemblyDir, "../../..");
 		static string ConfigPath { get; } = AssemblyDir.MakeRelativeTo(ProjectDir);
 
-		static FileReference HordeAgentFile { get; } = FileReference.Combine(ProjectDir, "../../Horde.Agent", ConfigPath, "HordeAgent.dll");
 		static FileReference RemoteServerFile { get; } = FileReference.Combine(ProjectDir, "../RemoteServer", ConfigPath, "RemoteServer.exe");
 
 		static async Task Main()
@@ -77,7 +76,8 @@ namespace RemoteClient
 		{
 			if (loopback)
 			{
-				return new AgentComputeClient(HordeAgentFile.FullName, 2000, logger);
+				DirectoryReference sandboxDir = DirectoryReference.Combine(new FileReference(Assembly.GetExecutingAssembly().Location).Directory, "Sandbox");
+				return new LocalComputeClient(2000, sandboxDir, logger);
 			}
 			else
 			{
diff --git a/Engine/Source/Programs/Shared/EpicGames.Horde/Compute/Clients/LocalComputeClient.cs b/Engine/Source/Programs/Shared/EpicGames.Horde/Compute/Clients/LocalComputeClient.cs
--- a/Engine/Source/Programs/Shared/EpicGames.Horde/Compute/Clients/LocalComputeClient.cs
+++ b/Engine/Source/Programs/Shared/EpicGames.Horde/Compute/Clients/LocalComputeClient.cs
@@ -8,6 +8,7 @@ using System.Threading;
 using System.Threading.Tasks;
 using EpicGames.Core;
 using EpicGames.Horde.Compute.Transports;
+using Microsoft.Extensions.Caching.Memory;
 using Microsoft.Extensions.Logging;
 
 namespace EpicGames.Horde.Compute.Clients
@@ -37,10 +38,10 @@ namespace EpicGames.Horde.Compute.Clients
 		/// <summary>
 		/// Constructor
 		/// </summary>
-		/// <param name="serverFunc">Callback to process the server side of the connection</param>
 		/// <param name="port">Port to connect on</param>
+		/// <param name="sandboxDir">Sandbox directory for the worker</param>
 		/// <param name="logger">Logger for diagnostic output</param>
-		public LocalComputeClient(Func<IComputeSocket, CancellationToken, Task> serverFunc, int port, ILogger logger)
+		public LocalComputeClient(int port, DirectoryReference sandboxDir, ILogger logger)
 		{
 			_logger = logger;
 
@@ -48,7 +49,7 @@ namespace EpicGames.Horde.Compute.Clients
 			_listener.Bind(new IPEndPoint(IPAddress.Loopback, port));
 			_listener.Listen();
 
-			_listenerTask = BackgroundTask.StartNew(ctx => RunListenerAsync(_listener, serverFunc, logger, ctx));
+			_listenerTask = BackgroundTask.StartNew(ctx => RunListenerAsync(_listener, sandboxDir, logger, ctx));
 
 			_socket = new Socket(SocketType.Stream, ProtocolType.IP);
 			_socket.Connect(IPAddress.Loopback, port);
@@ -65,13 +66,16 @@ namespace EpicGames.Horde.Compute.Clients
 		/// <summary>
 		/// Sets up the loopback listener and calls the server method
 		/// </summary>
-		static async Task RunListenerAsync(Socket listener, Func<IComputeSocket, CancellationToken, Task> serverFunc, ILogger logger, CancellationToken cancellationToken)
+		static async Task RunListenerAsync(Socket listener, DirectoryReference sandboxDir, ILogger logger, CancellationToken cancellationToken)
 		{
 			using Socket tcpSocket = await listener.AcceptAsync(cancellationToken);
 
+			using MemoryCache memoryCache = new MemoryCache(new MemoryCacheOptions { SizeLimit = 10 * 1024 * 1024 });
+
 			await using (ClientComputeSocket socket = new ClientComputeSocket(new TcpTransport(tcpSocket), logger))
 			{
-				await serverFunc(socket, cancellationToken);
+				ComputeWorker worker = new ComputeWorker(sandboxDir, memoryCache, logger);
+				await worker.RunAsync(socket, cancellationToken);
 				await socket.CloseAsync(cancellationToken);
 			}
 		}
diff --git a/Engine/Source/Programs/Shared/EpicGames.Horde/Compute/ComputeWorker.cs b/Engine/Source/Programs/Shared/EpicGames.Horde/Compute/ComputeWorker.cs
new file mode 100644
--- /dev/null
+++ b/Engine/Source/Programs/Shared/EpicGames.Horde/Compute/ComputeWorker.cs
@@ -0,0 +1,265 @@
+// Copyright Epic Games, Inc. All Rights Reserved.
+
+using System;
+using System.Collections.Generic;
+using System.Diagnostics;
+using System.Runtime.InteropServices;
+using System.Threading;
+using System.Threading.Tasks;
+using EpicGames.Core;
+using EpicGames.Horde.Compute.Buffers;
+using EpicGames.Horde.Storage;
+using EpicGames.Horde.Storage.Nodes;
+using Microsoft.Extensions.Caching.Memory;
+using Microsoft.Extensions.Logging;
+
+namespace EpicGames.Horde.Compute
+{
+	/// <summary>
+	/// Implements the remote end of a compute worker. 
+	/// </summary>
+	public class ComputeWorker
+	{
+		readonly DirectoryReference _sandboxDir;
+		readonly IMemoryCache _memoryCache;
+		readonly ILogger _logger;
+
+		/// <summary>
+		/// Constructor
+		/// </summary>
+		/// <param name="sandboxDir">Directory to use for reading/writing files</param>
+		/// <param name="memoryCache">Cache for nodes read from storage</param>
+		/// <param name="logger">Logger for diagnostics</param>
+		public ComputeWorker(DirectoryReference sandboxDir, IMemoryCache memoryCache, ILogger logger)
+		{
+			_sandboxDir = sandboxDir;
+			_memoryCache = memoryCache;
+			_logger = logger;
+		}
+
+		/// <summary>
+		/// Runs the worker using commands sent along the given socket
+		/// </summary>
+		/// <param name="socket">Socket to read from</param>
+		/// <param name="cancellationToken">Cancellation token for the operation</param>
+		public async Task RunAsync(IComputeSocket socket, CancellationToken cancellationToken)
+		{
+			await RunAsync(socket, 0, cancellationToken);
+		}
+
+		async Task RunAsync(IComputeSocket socket, int channelId, CancellationToken cancellationToken)
+		{
+			await using (IComputeMessageChannel channel = await socket.CreateMessageChannelAsync(channelId, 4 * 1024 * 1024, _logger, cancellationToken))
+			{
+				List<Task> childTasks = new List<Task>();
+				for (; ; )
+				{
+					using IComputeMessage message = await channel.ReceiveAsync(cancellationToken);
+					_logger.LogTrace("Compute Channel {ChannelId}: {MessageType}", channelId, message.Type);
+
+					switch (message.Type)
+					{
+						case ComputeMessageType.None:
+							await Task.WhenAll(childTasks);
+							return;
+						case ComputeMessageType.Fork:
+							{
+								ForkMessage fork = message.ParseForkMessage();
+								childTasks.Add(RunAsync(socket, fork.ChannelId, cancellationToken));
+							}
+							break;
+						case ComputeMessageType.WriteFiles:
+							{
+								UploadFilesMessage writeFiles = message.ParseUploadFilesMessage();
+								await WriteFilesAsync(channel, writeFiles.Name, writeFiles.Locator, cancellationToken);
+							}
+							break;
+						case ComputeMessageType.DeleteFiles:
+							{
+								DeleteFilesMessage deleteFiles = message.ParseDeleteFilesMessage();
+								DeleteFiles(deleteFiles.Filter);
+							}
+							break;
+						case ComputeMessageType.Execute:
+							{
+								ExecuteProcessMessage executeProcess = message.ParseExecuteProcessMessage();
+								await ExecuteProcessAsync(socket, channel, executeProcess.Executable, executeProcess.Arguments, executeProcess.WorkingDir, executeProcess.EnvVars, cancellationToken);
+							}
+							break;
+						case ComputeMessageType.XorRequest:
+							{
+								XorRequestMessage xorRequest = message.AsXorRequest();
+								await RunXor(channel, xorRequest.Data, xorRequest.Value, cancellationToken);
+							}
+							break;
+						default:
+							throw new ComputeInvalidMessageException(message);
+					}
+				}
+			}
+		}
+
+		static async ValueTask RunXor(IComputeMessageChannel channel, ReadOnlyMemory<byte> source, byte value, CancellationToken cancellationToken)
+		{
+			using IComputeMessageBuilder response = await channel.CreateMessageAsync(ComputeMessageType.XorResponse, source.Length, cancellationToken);
+			XorData(source.Span, response.GetSpanAndAdvance(source.Length), value);
+			response.Send();
+		}
+
+		static void XorData(ReadOnlySpan<byte> source, Span<byte> target, byte value)
+		{
+			for (int idx = 0; idx < source.Length; idx++)
+			{
+				target[idx] = (byte)(source[idx] ^ value);
+			}
+		}
+
+		async Task WriteFilesAsync(IComputeMessageChannel channel, string path, NodeLocator locator, CancellationToken cancellationToken)
+		{
+			using ComputeStorageClient store = new ComputeStorageClient(channel);
+			TreeReader reader = new TreeReader(store, _memoryCache, _logger);
+
+			DirectoryNode directoryNode = await reader.ReadNodeAsync<DirectoryNode>(locator, cancellationToken);
+
+			DirectoryReference outputDir = DirectoryReference.Combine(_sandboxDir, path);
+			if (!outputDir.IsUnderDirectory(_sandboxDir))
+			{
+				throw new InvalidOperationException("Cannot write files outside sandbox");
+			}
+
+			await directoryNode.CopyToDirectoryAsync(reader, outputDir.ToDirectoryInfo(), _logger, cancellationToken);
+
+			using (IComputeMessageBuilder message = await channel.CreateMessageAsync(ComputeMessageType.WriteFilesResponse, cancellationToken))
+			{
+				message.Send();
+			}
+		}
+
+		void DeleteFiles(IReadOnlyList<string> deleteFiles)
+		{
+			FileFilter filter = new FileFilter(deleteFiles);
+
+			List<FileReference> files = filter.ApplyToDirectory(_sandboxDir, false);
+			foreach (FileReference file in files)
+			{
+				FileUtils.ForceDeleteFile(file);
+			}
+		}
+
+		async Task ExecuteProcessAsync(IComputeSocket socket, IComputeMessageChannel channel, string executable, IReadOnlyList<string> arguments, string? workingDir, IReadOnlyDictionary<string, string?>? envVars, CancellationToken cancellationToken)
+		{
+			if (RuntimeInformation.IsOSPlatform(OSPlatform.Windows))
+			{
+				await ExecuteProcessWindowsAsync(socket, channel, executable, arguments, workingDir, envVars, cancellationToken);
+			}
+			else
+			{
+				await ExecuteProcessInternalAsync(channel, executable, arguments, workingDir, envVars, cancellationToken);
+			}
+		}
+
+		async Task ExecuteProcessWindowsAsync(IComputeSocket socket, IComputeMessageChannel channel, string executable, IReadOnlyList<string> arguments, string? workingDir, IReadOnlyDictionary<string, string?>? envVars, CancellationToken cancellationToken)
+		{
+			await using IpcComputeMessageChannel ipcChannel = new IpcComputeMessageChannel(4096, _logger);
+
+			Dictionary<string, string?> newEnvVars = new Dictionary<string, string?>();
+			if (envVars != null)
+			{
+				foreach ((string name, string? value) in envVars)
+				{
+					newEnvVars.Add(name, value);
+				}
+			}
+
+			newEnvVars[ComputeSocket.WorkerIpcEnvVar] = ipcChannel.GetStringHandle();
+
+			Task processTask = ExecuteProcessInternalAsync(channel, executable, arguments, workingDir, newEnvVars, cancellationToken);
+			processTask = processTask.ContinueWith(x => ipcChannel.ForceComplete(), cancellationToken, TaskContinuationOptions.None, TaskScheduler.Default);
+
+			List<IDisposable> buffers = new List<IDisposable>();
+			try
+			{
+				for (; ; )
+				{
+					using IComputeMessage message = await ipcChannel.ReceiveAsync(cancellationToken);
+					switch (message.Type)
+					{
+						case ComputeMessageType.None:
+							await processTask;
+							return;
+						case ComputeMessageType.AttachRecvBuffer:
+							{
+								AttachRecvBufferRequest attachRecvBuffer = IpcComputeMessageChannel.ParseAttachRecvBuffer(message);
+#pragma warning disable CA2000 // Dispose objects before losing scope
+								(IComputeBufferReader reader, IComputeBufferWriter writer) = SharedMemoryBuffer.OpenIpcHandle(attachRecvBuffer.Handle).ToShared();
+								reader.Dispose();
+								await socket.AttachRecvBufferAsync(attachRecvBuffer.ChannelId, writer, cancellationToken);
+#pragma warning restore CA2000 // Dispose objects before losing scope
+							}
+							break;
+						case ComputeMessageType.AttachSendBuffer:
+							{
+								AttachSendBufferRequest attachSendBuffer = IpcComputeMessageChannel.ParseAttachSendBuffer(message);
+#pragma warning disable CA2000 // Dispose objects before losing scope
+								(IComputeBufferReader reader, IComputeBufferWriter writer) = SharedMemoryBuffer.OpenIpcHandle(attachSendBuffer.Handle).ToShared();
+								writer.Dispose();
+								await socket.AttachSendBufferAsync(attachSendBuffer.ChannelId, reader, cancellationToken);
+#pragma warning restore CA2000 // Dispose objects before losing scope
+							}
+							break;
+						default:
+							throw new ComputeInvalidMessageException(message);
+					}
+				}
+			}
+			finally
+			{
+				foreach (SharedMemoryBuffer buffer in buffers)
+				{
+					buffer.Dispose();
+				}
+			}
+		}
+
+		async Task ExecuteProcessInternalAsync(IComputeMessageChannel channel, string executable, IReadOnlyList<string> arguments, string? workingDir, IReadOnlyDictionary<string, string?>? envVars, CancellationToken cancellationToken)
+		{
+			string resolvedExecutable = FileReference.Combine(_sandboxDir, executable).FullName;
+			string resolvedCommandLine = CommandLineArguments.Join(arguments);
+			string resolvedWorkingDir = DirectoryReference.Combine(_sandboxDir, workingDir ?? String.Empty).FullName;
+
+			Dictionary<string, string> resolvedEnvVars = ManagedProcess.GetCurrentEnvVars();
+			if (envVars != null)
+			{
+				foreach ((string key, string? value) in envVars)
+				{
+					if (value == null)
+					{
+						resolvedEnvVars.Remove(key);
+					}
+					else
+					{
+						resolvedEnvVars[key] = value;
+					}
+				}
+			}
+
+			using (ManagedProcessGroup group = new ManagedProcessGroup())
+			{
+				using (ManagedProcess process = new ManagedProcess(group, resolvedExecutable, resolvedCommandLine, resolvedWorkingDir, resolvedEnvVars, null, ProcessPriorityClass.Normal))
+				{
+					byte[] buffer = new byte[1024];
+					for (; ; )
+					{
+						int length = await process.ReadAsync(buffer, 0, buffer.Length, cancellationToken);
+						if (length == 0)
+						{
+							await channel.SendExecuteResultAsync(process.ExitCode, cancellationToken);
+							return;
+						}
+						await channel.SendExecuteOutputAsync(buffer.AsMemory(0, length), cancellationToken);
+					}
+				}
+			}
+		}
+	}
+}
