Fixed the unsized properties which have been broken since CL# 3494741 by assuming that int and unsigned int are always 32-bit.

diff --git a/Engine/Source/Programs/Shared/EpicGames.UHT/Types/Properties/UhtByteProperty.cs b/Engine/Source/Programs/Shared/EpicGames.UHT/Types/Properties/UhtByteProperty.cs
--- a/Engine/Source/Programs/Shared/EpicGames.UHT/Types/Properties/UhtByteProperty.cs
+++ b/Engine/Source/Programs/Shared/EpicGames.UHT/Types/Properties/UhtByteProperty.cs
@@ -41,9 +41,8 @@ namespace EpicGames.UHT.Types
 		/// Construct a new property
 		/// </summary>
 		/// <param name="propertySettings">Property settings</param>
-		/// <param name="intType">Integer type</param>
 		/// <param name="enumObj">Optional referenced enum</param>
-		public UhtByteProperty(UhtPropertySettings propertySettings, UhtPropertyIntType intType, UhtEnum? enumObj = null) : base(propertySettings, intType)
+		public UhtByteProperty(UhtPropertySettings propertySettings, UhtEnum? enumObj = null) : base(propertySettings)
 		{
 			Enum = enumObj;
 			PropertyCaps |= UhtPropertyCaps.CanExposeOnSpawn | UhtPropertyCaps.IsParameterSupportedByBlueprint |
@@ -170,7 +169,7 @@ namespace EpicGames.UHT.Types
 			}
 			else
 			{
-				return new UhtByteProperty(propertySettings, UhtPropertyIntType.Sized);
+				return new UhtByteProperty(propertySettings);
 			}
 		}
 		#endregion
diff --git a/Engine/Source/Programs/Shared/EpicGames.UHT/Types/Properties/UhtDoubleProperty.cs b/Engine/Source/Programs/Shared/EpicGames.UHT/Types/Properties/UhtDoubleProperty.cs
--- a/Engine/Source/Programs/Shared/EpicGames.UHT/Types/Properties/UhtDoubleProperty.cs
+++ b/Engine/Source/Programs/Shared/EpicGames.UHT/Types/Properties/UhtDoubleProperty.cs
@@ -27,7 +27,7 @@ namespace EpicGames.UHT.Types
 		/// Create new property
 		/// </summary>
 		/// <param name="propertySettings">Property settings</param>
-		public UhtDoubleProperty(UhtPropertySettings propertySettings) : base(propertySettings, UhtPropertyIntType.None)
+		public UhtDoubleProperty(UhtPropertySettings propertySettings) : base(propertySettings)
 		{
 			PropertyCaps |= UhtPropertyCaps.IsParameterSupportedByBlueprint | UhtPropertyCaps.IsMemberSupportedByBlueprint | UhtPropertyCaps.SupportsRigVM;
 		}
diff --git a/Engine/Source/Programs/Shared/EpicGames.UHT/Types/Properties/UhtEnumProperty.cs b/Engine/Source/Programs/Shared/EpicGames.UHT/Types/Properties/UhtEnumProperty.cs
--- a/Engine/Source/Programs/Shared/EpicGames.UHT/Types/Properties/UhtEnumProperty.cs
+++ b/Engine/Source/Programs/Shared/EpicGames.UHT/Types/Properties/UhtEnumProperty.cs
@@ -49,12 +49,6 @@ namespace EpicGames.UHT.Types
 		[JsonConverter(typeof(JsonStringEnumConverter))]
 		public UhtEnumUnderlyingType UnderlyingType => Enum.UnderlyingType != UhtEnumUnderlyingType.Unspecified ? Enum.UnderlyingType : UhtEnumUnderlyingType.Int32;
 
-		/// <summary>
-		/// Underlying type size.  Defaults to unsized if the referenced enum doesn't have an underlying type
-		/// </summary>
-		[JsonConverter(typeof(JsonStringEnumConverter))]
-		public UhtPropertyIntType UnderlyingTypeSize => Enum.UnderlyingType != UhtEnumUnderlyingType.Unspecified ? UhtPropertyIntType.Sized : UhtPropertyIntType.Unsized;
-
 		/// <summary>
 		/// Construct property
 		/// </summary>
@@ -316,21 +310,21 @@ namespace EpicGames.UHT.Types
 			switch (UnderlyingType)
 			{
 				case UhtEnumUnderlyingType.Int8:
-					return new UhtInt8Property(propertySettings, UnderlyingTypeSize);
+					return new UhtInt8Property(propertySettings);
 				case UhtEnumUnderlyingType.Int16:
-					return new UhtInt16Property(propertySettings, UnderlyingTypeSize);
+					return new UhtInt16Property(propertySettings);
 				case UhtEnumUnderlyingType.Int32:
-					return new UhtIntProperty(propertySettings, UnderlyingTypeSize);
+					return new UhtIntProperty(propertySettings);
 				case UhtEnumUnderlyingType.Int64:
-					return new UhtInt64Property(propertySettings, UnderlyingTypeSize);
+					return new UhtInt64Property(propertySettings);
 				case UhtEnumUnderlyingType.Uint8:
-					return new UhtByteProperty(propertySettings, UnderlyingTypeSize);
+					return new UhtByteProperty(propertySettings);
 				case UhtEnumUnderlyingType.Uint16:
-					return new UhtUInt16Property(propertySettings, UnderlyingTypeSize);
+					return new UhtUInt16Property(propertySettings);
 				case UhtEnumUnderlyingType.Uint32:
-					return new UhtUInt32Property(propertySettings, UnderlyingTypeSize);
+					return new UhtUInt32Property(propertySettings);
 				case UhtEnumUnderlyingType.Uint64:
-					return new UhtUInt64Property(propertySettings, UnderlyingTypeSize);
+					return new UhtUInt64Property(propertySettings);
 				default:
 					throw new UhtIceException("Unexpected underlying enum type");
 			}
@@ -354,7 +348,7 @@ namespace EpicGames.UHT.Types
 					enumObj = outer.FindType(UhtFindOptions.Enum | UhtFindOptions.SourceName, tokenList, tokenReader) as UhtEnum;
 				})
 				.Require('>');
-			return enumObj != null ? new UhtByteProperty(propertySettings, UhtPropertyIntType.None, enumObj) : null;
+			return enumObj != null ? new UhtByteProperty(propertySettings, enumObj) : null;
 		}
 		#endregion
 	}
diff --git a/Engine/Source/Programs/Shared/EpicGames.UHT/Types/Properties/UhtFloatProperty.cs b/Engine/Source/Programs/Shared/EpicGames.UHT/Types/Properties/UhtFloatProperty.cs
--- a/Engine/Source/Programs/Shared/EpicGames.UHT/Types/Properties/UhtFloatProperty.cs
+++ b/Engine/Source/Programs/Shared/EpicGames.UHT/Types/Properties/UhtFloatProperty.cs
@@ -26,7 +26,7 @@ namespace EpicGames.UHT.Types
 		/// Construct a new property
 		/// </summary>
 		/// <param name="propertySettings">Property settings</param>
-		public UhtFloatProperty(UhtPropertySettings propertySettings) : base(propertySettings, UhtPropertyIntType.None)
+		public UhtFloatProperty(UhtPropertySettings propertySettings) : base(propertySettings)
 		{
 			PropertyCaps |= UhtPropertyCaps.CanExposeOnSpawn | UhtPropertyCaps.IsParameterSupportedByBlueprint | UhtPropertyCaps.IsMemberSupportedByBlueprint | UhtPropertyCaps.SupportsRigVM;
 		}
diff --git a/Engine/Source/Programs/Shared/EpicGames.UHT/Types/Properties/UhtInt16Property.cs b/Engine/Source/Programs/Shared/EpicGames.UHT/Types/Properties/UhtInt16Property.cs
--- a/Engine/Source/Programs/Shared/EpicGames.UHT/Types/Properties/UhtInt16Property.cs
+++ b/Engine/Source/Programs/Shared/EpicGames.UHT/Types/Properties/UhtInt16Property.cs
@@ -25,8 +25,7 @@ namespace EpicGames.UHT.Types
 		/// Construct new property
 		/// </summary>
 		/// <param name="propertySettings">Property settings</param>
-		/// <param name="intType">Integer type</param>
-		public UhtInt16Property(UhtPropertySettings propertySettings, UhtPropertyIntType intType) : base(propertySettings, intType)
+		public UhtInt16Property(UhtPropertySettings propertySettings) : base(propertySettings)
 		{
 		}
 
@@ -62,7 +61,7 @@ namespace EpicGames.UHT.Types
 		[SuppressMessage("Style", "IDE0060:Remove unused parameter", Justification = "Attribute accessed method")]
 		private static UhtProperty? Int16Property(UhtPropertyResolvePhase resolvePhase, UhtPropertySettings propertySettings, IUhtTokenReader tokenReader, UhtToken matchedToken)
 		{
-			return new UhtInt16Property(propertySettings, UhtPropertyIntType.Sized);
+			return new UhtInt16Property(propertySettings);
 		}
 		#endregion
 	}
diff --git a/Engine/Source/Programs/Shared/EpicGames.UHT/Types/Properties/UhtInt64Property.cs b/Engine/Source/Programs/Shared/EpicGames.UHT/Types/Properties/UhtInt64Property.cs
--- a/Engine/Source/Programs/Shared/EpicGames.UHT/Types/Properties/UhtInt64Property.cs
+++ b/Engine/Source/Programs/Shared/EpicGames.UHT/Types/Properties/UhtInt64Property.cs
@@ -25,8 +25,7 @@ namespace EpicGames.UHT.Types
 		/// Construct new property
 		/// </summary>
 		/// <param name="propertySettings">Property settings</param>
-		/// <param name="intType">Integer type</param>
-		public UhtInt64Property(UhtPropertySettings propertySettings, UhtPropertyIntType intType) : base(propertySettings, intType)
+		public UhtInt64Property(UhtPropertySettings propertySettings) : base(propertySettings)
 		{
 			PropertyCaps |= UhtPropertyCaps.CanExposeOnSpawn | UhtPropertyCaps.IsParameterSupportedByBlueprint | UhtPropertyCaps.IsMemberSupportedByBlueprint;
 		}
@@ -64,7 +63,7 @@ namespace EpicGames.UHT.Types
 		[SuppressMessage("Style", "IDE0060:Remove unused parameter", Justification = "Attribute accessed method")]
 		private static UhtProperty? Int64Property(UhtPropertyResolvePhase resolvePhase, UhtPropertySettings propertySettings, IUhtTokenReader tokenReader, UhtToken matchedToken)
 		{
-			return new UhtInt64Property(propertySettings, UhtPropertyIntType.Sized);
+			return new UhtInt64Property(propertySettings);
 		}
 		#endregion
 	}
diff --git a/Engine/Source/Programs/Shared/EpicGames.UHT/Types/Properties/UhtInt8Property.cs b/Engine/Source/Programs/Shared/EpicGames.UHT/Types/Properties/UhtInt8Property.cs
--- a/Engine/Source/Programs/Shared/EpicGames.UHT/Types/Properties/UhtInt8Property.cs
+++ b/Engine/Source/Programs/Shared/EpicGames.UHT/Types/Properties/UhtInt8Property.cs
@@ -25,8 +25,7 @@ namespace EpicGames.UHT.Types
 		/// Construct a new property
 		/// </summary>
 		/// <param name="propertySettings">Property settings</param>
-		/// <param name="intType">Integer type</param>
-		public UhtInt8Property(UhtPropertySettings propertySettings, UhtPropertyIntType intType) : base(propertySettings, intType)
+		public UhtInt8Property(UhtPropertySettings propertySettings) : base(propertySettings)
 		{
 		}
 
@@ -62,7 +61,7 @@ namespace EpicGames.UHT.Types
 		[SuppressMessage("Style", "IDE0060:Remove unused parameter", Justification = "Attribute accessed method")]
 		private static UhtProperty? Int8Property(UhtPropertyResolvePhase resolvePhase, UhtPropertySettings propertySettings, IUhtTokenReader tokenReader, UhtToken matchedToken)
 		{
-			return new UhtInt8Property(propertySettings, UhtPropertyIntType.Sized);
+			return new UhtInt8Property(propertySettings);
 		}
 		#endregion
 	}
diff --git a/Engine/Source/Programs/Shared/EpicGames.UHT/Types/Properties/UhtIntProperty.cs b/Engine/Source/Programs/Shared/EpicGames.UHT/Types/Properties/UhtIntProperty.cs
--- a/Engine/Source/Programs/Shared/EpicGames.UHT/Types/Properties/UhtIntProperty.cs
+++ b/Engine/Source/Programs/Shared/EpicGames.UHT/Types/Properties/UhtIntProperty.cs
@@ -25,8 +25,7 @@ namespace EpicGames.UHT.Types
 		/// Construct new property
 		/// </summary>
 		/// <param name="propertySettings">Property settings</param>
-		/// <param name="intType">Integer type</param>
-		public UhtIntProperty(UhtPropertySettings propertySettings, UhtPropertyIntType intType) : base(propertySettings, intType)
+		public UhtIntProperty(UhtPropertySettings propertySettings) : base(propertySettings)
 		{
 			PropertyCaps |= UhtPropertyCaps.CanExposeOnSpawn | UhtPropertyCaps.IsParameterSupportedByBlueprint | UhtPropertyCaps.IsMemberSupportedByBlueprint | UhtPropertyCaps.SupportsRigVM;
 		}
@@ -34,16 +33,13 @@ namespace EpicGames.UHT.Types
 		/// <inheritdoc/>
 		public override StringBuilder AppendMemberDecl(StringBuilder builder, IUhtPropertyMemberContext context, string name, string nameSuffix, int tabs)
 		{
-			return AppendMemberDecl(builder, context, name, nameSuffix, tabs,
-				IntType == UhtPropertyIntType.Unsized ? "FUnsizedIntPropertyParams" : "FIntPropertyParams");
+			return AppendMemberDecl(builder, context, name, nameSuffix, tabs, "FIntPropertyParams");
 		}
 
 		/// <inheritdoc/>
 		public override StringBuilder AppendMemberDef(StringBuilder builder, IUhtPropertyMemberContext context, string name, string nameSuffix, string? offset, int tabs)
 		{
-			AppendMemberDefStart(builder, context, name, nameSuffix, offset, tabs,
-				IntType == UhtPropertyIntType.Unsized ? "FUnsizedIntPropertyParams" : "FIntPropertyParams",
-				"UECodeGen_Private::EPropertyGenFlags::Int");
+			AppendMemberDefStart(builder, context, name, nameSuffix, offset, tabs, "FIntPropertyParams", "UECodeGen_Private::EPropertyGenFlags::Int");
 			AppendMemberDefEnd(builder, context, name, nameSuffix);
 			return builder;
 		}
@@ -67,7 +63,7 @@ namespace EpicGames.UHT.Types
 		[SuppressMessage("Style", "IDE0060:Remove unused parameter", Justification = "Attribute accessed method")]
 		private static UhtProperty? Int32Property(UhtPropertyResolvePhase resolvePhase, UhtPropertySettings propertySettings, IUhtTokenReader tokenReader, UhtToken matchedToken)
 		{
-			return new UhtIntProperty(propertySettings, UhtPropertyIntType.Sized);
+			return new UhtIntProperty(propertySettings);
 		}
 
 		[UhtPropertyType(Keyword = "int", Options = UhtPropertyTypeOptions.Simple | UhtPropertyTypeOptions.Immediate)]
@@ -75,7 +71,7 @@ namespace EpicGames.UHT.Types
 		[SuppressMessage("Style", "IDE0060:Remove unused parameter", Justification = "Attribute accessed method")]
 		private static UhtProperty? IntProperty(UhtPropertyResolvePhase resolvePhase, UhtPropertySettings propertySettings, IUhtTokenReader tokenReader, UhtToken matchedToken)
 		{
-			return new UhtIntProperty(propertySettings, UhtPropertyIntType.Unsized);
+			return new UhtIntProperty(propertySettings);
 		}
 
 		[UhtPropertyType(Keyword = "signed", Options = UhtPropertyTypeOptions.Immediate)]
@@ -86,7 +82,7 @@ namespace EpicGames.UHT.Types
 			tokenReader
 				.Require("signed")
 				.Optional("int");
-			return new UhtIntProperty(propertySettings, UhtPropertyIntType.Unsized);
+			return new UhtIntProperty(propertySettings);
 		}
 		#endregion
 	}
diff --git a/Engine/Source/Programs/Shared/EpicGames.UHT/Types/Properties/UhtLargeWorldCoordinatesRealProperty.cs b/Engine/Source/Programs/Shared/EpicGames.UHT/Types/Properties/UhtLargeWorldCoordinatesRealProperty.cs
--- a/Engine/Source/Programs/Shared/EpicGames.UHT/Types/Properties/UhtLargeWorldCoordinatesRealProperty.cs
+++ b/Engine/Source/Programs/Shared/EpicGames.UHT/Types/Properties/UhtLargeWorldCoordinatesRealProperty.cs
@@ -26,7 +26,7 @@ namespace EpicGames.UHT.Types
 		/// Construct new property
 		/// </summary>
 		/// <param name="propertySettings">Property settings</param>
-		public UhtLargeWorldCoordinatesRealProperty(UhtPropertySettings propertySettings) : base(propertySettings, UhtPropertyIntType.None)
+		public UhtLargeWorldCoordinatesRealProperty(UhtPropertySettings propertySettings) : base(propertySettings)
 		{
 			PropertyCaps |= UhtPropertyCaps.IsParameterSupportedByBlueprint | UhtPropertyCaps.IsMemberSupportedByBlueprint;
 		}
diff --git a/Engine/Source/Programs/Shared/EpicGames.UHT/Types/Properties/UhtNumericProperty.cs b/Engine/Source/Programs/Shared/EpicGames.UHT/Types/Properties/UhtNumericProperty.cs
--- a/Engine/Source/Programs/Shared/EpicGames.UHT/Types/Properties/UhtNumericProperty.cs
+++ b/Engine/Source/Programs/Shared/EpicGames.UHT/Types/Properties/UhtNumericProperty.cs
@@ -17,20 +17,12 @@ namespace EpicGames.UHT.Types
 		/// <inheritdoc/>
 		protected override UhtPGetArgumentType PGetTypeArgument => UhtPGetArgumentType.EngineClass;
 
-		/// <summary>
-		/// Describes the integer as either being sized on unsized
-		/// </summary>
-		[JsonConverter(typeof(JsonStringEnumConverter))]
-		public UhtPropertyIntType IntType { get; }
-
 		/// <summary>
 		/// Construct new property
 		/// </summary>
 		/// <param name="propertySettings">Property settings</param>
-		/// <param name="intType">Type of integer</param>
-		protected UhtNumericProperty(UhtPropertySettings propertySettings, UhtPropertyIntType intType) : base(propertySettings)
+		protected UhtNumericProperty(UhtPropertySettings propertySettings) : base(propertySettings)
 		{
-			IntType = intType;
 			PropertyCaps |= UhtPropertyCaps.RequiresNullConstructorArg;
 		}
 
diff --git a/Engine/Source/Programs/Shared/EpicGames.UHT/Types/Properties/UhtUInt16Property.cs b/Engine/Source/Programs/Shared/EpicGames.UHT/Types/Properties/UhtUInt16Property.cs
--- a/Engine/Source/Programs/Shared/EpicGames.UHT/Types/Properties/UhtUInt16Property.cs
+++ b/Engine/Source/Programs/Shared/EpicGames.UHT/Types/Properties/UhtUInt16Property.cs
@@ -26,21 +26,20 @@ namespace EpicGames.UHT.Types
 		/// Construct a new property
 		/// </summary>
 		/// <param name="propertySettings">Property settings</param>
-		/// <param name="intType">Integer type</param>
-		public UhtUInt16Property(UhtPropertySettings propertySettings, UhtPropertyIntType intType) : base(propertySettings, intType)
+		public UhtUInt16Property(UhtPropertySettings propertySettings) : base(propertySettings)
 		{
 		}
 
 		/// <inheritdoc/>
 		public override StringBuilder AppendMemberDecl(StringBuilder builder, IUhtPropertyMemberContext context, string name, string nameSuffix, int tabs)
 		{
-			return AppendMemberDecl(builder, context, name, nameSuffix, tabs, "FFInt16PropertyParams");
+			return AppendMemberDecl(builder, context, name, nameSuffix, tabs, "FUInt16PropertyParams");
 		}
 
 		/// <inheritdoc/>
 		public override StringBuilder AppendMemberDef(StringBuilder builder, IUhtPropertyMemberContext context, string name, string nameSuffix, string? offset, int tabs)
 		{
-			AppendMemberDefStart(builder, context, name, nameSuffix, offset, tabs, "FFInt16PropertyParams", "UECodeGen_Private::EPropertyGenFlags::UInt16");
+			AppendMemberDefStart(builder, context, name, nameSuffix, offset, tabs, "FUInt16PropertyParams", "UECodeGen_Private::EPropertyGenFlags::UInt16");
 			AppendMemberDefEnd(builder, context, name, nameSuffix);
 			return builder;
 		}
@@ -69,7 +68,7 @@ namespace EpicGames.UHT.Types
 			}
 			else
 			{
-				return new UhtUInt16Property(propertySettings, UhtPropertyIntType.Sized);
+				return new UhtUInt16Property(propertySettings);
 			}
 		}
 		#endregion
diff --git a/Engine/Source/Programs/Shared/EpicGames.UHT/Types/Properties/UhtUInt32Property.cs b/Engine/Source/Programs/Shared/EpicGames.UHT/Types/Properties/UhtUInt32Property.cs
--- a/Engine/Source/Programs/Shared/EpicGames.UHT/Types/Properties/UhtUInt32Property.cs
+++ b/Engine/Source/Programs/Shared/EpicGames.UHT/Types/Properties/UhtUInt32Property.cs
@@ -26,8 +26,7 @@ namespace EpicGames.UHT.Types
 		/// Construct a new property
 		/// </summary>
 		/// <param name="propertySettings">Property settings</param>
-		/// <param name="intType">Integer type</param>
-		public UhtUInt32Property(UhtPropertySettings propertySettings, UhtPropertyIntType intType) : base(propertySettings, intType)
+		public UhtUInt32Property(UhtPropertySettings propertySettings) : base(propertySettings)
 		{
 			PropertyCaps |= UhtPropertyCaps.SupportsRigVM;
 		}
@@ -35,16 +34,13 @@ namespace EpicGames.UHT.Types
 		/// <inheritdoc/>
 		public override StringBuilder AppendMemberDecl(StringBuilder builder, IUhtPropertyMemberContext context, string name, string nameSuffix, int tabs)
 		{
-			return AppendMemberDecl(builder, context, name, nameSuffix, tabs,
-				IntType == UhtPropertyIntType.Unsized ? "FUnsizedFIntPropertyParams" : "FUInt32PropertyParams");
+			return AppendMemberDecl(builder, context, name, nameSuffix, tabs, "FUInt32PropertyParams");
 		}
 
 		/// <inheritdoc/>
 		public override StringBuilder AppendMemberDef(StringBuilder builder, IUhtPropertyMemberContext context, string name, string nameSuffix, string? offset, int tabs)
 		{
-			AppendMemberDefStart(builder, context, name, nameSuffix, offset, tabs,
-				IntType == UhtPropertyIntType.Unsized ? "FUnsizedFIntPropertyParams" : "FUInt32PropertyParams",
-				"UECodeGen_Private::EPropertyGenFlags::UInt32");
+			AppendMemberDefStart(builder, context, name, nameSuffix, offset, tabs, "FUInt32PropertyParams", "UECodeGen_Private::EPropertyGenFlags::UInt32");
 			AppendMemberDefEnd(builder, context, name, nameSuffix);
 			return builder;
 		}
@@ -73,7 +69,7 @@ namespace EpicGames.UHT.Types
 			}
 			else
 			{
-				return new UhtUInt32Property(propertySettings, UhtPropertyIntType.Sized);
+				return new UhtUInt32Property(propertySettings);
 			}
 		}
 
@@ -85,7 +81,7 @@ namespace EpicGames.UHT.Types
 			tokenReader
 				.Require("unsigned")
 				.Optional("int");
-			return new UhtUInt32Property(propertySettings, UhtPropertyIntType.Unsized);
+			return new UhtUInt32Property(propertySettings);
 		}
 		#endregion
 	}
diff --git a/Engine/Source/Programs/Shared/EpicGames.UHT/Types/Properties/UhtUInt64Property.cs b/Engine/Source/Programs/Shared/EpicGames.UHT/Types/Properties/UhtUInt64Property.cs
--- a/Engine/Source/Programs/Shared/EpicGames.UHT/Types/Properties/UhtUInt64Property.cs
+++ b/Engine/Source/Programs/Shared/EpicGames.UHT/Types/Properties/UhtUInt64Property.cs
@@ -26,21 +26,20 @@ namespace EpicGames.UHT.Types
 		/// Construct a new property
 		/// </summary>
 		/// <param name="propertySettings">Property settings</param>
-		/// <param name="intType">Integer type</param>
-		public UhtUInt64Property(UhtPropertySettings propertySettings, UhtPropertyIntType intType) : base(propertySettings, intType)
+		public UhtUInt64Property(UhtPropertySettings propertySettings) : base(propertySettings)
 		{
 		}
 
 		/// <inheritdoc/>
 		public override StringBuilder AppendMemberDecl(StringBuilder builder, IUhtPropertyMemberContext context, string name, string nameSuffix, int tabs)
 		{
-			return AppendMemberDecl(builder, context, name, nameSuffix, tabs, "FFInt64PropertyParams");
+			return AppendMemberDecl(builder, context, name, nameSuffix, tabs, "FUInt64PropertyParams");
 		}
 
 		/// <inheritdoc/>
 		public override StringBuilder AppendMemberDef(StringBuilder builder, IUhtPropertyMemberContext context, string name, string nameSuffix, string? offset, int tabs)
 		{
-			AppendMemberDefStart(builder, context, name, nameSuffix, offset, tabs, "FFInt64PropertyParams", "UECodeGen_Private::EPropertyGenFlags::UInt64");
+			AppendMemberDefStart(builder, context, name, nameSuffix, offset, tabs, "FUInt64PropertyParams", "UECodeGen_Private::EPropertyGenFlags::UInt64");
 			AppendMemberDefEnd(builder, context, name, nameSuffix);
 			return builder;
 		}
@@ -69,7 +68,7 @@ namespace EpicGames.UHT.Types
 			}
 			else
 			{
-				return new UhtUInt64Property(propertySettings, UhtPropertyIntType.Sized);
+				return new UhtUInt64Property(propertySettings);
 			}
 		}
 		#endregion
diff --git a/Engine/Source/Programs/Shared/EpicGames.UHT/Types/UhtProperty.cs b/Engine/Source/Programs/Shared/EpicGames.UHT/Types/UhtProperty.cs
--- a/Engine/Source/Programs/Shared/EpicGames.UHT/Types/UhtProperty.cs
+++ b/Engine/Source/Programs/Shared/EpicGames.UHT/Types/UhtProperty.cs
@@ -211,28 +211,6 @@ namespace EpicGames.UHT.Types
 		Native,
 	}
 
-	/// <summary>
-	/// Size type of an integer
-	/// </summary>
-	public enum UhtPropertyIntType
-	{
-
-		/// <summary>
-		/// Property is not an integer
-		/// </summary>
-		None,
-
-		/// <summary>
-		/// Property is a sized integer
-		/// </summary>
-		Sized,
-
-		/// <summary>
-		/// Property is an unsized integer
-		/// </summary>
-		Unsized,
-	};
-
 	/// <summary>
 	/// Type of reference
 	/// </summary>
diff --git a/Engine/Source/Programs/UnrealHeaderTool/Private/CodeGenerator.cpp b/Engine/Source/Programs/UnrealHeaderTool/Private/CodeGenerator.cpp
--- a/Engine/Source/Programs/UnrealHeaderTool/Private/CodeGenerator.cpp
+++ b/Engine/Source/Programs/UnrealHeaderTool/Private/CodeGenerator.cpp
@@ -905,7 +905,7 @@ void FNativeClassHeaderGenerator::OutputProperty(FOutputDevice& DeclOut, FOutput
 	{
 		FString MetaDataParams = OutputMetaDataCodeForObject(DeclOut, Out, PropertyDef, *FString::Printf(TEXT("%s_MetaData"), *Name), DeclSpaces, Spaces);
 
-		const TCHAR* PropTypeName = PropertyDef.IsUnsized() ? TEXT("FUnsizedIntPropertyParams") : TEXT("FIntPropertyParams");
+		const TCHAR* PropTypeName = TEXT("FIntPropertyParams");
 
 		DeclOut.Logf(TEXT("%sstatic const UECodeGen_Private::%s %s;\r\n"), DeclSpaces, PropTypeName, *NameWithoutScope);
 
@@ -956,10 +956,10 @@ void FNativeClassHeaderGenerator::OutputProperty(FOutputDevice& DeclOut, FOutput
 	{
 		FString MetaDataParams = OutputMetaDataCodeForObject(DeclOut, Out, PropertyDef, *FString::Printf(TEXT("%s_MetaData"), *Name), DeclSpaces, Spaces);
 
-		DeclOut.Logf(TEXT("%sstatic const UECodeGen_Private::FFInt16PropertyParams %s;\r\n"), DeclSpaces, *NameWithoutScope);
+		DeclOut.Logf(TEXT("%sstatic const UECodeGen_Private::FUInt16PropertyParams %s;\r\n"), DeclSpaces, *NameWithoutScope);
 
 		Out.Logf(
-			TEXT("%sconst UECodeGen_Private::FFInt16PropertyParams %s = { %s, %s, (EPropertyFlags)0x%016llx, UECodeGen_Private::EPropertyGenFlags::UInt16, %s, %s, %s, %s, %s, %s };%s\r\n"),
+			TEXT("%sconst UECodeGen_Private::FUInt16PropertyParams %s = { %s, %s, (EPropertyFlags)0x%016llx, UECodeGen_Private::EPropertyGenFlags::UInt16, %s, %s, %s, %s, %s, %s };%s\r\n"),
 			Spaces,
 			*Name,
 			*PropName,
@@ -980,7 +980,7 @@ void FNativeClassHeaderGenerator::OutputProperty(FOutputDevice& DeclOut, FOutput
 	{
 		FString MetaDataParams = OutputMetaDataCodeForObject(DeclOut, Out, PropertyDef, *FString::Printf(TEXT("%s_MetaData"), *Name), DeclSpaces, Spaces);
 
-		const TCHAR* PropTypeName = PropertyDef.IsUnsized() ? TEXT("FUnsizedFIntPropertyParams") : TEXT("FUInt32PropertyParams");
+		const TCHAR* PropTypeName = TEXT("FUInt32PropertyParams");
 
 		DeclOut.Logf(TEXT("%sstatic const UECodeGen_Private::%s %s;\r\n"), DeclSpaces, PropTypeName, *NameWithoutScope);
 
@@ -1007,10 +1007,10 @@ void FNativeClassHeaderGenerator::OutputProperty(FOutputDevice& DeclOut, FOutput
 	{
 		FString MetaDataParams = OutputMetaDataCodeForObject(DeclOut, Out, PropertyDef, *FString::Printf(TEXT("%s_MetaData"), *Name), DeclSpaces, Spaces);
 
-		DeclOut.Logf(TEXT("%sstatic const UECodeGen_Private::FFInt64PropertyParams %s;\r\n"), DeclSpaces, *NameWithoutScope);
+		DeclOut.Logf(TEXT("%sstatic const UECodeGen_Private::FUInt64PropertyParams %s;\r\n"), DeclSpaces, *NameWithoutScope);
 
 		Out.Logf(
-			TEXT("%sconst UECodeGen_Private::FFInt64PropertyParams %s = { %s, %s, (EPropertyFlags)0x%016llx, UECodeGen_Private::EPropertyGenFlags::UInt64, %s, %s, %s, %s, %s, %s };%s\r\n"),
+			TEXT("%sconst UECodeGen_Private::FUInt64PropertyParams %s = { %s, %s, (EPropertyFlags)0x%016llx, UECodeGen_Private::EPropertyGenFlags::UInt64, %s, %s, %s, %s, %s, %s };%s\r\n"),
 			Spaces,
 			*Name,
 			*PropName,
diff --git a/Engine/Source/Programs/UnrealHeaderTool/Private/HeaderParser.cpp b/Engine/Source/Programs/UnrealHeaderTool/Private/HeaderParser.cpp
--- a/Engine/Source/Programs/UnrealHeaderTool/Private/HeaderParser.cpp
+++ b/Engine/Source/Programs/UnrealHeaderTool/Private/HeaderParser.cpp
@@ -3839,17 +3839,17 @@ void FHeaderParser::GetVarType(
 	}
 	else if ( VarType.IsValue(TEXT("int"), ESearchCase::CaseSensitive) )
 	{
-		VarProperty = FPropertyBase(CPT_Int, EIntType::Unsized);
+		VarProperty = FPropertyBase(CPT_Int);
 	}
 	else if ( VarType.IsValue(TEXT("signed"), ESearchCase::CaseSensitive) )
 	{
 		MatchIdentifier(TEXT("int"), ESearchCase::CaseSensitive);
-		VarProperty = FPropertyBase(CPT_Int, EIntType::Unsized);
+		VarProperty = FPropertyBase(CPT_Int);
 	}
 	else if (VarType.IsValue(TEXT("unsigned"), ESearchCase::CaseSensitive))
 	{
 		MatchIdentifier(TEXT("int"), ESearchCase::CaseSensitive);
-		VarProperty = FPropertyBase(CPT_UInt32, EIntType::Unsized);
+		VarProperty = FPropertyBase(CPT_UInt32);
 	}
 	else if ( VarType.IsValue(TEXT("bool"), ESearchCase::CaseSensitive) )
 	{
diff --git a/Engine/Source/Programs/UnrealHeaderTool/Private/ParserHelper.cpp b/Engine/Source/Programs/UnrealHeaderTool/Private/ParserHelper.cpp
--- a/Engine/Source/Programs/UnrealHeaderTool/Private/ParserHelper.cpp
+++ b/Engine/Source/Programs/UnrealHeaderTool/Private/ParserHelper.cpp
@@ -18,20 +18,12 @@ extern FUnrealClassDefinitionInfo* GUInterfaceDef;
 
 FPropertyBase::FPropertyBase(EPropertyType InType)
 	: Type(InType)
-	, IntType(GetSizedIntTypeFromPropertyType(InType))
-{
-}
-
-FPropertyBase::FPropertyBase(EPropertyType InType, EIntType InIntType)
-	: Type(InType)
-	, IntType(InIntType)
 {
 }
 
 FPropertyBase::FPropertyBase(FUnrealEnumDefinitionInfo& InEnumDef, EPropertyType InType)
 	: Type(InType)
 	, EnumDef(&InEnumDef)
-	, IntType(GetSizedIntTypeFromPropertyType(InType))
 {
 }
 
@@ -54,7 +46,6 @@ FPropertyBase::FPropertyBase(FUnrealScriptStructDefinitionInfo& InStructDef)
 FPropertyBase::FPropertyBase(FName InFieldClassName, EPropertyType InType)
 	: Type(InType)
 	, FieldClassName(InFieldClassName)
-	, IntType(GetSizedIntTypeFromPropertyType(InType))
 {
 }
 
diff --git a/Engine/Source/Programs/UnrealHeaderTool/Private/ParserHelper.h b/Engine/Source/Programs/UnrealHeaderTool/Private/ParserHelper.h
--- a/Engine/Source/Programs/UnrealHeaderTool/Private/ParserHelper.h
+++ b/Engine/Source/Programs/UnrealHeaderTool/Private/ParserHelper.h
@@ -91,13 +91,6 @@ struct ERefQualifier
 	};
 };
 
-enum class EIntType
-{
-	None,
-	Sized,  // e.g. int32, int16
-	Unsized // e.g. int, unsigned int
-};
-
 enum class EAllocatorType
 {
 	Default,
@@ -224,7 +217,6 @@ public:
 	TMap<FName, FString> MetaData;
 
 	EPointerType::Type PointerType = EPointerType::None;
-	EIntType IntType = EIntType::None;
 
 	/** Setter function name. Either specified in the UPROPERTY macro or autogenerated while parsing class declaration */
 	FString SetterName;	
@@ -251,8 +243,6 @@ public:
 
 	explicit FPropertyBase(EPropertyType InType);
 
-	explicit FPropertyBase(EPropertyType InType, EIntType InIntType);
-
 	explicit FPropertyBase(FUnrealEnumDefinitionInfo& InEnumDef, EPropertyType InType);
 
 	explicit FPropertyBase(FUnrealClassDefinitionInfo& InClassDef, EPropertyType InType, bool bWeakIsAuto = false);
@@ -397,25 +387,6 @@ public:
 	bool MatchesType(const FPropertyBase& Other, bool bDisallowGeneralization, bool bIgnoreImplementedInterfaces = false, bool bEmulateSameType = false) const;
 	//@}
 
-	EIntType GetSizedIntTypeFromPropertyType(EPropertyType PropType)
-	{
-		switch (PropType)
-		{
-			case CPT_Byte:
-			case CPT_UInt16:
-			case CPT_UInt32:
-			case CPT_UInt64:
-			case CPT_Int8:
-			case CPT_Int16:
-			case CPT_Int:
-			case CPT_Int64:
-				return EIntType::Sized;
-
-			default:
-				return EIntType::None;
-		}
-	}
-
 	EUHTPropertyType GetUHTPropertyType() const;
 
 	friend struct FPropertyBaseArchiveProxy;
diff --git a/Engine/Source/Programs/UnrealHeaderTool/Private/PropertyTypes.cpp b/Engine/Source/Programs/UnrealHeaderTool/Private/PropertyTypes.cpp
--- a/Engine/Source/Programs/UnrealHeaderTool/Private/PropertyTypes.cpp
+++ b/Engine/Source/Programs/UnrealHeaderTool/Private/PropertyTypes.cpp
@@ -659,8 +659,6 @@ struct FPropertyTypeTraitsInt : public FPropertyTypeTraitsNumericBase
 
 	static void CreateProperty(FUnrealPropertyDefinitionInfo& PropDef)
 	{
-		const FPropertyBase& VarProperty = PropDef.GetPropertyBase();
-		PropDef.SetUnsized(VarProperty.IntType == EIntType::Unsized);
 	}
 
 	static bool IsSupportedByBlueprint(const FUnrealPropertyDefinitionInfo& PropDef, bool bMemberVariable)
@@ -733,8 +731,6 @@ struct FPropertyTypeTraitsUInt32 : public FPropertyTypeTraitsNumericBase
 {
 	static void CreateProperty(FUnrealPropertyDefinitionInfo& PropDef)
 	{
-		const FPropertyBase& VarProperty = PropDef.GetPropertyBase();
-		PropDef.SetUnsized(VarProperty.IntType == EIntType::Unsized);
 	}
 
 	static FString GetEngineClassName(const FUnrealPropertyDefinitionInfo& PropDef)
@@ -1046,7 +1042,6 @@ struct FPropertyTypeTraitsEnum : public FPropertyTypeTraitsBase
 		case EUnderlyingEnumType::uint64:      UnderlyingProperty.Type = CPT_UInt64; break;
 		case EUnderlyingEnumType::Unspecified:
 			UnderlyingProperty.Type = CPT_Int;
-			UnderlyingProperty.IntType = EIntType::Unsized;
 			break;
 
 		default:
diff --git a/Engine/Source/Programs/UnrealHeaderTool/Private/UnrealTypeDefinitionInfo.h b/Engine/Source/Programs/UnrealHeaderTool/Private/UnrealTypeDefinitionInfo.h
--- a/Engine/Source/Programs/UnrealHeaderTool/Private/UnrealTypeDefinitionInfo.h
+++ b/Engine/Source/Programs/UnrealHeaderTool/Private/UnrealTypeDefinitionInfo.h
@@ -980,22 +980,6 @@ public:
 		return AccessSpecifier;
 	}
 
-	/**
-	 * Return true if the property is unsized
-	 */
-	bool IsUnsized() const
-	{
-		return bUnsized;
-	}
-
-	/**
-	 * Set the unsized flag
-	 */
-	void SetUnsized(bool bInUnsized)
-	{
-		bUnsized = bInUnsized;
-	}
-
 	/**
 	 * Return the allocator type
 	 */
@@ -1165,7 +1149,6 @@ private:
 	EVariableCategory VariableCategory = EVariableCategory::Member;
 	EAccessSpecifier AccessSpecifier = ACCESS_Public;
 	EAllocatorType AllocatorType = EAllocatorType::Default;
-	bool bUnsized = false;
 	bool bSignatureSet = false; // temporary flag to know if we have set the signature
 };
 
diff --git a/Engine/Source/Runtime/CoreUObject/Private/UObject/PropertyNumeric.cpp b/Engine/Source/Runtime/CoreUObject/Private/UObject/PropertyNumeric.cpp
--- a/Engine/Source/Runtime/CoreUObject/Private/UObject/PropertyNumeric.cpp
+++ b/Engine/Source/Runtime/CoreUObject/Private/UObject/PropertyNumeric.cpp
@@ -300,7 +300,7 @@ FInt64Property::FInt64Property(FFieldVariant InOwner, const UECodeGen_Private::F
 {
 }
 
-FUInt16Property::FUInt16Property(FFieldVariant InOwner, const UECodeGen_Private::FFInt16PropertyParams& Prop)
+FUInt16Property::FUInt16Property(FFieldVariant InOwner, const UECodeGen_Private::FUInt16PropertyParams& Prop)
 	: TProperty_Numeric(InOwner, (const UECodeGen_Private::FPropertyParamsBaseWithOffset&)Prop)
 {
 }
@@ -310,7 +310,7 @@ FUInt32Property::FUInt32Property(FFieldVariant InOwner, const UECodeGen_Private:
 {
 }
 
-FUInt64Property::FUInt64Property(FFieldVariant InOwner, const UECodeGen_Private::FUnsizedIntPropertyParams& Prop)
+FUInt64Property::FUInt64Property(FFieldVariant InOwner, const UECodeGen_Private::FUInt64PropertyParams& Prop)
 	: TProperty_Numeric(InOwner, (const UECodeGen_Private::FPropertyParamsBaseWithOffset&)Prop)
 {
 }
diff --git a/Engine/Source/Runtime/CoreUObject/Private/UObject/UObjectGlobals.cpp b/Engine/Source/Runtime/CoreUObject/Private/UObject/UObjectGlobals.cpp
--- a/Engine/Source/Runtime/CoreUObject/Private/UObject/UObjectGlobals.cpp
+++ b/Engine/Source/Runtime/CoreUObject/Private/UObject/UObjectGlobals.cpp
@@ -5629,7 +5629,7 @@ namespace UECodeGen_Private
 
 			case EPropertyGenFlags::UInt16:
 			{
-				NewProp = NewFProperty<FUInt16Property, FFInt16PropertyParams>(Outer, *PropBase);
+				NewProp = NewFProperty<FUInt16Property, FUInt16PropertyParams>(Outer, *PropBase);
 			}
 			break;
 
@@ -5641,19 +5641,7 @@ namespace UECodeGen_Private
 
 			case EPropertyGenFlags::UInt64:
 			{
-				NewProp = NewFProperty<FUInt64Property, FFInt64PropertyParams>(Outer, *PropBase);
-			}
-			break;
-
-			case EPropertyGenFlags::UnsizedInt:
-			{
-				NewProp = NewFProperty<FUInt64Property, FUnsizedIntPropertyParams>(Outer, *PropBase);
-			}
-			break;
-
-			case EPropertyGenFlags::UnsizedUInt:
-			{
-				NewProp = NewFProperty<FUInt64Property, FUnsizedFIntPropertyParams>(Outer, *PropBase);
+				NewProp = NewFProperty<FUInt64Property, FUInt64PropertyParams>(Outer, *PropBase);
 			}
 			break;
 
diff --git a/Engine/Source/Runtime/CoreUObject/Public/UObject/UObjectGlobals.h b/Engine/Source/Runtime/CoreUObject/Public/UObject/UObjectGlobals.h
--- a/Engine/Source/Runtime/CoreUObject/Public/UObject/UObjectGlobals.h
+++ b/Engine/Source/Runtime/CoreUObject/Public/UObject/UObjectGlobals.h
@@ -3100,8 +3100,8 @@ namespace UECodeGen_Private
 		UInt16            = 0x05,
 		UInt32            = 0x06,
 		UInt64            = 0x07,
-		UnsizedInt        = 0x08,
-		UnsizedUInt       = 0x09,
+		//                = 0x08,
+		//                = 0x09,
 		Float             = 0x0A,
 		Double            = 0x0B,
 		Bool              = 0x0C,
@@ -3132,6 +3132,9 @@ namespace UECodeGen_Private
 
 	};
 
+	static_assert(std::is_same_v<int32,  int         >, "CoreUObject property system expects int32 to be an int");
+	static_assert(std::is_same_v<uint32, unsigned int>, "CoreUObject property system expects uint32 to be an unsigned int");
+
 	ENUM_CLASS_FLAGS(EPropertyGenFlags)
 
 	// Value which masks out the type of combined EPropertyGenFlags.
@@ -3483,11 +3486,9 @@ namespace UECodeGen_Private
 	typedef FGenericPropertyParams FInt16PropertyParams;
 	typedef FGenericPropertyParams FIntPropertyParams;
 	typedef FGenericPropertyParams FInt64PropertyParams;
-	typedef FGenericPropertyParams FFInt16PropertyParams;
+	typedef FGenericPropertyParams FUInt16PropertyParams;
 	typedef FGenericPropertyParams FUInt32PropertyParams;
-	typedef FGenericPropertyParams FFInt64PropertyParams;
-	typedef FGenericPropertyParams FUnsizedIntPropertyParams;
-	typedef FGenericPropertyParams FUnsizedFIntPropertyParams;
+	typedef FGenericPropertyParams FUInt64PropertyParams;
 	typedef FGenericPropertyParams FFloatPropertyParams;
 	typedef FGenericPropertyParams FDoublePropertyParams;
 	typedef FGenericPropertyParams FLargeWorldCoordinatesRealPropertyParams;
diff --git a/Engine/Source/Runtime/CoreUObject/Public/UObject/UnrealType.h b/Engine/Source/Runtime/CoreUObject/Public/UObject/UnrealType.h
--- a/Engine/Source/Runtime/CoreUObject/Public/UObject/UnrealType.h
+++ b/Engine/Source/Runtime/CoreUObject/Public/UObject/UnrealType.h
@@ -2226,7 +2226,7 @@ PRAGMA_ENABLE_DEPRECATION_WARNINGS
 	 * @param InOwner Owner of the property
 	 * @param PropBase Pointer to the compiled in structure describing the property
 	 **/
-	FUInt16Property(FFieldVariant InOwner, const UECodeGen_Private::FFInt16PropertyParams& Prop);
+	FUInt16Property(FFieldVariant InOwner, const UECodeGen_Private::FUInt16PropertyParams& Prop);
 
 #if WITH_EDITORONLY_DATA
 	explicit FUInt16Property(UField* InField)
@@ -2304,7 +2304,7 @@ PRAGMA_ENABLE_DEPRECATION_WARNINGS
 	 * @param InOwner Owner of the property
 	 * @param PropBase Pointer to the compiled in structure describing the property
 	 **/
-	FUInt64Property(FFieldVariant InOwner, const UECodeGen_Private::FUnsizedIntPropertyParams& Prop);
+	FUInt64Property(FFieldVariant InOwner, const UECodeGen_Private::FUInt64PropertyParams& Prop);
 
 #if WITH_EDITORONLY_DATA
 	explicit FUInt64Property(UField* InField)
@@ -2315,28 +2315,6 @@ PRAGMA_ENABLE_DEPRECATION_WARNINGS
 };
 
 
-/*-----------------------------------------------------------------------------
-	Aliases for implicitly-sized integer properties.
------------------------------------------------------------------------------*/
-
-namespace UETypes_Private
-{
-	template <typename IntType> struct TIntegerPropertyMapping;
-
-	template <> struct TIntegerPropertyMapping<int8>   { typedef FInt8Property   Type; };
-	template <> struct TIntegerPropertyMapping<int16>  { typedef FInt16Property  Type; };
-	template <> struct TIntegerPropertyMapping<int32>  { typedef FIntProperty    Type; };
-	template <> struct TIntegerPropertyMapping<int64>  { typedef FInt64Property  Type; };
-	template <> struct TIntegerPropertyMapping<uint8>  { typedef FByteProperty   Type; };
-	template <> struct TIntegerPropertyMapping<uint16> { typedef FUInt16Property Type; };
-	template <> struct TIntegerPropertyMapping<uint32> { typedef FUInt32Property Type; };
-	template <> struct TIntegerPropertyMapping<uint64> { typedef FUInt64Property Type; };
-}
-
-typedef UETypes_Private::TIntegerPropertyMapping<signed int>::Type UUnsizedIntProperty;
-typedef UETypes_Private::TIntegerPropertyMapping<unsigned int>::Type UUnsizedFIntProperty;
-
-
 /*-----------------------------------------------------------------------------
 	FFloatProperty.
 -----------------------------------------------------------------------------*/
