Horde: Fix files not being truncated when extracted from a bundle, resulting in files being longer than they should be.

diff --git a/Engine/Source/Programs/Shared/EpicGames.Horde/Storage/Nodes/FileNode.cs b/Engine/Source/Programs/Shared/EpicGames.Horde/Storage/Nodes/FileNode.cs
--- a/Engine/Source/Programs/Shared/EpicGames.Horde/Storage/Nodes/FileNode.cs
+++ b/Engine/Source/Programs/Shared/EpicGames.Horde/Storage/Nodes/FileNode.cs
@@ -112,7 +112,7 @@ namespace EpicGames.Horde.Storage.Nodes
 			{
 				file.Attributes &= ~FileAttributes.ReadOnly;
 			}
-			using (FileStream stream = file.OpenWrite())
+			using (FileStream stream = file.Open(FileMode.Create, FileAccess.Write, FileShare.Read))
 			{
 				await CopyToStreamAsync(reader, stream, cancellationToken);
 			}
