Horde: Add a test which creates a bundle manually and using TreeWriter, and checks them for equality.

diff --git a/Engine/Source/Programs/Shared/EpicGames.Horde.Tests/BundleStoreTests.cs b/Engine/Source/Programs/Shared/EpicGames.Horde.Tests/BundleStoreTests.cs
--- a/Engine/Source/Programs/Shared/EpicGames.Horde.Tests/BundleStoreTests.cs
+++ b/Engine/Source/Programs/Shared/EpicGames.Horde.Tests/BundleStoreTests.cs
@@ -12,7 +12,6 @@ using System.Threading.Tasks;
 using System.Threading;
 using System.Text;
 using System.IO;
-using System.Security.Cryptography;
 using EpicGames.Core;
 using Microsoft.Extensions.Caching.Memory;
 using Microsoft.Extensions.Logging.Abstractions;
@@ -24,6 +23,67 @@ namespace EpicGames.Horde.Tests
 	{
 		static readonly TreeReaderOptions s_readOptions = new TreeReaderOptions(typeof(SimpleNode));
 
+		[TestMethod]
+		public async Task CreateBundlesManuallyAsync()
+		{
+			Bundle a = CreateBundleManually();
+			Bundle b = await CreateBundleNormalAsync();
+
+			byte[] bytesA = a.AsSequence().ToArray();
+			byte[] bytesB = b.AsSequence().ToArray();
+			Assert.IsTrue(bytesA.SequenceEqual(bytesB));
+		}
+
+		[TreeNode("{F63606D4-5DBB-4061-A655-6F444F65229E}")]
+		class TextNode : TreeNode
+		{
+			public string Text { get; }
+
+			public TextNode(string text) => Text = text;
+
+			public TextNode(ITreeNodeReader reader)
+			{
+				Text = reader.ReadString();
+			}
+
+			public override void Serialize(ITreeNodeWriter writer)
+			{
+				writer.WriteString(Text);
+			}
+
+			public override IEnumerable<TreeNodeRef> EnumerateRefs() => Array.Empty<TreeNodeRef>();
+		}
+
+		static async Task<Bundle> CreateBundleNormalAsync()
+		{
+			MemoryStorageClient store = new MemoryStorageClient();
+			using TreeWriter writer = new TreeWriter(store, new TreeOptions { CompressionFormat = BundleCompressionFormat.None });
+
+			TextNode node = new TextNode("Hello world");
+			NodeHandle handle = await writer.FlushAsync(node, CancellationToken.None);
+
+			return await store.ReadBundleAsync(handle.Locator.Blob);
+		}
+
+		static Bundle CreateBundleManually()
+		{
+			ArrayMemoryWriter payloadWriter = new ArrayMemoryWriter(200);
+			payloadWriter.WriteString("Hello world");
+			byte[] payload = payloadWriter.WrittenMemory.ToArray();
+
+			List<BundleType> types = new List<BundleType>();
+			types.Add(new BundleType(Guid.Parse("F63606D4-5DBB-4061-A655-6F444F65229E"), 1));
+
+			List<BundleExport> exports = new List<BundleExport>();
+			exports.Add(new BundleExport(0, IoHash.Compute(payload), payload.Length, Array.Empty<int>()));
+
+			List<BundlePacket> packets = new List<BundlePacket>();
+			packets.Add(new BundlePacket(payload.Length, payload.Length));
+
+			BundleHeader header = new BundleHeader(BundleCompressionFormat.None, types, Array.Empty<BundleImport>(), exports, packets);
+			return new Bundle(header, new List<ReadOnlyMemory<byte>> { payload });
+		}
+
 		[TestMethod]
 		public async Task TestTreeAsync()
 		{
@@ -47,7 +107,7 @@ namespace EpicGames.Horde.Tests
 			Assert.AreEqual(5, blobStore.Blobs.Count);
 			Assert.AreEqual(1, blobStore.Refs.Count);
 		}
-
+		
 		[TreeNode("{F63606D4-5DBB-4061-A655-6F444F65229F}")]
 		class SimpleNode : TreeNode
 		{
