Horde: Additional logging for spawning child processes.

diff --git a/Engine/Source/Programs/Shared/EpicGames.Horde/Compute/ComputeWorker.cs b/Engine/Source/Programs/Shared/EpicGames.Horde/Compute/ComputeWorker.cs
--- a/Engine/Source/Programs/Shared/EpicGames.Horde/Compute/ComputeWorker.cs
+++ b/Engine/Source/Programs/Shared/EpicGames.Horde/Compute/ComputeWorker.cs
@@ -173,6 +173,7 @@ namespace EpicGames.Horde.Compute
 
 			newEnvVars[ComputeSocket.WorkerIpcEnvVar] = ipcChannel.GetStringHandle();
 
+			_logger.LogInformation("Launching {Executable} {Arguments}...", CommandLineArguments.Quote(executable), arguments);
 			Task processTask = ExecuteProcessInternalAsync(channel, executable, arguments, workingDir, newEnvVars, cancellationToken);
 			processTask = processTask.ContinueWith(x => ipcChannel.ForceComplete(), cancellationToken, TaskContinuationOptions.None, TaskScheduler.Default);
 
@@ -185,11 +186,17 @@ namespace EpicGames.Horde.Compute
 					switch (message.Type)
 					{
 						case ComputeMessageType.None:
+							if (!processTask.IsCompleted)
+							{
+								_logger.LogInformation("Waiting for child process to terminate");
+							}
 							await processTask;
+							_logger.LogInformation("Finished executing process");
 							return;
 						case ComputeMessageType.AttachRecvBuffer:
 							{
 								AttachRecvBufferRequest attachRecvBuffer = IpcComputeMessageChannel.ParseAttachRecvBuffer(message);
+								_logger.LogInformation("Attaching receive buffer to channel {Id}", attachRecvBuffer.ChannelId);
 #pragma warning disable CA2000 // Dispose objects before losing scope
 								(IComputeBufferReader reader, IComputeBufferWriter writer) = SharedMemoryBuffer.OpenIpcHandle(attachRecvBuffer.Handle).ToShared();
 								reader.Dispose();
@@ -200,6 +207,7 @@ namespace EpicGames.Horde.Compute
 						case ComputeMessageType.AttachSendBuffer:
 							{
 								AttachSendBufferRequest attachSendBuffer = IpcComputeMessageChannel.ParseAttachSendBuffer(message);
+								_logger.LogInformation("Attaching send buffer to channel {Id}", attachSendBuffer.ChannelId);
 #pragma warning disable CA2000 // Dispose objects before losing scope
 								(IComputeBufferReader reader, IComputeBufferWriter writer) = SharedMemoryBuffer.OpenIpcHandle(attachSendBuffer.Handle).ToShared();
 								writer.Dispose();
