Horde: Support cancellation of the child process during remote execution.

diff --git a/Engine/Source/Programs/Shared/EpicGames.Horde/Compute/ComputeWorker.cs b/Engine/Source/Programs/Shared/EpicGames.Horde/Compute/ComputeWorker.cs
--- a/Engine/Source/Programs/Shared/EpicGames.Horde/Compute/ComputeWorker.cs
+++ b/Engine/Source/Programs/Shared/EpicGames.Horde/Compute/ComputeWorker.cs
@@ -175,14 +175,23 @@ namespace EpicGames.Horde.Compute
 
 			_logger.LogInformation("Launching {Executable} {Arguments}...", CommandLineArguments.Quote(executable), arguments);
 			Task processTask = ExecuteProcessInternalAsync(channel, executable, arguments, workingDir, newEnvVars, cancellationToken);
-			processTask = processTask.ContinueWith(x => ipcChannel.ForceComplete(), cancellationToken, TaskContinuationOptions.None, TaskScheduler.Default);
+
+			using CancellationTokenSource cancellationSource = CancellationTokenSource.CreateLinkedTokenSource(cancellationToken);
+
+			void OnProcessExit()
+			{
+				ipcChannel.ForceComplete();
+				cancellationSource.Cancel();
+			}
+
+			processTask = processTask.ContinueWith(_ => OnProcessExit(), cancellationToken, TaskContinuationOptions.None, TaskScheduler.Default);
 
 			List<IDisposable> buffers = new List<IDisposable>();
 			try
 			{
 				for (; ; )
 				{
-					using IComputeMessage message = await ipcChannel.ReceiveAsync(cancellationToken);
+					using IComputeMessage message = await ipcChannel.ReceiveAsync(cancellationSource.Token);
 					switch (message.Type)
 					{
 						case ComputeMessageType.None:
@@ -200,7 +209,7 @@ namespace EpicGames.Horde.Compute
 #pragma warning disable CA2000 // Dispose objects before losing scope
 								(IComputeBufferReader reader, IComputeBufferWriter writer) = SharedMemoryBuffer.OpenIpcHandle(attachRecvBuffer.Handle).ToShared();
 								reader.Dispose();
-								await socket.AttachRecvBufferAsync(attachRecvBuffer.ChannelId, writer, cancellationToken);
+								await socket.AttachRecvBufferAsync(attachRecvBuffer.ChannelId, writer, cancellationSource.Token);
 #pragma warning restore CA2000 // Dispose objects before losing scope
 							}
 							break;
@@ -211,7 +220,7 @@ namespace EpicGames.Horde.Compute
 #pragma warning disable CA2000 // Dispose objects before losing scope
 								(IComputeBufferReader reader, IComputeBufferWriter writer) = SharedMemoryBuffer.OpenIpcHandle(attachSendBuffer.Handle).ToShared();
 								writer.Dispose();
-								await socket.AttachSendBufferAsync(attachSendBuffer.ChannelId, reader, cancellationToken);
+								await socket.AttachSendBufferAsync(attachSendBuffer.ChannelId, reader, cancellationSource.Token);
 #pragma warning restore CA2000 // Dispose objects before losing scope
 							}
 							break;
@@ -220,6 +229,10 @@ namespace EpicGames.Horde.Compute
 					}
 				}
 			}
+			catch (OperationCanceledException)
+			{
+				_logger.LogInformation("Cancelled during process exec");
+			}
 			finally
 			{
 				foreach (SharedMemoryBuffer buffer in buffers)
