Horde: Improve cancellation handling for compute workers. Also add default cancellation token parameters to various methods.

diff --git a/Engine/Source/Programs/Shared/EpicGames.Horde/Compute/ComputeWorker.cs b/Engine/Source/Programs/Shared/EpicGames.Horde/Compute/ComputeWorker.cs
--- a/Engine/Source/Programs/Shared/EpicGames.Horde/Compute/ComputeWorker.cs
+++ b/Engine/Source/Programs/Shared/EpicGames.Horde/Compute/ComputeWorker.cs
@@ -148,59 +148,60 @@ namespace EpicGames.Horde.Compute
 
 		async Task ExecuteProcessAsync(IComputeSocket socket, IComputeMessageChannel channel, string executable, IReadOnlyList<string> arguments, string? workingDir, IReadOnlyDictionary<string, string?>? envVars, CancellationToken cancellationToken)
 		{
-			if (RuntimeInformation.IsOSPlatform(OSPlatform.Windows))
+			try
 			{
-				await ExecuteProcessWindowsAsync(socket, channel, executable, arguments, workingDir, envVars, cancellationToken);
+				if (RuntimeInformation.IsOSPlatform(OSPlatform.Windows))
+				{
+					await ExecuteProcessWindowsAsync(socket, channel, executable, arguments, workingDir, envVars, cancellationToken);
+				}
+				else
+				{
+					await ExecuteProcessInternalAsync(channel, executable, arguments, workingDir, envVars, cancellationToken);
+				}
 			}
-			else
+			catch (Exception ex)
 			{
-				await ExecuteProcessInternalAsync(channel, executable, arguments, workingDir, envVars, cancellationToken);
+				await channel.SendExceptionAsync(ex, cancellationToken);
 			}
 		}
 
 		async Task ExecuteProcessWindowsAsync(IComputeSocket socket, IComputeMessageChannel channel, string executable, IReadOnlyList<string> arguments, string? workingDir, IReadOnlyDictionary<string, string?>? envVars, CancellationToken cancellationToken)
 		{
 			await using IpcComputeMessageChannel ipcChannel = new IpcComputeMessageChannel(4096, _logger);
-
-			Dictionary<string, string?> newEnvVars = new Dictionary<string, string?>();
-			if (envVars != null)
+			await using (BackgroundTask ipcTask = BackgroundTask.StartNew(ctx => RunIpcChannel(socket, ipcChannel, ctx)))
 			{
-				foreach ((string name, string? value) in envVars)
+				Dictionary<string, string?> newEnvVars = new Dictionary<string, string?>();
+				if (envVars != null)
 				{
-					newEnvVars.Add(name, value);
+					foreach ((string name, string? value) in envVars)
+					{
+						newEnvVars.Add(name, value);
+					}
 				}
-			}
-
-			newEnvVars[ComputeSocket.WorkerIpcEnvVar] = ipcChannel.GetStringHandle();
 
-			_logger.LogInformation("Launching {Executable} {Arguments}...", CommandLineArguments.Quote(executable), arguments);
-			Task processTask = ExecuteProcessInternalAsync(channel, executable, arguments, workingDir, newEnvVars, cancellationToken);
+				newEnvVars[ComputeSocket.WorkerIpcEnvVar] = ipcChannel.GetStringHandle();
 
-			using CancellationTokenSource cancellationSource = CancellationTokenSource.CreateLinkedTokenSource(cancellationToken);
+				_logger.LogInformation("Launching {Executable} {Arguments}...", CommandLineArguments.Quote(executable), arguments);
+				await ExecuteProcessInternalAsync(channel, executable, arguments, workingDir, newEnvVars, cancellationToken);
+				_logger.LogInformation("Finished executing process");
 
-			void OnProcessExit()
-			{
 				ipcChannel.ForceComplete();
-				cancellationSource.Cancel();
+				await ipcTask.Task.WaitAsync(TimeSpan.FromSeconds(5.0), cancellationToken);
 			}
+			_logger.LogInformation("Child process has shut down");
+		}
 
-			processTask = processTask.ContinueWith(_ => OnProcessExit(), cancellationToken, TaskContinuationOptions.None, TaskScheduler.Default);
-
+		async Task RunIpcChannel(IComputeSocket socket, IpcComputeMessageChannel ipcChannel, CancellationToken cancellationToken)
+		{
 			List<IDisposable> buffers = new List<IDisposable>();
 			try
 			{
 				for (; ; )
 				{
-					using IComputeMessage message = await ipcChannel.ReceiveAsync(cancellationSource.Token);
+					using IComputeMessage message = await ipcChannel.ReceiveAsync(cancellationToken);
 					switch (message.Type)
 					{
 						case ComputeMessageType.None:
-							if (!processTask.IsCompleted)
-							{
-								_logger.LogInformation("Waiting for child process to terminate");
-							}
-							await processTask;
-							_logger.LogInformation("Finished executing process");
 							return;
 						case ComputeMessageType.AttachRecvBuffer:
 							{
@@ -209,7 +210,7 @@ namespace EpicGames.Horde.Compute
 #pragma warning disable CA2000 // Dispose objects before losing scope
 								(IComputeBufferReader reader, IComputeBufferWriter writer) = SharedMemoryBuffer.OpenIpcHandle(attachRecvBuffer.Handle).ToShared();
 								reader.Dispose();
-								await socket.AttachRecvBufferAsync(attachRecvBuffer.ChannelId, writer, cancellationSource.Token);
+								await socket.AttachRecvBufferAsync(attachRecvBuffer.ChannelId, writer, cancellationToken);
 #pragma warning restore CA2000 // Dispose objects before losing scope
 							}
 							break;
@@ -220,7 +221,7 @@ namespace EpicGames.Horde.Compute
 #pragma warning disable CA2000 // Dispose objects before losing scope
 								(IComputeBufferReader reader, IComputeBufferWriter writer) = SharedMemoryBuffer.OpenIpcHandle(attachSendBuffer.Handle).ToShared();
 								writer.Dispose();
-								await socket.AttachSendBufferAsync(attachSendBuffer.ChannelId, reader, cancellationSource.Token);
+								await socket.AttachSendBufferAsync(attachSendBuffer.ChannelId, reader, cancellationToken);
 #pragma warning restore CA2000 // Dispose objects before losing scope
 							}
 							break;
diff --git a/Engine/Source/Programs/Shared/EpicGames.Horde/Compute/IComputeBuffer.cs b/Engine/Source/Programs/Shared/EpicGames.Horde/Compute/IComputeBuffer.cs
--- a/Engine/Source/Programs/Shared/EpicGames.Horde/Compute/IComputeBuffer.cs
+++ b/Engine/Source/Programs/Shared/EpicGames.Horde/Compute/IComputeBuffer.cs
@@ -65,7 +65,7 @@ namespace EpicGames.Horde.Compute
 		/// </summary>
 		/// <param name="currentLength">Current length of the buffer</param>
 		/// <param name="cancellationToken">Cancellation token for the operation</param>
-		ValueTask WaitForDataAsync(int currentLength, CancellationToken cancellationToken);
+		ValueTask WaitForDataAsync(int currentLength, CancellationToken cancellationToken = default);
 	}
 
 	/// <summary>
diff --git a/Engine/Source/Programs/Shared/EpicGames.Horde/Compute/IComputeClient.cs b/Engine/Source/Programs/Shared/EpicGames.Horde/Compute/IComputeClient.cs
--- a/Engine/Source/Programs/Shared/EpicGames.Horde/Compute/IComputeClient.cs
+++ b/Engine/Source/Programs/Shared/EpicGames.Horde/Compute/IComputeClient.cs
@@ -17,6 +17,6 @@ namespace EpicGames.Horde.Compute
 		/// <param name="clusterId">Cluster to execute the request</param>
 		/// <param name="requirements">Requirements for the agent</param>
 		/// <param name="cancellationToken">Cancellation token for the operation</param>
-		public Task<IComputeLease?> TryAssignWorkerAsync(ClusterId clusterId, Requirements? requirements, CancellationToken cancellationToken);
+		public Task<IComputeLease?> TryAssignWorkerAsync(ClusterId clusterId, Requirements? requirements, CancellationToken cancellationToken = default);
 	}
 }
diff --git a/Engine/Source/Programs/Shared/EpicGames.Horde/Compute/IComputeMessage.cs b/Engine/Source/Programs/Shared/EpicGames.Horde/Compute/IComputeMessage.cs
--- a/Engine/Source/Programs/Shared/EpicGames.Horde/Compute/IComputeMessage.cs
+++ b/Engine/Source/Programs/Shared/EpicGames.Horde/Compute/IComputeMessage.cs
@@ -284,7 +284,7 @@ namespace EpicGames.Horde.Compute
 		/// <param name="locator">Location of a <see cref="DirectoryNode"/> describing contents of the sandbox</param>
 		/// <param name="storage">Storage for the sandbox data</param>
 		/// <param name="cancellationToken">Cancellation token for the operation</param>
-		public static async Task UploadFilesAsync(this IComputeMessageChannel channel, string path, NodeLocator locator, IStorageClient storage, CancellationToken cancellationToken)
+		public static async Task UploadFilesAsync(this IComputeMessageChannel channel, string path, NodeLocator locator, IStorageClient storage, CancellationToken cancellationToken = default)
 		{
 			using (IComputeMessageBuilder request = await channel.CreateMessageAsync(ComputeMessageType.WriteFiles, cancellationToken))
 			{
@@ -389,7 +389,7 @@ namespace EpicGames.Horde.Compute
 		/// <param name="outputHandler">Output callback for stdout</param>
 		/// <param name="envVars">Environment variables for the child process</param>
 		/// <param name="cancellationToken">Cancellation token for the operation</param>
-		public static Task<int> ExecuteAsync(this IComputeMessageChannel channel, string executable, IReadOnlyList<string> arguments, string? workingDir, IReadOnlyDictionary<string, string?>? envVars, Action<string> outputHandler, CancellationToken cancellationToken)
+		public static Task<int> ExecuteAsync(this IComputeMessageChannel channel, string executable, IReadOnlyList<string> arguments, string? workingDir, IReadOnlyDictionary<string, string?>? envVars, Action<string> outputHandler, CancellationToken cancellationToken = default)
 		{
 			StringOutputWriter writer = new StringOutputWriter(outputHandler);
 			return ExecuteAsync(channel, executable, arguments, workingDir, envVars, writer.WriteData, cancellationToken);
@@ -405,7 +405,7 @@ namespace EpicGames.Horde.Compute
 		/// <param name="outputHandler">Output callback for stdout</param>
 		/// <param name="envVars">Environment variables for the child process</param>
 		/// <param name="cancellationToken">Cancellation token for the operation</param>
-		public static async Task<int> ExecuteAsync(this IComputeMessageChannel channel, string executable, IReadOnlyList<string> arguments, string? workingDir, IReadOnlyDictionary<string, string?>? envVars, Action<ReadOnlyMemory<byte>> outputHandler, CancellationToken cancellationToken)
+		public static async Task<int> ExecuteAsync(this IComputeMessageChannel channel, string executable, IReadOnlyList<string> arguments, string? workingDir, IReadOnlyDictionary<string, string?>? envVars, Action<ReadOnlyMemory<byte>> outputHandler, CancellationToken cancellationToken = default)
 		{
 			using (IComputeMessageBuilder request = await channel.CreateMessageAsync(ComputeMessageType.Execute, cancellationToken))
 			{
@@ -451,7 +451,7 @@ namespace EpicGames.Horde.Compute
 		/// <summary>
 		/// Sends output from a child process
 		/// </summary>
-		public static async ValueTask SendExecuteOutputAsync(this IComputeMessageChannel channel, ReadOnlyMemory<byte> data, CancellationToken cancellationToken)
+		public static async ValueTask SendExecuteOutputAsync(this IComputeMessageChannel channel, ReadOnlyMemory<byte> data, CancellationToken cancellationToken = default)
 		{
 			using IComputeMessageBuilder message = await channel.CreateMessageAsync(ComputeMessageType.ExecuteOutput, cancellationToken);
 			message.WriteFixedLengthBytes(data.Span);
@@ -464,7 +464,7 @@ namespace EpicGames.Horde.Compute
 		/// <param name="channel"></param>
 		/// <param name="exitCode">Exit code from the process</param>
 		/// <param name="cancellationToken">Cancellation token for the operation</param>
-		public static async ValueTask SendExecuteResultAsync(this IComputeMessageChannel channel, int exitCode, CancellationToken cancellationToken)
+		public static async ValueTask SendExecuteResultAsync(this IComputeMessageChannel channel, int exitCode, CancellationToken cancellationToken = default)
 		{
 			using IComputeMessageBuilder builder = await channel.CreateMessageAsync(ComputeMessageType.ExecuteResult, cancellationToken);
 			builder.WriteInt32(exitCode);
@@ -530,7 +530,7 @@ namespace EpicGames.Horde.Compute
 		/// <param name="length">Length of data to return</param>
 		/// <param name="cancellationToken">Cancellation token for the operation</param>
 		/// <returns>Stream containing the blob data</returns>
-		public static async Task<Stream> ReadBlobAsync(this IComputeMessageChannel channel, BlobLocator locator, int offset, int length, CancellationToken cancellationToken)
+		public static async Task<Stream> ReadBlobAsync(this IComputeMessageChannel channel, BlobLocator locator, int offset, int length, CancellationToken cancellationToken = default)
 		{
 			using (IComputeMessageBuilder request = await channel.CreateMessageAsync(ComputeMessageType.ReadBlob, cancellationToken))
 			{
@@ -557,7 +557,7 @@ namespace EpicGames.Horde.Compute
 		/// <param name="message">The read request</param>
 		/// <param name="storage">Storage client to retrieve the blob from</param>
 		/// <param name="cancellationToken">Cancellation token for the operation</param>
-		public static Task SendBlobDataAsync(this IComputeMessageChannel channel, ReadBlobMessage message, IStorageClient storage, CancellationToken cancellationToken)
+		public static Task SendBlobDataAsync(this IComputeMessageChannel channel, ReadBlobMessage message, IStorageClient storage, CancellationToken cancellationToken = default)
 		{
 			return SendBlobDataAsync(channel, message.Locator, message.Offset, message.Length, storage, cancellationToken);
 		}
@@ -571,7 +571,7 @@ namespace EpicGames.Horde.Compute
 		/// <param name="length">Length of the data</param>
 		/// <param name="storage">Storage client to retrieve the blob from</param>
 		/// <param name="cancellationToken">Cancellation token for the operation</param>
-		public static async Task SendBlobDataAsync(this IComputeMessageChannel channel, BlobLocator locator, int offset, int length, IStorageClient storage, CancellationToken cancellationToken)
+		public static async Task SendBlobDataAsync(this IComputeMessageChannel channel, BlobLocator locator, int offset, int length, IStorageClient storage, CancellationToken cancellationToken = default)
 		{
 			byte[] data;
 			if (offset == 0 && length == 0)
@@ -607,7 +607,7 @@ namespace EpicGames.Horde.Compute
 		/// <summary>
 		/// Send a message to request that a byte string be xor'ed with a particular value
 		/// </summary>
-		public static async ValueTask SendXorRequestAsync(this IComputeMessageChannel channel, ReadOnlyMemory<byte> data, byte value, CancellationToken cancellationToken)
+		public static async ValueTask SendXorRequestAsync(this IComputeMessageChannel channel, ReadOnlyMemory<byte> data, byte value, CancellationToken cancellationToken = default)
 		{
 			using IComputeMessageBuilder message = await channel.CreateMessageAsync(ComputeMessageType.XorRequest, cancellationToken);
 			message.WriteFixedLengthBytes(data.Span);
diff --git a/Engine/Source/Programs/Shared/EpicGames.Horde/Compute/IComputeSocket.cs b/Engine/Source/Programs/Shared/EpicGames.Horde/Compute/IComputeSocket.cs
--- a/Engine/Source/Programs/Shared/EpicGames.Horde/Compute/IComputeSocket.cs
+++ b/Engine/Source/Programs/Shared/EpicGames.Horde/Compute/IComputeSocket.cs
@@ -83,7 +83,7 @@ namespace EpicGames.Horde.Compute
 		/// <param name="channelId">Channel to receive data on</param>
 		/// <param name="capacity">Capacity for the buffer</param>
 		/// <param name="cancellationToken">Cancellation token for the operation</param>
-		public static ValueTask<IComputeBufferReader> AttachRecvBufferAsync(this IComputeSocket socket, int channelId, long capacity, CancellationToken cancellationToken)
+		public static ValueTask<IComputeBufferReader> AttachRecvBufferAsync(this IComputeSocket socket, int channelId, long capacity, CancellationToken cancellationToken = default)
 		{
 			return AttachRecvBufferAsync(socket, channelId, socket.CreateBuffer(capacity), cancellationToken);
 		}
@@ -95,7 +95,7 @@ namespace EpicGames.Horde.Compute
 		/// <param name="channelId">Channel to receive data on</param>
 		/// <param name="buffer">Buffer to attach to the socket</param>
 		/// <param name="cancellationToken">Cancellation token for the operation</param>
-		public static async ValueTask<IComputeBufferReader> AttachRecvBufferAsync(this IComputeSocket socket, int channelId, IComputeBuffer buffer, CancellationToken cancellationToken)
+		public static async ValueTask<IComputeBufferReader> AttachRecvBufferAsync(this IComputeSocket socket, int channelId, IComputeBuffer buffer, CancellationToken cancellationToken = default)
 		{
 			(IComputeBufferReader reader, IComputeBufferWriter writer) = buffer.ToShared();
 			await socket.AttachRecvBufferAsync(channelId, writer, cancellationToken);
@@ -109,7 +109,7 @@ namespace EpicGames.Horde.Compute
 		/// <param name="channelId">Channel to receive data on</param>
 		/// <param name="capacity">Capacity for the buffer</param>
 		/// <param name="cancellationToken">Cancellation token for the operation</param>
-		public static ValueTask<IComputeBufferWriter> AttachSendBufferAsync(this IComputeSocket socket, int channelId, long capacity, CancellationToken cancellationToken) 
+		public static ValueTask<IComputeBufferWriter> AttachSendBufferAsync(this IComputeSocket socket, int channelId, long capacity, CancellationToken cancellationToken = default) 
 			=> AttachSendBufferAsync(socket, channelId, socket.CreateBuffer(capacity), cancellationToken);
 
 		/// <summary>
@@ -119,7 +119,7 @@ namespace EpicGames.Horde.Compute
 		/// <param name="channelId">Channel to receive data on</param>
 		/// <param name="buffer">Buffer for the data to send</param>
 		/// <param name="cancellationToken">Cancellation token for the operation</param>
-		public static async ValueTask<IComputeBufferWriter> AttachSendBufferAsync(this IComputeSocket socket, int channelId, IComputeBuffer buffer, CancellationToken cancellationToken)
+		public static async ValueTask<IComputeBufferWriter> AttachSendBufferAsync(this IComputeSocket socket, int channelId, IComputeBuffer buffer, CancellationToken cancellationToken = default)
 		{
 			(IComputeBufferReader reader, IComputeBufferWriter writer) = buffer.ToShared();
 			await socket.AttachSendBufferAsync(channelId, reader, cancellationToken);
@@ -133,7 +133,7 @@ namespace EpicGames.Horde.Compute
 		/// <param name="channelId">Identifier for the channel</param>
 		/// <param name="logger">Logger for the channel</param>
 		/// <param name="cancellationToken">Cancellation token for the operation</param>
-		public static ValueTask<IComputeMessageChannel> CreateMessageChannelAsync(this IComputeSocket socket, int channelId, ILogger logger, CancellationToken cancellationToken) 
+		public static ValueTask<IComputeMessageChannel> CreateMessageChannelAsync(this IComputeSocket socket, int channelId, ILogger logger, CancellationToken cancellationToken = default) 
 			=> socket.CreateMessageChannelAsync(channelId, 65536, logger, cancellationToken);
 
 		/// <summary>
@@ -144,7 +144,7 @@ namespace EpicGames.Horde.Compute
 		/// <param name="bufferSize">Size of the send and receive buffer</param>
 		/// <param name="logger">Logger for the channel</param>
 		/// <param name="cancellationToken">Cancellation token for the operation</param>
-		public static ValueTask<IComputeMessageChannel> CreateMessageChannelAsync(this IComputeSocket socket, int channelId, long bufferSize, ILogger logger, CancellationToken cancellationToken)
+		public static ValueTask<IComputeMessageChannel> CreateMessageChannelAsync(this IComputeSocket socket, int channelId, long bufferSize, ILogger logger, CancellationToken cancellationToken = default)
 			=> socket.CreateMessageChannelAsync(channelId, bufferSize, bufferSize, logger, cancellationToken);
 
 		/// <summary>
@@ -156,7 +156,7 @@ namespace EpicGames.Horde.Compute
 		/// <param name="recvBufferSize">Size of the recieve buffer</param>
 		/// <param name="logger">Logger for the channel</param>
 		/// <param name="cancellationToken">Cancellation token for the operation</param>
-		public static async ValueTask<IComputeMessageChannel> CreateMessageChannelAsync(this IComputeSocket socket, int channelId, long sendBufferSize, long recvBufferSize, ILogger logger, CancellationToken cancellationToken)
+		public static async ValueTask<IComputeMessageChannel> CreateMessageChannelAsync(this IComputeSocket socket, int channelId, long sendBufferSize, long recvBufferSize, ILogger logger, CancellationToken cancellationToken = default)
 		{
 			IComputeBufferReader reader = await socket.AttachRecvBufferAsync(channelId, recvBufferSize, cancellationToken);
 			IComputeBufferWriter writer = await socket.AttachSendBufferAsync(channelId, sendBufferSize, cancellationToken);
diff --git a/Engine/Source/Programs/Shared/EpicGames.Horde/Storage/Nodes/DirectoryNode.cs b/Engine/Source/Programs/Shared/EpicGames.Horde/Storage/Nodes/DirectoryNode.cs
--- a/Engine/Source/Programs/Shared/EpicGames.Horde/Storage/Nodes/DirectoryNode.cs
+++ b/Engine/Source/Programs/Shared/EpicGames.Horde/Storage/Nodes/DirectoryNode.cs
@@ -847,7 +847,7 @@ namespace EpicGames.Horde.Storage.Nodes
 		/// <param name="writer">Writer for new node data</param>
 		/// <param name="progress">Feedback interface for progress updates</param>
 		/// <param name="cancellationToken">Cancellation token for the operation</param>
-		public async Task CopyFromDirectoryAsync(DirectoryInfo directoryInfo, ChunkingOptions options, TreeWriter writer, IProgress<ICopyStats>? progress, CancellationToken cancellationToken)
+		public async Task CopyFromDirectoryAsync(DirectoryInfo directoryInfo, ChunkingOptions options, TreeWriter writer, IProgress<ICopyStats>? progress, CancellationToken cancellationToken = default)
 		{
 			// Enumerate all the files below this directory
 			List<(DirectoryNode DirectoryNode, FileInfo FileInfo)> files = new List<(DirectoryNode, FileInfo)>();
@@ -864,7 +864,7 @@ namespace EpicGames.Horde.Storage.Nodes
 		/// <param name="progress">Progress notification object</param>
 		/// <param name="cancellationToken">Cancellation token for the operation</param>
 		/// <returns></returns>
-		public static async Task CopyFromDirectoryAsync(List<(DirectoryNode DirectoryNode, FileInfo FileInfo)> files, ChunkingOptions options, TreeWriter writer, IProgress<ICopyStats>? progress, CancellationToken cancellationToken)
+		public static async Task CopyFromDirectoryAsync(List<(DirectoryNode DirectoryNode, FileInfo FileInfo)> files, ChunkingOptions options, TreeWriter writer, IProgress<ICopyStats>? progress, CancellationToken cancellationToken = default)
 		{
 			const int MaxWriters = 32;
 			const long MinSizePerWriter = 1024 * 1024;
diff --git a/Engine/Source/Programs/Shared/EpicGames.Horde/Storage/TreeNodeRef.cs b/Engine/Source/Programs/Shared/EpicGames.Horde/Storage/TreeNodeRef.cs
--- a/Engine/Source/Programs/Shared/EpicGames.Horde/Storage/TreeNodeRef.cs
+++ b/Engine/Source/Programs/Shared/EpicGames.Horde/Storage/TreeNodeRef.cs
@@ -387,7 +387,7 @@ namespace EpicGames.Horde.Storage
 		/// <param name="root">Root for the tree</param>
 		/// <param name="cancellationToken">Cancellation token for the operation</param>
 		/// <returns></returns>
-		public static async Task<NodeHandle> FlushAsync(this TreeWriter writer, TreeNode root, CancellationToken cancellationToken)
+		public static async Task<NodeHandle> FlushAsync(this TreeWriter writer, TreeNode root, CancellationToken cancellationToken = default)
 		{
 			TreeNodeRef rootRef = new TreeNodeRef(root);
 
diff --git a/Engine/Source/Programs/Shared/EpicGames.Horde/Storage/TreeWriter.cs b/Engine/Source/Programs/Shared/EpicGames.Horde/Storage/TreeWriter.cs
--- a/Engine/Source/Programs/Shared/EpicGames.Horde/Storage/TreeWriter.cs
+++ b/Engine/Source/Programs/Shared/EpicGames.Horde/Storage/TreeWriter.cs
@@ -785,7 +785,7 @@ namespace EpicGames.Horde.Storage
 		/// </summary>
 		/// <param name="cancellationToken">Cancellation token for the operation</param>
 		/// <returns></returns>
-		public async Task FlushAsync(CancellationToken cancellationToken)
+		public async Task FlushAsync(CancellationToken cancellationToken = default)
 		{
 			Complete();
 			while (_writeQueue.Count > 0)
