Horde: Fix compute exceptions being sent with optional descriptions.

diff --git a/Engine/Source/Programs/Shared/EpicGames.Horde/Compute/IComputeMessage.cs b/Engine/Source/Programs/Shared/EpicGames.Horde/Compute/IComputeMessage.cs
--- a/Engine/Source/Programs/Shared/EpicGames.Horde/Compute/IComputeMessage.cs
+++ b/Engine/Source/Programs/Shared/EpicGames.Horde/Compute/IComputeMessage.cs
@@ -200,16 +200,16 @@ namespace EpicGames.Horde.Compute
 		/// <summary>
 		/// Sends an exception response to the remote
 		/// </summary>
-		public static ValueTask SendExceptionAsync(this IComputeMessageChannel channel, Exception ex, CancellationToken cancellationToken) => SendExceptionAsync(channel, ex.Message, ex.StackTrace, cancellationToken);
+		public static ValueTask SendExceptionAsync(this IComputeMessageChannel channel, Exception ex, CancellationToken cancellationToken) => SendExceptionAsync(channel, ex.Message, ex.ToString(), cancellationToken);
 
 		/// <summary>
 		/// Sends an exception response to the remote
 		/// </summary>
-		public static async ValueTask SendExceptionAsync(this IComputeMessageChannel channel, string description, string? trace, CancellationToken cancellationToken)
+		public static async ValueTask SendExceptionAsync(this IComputeMessageChannel channel, string description, string trace, CancellationToken cancellationToken)
 		{
 			using IComputeMessageBuilder message = await channel.CreateMessageAsync(ComputeMessageType.Exception, cancellationToken);
 			message.WriteString(description);
-			message.WriteOptionalString(trace);
+			message.WriteString(trace);
 			message.Send();
 		}
 
