Updated finding rule files to be thread safe.

diff --git a/Engine/Source/Programs/Shared/EpicGames.Build/System/Rules.cs b/Engine/Source/Programs/Shared/EpicGames.Build/System/Rules.cs
--- a/Engine/Source/Programs/Shared/EpicGames.Build/System/Rules.cs
+++ b/Engine/Source/Programs/Shared/EpicGames.Build/System/Rules.cs
@@ -3,6 +3,7 @@
 using EpicGames.Core;
 using System;
 using System.Collections.Generic;
+using System.IO;
 using System.Linq;
 using System.Text;
 
@@ -134,7 +135,10 @@ namespace UnrealBuildBase
         public static void InvalidateRulesFileCache(string DirectoryPath)
         {
             DirectoryReference Directory = new DirectoryReference(DirectoryPath);
-            RootFolderToRulesFileCache.Remove(Directory);
+			lock(RootFolderToRulesFileCache)
+			{
+				RootFolderToRulesFileCache.Remove(Directory);
+			}
             DirectoryLookupCache.InvalidateCachedDirectory(Directory);
         }
 
@@ -144,34 +148,52 @@ namespace UnrealBuildBase
 		/// <param name="Directories">The directories to cache</param>
 		public static void PrefetchRulesFiles(IEnumerable<DirectoryReference> Directories)
 		{
-			ThreadPoolWorkQueue? Queue = null;
-			try
+			PrefetchRulesFilesInternal(Directories);
+		}
+
+		/// <summary>
+		/// Prefetch multiple directories in parallel
+		/// </summary>
+		/// <param name="Directories">The directories to cache</param>
+		/// <returns>The list of caches for each directory</returns>
+		private static List<(DirectoryReference, RulesFileCache)> PrefetchRulesFilesInternal(IEnumerable<DirectoryReference> Directories)
+		{
+			List<(DirectoryReference, RulesFileCache)> Caches = new List<(DirectoryReference, RulesFileCache)>(Directories.Count());
+			lock (RootFolderToRulesFileCache)
 			{
-				foreach(DirectoryReference Directory in Directories)
+				using (ThreadPoolWorkQueue Queue = new ThreadPoolWorkQueue())
 				{
-					if(!RootFolderToRulesFileCache.ContainsKey(Directory))
+					foreach (DirectoryReference Directory in Directories)
 					{
-						RulesFileCache Cache = new RulesFileCache();
-						RootFolderToRulesFileCache[Directory] = Cache;
-
-						if(Queue == null)
+						RulesFileCache? Cache;
+						if (!RootFolderToRulesFileCache.TryGetValue(Directory, out Cache))
 						{
-							Queue = new ThreadPoolWorkQueue();
+							Cache = new RulesFileCache();
+							Queue.Enqueue(() => FindAllRulesFilesRecursively(DirectoryItem.GetItemByDirectoryReference(Directory), Cache, Queue));
 						}
-
-						DirectoryItem DirectoryItem = DirectoryItem.GetItemByDirectoryReference(Directory);
-						Queue.Enqueue(() => FindAllRulesFilesRecursively(DirectoryItem, Cache, Queue));
+						Caches.Add((Directory, Cache));
 					}
 				}
-			}
-			finally
-			{
-				if(Queue != null)
+
+				using (ThreadPoolWorkQueue SortQueue = new ThreadPoolWorkQueue())
 				{
-					Queue.Dispose();
-					Queue = null;
+					foreach ((DirectoryReference Directory, RulesFileCache Cache) in Caches)
+					{
+						if (!RootFolderToRulesFileCache.ContainsKey(Directory))
+						{
+							SortQueue.Enqueue(() =>
+							{
+								Cache.ModuleRules.Sort((A, B) => A.FullName.CompareTo(B.FullName));
+								Cache.TargetRules.Sort((A, B) => A.FullName.CompareTo(B.FullName));
+								Cache.AutomationModules.Sort((A, B) => A.FullName.CompareTo(B.FullName));
+								Cache.UbtPlugins.Sort((A, B) => A.FullName.CompareTo(B.FullName));
+							});
+							RootFolderToRulesFileCache[Directory] = Cache;
+						}
+					}
 				}
 			}
+			return Caches;
 		}
 
 		/// <summary>
@@ -193,34 +215,10 @@ namespace UnrealBuildBase
 		/// <returns>List of rules files of the given type</returns>
 		public static List<FileReference> FindAllRulesFiles(IEnumerable<DirectoryReference> Directories, RulesFileType Type)
 		{
-			List<(DirectoryReference, RulesFileCache)> Caches = new List<(DirectoryReference, RulesFileCache)>(Directories.Count());
-			using (ThreadPoolWorkQueue Queue = new ThreadPoolWorkQueue())
-			{
-				foreach (DirectoryReference Directory in Directories)
-				{
-					RulesFileCache? Cache;
-					if (!RootFolderToRulesFileCache.TryGetValue(Directory, out Cache))
-					{
-						Cache = new RulesFileCache();
-						Queue.Enqueue(() => FindAllRulesFilesRecursively(DirectoryItem.GetItemByDirectoryReference(Directory), Cache, Queue));
-					}
-					Caches.Add((Directory, Cache));
-				}
-			}
-
-			List<FileReference> Files = new List<FileReference>();
-			
+			List<(DirectoryReference, RulesFileCache)> Caches = PrefetchRulesFilesInternal(Directories);
+			List<FileReference> Files = new List<FileReference>();			
 			foreach ((DirectoryReference Directory, RulesFileCache Cache) in Caches)
 			{
-				if (!RootFolderToRulesFileCache.ContainsKey(Directory))
-				{
-					Cache.ModuleRules.Sort((A, B) => A.FullName.CompareTo(B.FullName));
-					Cache.TargetRules.Sort((A, B) => A.FullName.CompareTo(B.FullName));
-					Cache.AutomationModules.Sort((A, B) => A.FullName.CompareTo(B.FullName));
-					Cache.UbtPlugins.Sort((A, B) => A.FullName.CompareTo(B.FullName));
-					RootFolderToRulesFileCache[Directory] = Cache;
-				}
-
 				// Get the list of files of the type we're looking for
 				if (Type == RulesFileType.Module)
 				{
@@ -243,7 +241,6 @@ namespace UnrealBuildBase
 					throw new Exception($"Unhandled rules type: {Type}");
 				}
 			}
-
 			return Files;
 		}
 
@@ -321,7 +318,5 @@ namespace UnrealBuildBase
 				}
 			}
 		}
-
-
 	}
 }
