Horde: Remote execution improvements.

* Make it an error to send data to a buffer that is not yet attached.
* Send a 'ready' message whenever the remote attaches a message buffer, allowing synchronizing the two machines before streaming data.
* Add a class for tracking remote processes, and reading output data line-by-line.

diff --git a/Engine/Source/Programs/Horde/Horde.Agent/Commands/Compute/CppComputeCommand.cs b/Engine/Source/Programs/Horde/Horde.Agent/Commands/Compute/CppComputeCommand.cs
--- a/Engine/Source/Programs/Horde/Horde.Agent/Commands/Compute/CppComputeCommand.cs
+++ b/Engine/Source/Programs/Horde/Horde.Agent/Commands/Compute/CppComputeCommand.cs
@@ -59,7 +59,15 @@ namespace Horde.Agent.Commands.Compute
 			await using (IComputeMessageChannel channel = await lease.Socket.CreateMessageChannelAsync(0, 4 * 1024 * 1024, _logger, cancellationToken))
 			{
 				await channel.UploadFilesAsync("", sandbox, storage, cancellationToken);
-				await channel.ExecuteAsync(jsonComputeTask.Executable, jsonComputeTask.Arguments, jsonComputeTask.WorkingDir, jsonComputeTask.EnvVars, (string x) => _logger.LogInformation("Child Process: {Text}", x), cancellationToken);
+
+				await using (IComputeProcess process = await channel.ExecuteAsync(jsonComputeTask.Executable, jsonComputeTask.Arguments, jsonComputeTask.WorkingDir, jsonComputeTask.EnvVars, cancellationToken))
+				{
+					string? line;
+					while ((line = await process.ReadLineAsync(cancellationToken)) != null)
+					{
+						_logger.LogInformation("Child Process: {Text}", line);
+					}
+				}
 			}
 
 			return true;
diff --git a/Engine/Source/Programs/Shared/EpicGames.Horde/Compute/ComputeProcess.cs b/Engine/Source/Programs/Shared/EpicGames.Horde/Compute/ComputeProcess.cs
new file mode 100644
--- /dev/null
+++ b/Engine/Source/Programs/Shared/EpicGames.Horde/Compute/ComputeProcess.cs
@@ -0,0 +1,133 @@
+// Copyright Epic Games, Inc. All Rights Reserved.
+
+using System;
+using System.Text;
+using System.Threading;
+using System.Threading.Channels;
+using System.Threading.Tasks;
+using EpicGames.Core;
+
+namespace EpicGames.Horde.Compute
+{
+	sealed class ComputeProcess : IComputeProcess
+	{
+		readonly Channel<string> _output;
+		readonly BackgroundTask _backgroundTask;
+		readonly TaskCompletionSource<int> _result = new TaskCompletionSource<int>();
+
+		byte[] _buffer = new byte[1024];
+		int _bufferLength;
+
+		/// <inheritdoc/>
+		public int ExitCode => _result.Task.Result;
+
+		/// <inheritdoc/>
+		public bool HasExited => _result.Task.IsCompleted;
+
+		/// <summary>
+		/// Constructor
+		/// </summary>
+		public ComputeProcess(IComputeMessageChannel channel)
+		{
+			_output = Channel.CreateUnbounded<string>();
+			_backgroundTask = BackgroundTask.StartNew(ctx => RunAsync(channel, ctx));
+		}
+
+		/// <inheritdoc/>
+		public ValueTask DisposeAsync() => _backgroundTask.DisposeAsync();
+
+		/// <inheritdoc/>
+		public async ValueTask<string?> ReadLineAsync(CancellationToken cancellationToken)
+		{
+			if (!await _output.Reader.WaitToReadAsync(cancellationToken))
+			{
+				return null;
+			}
+			return await _output.Reader.ReadAsync(cancellationToken);
+		}
+
+		/// <inheritdoc/>
+		public Task<int> WaitForExitAsync(CancellationToken cancellationToken) => _result.Task.WaitAsync(cancellationToken);
+
+		async Task RunAsync(IComputeMessageChannel channel, CancellationToken cancellationToken)
+		{
+			try
+			{
+				await RunInternalAsync(channel, cancellationToken);
+			}
+			catch (Exception ex)
+			{
+				_result.TrySetException(ex);
+			}
+		}
+
+		async Task RunInternalAsync(IComputeMessageChannel channel, CancellationToken cancellationToken)
+		{
+			for (; ; )
+			{
+				using IComputeMessage message = await channel.ReceiveAsync(cancellationToken);
+				switch (message.Type)
+				{
+					case ComputeMessageType.Exception:
+						ExceptionMessage exception = message.ParseExceptionMessage();
+						throw new ComputeRemoteException(exception);
+					case ComputeMessageType.ExecuteOutput:
+						AppendData(message.Data.Span);
+						break;
+					case ComputeMessageType.ExecuteResult:
+						ExecuteProcessResponseMessage executeProcessResponse = message.ParseExecuteProcessResponse();
+						_result.TrySetResult(executeProcessResponse.ExitCode);
+						return;
+					default:
+						throw new ComputeInvalidMessageException(message);
+				}
+			}
+		}
+
+		void AppendData(ReadOnlySpan<byte> data)
+		{
+			for (; ; )
+			{
+				int lineEnd = data.IndexOf((byte)'\n');
+				if (lineEnd == -1)
+				{
+					AppendToBuffer(data);
+					break;
+				}
+
+				int lineLen = lineEnd;
+				if (lineLen > 0 && data[lineLen - 1] == '\r')
+				{
+					lineLen--;
+				}
+
+				string str;
+				if (_bufferLength == 0)
+				{
+					str = Encoding.UTF8.GetString(data.Slice(0, lineLen));
+				}
+				else
+				{
+					AppendToBuffer(data.Slice(0, lineLen));
+					str = Encoding.UTF8.GetString(data.Slice(0, _bufferLength + lineLen));
+					_bufferLength = 0;
+				}
+
+				_output.Writer.TryWrite(str);
+				data = data.Slice(lineEnd + 1);
+			}
+		}
+
+		void AppendToBuffer(ReadOnlySpan<byte> data)
+		{
+			if (_bufferLength + data.Length > _buffer.Length)
+			{
+				byte[] newBuffer = new byte[_bufferLength + data.Length + 256];
+				_buffer.AsSpan(0, _bufferLength).CopyTo(newBuffer);
+				_buffer = newBuffer;
+			}
+			data.CopyTo(_buffer.AsSpan(_bufferLength));
+			_bufferLength += data.Length;
+		}
+	}
+}
diff --git a/Engine/Source/Programs/Shared/EpicGames.Horde/Compute/ComputeSocketBase.cs b/Engine/Source/Programs/Shared/EpicGames.Horde/Compute/ComputeSocketBase.cs
--- a/Engine/Source/Programs/Shared/EpicGames.Horde/Compute/ComputeSocketBase.cs
+++ b/Engine/Source/Programs/Shared/EpicGames.Horde/Compute/ComputeSocketBase.cs
@@ -27,7 +27,6 @@ namespace EpicGames.Horde.Compute
 		CancellationTokenSource _cancellationSource = new CancellationTokenSource();
 
 		readonly Dictionary<int, IComputeBufferWriter> _receiveBuffers = new Dictionary<int, IComputeBufferWriter>();
-		readonly AsyncEvent _recvBufferWritersChangedEvent = new AsyncEvent();
 
 		readonly SemaphoreSlim _sendSemaphore = new SemaphoreSlim(1, 1);
 		readonly Dictionary<int, IComputeBufferReader> _sendBuffers = new Dictionary<int, IComputeBufferReader>();
@@ -128,24 +127,8 @@ namespace EpicGames.Horde.Compute
 					}
 					else
 					{
-						IComputeBufferWriter writer = await GetReceiveBufferAsync(cachedWriters, id);
-						for (int offset = 0; offset < size;)
-						{
-							Memory<byte> memory = writer.GetMemory();
-							if (memory.Length == 0)
-							{
-								_logger.LogTrace("No space in buffer {Id}, flushing", id);
-								await writer.FlushAsync(cancellationToken);
-							}
-							else
-							{
-								int maxRead = Math.Min(size - offset, memory.Length);
-								int read = await transport.ReadPartialAsync(memory.Slice(0, maxRead), cancellationToken);
-								writer.Advance(read);
-
-								offset += read;
-							}
-						}
+						IComputeBufferWriter writer = GetReceiveBuffer(cachedWriters, id);
+						await ReadPacketAsync(transport, id, size, writer, cancellationToken);
 					}
 				}
 			}
@@ -166,7 +149,28 @@ namespace EpicGames.Horde.Compute
 			_logger.LogTrace("Closing reader");
 		}
 
-		async Task<IComputeBufferWriter> GetReceiveBufferAsync(Dictionary<int, IComputeBufferWriter> cachedWriters, int id)
+		async Task ReadPacketAsync(IComputeTransport transport, int id, int size, IComputeBufferWriter writer, CancellationToken cancellationToken)
+		{
+			for (int offset = 0; offset < size;)
+			{
+				Memory<byte> memory = writer.GetMemory();
+				if (memory.Length == 0)
+				{
+					_logger.LogTrace("No space in buffer {Id}, flushing", id);
+					await writer.FlushAsync(cancellationToken);
+				}
+				else
+				{
+					int maxRead = Math.Min(size - offset, memory.Length);
+					int read = await transport.ReadPartialAsync(memory.Slice(0, maxRead), cancellationToken);
+					writer.Advance(read);
+
+					offset += read;
+				}
+			}
+		}
+
+		IComputeBufferWriter GetReceiveBuffer(Dictionary<int, IComputeBufferWriter> cachedWriters, int id)
 		{
 			IComputeBufferWriter? writer;
 			if (cachedWriters.TryGetValue(id, out writer))
@@ -174,38 +178,16 @@ namespace EpicGames.Horde.Compute
 				return writer;
 			}
 
-			int delayTime = 1;
-			Task? delayTask = null;
-			Stopwatch? timer = null;
-
-			for (; ; )
+			lock (_lockObject)
 			{
-				Task waitTask;
-				lock (_lockObject)
-				{
-					if (_receiveBuffers.TryGetValue(id, out writer))
-					{
-						cachedWriters.Add(id, writer);
-						return writer;
-					}
-					waitTask = _recvBufferWritersChangedEvent.Task;
-				}
-
-				timer ??= Stopwatch.StartNew();
-				delayTask ??= Task.Delay(TimeSpan.FromSeconds(delayTime));
-
-				while (!waitTask.IsCompleted)
+				if (_receiveBuffers.TryGetValue(id, out writer))
 				{
-					if (await Task.WhenAny(waitTask, delayTask) == delayTask)
-					{
-						_logger.LogWarning("Socket stalled for {Time}s waiting for a buffer to be attached to channel {ChannelId}", (int)timer.Elapsed.TotalSeconds, id);
-						delayTime = Math.Min(15, delayTime + 2);
-						delayTask = Task.Delay(TimeSpan.FromSeconds(delayTime));
-					}
+					cachedWriters.Add(id, writer);
+					return writer;
 				}
-
-				await waitTask;
 			}
+
+			throw new ComputeInternalException($"No buffer is attached to channel {id}");
 		}
 
 		class SendSegment : ReadOnlySequenceSegment<byte>
@@ -285,17 +267,11 @@ namespace EpicGames.Horde.Compute
 				}
 			}
 
-			if (recvBufferWriter != null)
+			if (recvBufferWriter != null && complete)
 			{
-				if (complete)
-				{
-					recvBufferWriter.MarkComplete();
-				}
-				else
-				{
-					_recvBufferWritersChangedEvent.Set();
-				}
+				recvBufferWriter.MarkComplete();
 			}
+
 			return new ValueTask();
 		}
 
diff --git a/Engine/Source/Programs/Shared/EpicGames.Horde/Compute/ComputeWorker.cs b/Engine/Source/Programs/Shared/EpicGames.Horde/Compute/ComputeWorker.cs
--- a/Engine/Source/Programs/Shared/EpicGames.Horde/Compute/ComputeWorker.cs
+++ b/Engine/Source/Programs/Shared/EpicGames.Horde/Compute/ComputeWorker.cs
@@ -51,6 +51,8 @@ namespace EpicGames.Horde.Compute
 		{
 			await using (IComputeMessageChannel channel = await socket.CreateMessageChannelAsync(channelId, 4 * 1024 * 1024, _logger, cancellationToken))
 			{
+				await channel.SendReadyAsync(cancellationToken);
+
 				List<Task> childTasks = new List<Task>();
 				for (; ; )
 				{
diff --git a/Engine/Source/Programs/Shared/EpicGames.Horde/Compute/IComputeMessage.cs b/Engine/Source/Programs/Shared/EpicGames.Horde/Compute/IComputeMessage.cs
--- a/Engine/Source/Programs/Shared/EpicGames.Horde/Compute/IComputeMessage.cs
+++ b/Engine/Source/Programs/Shared/EpicGames.Horde/Compute/IComputeMessage.cs
@@ -24,20 +24,25 @@ namespace EpicGames.Horde.Compute
 		/// </summary>
 		None = 0x00,
 
+		/// <summary>
+		/// Indicates that the remote is ready to receive messages on the current channel
+		/// </summary>
+		Ready = 0x01,
+
 		/// <summary>
 		/// Sent in place of a regular response if an error occurs on the remote
 		/// </summary>
-		Exception = 0x01,
+		Exception = 0x02,
 
 		/// <summary>
 		/// Fork a new request channel
 		/// </summary>
-		Fork = 0x02,
+		Fork = 0x03,
 
 		/// <summary>
 		/// Untyped user data
 		/// </summary>
-		UserData = 0x03,
+		UserData = 0x04,
 
 		#region Process Management
 
@@ -62,14 +67,14 @@ namespace EpicGames.Horde.Compute
 		Execute = 0x16,
 
 		/// <summary>
-		/// Returns the process exit code (Remote -> Initiator)
+		/// Returns output from the child process to the caller (Remote -> Initiator)
 		/// </summary>
-		ExecuteResult = 0x17,
+		ExecuteOutput = 0x17,
 
 		/// <summary>
-		/// Returns output from the child process to the caller (Remote -> Initiator)
+		/// Returns the process exit code (Remote -> Initiator)
 		/// </summary>
-		ExecuteOutput = 0x18,
+		ExecuteResult = 0x18,
 
 		#endregion
 
@@ -197,6 +202,15 @@ namespace EpicGames.Horde.Compute
 	/// </summary>
 	public static class ComputeMessageExtensions
 	{
+		/// <summary>
+		/// Sends a message to the remote indicating that we're ready to receive messages on the given channel
+		/// </summary>
+		public static async ValueTask SendReadyAsync(this IComputeMessageChannel channel, CancellationToken cancellationToken)
+		{
+			using IComputeMessageBuilder message = await channel.CreateMessageAsync(ComputeMessageType.Ready, cancellationToken);
+			message.Send();
+		}
+
 		/// <summary>
 		/// Sends an exception response to the remote
 		/// </summary>
@@ -332,53 +346,6 @@ namespace EpicGames.Horde.Compute
 			return new DeleteFilesMessage(files);
 		}
 
-		class StringOutputWriter
-		{
-			byte[] _buffer;
-			int _length;
-			readonly Action<string> _handler;
-
-			public StringOutputWriter(Action<string> handler)
-			{
-				_buffer = new byte[10];
-				_handler = handler;
-			}
-
-			public void WriteData(ReadOnlyMemory<byte> data)
-			{
-				int requiredLength = _length + data.Length;
-				if (requiredLength > _buffer.Length)
-				{
-					Array.Resize(ref _buffer, requiredLength);
-				}
-
-				data.CopyTo(_buffer.AsMemory(_length));
-				_length += data.Length;
-
-				ReadOnlySpan<byte> output = _buffer.AsSpan(0, _length);
-				for (; ; )
-				{
-					int newLineIdx = output.IndexOf((byte)'\n');
-					if (newLineIdx == -1)
-					{
-						break;
-					}
-
-					int length = newLineIdx;
-					if (length > 0 && output[length - 1] == (byte)'\r')
-					{
-						length--;
-					}
-
-					_handler(Encoding.UTF8.GetString(output.Slice(0, length)));
-					output = output.Slice(newLineIdx + 1);
-				}
-
-				_length = output.Length;
-				output.CopyTo(_buffer);
-			}
-		}
-
 		/// <summary>
 		/// Executes a remote process
 		/// </summary>
@@ -386,26 +353,9 @@ namespace EpicGames.Horde.Compute
 		/// <param name="executable">Executable to run, relative to the sandbox root</param>
 		/// <param name="arguments">Arguments for the child process</param>
 		/// <param name="workingDir">Working directory for the process</param>
-		/// <param name="outputHandler">Output callback for stdout</param>
 		/// <param name="envVars">Environment variables for the child process</param>
 		/// <param name="cancellationToken">Cancellation token for the operation</param>
-		public static Task<int> ExecuteAsync(this IComputeMessageChannel channel, string executable, IReadOnlyList<string> arguments, string? workingDir, IReadOnlyDictionary<string, string?>? envVars, Action<string> outputHandler, CancellationToken cancellationToken = default)
-		{
-			StringOutputWriter writer = new StringOutputWriter(outputHandler);
-			return ExecuteAsync(channel, executable, arguments, workingDir, envVars, writer.WriteData, cancellationToken);
-		}
-
-		/// <summary>
-		/// Executes a remote process
-		/// </summary>
-		/// <param name="channel">Current channel</param>
-		/// <param name="executable">Executable to run, relative to the sandbox root</param>
-		/// <param name="arguments">Arguments for the child process</param>
-		/// <param name="workingDir">Working directory for the process</param>
-		/// <param name="outputHandler">Output callback for stdout</param>
-		/// <param name="envVars">Environment variables for the child process</param>
-		/// <param name="cancellationToken">Cancellation token for the operation</param>
-		public static async Task<int> ExecuteAsync(this IComputeMessageChannel channel, string executable, IReadOnlyList<string> arguments, string? workingDir, IReadOnlyDictionary<string, string?>? envVars, Action<ReadOnlyMemory<byte>> outputHandler, CancellationToken cancellationToken = default)
+		public static async Task<IComputeProcess> ExecuteAsync(this IComputeMessageChannel channel, string executable, IReadOnlyList<string> arguments, string? workingDir, IReadOnlyDictionary<string, string?>? envVars, CancellationToken cancellationToken = default)
 		{
 			using (IComputeMessageBuilder request = await channel.CreateMessageAsync(ComputeMessageType.Execute, cancellationToken))
 			{
@@ -415,25 +365,7 @@ namespace EpicGames.Horde.Compute
 				request.WriteDictionary(envVars ?? new Dictionary<string, string?>(), MemoryWriterExtensions.WriteString, MemoryWriterExtensions.WriteOptionalString);
 				request.Send();
 			}
-
-			for (; ; )
-			{
-				using IComputeMessage message = await channel.ReceiveAsync(cancellationToken);
-				switch (message.Type)
-				{
-					case ComputeMessageType.Exception:
-						ExceptionMessage exception = message.ParseExceptionMessage();
-						throw new ComputeRemoteException(exception);
-					case ComputeMessageType.ExecuteOutput:
-						outputHandler(message.Data);
-						break;
-					case ComputeMessageType.ExecuteResult:
-						ExecuteProcessResponseMessage executeProcessResponse = message.ParseExecuteProcessResponse();
-						return executeProcessResponse.ExitCode;
-					default:
-						throw new ComputeInvalidMessageException(message);
-				}
-			}
+			return new ComputeProcess(channel);
 		}
 
 		/// <summary>
diff --git a/Engine/Source/Programs/Shared/EpicGames.Horde/Compute/IComputeMessageChannel.cs b/Engine/Source/Programs/Shared/EpicGames.Horde/Compute/IComputeMessageChannel.cs
--- a/Engine/Source/Programs/Shared/EpicGames.Horde/Compute/IComputeMessageChannel.cs
+++ b/Engine/Source/Programs/Shared/EpicGames.Horde/Compute/IComputeMessageChannel.cs
@@ -45,6 +45,20 @@ namespace EpicGames.Horde.Compute
 	/// </summary>
 	public static class ComputeChannelExtensions
 	{
+		/// <summary>
+		/// Waits for the remote machine to send a 'ready' response
+		/// </summary>
+		/// <param name="channel">Channel to receive on</param>
+		/// <param name="cancellationToken">Cancellation token for the operation</param>
+		public static async ValueTask WaitForAttachAsync(this IComputeMessageChannel channel, CancellationToken cancellationToken = default)
+		{
+			IComputeMessage message = await channel.ReceiveAsync(cancellationToken);
+			if (message.Type != ComputeMessageType.Ready)
+			{
+				throw new ComputeInvalidMessageException(message);
+			}
+		}
+
 		/// <summary>
 		/// Creates a new builder for a message
 		/// </summary>
diff --git a/Engine/Source/Programs/Shared/EpicGames.Horde/Compute/IComputeProcess.cs b/Engine/Source/Programs/Shared/EpicGames.Horde/Compute/IComputeProcess.cs
new file mode 100644
--- /dev/null
+++ b/Engine/Source/Programs/Shared/EpicGames.Horde/Compute/IComputeProcess.cs
@@ -0,0 +1,38 @@
+// Copyright Epic Games, Inc. All Rights Reserved.
+
+using System;
+using System.Threading;
+using System.Threading.Tasks;
+
+namespace EpicGames.Horde.Compute
+{
+	/// <summary>
+	/// Represents a remotely executing process
+	/// </summary>
+	public interface IComputeProcess : IAsyncDisposable
+	{
+		/// <summary>
+		/// Exit code for the process.
+		/// </summary>
+		int ExitCode { get; }
+
+		/// <summary>
+		/// Whether the process has exited yet or not
+		/// </summary>
+		bool HasExited { get; }
+
+		/// <summary>
+		/// Read a line from the child process
+		/// </summary>
+		/// <param name="cancellationToken">Cancellation token for the operation</param>
+		/// <returns>Line that was read</returns>
+		ValueTask<string?> ReadLineAsync(CancellationToken cancellationToken = default);
+
+		/// <summary>
+		/// Wait until the process exits
+		/// </summary>
+		/// <param name="cancellationToken">Cancellation token for the operation</param>
+		/// <returns>Exit code</returns>
+		Task<int> WaitForExitAsync(CancellationToken cancellationToken = default);
+	}
+}
