Horde: Fail gracefully if a remote machine is not available for a compute request.

diff --git a/Engine/Source/Programs/Shared/EpicGames.Horde/Compute/Clients/ServerComputeClient.cs b/Engine/Source/Programs/Shared/EpicGames.Horde/Compute/Clients/ServerComputeClient.cs
--- a/Engine/Source/Programs/Shared/EpicGames.Horde/Compute/Clients/ServerComputeClient.cs
+++ b/Engine/Source/Programs/Shared/EpicGames.Horde/Compute/Clients/ServerComputeClient.cs
@@ -5,6 +5,7 @@ using EpicGames.Horde.Compute.Transports;
 using Microsoft.Extensions.Logging;
 using System;
 using System.Collections.Generic;
+using System.Linq;
 using System.Net;
 using System.Net.Http;
 using System.Net.Http.Headers;
@@ -148,7 +149,14 @@ namespace EpicGames.Horde.Compute.Clients
 			AssignComputeResponse? responseMessage;
 			using (HttpResponseMessage response = await client.PostAsync($"api/v2/compute/{clusterId}", request, _cancellationSource.Token))
 			{
+				if (response.StatusCode == HttpStatusCode.ServiceUnavailable)
+				{
+					_logger.LogInformation("No compute resource is available.");
+					yield break;
+				}
+
 				response.EnsureSuccessStatusCode();
+
 				responseMessage = await response.Content.ReadFromJsonAsync<AssignComputeResponse>(cancellationToken: cancellationToken);
 				if (responseMessage == null)
 				{
