Horde: Support for downloading files and zips for artifact data.

diff --git a/Engine/Source/Programs/Horde/Horde.Server/Artifacts/ArtifactsController.cs b/Engine/Source/Programs/Horde/Horde.Server/Artifacts/ArtifactsController.cs
--- a/Engine/Source/Programs/Horde/Horde.Server/Artifacts/ArtifactsController.cs
+++ b/Engine/Source/Programs/Horde/Horde.Server/Artifacts/ArtifactsController.cs
@@ -2,7 +2,9 @@
 
 using System;
 using System.Collections.Generic;
+using System.IO;
 using System.Linq;
+using System.Net.Http.Headers;
 using System.Threading;
 using System.Threading.Tasks;
 using EpicGames.Core;
@@ -250,6 +252,77 @@ namespace Horde.Server.Artifacts
 			return PropertyFilter.Apply(new GetArtifactDirectoryResponse(directory), filter);
 		}
 
+		/// <summary>
+		/// Downloads an individual file from an artifact
+		/// </summary>
+		/// <param name="id">Identifier of the artifact to retrieve</param>
+		/// <param name="path">Path to fetch</param>
+		/// <param name="download">Whether to request the file be downloaded vs displayed inline</param>
+		/// <param name="cancellationToken">Cancellation token for the operation</param>
+		/// <returns>Information about all the artifacts</returns>
+		[HttpGet]
+		[Route("/api/v2/artifacts/{id}/file")]
+		public async Task<ActionResult<object>> GetFileAsync(ArtifactId id, [FromQuery] string path, [FromQuery] bool download = false, CancellationToken cancellationToken = default)
+		{
+			IArtifact? artifact = await _artifactCollection.GetAsync(id, cancellationToken);
+			if (artifact == null)
+			{
+				return NotFound(id);
+			}
+			if (!_globalConfig.Authorize(artifact.AclScope, AclAction.ReadArtifact, User))
+			{
+				return Forbid(AclAction.ReadArtifact, artifact.AclScope);
+			}
+
+			TreeReader reader = await _storageService.GetReaderAsync(artifact.NamespaceId, cancellationToken);
+			TreeNodeRef<DirectoryNode> directoryRef = await reader.ReadNodeRefAsync<DirectoryNode>(artifact.RefName, DateTime.UtcNow.AddHours(1.0), cancellationToken);
+			DirectoryNode directory = await directoryRef.ExpandAsync(reader, cancellationToken);
+
+			FileEntry? fileEntry = await directory.GetFileEntryByPathAsync(reader, path, cancellationToken);
+			if (fileEntry == null)
+			{
+				return NotFound($"Unable to find file {path}");
+			}
+
+			Stream stream = fileEntry.AsStream(reader);
+			if (download)
+			{
+				return new FileStreamResult(stream, "application/octet-stream") { FileDownloadName = Path.GetFileName(path) };
+			}
+			else
+			{
+				return new InlineFileStreamResult(stream, "application/octet-stream", Path.GetFileName(path));
+			}
+		}
+
+		/// <summary>
+		/// Downloads an individual file from an artifact
+		/// </summary>
+		/// <param name="id">Identifier of the artifact to retrieve</param>
+		/// <param name="cancellationToken">Cancellation token for the operation</param>
+		/// <returns>Information about all the artifacts</returns>
+		[HttpGet]
+		[Route("/api/v2/artifacts/{id}/zip")]
+		public async Task<ActionResult<object>> GetZipAsync(ArtifactId id, CancellationToken cancellationToken = default)
+		{
+			IArtifact? artifact = await _artifactCollection.GetAsync(id, cancellationToken);
+			if (artifact == null)
+			{
+				return NotFound(id);
+			}
+			if (!_globalConfig.Authorize(artifact.AclScope, AclAction.ReadArtifact, User))
+			{
+				return Forbid(AclAction.ReadArtifact, artifact.AclScope);
+			}
+
+			TreeReader reader = await _storageService.GetReaderAsync(artifact.NamespaceId, cancellationToken);
+			TreeNodeRef<DirectoryNode> directoryRef = await reader.ReadNodeRefAsync<DirectoryNode>(artifact.RefName, DateTime.UtcNow.AddHours(1.0), cancellationToken);
+			DirectoryNode directory = await directoryRef.ExpandAsync(reader, cancellationToken);
+
+			Stream stream = directory.AsZipStream(reader);
+			return new FileStreamResult(stream, "application/octet-stream") { FileDownloadName = $"{artifact.RefName}.zip" };
+		}
+
 		/// <summary>
 		/// Gets metadata about an artifact object
 		/// </summary>
diff --git a/Engine/Source/Programs/Horde/Horde.Server/Jobs/Artifacts/ArtifactsControllerV1.cs b/Engine/Source/Programs/Horde/Horde.Server/Jobs/Artifacts/ArtifactsControllerV1.cs
--- a/Engine/Source/Programs/Horde/Horde.Server/Jobs/Artifacts/ArtifactsControllerV1.cs
+++ b/Engine/Source/Programs/Horde/Horde.Server/Jobs/Artifacts/ArtifactsControllerV1.cs
@@ -247,40 +247,6 @@ namespace Horde.Server.Jobs.Artifacts
 			return new FileStreamResult(await _artifactCollection.OpenArtifactReadStreamAsync(artifact), artifact.MimeType);
 		}
 
-		/// <summary>
-		/// Class to return a file stream without the "content-disposition: attachment" header
-		/// </summary>
-		class InlineFileStreamResult : FileStreamResult
-		{
-			/// <summary>
-			/// The suggested download filename
-			/// </summary>
-			readonly string _fileName;
-
-			/// <summary>
-			/// Constructor
-			/// </summary>
-			/// <param name="stream"></param>
-			/// <param name="mimeType"></param>
-			/// <param name="fileName"></param>
-			public InlineFileStreamResult(System.IO.Stream stream, string mimeType, string fileName)
-				: base(stream, mimeType)
-			{
-				_fileName = fileName;
-			}
-
-			/// <inheritdoc/>
-			public override Task ExecuteResultAsync(ActionContext context)
-			{
-				ContentDisposition contentDisposition = new ContentDisposition();
-				contentDisposition.Inline = true;
-				contentDisposition.FileName = _fileName;
-				context.HttpContext.Response.Headers.Add("Content-Disposition", contentDisposition.ToString());
-				
-				return base.ExecuteResultAsync(context);
-			}
-		}
-
 		/// <summary>
 		/// Retrieve raw data for an artifact
 		/// </summary>
diff --git a/Engine/Source/Programs/Horde/Horde.Server/Tools/ToolCollection.cs b/Engine/Source/Programs/Horde/Horde.Server/Tools/ToolCollection.cs
--- a/Engine/Source/Programs/Horde/Horde.Server/Tools/ToolCollection.cs
+++ b/Engine/Source/Programs/Horde/Horde.Server/Tools/ToolCollection.cs
@@ -412,7 +412,7 @@ namespace Horde.Server.Tools
 			TreeReader reader = new TreeReader(client, _cache, _logger);
 			DirectoryNode node = await reader.ReadNodeAsync<DirectoryNode>(deployment.RefName, DateTime.UtcNow - TimeSpan.FromDays(2.0), cancellationToken);
 
-			return node.ToZipStream(reader);
+			return node.AsZipStream(reader);
 		}
 	}
 }
diff --git a/Engine/Source/Programs/Horde/Horde.Server/Utilities/InlineFileStreamResult.cs b/Engine/Source/Programs/Horde/Horde.Server/Utilities/InlineFileStreamResult.cs
new file mode 100644
--- /dev/null
+++ b/Engine/Source/Programs/Horde/Horde.Server/Utilities/InlineFileStreamResult.cs
@@ -0,0 +1,36 @@
+// Copyright Epic Games, Inc. All Rights Reserved.
+
+using System.Net.Mime;
+using System.Threading.Tasks;
+using Microsoft.AspNetCore.Mvc;
+
+namespace Horde.Server.Utilities
+{
+	/// <summary>
+	/// Class to return a file stream without the "content-disposition: attachment" header
+	/// </summary>
+	class InlineFileStreamResult : FileStreamResult
+	{
+		/// <summary>
+		/// The suggested download filename
+		/// </summary>
+		readonly string _fileName;
+
+		public InlineFileStreamResult(System.IO.Stream stream, string mimeType, string fileName)
+			: base(stream, mimeType)
+		{
+			_fileName = fileName;
+		}
+
+		/// <inheritdoc/>
+		public override Task ExecuteResultAsync(ActionContext context)
+		{
+			ContentDisposition contentDisposition = new ContentDisposition();
+			contentDisposition.Inline = true;
+			contentDisposition.FileName = _fileName;
+			context.HttpContext.Response.Headers.Add("Content-Disposition", contentDisposition.ToString());
+
+			return base.ExecuteResultAsync(context);
+		}
+	}
+}
diff --git a/Engine/Source/Programs/Shared/EpicGames.Horde/Storage/Nodes/DirectoryEntry.cs b/Engine/Source/Programs/Shared/EpicGames.Horde/Storage/Nodes/DirectoryEntry.cs
new file mode 100644
--- /dev/null
+++ b/Engine/Source/Programs/Shared/EpicGames.Horde/Storage/Nodes/DirectoryEntry.cs
@@ -0,0 +1,58 @@
+// Copyright Epic Games, Inc. All Rights Reserved.
+
+using EpicGames.Core;
+
+namespace EpicGames.Horde.Storage.Nodes
+{
+	/// <summary>
+	/// Entry for a directory within a directory node
+	/// </summary>
+	public class DirectoryEntry : DirectoryNodeRef
+	{
+		/// <summary>
+		/// Name of this directory
+		/// </summary>
+		public Utf8String Name { get; }
+
+		/// <summary>
+		/// Constructor
+		/// </summary>
+		public DirectoryEntry(Utf8String name)
+			: this(name, new DirectoryNode())
+		{
+		}
+
+		/// <summary>
+		/// Constructor
+		/// </summary>
+		public DirectoryEntry(Utf8String name, DirectoryNode node)
+			: base(node)
+		{
+			Name = name;
+		}
+
+		/// <summary>
+		/// Deserializing constructor
+		/// </summary>
+		/// <param name="reader"></param>
+		public DirectoryEntry(ITreeNodeReader reader)
+			: base(reader)
+		{
+			Name = reader.ReadUtf8String();
+		}
+
+		/// <summary>
+		/// Serialize this directory entry to disk
+		/// </summary>
+		/// <param name="writer"></param>
+		public override void Serialize(ITreeNodeWriter writer)
+		{
+			base.Serialize(writer);
+
+			writer.WriteUtf8String(Name);
+		}
+
+		/// <inheritdoc/>
+		public override string ToString() => Name.ToString();
+	}
+}
diff --git a/Engine/Source/Programs/Shared/EpicGames.Horde/Storage/Nodes/DirectoryNode.cs b/Engine/Source/Programs/Shared/EpicGames.Horde/Storage/Nodes/DirectoryNode.cs
--- a/Engine/Source/Programs/Shared/EpicGames.Horde/Storage/Nodes/DirectoryNode.cs
+++ b/Engine/Source/Programs/Shared/EpicGames.Horde/Storage/Nodes/DirectoryNode.cs
@@ -16,267 +16,6 @@ using System.Threading.Tasks;
 
 namespace EpicGames.Horde.Storage.Nodes
 {
-	/// <summary>
-	/// Flags for a file entry
-	/// </summary>
-	[Flags]
-	public enum FileEntryFlags
-	{
-		/// <summary>
-		/// No other flags set
-		/// </summary>
-		None = 0,
-
-		/// <summary>
-		/// Indicates that the referenced file is executable
-		/// </summary>
-		Executable = 4,
-
-		/// <summary>
-		/// File should be stored as read-only
-		/// </summary>
-		ReadOnly = 8,
-
-		/// <summary>
-		/// File contents are utf-8 encoded text. Client may want to replace line-endings with OS-specific format.
-		/// </summary>
-		Text = 16,
-
-		/// <summary>
-		/// Used to indicate that custom data is included in the output. Used internally for serialization; not exposed to users.
-		/// </summary>
-		HasCustomData = 32,
-
-		/// <summary>
-		/// File should be materialized as UTF-16 (but is stored as a UTF-8 source)
-		/// </summary>
-		Utf16 = 64,
-	}
-
-	/// <summary>
-	/// Entry for a file within a directory node
-	/// </summary>
-	public sealed class FileEntry : TreeNodeRef<FileNode>
-	{
-		/// <summary>
-		/// Name of this file
-		/// </summary>
-		public Utf8String Name { get; }
-
-		/// <summary>
-		/// Flags for this file
-		/// </summary>
-		public FileEntryFlags Flags { get; set; }
-
-		/// <summary>
-		/// Length of this entry
-		/// </summary>
-		public long Length { get; }
-
-		/// <summary>
-		/// Hash of the target node
-		/// </summary>
-		public IoHash Hash { get; }
-
-		/// <summary>
-		/// Custom user data for this file entry
-		/// </summary>
-		public ReadOnlyMemory<byte> CustomData { get; set; } = ReadOnlyMemory<byte>.Empty;
-
-		/// <summary>
-		/// Constructor
-		/// </summary>
-		public FileEntry(Utf8String name, FileEntryFlags flags, long length, NodeHandle node)
-			: base(node)
-		{
-			Name = name;
-			Flags = flags;
-			Length = length;
-			Hash = node.Hash;
-		}
-
-		/// <summary>
-		/// Deserialize from a buffer
-		/// </summary>
-		/// <param name="reader"></param>
-		public FileEntry(ITreeNodeReader reader)
-			: base(reader)
-		{
-			Name = reader.ReadUtf8String();
-			Flags = (FileEntryFlags)reader.ReadUnsignedVarInt();
-			Length = (long)reader.ReadUnsignedVarInt();
-			Hash = reader.ReadIoHash();
-
-			if ((Flags & FileEntryFlags.HasCustomData) != 0)
-			{
-				CustomData = reader.ReadVariableLengthBytes();
-				Flags &= ~FileEntryFlags.HasCustomData;
-			}
-		}
-
-		/// <summary>
-		/// Serialize this entry
-		/// </summary>
-		/// <param name="writer"></param>
-		public override void Serialize(ITreeNodeWriter writer)
-		{
-			base.Serialize(writer);
-
-			FileEntryFlags flags = (CustomData.Length > 0) ? (Flags | FileEntryFlags.HasCustomData) : (Flags & ~FileEntryFlags.HasCustomData);
-
-			writer.WriteUtf8String(Name);
-			writer.WriteUnsignedVarInt((ulong)flags);
-			writer.WriteUnsignedVarInt((ulong)Length);
-			writer.WriteIoHash(Hash);
-
-			if ((flags & FileEntryFlags.HasCustomData) != 0)
-			{
-				writer.WriteVariableLengthBytes(CustomData.Span);
-			}
-		}
-
-		/// <summary>
-		/// Copies the contents of this node and its children to the given output stream
-		/// </summary>
-		/// <param name="reader">Reader for retrieving existing node data</param>
-		/// <param name="outputStream">The output stream to receive the data</param>
-		/// <param name="cancellationToken">Cancellation token for the operation</param>
-		public async Task CopyToStreamAsync(TreeReader reader, Stream outputStream, CancellationToken cancellationToken)
-		{
-			FileNode node = await ExpandAsync(reader, cancellationToken);
-			await node.CopyToStreamAsync(reader, outputStream, cancellationToken);
-		}
-
-		/// <summary>
-		/// Extracts the contents of this node to a file
-		/// </summary>
-		/// <param name="reader">Reader for retrieving existing node data</param>
-		/// <param name="file">File to write with the contents of this node</param>
-		/// <param name="cancellationToken">Cancellation token for the operation</param>
-		/// <returns></returns>
-		public async Task CopyToFileAsync(TreeReader reader, FileInfo file, CancellationToken cancellationToken)
-		{
-			FileNode node = await ExpandAsync(reader, cancellationToken);
-			await node.CopyToFileAsync(reader, file, cancellationToken);
-		}
-
-		/// <inheritdoc/>
-		public override string ToString() => Name.ToString();
-	}
-
-	/// <summary>
-	/// Reference to a directory node, including the target hash and length
-	/// </summary>
-	public class DirectoryNodeRef : TreeNodeRef<DirectoryNode>
-	{
-		/// <summary>
-		/// Length of this directory tree
-		/// </summary>
-		public long Length => (Target == null) ? _cachedLength : Target.Length;
-
-		/// <summary>
-		/// Hash of the target node
-		/// </summary>
-		public IoHash Hash { get; private set; }
-
-		/// <summary>
-		/// Cached value for the length of this tree
-		/// </summary>
-		long _cachedLength;
-
-		/// <summary>
-		/// Constructor
-		/// </summary>
-		public DirectoryNodeRef(DirectoryNode node)
-			: base(node)
-		{
-		}
-
-		/// <summary>
-		/// Constructor
-		/// </summary>
-		/// <param name="reader"></param>
-		public DirectoryNodeRef(ITreeNodeReader reader)
-			: base(reader)
-		{
-			_cachedLength = (long)reader.ReadUnsignedVarInt();
-			Hash = reader.ReadIoHash();
-		}
-
-		/// <summary>
-		/// Serialize this directory entry to disk
-		/// </summary>
-		/// <param name="writer"></param>
-		public override void Serialize(ITreeNodeWriter writer)
-		{
-			base.Serialize(writer);
-
-			writer.WriteUnsignedVarInt((ulong)Length);
-			writer.WriteIoHash(Hash);
-		}
-
-		/// <inheritdoc/>
-		protected override void OnCollapse()
-		{
-			base.OnCollapse();
-
-			Hash = Target!.Hash;
-			_cachedLength = Target!.Length;
-		}
-	}
-
-	/// <summary>
-	/// Entry for a directory within a directory node
-	/// </summary>
-	public class DirectoryEntry : DirectoryNodeRef
-	{
-		/// <summary>
-		/// Name of this directory
-		/// </summary>
-		public Utf8String Name { get; }
-
-		/// <summary>
-		/// Constructor
-		/// </summary>
-		public DirectoryEntry(Utf8String name)
-			: this(name, new DirectoryNode())
-		{
-		}
-
-		/// <summary>
-		/// Constructor
-		/// </summary>
-		public DirectoryEntry(Utf8String name, DirectoryNode node)
-			: base(node)
-		{
-			Name = name;
-		}
-
-		/// <summary>
-		/// Deserializing constructor
-		/// </summary>
-		/// <param name="reader"></param>
-		public DirectoryEntry(ITreeNodeReader reader)
-			: base(reader)
-		{
-			Name = reader.ReadUtf8String();
-		}
-
-		/// <summary>
-		/// Serialize this directory entry to disk
-		/// </summary>
-		/// <param name="writer"></param>
-		public override void Serialize(ITreeNodeWriter writer)
-		{
-			base.Serialize(writer);
-
-			writer.WriteUtf8String(Name);
-		}
-
-		/// <inheritdoc/>
-		public override string ToString() => Name.ToString();
-	}
-
 	/// <summary>
 	/// Stats reported for copy operations
 	/// </summary>
@@ -400,7 +139,7 @@ namespace EpicGames.Horde.Storage.Nodes
 
 			int directoryCount = (int)reader.ReadUnsignedVarInt();
 			_nameToDirectoryEntry.EnsureCapacity(directoryCount);
-			
+
 			for (int idx = 0; idx < directoryCount; idx++)
 			{
 				DirectoryEntry entry = new DirectoryEntry(reader);
@@ -558,6 +297,79 @@ namespace EpicGames.Horde.Storage.Nodes
 			return false;
 		}
 
+		/// <summary>
+		/// Attempts to get a file entry from a path
+		/// </summary>
+		/// <param name="reader">Reader for node data</param>
+		/// <param name="path">Path to the directory</param>
+		/// <param name="cancellationToken">Cancellation token</param>
+		/// <returns>The directory with the given path, or null if it was not found</returns>
+		public async ValueTask<FileEntry?> GetFileEntryByPathAsync(TreeReader reader, Utf8String path, CancellationToken cancellationToken)
+		{
+			FileEntry? fileEntry;
+
+			int slashIdx = path.LastIndexOf('/');
+			if (slashIdx == -1)
+			{
+				if (!TryGetFileEntry(path, out fileEntry))
+				{
+					return null;
+				}
+			}
+			else
+			{
+				DirectoryNode? directoryNode = await GetDirectoryByPathAsync(reader, path.Slice(0, slashIdx), cancellationToken);
+				if (directoryNode == null)
+				{
+					return null;
+				}
+				if (!directoryNode.TryGetFileEntry(path.Slice(slashIdx + 1), out fileEntry))
+				{
+					return null;
+				}
+			}
+
+			return fileEntry;
+		}
+
+		/// <summary>
+		/// Attempts to get a directory entry from a path
+		/// </summary>
+		/// <param name="reader">Reader for node data</param>
+		/// <param name="path">Path to the directory</param>
+		/// <param name="cancellationToken">Cancellation token</param>
+		/// <returns>The directory with the given path, or null if it was not found</returns>
+		public ValueTask<DirectoryNode?> GetDirectoryByPathAsync(TreeReader reader, Utf8String path, CancellationToken cancellationToken) => GetDirectoryByPathAsync(this, reader, path, cancellationToken);
+
+		static async ValueTask<DirectoryNode?> GetDirectoryByPathAsync(DirectoryNode directoryNode, TreeReader reader, Utf8String path, CancellationToken cancellationToken)
+		{
+			while(path.Length > 0)
+			{
+				Utf8String directoryName;
+
+				int slashIdx = path.IndexOf('/');
+				if (slashIdx == -1)
+				{
+					directoryName = path;
+					path = Utf8String.Empty;
+				}
+				else
+				{
+					directoryName = path.Slice(0, slashIdx);
+					path = path.Slice(slashIdx + 1);
+				}
+
+				DirectoryEntry? directoryEntry;
+				if (!directoryNode.TryGetDirectoryEntry(directoryName, out directoryEntry))
+				{
+					return null;
+				}
+
+				directoryNode = await directoryEntry.ExpandAsync(reader, cancellationToken);
+			}
+			return directoryNode;
+		}
+
 		/// <summary>
 		/// Deletes a file with the given path
 		/// </summary>
@@ -996,13 +808,13 @@ namespace EpicGames.Horde.Storage.Nodes
 		/// </summary>
 		/// <param name="reader">Reader for other nodes</param>
 		/// <returns>Stream containing zipped archive data</returns>
-		public Stream ToZipStream(TreeReader reader) => new ZippedDirectoryStream(reader, this);
+		public Stream AsZipStream(TreeReader reader) => new DirectoryNodeZipStream(reader, this);
 	}
 
 	/// <summary>
 	/// Stream which zips a directory node tree dynamically
 	/// </summary>
-	class ZippedDirectoryStream : Stream
+	class DirectoryNodeZipStream : Stream
 	{
 		/// <inheritdoc/>
 		public override bool CanRead => true;
@@ -1030,7 +842,7 @@ namespace EpicGames.Horde.Storage.Nodes
 		/// </summary>
 		/// <param name="reader">Reader to read nodes through</param>
 		/// <param name="node">Root node to copy from</param>
-		public ZippedDirectoryStream(TreeReader reader, DirectoryNode node)
+		public DirectoryNodeZipStream(TreeReader reader, DirectoryNode node)
 		{
 			_pipe = new Pipe();
 			_backgroundTask = BackgroundTask.StartNew(ctx => CopyToPipeAsync(reader, node, _pipe.Writer, ctx));
diff --git a/Engine/Source/Programs/Shared/EpicGames.Horde/Storage/Nodes/DirectoryNodeRef.cs b/Engine/Source/Programs/Shared/EpicGames.Horde/Storage/Nodes/DirectoryNodeRef.cs
new file mode 100644
--- /dev/null
+++ b/Engine/Source/Programs/Shared/EpicGames.Horde/Storage/Nodes/DirectoryNodeRef.cs
@@ -0,0 +1,72 @@
+// Copyright Epic Games, Inc. All Rights Reserved.
+
+using System;
+using System.Collections.Generic;
+using System.Linq;
+using System.Text;
+using System.Threading.Tasks;
+using EpicGames.Core;
+
+namespace EpicGames.Horde.Storage.Nodes
+{
+	/// <summary>
+	/// Reference to a directory node, including the target hash and length
+	/// </summary>
+	public class DirectoryNodeRef : TreeNodeRef<DirectoryNode>
+	{
+		/// <summary>
+		/// Length of this directory tree
+		/// </summary>
+		public long Length => (Target == null) ? _cachedLength : Target.Length;
+
+		/// <summary>
+		/// Hash of the target node
+		/// </summary>
+		public IoHash Hash { get; private set; }
+
+		/// <summary>
+		/// Cached value for the length of this tree
+		/// </summary>
+		long _cachedLength;
+
+		/// <summary>
+		/// Constructor
+		/// </summary>
+		public DirectoryNodeRef(DirectoryNode node)
+			: base(node)
+		{
+		}
+
+		/// <summary>
+		/// Constructor
+		/// </summary>
+		/// <param name="reader"></param>
+		public DirectoryNodeRef(ITreeNodeReader reader)
+			: base(reader)
+		{
+			_cachedLength = (long)reader.ReadUnsignedVarInt();
+			Hash = reader.ReadIoHash();
+		}
+
+		/// <summary>
+		/// Serialize this directory entry to disk
+		/// </summary>
+		/// <param name="writer"></param>
+		public override void Serialize(ITreeNodeWriter writer)
+		{
+			base.Serialize(writer);
+
+			writer.WriteUnsignedVarInt((ulong)Length);
+			writer.WriteIoHash(Hash);
+		}
+
+		/// <inheritdoc/>
+		protected override void OnCollapse()
+		{
+			base.OnCollapse();
+
+			Hash = Target!.Hash;
+			_cachedLength = Target!.Length;
+		}
+	}
+}
diff --git a/Engine/Source/Programs/Shared/EpicGames.Horde/Storage/Nodes/FileEntry.cs b/Engine/Source/Programs/Shared/EpicGames.Horde/Storage/Nodes/FileEntry.cs
new file mode 100644
--- /dev/null
+++ b/Engine/Source/Programs/Shared/EpicGames.Horde/Storage/Nodes/FileEntry.cs
@@ -0,0 +1,241 @@
+// Copyright Epic Games, Inc. All Rights Reserved.
+
+using System;
+using System.IO;
+using System.IO.Pipelines;
+using System.Threading;
+using System.Threading.Tasks;
+using EpicGames.Core;
+
+namespace EpicGames.Horde.Storage.Nodes
+{
+	/// <summary>
+	/// Flags for a file entry
+	/// </summary>
+	[Flags]
+	public enum FileEntryFlags
+	{
+		/// <summary>
+		/// No other flags set
+		/// </summary>
+		None = 0,
+
+		/// <summary>
+		/// Indicates that the referenced file is executable
+		/// </summary>
+		Executable = 4,
+
+		/// <summary>
+		/// File should be stored as read-only
+		/// </summary>
+		ReadOnly = 8,
+
+		/// <summary>
+		/// File contents are utf-8 encoded text. Client may want to replace line-endings with OS-specific format.
+		/// </summary>
+		Text = 16,
+
+		/// <summary>
+		/// Used to indicate that custom data is included in the output. Used internally for serialization; not exposed to users.
+		/// </summary>
+		HasCustomData = 32,
+
+		/// <summary>
+		/// File should be materialized as UTF-16 (but is stored as a UTF-8 source)
+		/// </summary>
+		Utf16 = 64,
+	}
+
+	/// <summary>
+	/// Entry for a file within a directory node
+	/// </summary>
+	public sealed class FileEntry : TreeNodeRef<FileNode>
+	{
+		/// <summary>
+		/// Name of this file
+		/// </summary>
+		public Utf8String Name { get; }
+
+		/// <summary>
+		/// Flags for this file
+		/// </summary>
+		public FileEntryFlags Flags { get; set; }
+
+		/// <summary>
+		/// Length of this entry
+		/// </summary>
+		public long Length { get; }
+
+		/// <summary>
+		/// Hash of the target node
+		/// </summary>
+		public IoHash Hash { get; }
+
+		/// <summary>
+		/// Custom user data for this file entry
+		/// </summary>
+		public ReadOnlyMemory<byte> CustomData { get; set; } = ReadOnlyMemory<byte>.Empty;
+
+		/// <summary>
+		/// Constructor
+		/// </summary>
+		public FileEntry(Utf8String name, FileEntryFlags flags, long length, NodeHandle node)
+			: base(node)
+		{
+			Name = name;
+			Flags = flags;
+			Length = length;
+			Hash = node.Hash;
+		}
+
+		/// <summary>
+		/// Deserialize from a buffer
+		/// </summary>
+		/// <param name="reader"></param>
+		public FileEntry(ITreeNodeReader reader)
+			: base(reader)
+		{
+			Name = reader.ReadUtf8String();
+			Flags = (FileEntryFlags)reader.ReadUnsignedVarInt();
+			Length = (long)reader.ReadUnsignedVarInt();
+			Hash = reader.ReadIoHash();
+
+			if ((Flags & FileEntryFlags.HasCustomData) != 0)
+			{
+				CustomData = reader.ReadVariableLengthBytes();
+				Flags &= ~FileEntryFlags.HasCustomData;
+			}
+		}
+
+		/// <summary>
+		/// Serialize this entry
+		/// </summary>
+		/// <param name="writer"></param>
+		public override void Serialize(ITreeNodeWriter writer)
+		{
+			base.Serialize(writer);
+
+			FileEntryFlags flags = (CustomData.Length > 0) ? (Flags | FileEntryFlags.HasCustomData) : (Flags & ~FileEntryFlags.HasCustomData);
+
+			writer.WriteUtf8String(Name);
+			writer.WriteUnsignedVarInt((ulong)flags);
+			writer.WriteUnsignedVarInt((ulong)Length);
+			writer.WriteIoHash(Hash);
+
+			if ((flags & FileEntryFlags.HasCustomData) != 0)
+			{
+				writer.WriteVariableLengthBytes(CustomData.Span);
+			}
+		}
+
+		/// <summary>
+		/// Creates a stream that returns the contents of this file
+		/// </summary>
+		/// <param name="reader">Reader for node data</param>
+		/// <returns>The content stream</returns>
+		public Stream AsStream(TreeReader reader) => new FileEntryContentStream(reader, this);
+
+		/// <summary>
+		/// Copies the contents of this node and its children to the given output stream
+		/// </summary>
+		/// <param name="reader">Reader for retrieving existing node data</param>
+		/// <param name="outputStream">The output stream to receive the data</param>
+		/// <param name="cancellationToken">Cancellation token for the operation</param>
+		public async Task CopyToStreamAsync(TreeReader reader, Stream outputStream, CancellationToken cancellationToken)
+		{
+			FileNode node = await ExpandAsync(reader, cancellationToken);
+			await node.CopyToStreamAsync(reader, outputStream, cancellationToken);
+		}
+
+		/// <summary>
+		/// Extracts the contents of this node to a file
+		/// </summary>
+		/// <param name="reader">Reader for retrieving existing node data</param>
+		/// <param name="file">File to write with the contents of this node</param>
+		/// <param name="cancellationToken">Cancellation token for the operation</param>
+		/// <returns></returns>
+		public async Task CopyToFileAsync(TreeReader reader, FileInfo file, CancellationToken cancellationToken)
+		{
+			FileNode node = await ExpandAsync(reader, cancellationToken);
+			await node.CopyToFileAsync(reader, file, cancellationToken);
+		}
+
+		/// <inheritdoc/>
+		public override string ToString() => Name.ToString();
+	}
+
+	/// <summary>
+	/// Stream which returns the content of a file
+	/// </summary>
+	class FileEntryContentStream : Stream
+	{
+		/// <inheritdoc/>
+		public override bool CanRead => _readStream.CanRead;
+
+		/// <inheritdoc/>
+		public override bool CanSeek => _readStream.CanSeek;
+
+		/// <inheritdoc/>
+		public override bool CanWrite => _readStream.CanWrite;
+
+		/// <inheritdoc/>
+		public override long Length { get; }
+
+		/// <inheritdoc/>
+		public override long Position { get => _readStream.Position; set => throw new NotImplementedException(); }
+
+		readonly Pipe _pipe;
+		readonly Stream _readStream;
+		readonly BackgroundTask _writeTask;
+
+		/// <summary>
+		/// Constructor
+		/// </summary>
+		/// <param name="reader">Reader to read nodes through</param>
+		/// <param name="entry">The file entry to copy from</param>
+		public FileEntryContentStream(TreeReader reader, FileEntry entry)
+		{
+			_pipe = new Pipe();
+			_readStream = _pipe.Reader.AsStream();
+			_writeTask = BackgroundTask.StartNew(ctx => WriteDataAsync(reader, entry, _pipe.Writer, ctx));
+
+			Length = entry.Length;
+		}
+
+		/// <inheritdoc/>
+		public override async ValueTask DisposeAsync()
+		{
+			await base.DisposeAsync();
+
+			await _writeTask.DisposeAsync();
+
+			_readStream.Dispose();
+		}
+
+		static async Task WriteDataAsync(TreeReader reader, FileEntry entry, PipeWriter pipeWriter, CancellationToken cancellationToken)
+		{
+			using Stream stream = pipeWriter.AsStream();
+			await entry.CopyToStreamAsync(reader, stream, cancellationToken);
+		}
+
+		/// <inheritdoc/>
+		public override ValueTask<int> ReadAsync(Memory<byte> buffer, CancellationToken cancellationToken = default) => _readStream.ReadAsync(buffer, cancellationToken);
+
+		/// <inheritdoc/>
+		public override void Flush()
+		{
+		}
+
+		/// <inheritdoc/>
+		public override int Read(byte[] buffer, int offset, int count) => _readStream.Read(buffer, offset, count);
+
+		/// <inheritdoc/>
+		public override long Seek(long offset, SeekOrigin origin) => throw new NotSupportedException();
+
+		/// <inheritdoc/>
+		public override void SetLength(long value) => throw new NotSupportedException();
+
+		/// <inheritdoc/>
+		public override void Write(byte[] buffer, int offset, int count) => throw new NotSupportedException();
+	}
+}
diff --git a/Engine/Source/Programs/Shared/EpicGames.Horde/Storage/Nodes/FileNode.cs b/Engine/Source/Programs/Shared/EpicGames.Horde/Storage/Nodes/FileNode.cs
--- a/Engine/Source/Programs/Shared/EpicGames.Horde/Storage/Nodes/FileNode.cs
+++ b/Engine/Source/Programs/Shared/EpicGames.Horde/Storage/Nodes/FileNode.cs
@@ -357,304 +357,4 @@ namespace EpicGames.Horde.Storage.Nodes
 			}
 		}
 	}
-
-	/// <summary>
-	/// Utility class for generating FileNode data directly into <see cref="TreeWriter"/> instances, without constructing node representations first.
-	/// </summary>
-	public class FileNodeWriter
-	{
-		class InteriorNodeState
-		{
-			public InteriorNodeState? _parent;
-			public readonly ArrayMemoryWriter _data;
-			public readonly List<NodeHandle> _children = new List<NodeHandle>();
-
-			public uint _rollingHash;
-
-			public InteriorNodeState(int maxSize)
-			{
-				_data = new ArrayMemoryWriter(maxSize);
-			}
-
-			public void Reset()
-			{
-				_rollingHash = 0;
-				_data.Clear();
-				_children.Clear();
-			}
-
-			public void Write(NodeHandle handle)
-			{
-				_children.Add(handle);
-				_data.WriteIoHash(handle.Hash);
-			}
-		}
-
-		static readonly BundleType s_leafNodeType = TreeNode.GetBundleType<LeafFileNode>();
-		static readonly BundleType s_interiorNodeType = TreeNode.GetBundleType<InteriorFileNode>();
-
-		readonly TreeWriter _writer;
-		readonly ChunkingOptions _options;
-
-		// Tree state
-		InteriorNodeState? _topInteriorNode = null;
-		long _totalLength;
-		readonly Stack<InteriorNodeState> _freeInteriorNodes = new Stack<InteriorNodeState>();
-
-		// Leaf node state
-		uint _leafHash;
-		int _leafLength;
-
-		/// <summary>
-		/// Length of the file so far
-		/// </summary>
-		public long Length => _totalLength;
-
-		/// <summary>
-		/// Constructor
-		/// </summary>
-		/// <param name="writer">Writer for new nodes</param>
-		/// <param name="options">Chunking options</param>
-		public FileNodeWriter(TreeWriter writer, ChunkingOptions options)
-		{
-			_writer = writer;
-			_options = options;
-		}
-
-		/// <summary>
-		/// Reset the current state
-		/// </summary>
-		public void Reset()
-		{
-			FreeInteriorNodes();
-			ResetLeafState();
-			_totalLength = 0;
-		}
-
-		/// <summary>
-		/// Resets the state of the current leaf node
-		/// </summary>
-		void ResetLeafState()
-		{
-			_leafHash = 0;
-			_leafLength = 0;
-		}
-
-		/// <summary>
-		/// Creates a new interior node state
-		/// </summary>
-		/// <returns>State object</returns>
-		InteriorNodeState CreateInteriorNode()
-		{
-			InteriorNodeState? result;
-			if (!_freeInteriorNodes.TryPop(out result))
-			{
-				result = new InteriorNodeState(_options.InteriorOptions.MaxSize);
-			}
-			return result;
-		}
-
-		/// <summary>
-		/// Free all the current interior nodes
-		/// </summary>
-		void FreeInteriorNodes()
-		{
-			while (_topInteriorNode != null)
-			{
-				InteriorNodeState current = _topInteriorNode;
-				_topInteriorNode = _topInteriorNode._parent;
-				current.Reset();
-				_freeInteriorNodes.Push(current);
-			}
-		}
-
-		/// <summary>
-		/// Creates data for the given file
-		/// </summary>
-		/// <param name="fileInfo">File to append</param>
-		/// <param name="cancellationToken">Cancellation token for the operation</param>
-		public async Task<NodeHandle> CreateAsync(FileInfo fileInfo, CancellationToken cancellationToken)
-		{
-			using (FileStream stream = fileInfo.OpenRead())
-			{
-				return await CreateAsync(stream, cancellationToken);
-			}
-		}
-
-		/// <summary>
-		/// Creates data from the given stream
-		/// </summary>
-		/// <param name="stream">Stream to append</param>
-		/// <param name="cancellationToken">Cancellation token for the operation</param>
-		public async Task<NodeHandle> CreateAsync(Stream stream, CancellationToken cancellationToken)
-		{
-			Reset();
-			await AppendAsync(stream, cancellationToken);
-			return await CompleteAsync(cancellationToken);
-		}
-
-		/// <summary>
-		/// Creates data from the given data
-		/// </summary>
-		/// <param name="data">Stream to append</param>
-		/// <param name="cancellationToken">Cancellation token for the operation</param>
-		public async Task<NodeHandle> CreateAsync(ReadOnlyMemory<byte> data, CancellationToken cancellationToken)
-		{
-			Reset();
-			await AppendAsync(data, cancellationToken);
-			return await FlushAsync(cancellationToken);
-		}
-
-		/// <summary>
-		/// Appends data to the current file
-		/// </summary>
-		/// <param name="stream">Stream containing data to append</param>
-		/// <param name="cancellationToken">Cancellation token for the operation</param>
-		public async Task AppendAsync(Stream stream, CancellationToken cancellationToken)
-		{
-			const int BufferLength = 32 * 1024;
-
-			using IMemoryOwner<byte> owner = MemoryPool<byte>.Shared.Rent(BufferLength * 2);
-			Memory<byte> buffer = owner.Memory;
-
-			int readBufferOffset = 0;
-			Memory<byte> appendBuffer = Memory<byte>.Empty;
-			for (; ; )
-			{
-				// Start a read into memory
-				Memory<byte> readBuffer = buffer.Slice(readBufferOffset, BufferLength);
-				Task<int> readTask = Task.Run(async () => await stream.ReadAsync(readBuffer, cancellationToken), cancellationToken);
-
-				// In the meantime, append the last data that was read to the tree
-				if (appendBuffer.Length > 0)
-				{
-					await AppendAsync(appendBuffer, cancellationToken);
-				}
-
-				// Wait for the read to finish
-				int numBytes = await readTask;
-				if (numBytes == 0)
-				{
-					break;
-				}
-
-				// Switch the buffers around
-				appendBuffer = readBuffer.Slice(0, numBytes);
-				readBufferOffset ^= BufferLength;
-			}
-		}
-
-		/// <summary>
-		/// Appends data to the current file
-		/// </summary>
-		/// <param name="data">Data to append</param>
-		/// <param name="cancellationToken">Cancellation token for the operation</param>
-		public async Task AppendAsync(ReadOnlyMemory<byte> data, CancellationToken cancellationToken)
-		{
-			Memory<byte> buffer = _writer.GetOutputBuffer(_leafLength, _leafLength);
-			for (; ; )
-			{
-				// Append data to the current leaf node
-				int appendLength = LeafFileNode.AppendData(buffer.Span.Slice(0, _leafLength), data.Span, ref _leafHash, _options.LeafOptions);
-
-				buffer = _writer.GetOutputBuffer(_leafLength, _leafLength + appendLength);
-				data.Slice(0, appendLength).CopyTo(buffer.Slice(_leafLength));
-
-				_leafLength += appendLength;
-				data = data.Slice(appendLength);
-
-				_totalLength += appendLength;
-
-				if (data.Length == 0)
-				{
-					break;
-				}
-
-				// Flush the leaf node and any interior nodes that are full
-				NodeHandle handle = await WriteLeafNodeAsync(cancellationToken);
-				ResetLeafState();
-				await AddToInteriorNodeAsync(handle, cancellationToken);
-			}
-		}
-
-		async Task AddToInteriorNodeAsync(NodeHandle handle, CancellationToken cancellationToken)
-		{
-			_topInteriorNode ??= CreateInteriorNode();
-			await AddToInteriorNodeAsync(_topInteriorNode, handle, cancellationToken);
-		}
-
-		async Task AddToInteriorNodeAsync(InteriorNodeState interiorNode, NodeHandle handle, CancellationToken cancellationToken)
-		{
-			// If the node is already full, flush it
-			if (InteriorFileNode.IsComplete(interiorNode._data.WrittenSpan, interiorNode._rollingHash, _options.InteriorOptions))
-			{
-				NodeHandle interiorNodeHandle = await WriteInteriorNodeAndResetAsync(interiorNode, cancellationToken);
-				interiorNode._parent ??= CreateInteriorNode();
-				await AddToInteriorNodeAsync(interiorNode._parent, interiorNodeHandle, cancellationToken);
-			}
-
-			// Add this handle
-			InteriorFileNode.AppendData(interiorNode._data.WrittenSpan, handle.Hash, ref interiorNode._rollingHash, _options.InteriorOptions);
-			interiorNode.Write(handle);
-		}
-
-		/// <summary>
-		/// Complete the current file, and write all open nodes to the underlying writer
-		/// </summary>
-		/// <param name="cancellationToken">Cancellation token for the operation</param>
-		/// <returns>Handle to the root node</returns>
-		public async Task<NodeHandle> CompleteAsync(CancellationToken cancellationToken)
-		{
-			NodeHandle handle = await WriteLeafNodeAsync(cancellationToken);
-			ResetLeafState();
-
-			for (InteriorNodeState? state = _topInteriorNode; state != null; state = state._parent)
-			{
-				await AddToInteriorNodeAsync(state, handle, cancellationToken);
-				handle = await WriteInteriorNodeAndResetAsync(state, cancellationToken);
-			}
-
-			FreeInteriorNodes();
-			return handle;
-		}
-
-		/// <summary>
-		/// Flush the state of the writer
-		/// </summary>
-		/// <param name="cancellationToken">Cancellation token for the operation</param>
-		/// <returns>Handle to the root FileNode</returns>
-		public async Task<NodeHandle> FlushAsync(CancellationToken cancellationToken)
-		{
-			NodeHandle handle = await CompleteAsync(cancellationToken);
-			await _writer.FlushAsync(cancellationToken);
-			return handle;
-		}
-
-		/// <summary>
-		/// Writes the state of the given interior node to storage
-		/// </summary>
-		/// <param name="state"></param>
-		/// <param name="cancellationToken"></param>
-		/// <returns></returns>
-		async ValueTask<NodeHandle> WriteInteriorNodeAndResetAsync(InteriorNodeState state, CancellationToken cancellationToken)
-		{
-			Memory<byte> buffer = _writer.GetOutputBuffer(0, state._data.Length);
-			state._data.WrittenMemory.CopyTo(buffer);
-
-			NodeHandle handle = await _writer.WriteNodeAsync(state._data.Length, state._children, s_interiorNodeType, cancellationToken);
-			state.Reset();
-
-			return handle;
-		}
-
-		/// <summary>
-		/// Writes the contents of the current leaf node to storage
-		/// </summary>
-		/// <param name="cancellationToken">Cancellation token for the operation</param>
-		/// <returns>Handle to the written leaf node</returns>
-		async ValueTask<NodeHandle> WriteLeafNodeAsync(CancellationToken cancellationToken)
-		{
-			return await _writer.WriteNodeAsync(_leafLength, Array.Empty<NodeHandle>(), s_leafNodeType, cancellationToken);
-		}
-	}
 }
diff --git a/Engine/Source/Programs/Shared/EpicGames.Horde/Storage/Nodes/FileNodeWriter.cs b/Engine/Source/Programs/Shared/EpicGames.Horde/Storage/Nodes/FileNodeWriter.cs
new file mode 100644
--- /dev/null
+++ b/Engine/Source/Programs/Shared/EpicGames.Horde/Storage/Nodes/FileNodeWriter.cs
@@ -0,0 +1,312 @@
+// Copyright Epic Games, Inc. All Rights Reserved.
+
+using System;
+using System.Buffers;
+using System.Collections.Generic;
+using System.IO;
+using System.Threading;
+using System.Threading.Tasks;
+using EpicGames.Core;
+
+namespace EpicGames.Horde.Storage.Nodes
+{
+	/// <summary>
+	/// Utility class for generating FileNode data directly into <see cref="TreeWriter"/> instances, without constructing node representations first.
+	/// </summary>
+	public class FileNodeWriter
+	{
+		class InteriorNodeState
+		{
+			public InteriorNodeState? _parent;
+			public readonly ArrayMemoryWriter _data;
+			public readonly List<NodeHandle> _children = new List<NodeHandle>();
+
+			public uint _rollingHash;
+
+			public InteriorNodeState(int maxSize)
+			{
+				_data = new ArrayMemoryWriter(maxSize);
+			}
+
+			public void Reset()
+			{
+				_rollingHash = 0;
+				_data.Clear();
+				_children.Clear();
+			}
+
+			public void Write(NodeHandle handle)
+			{
+				_children.Add(handle);
+				_data.WriteIoHash(handle.Hash);
+			}
+		}
+
+		static readonly BundleType s_leafNodeType = TreeNode.GetBundleType<LeafFileNode>();
+		static readonly BundleType s_interiorNodeType = TreeNode.GetBundleType<InteriorFileNode>();
+
+		readonly TreeWriter _writer;
+		readonly ChunkingOptions _options;
+
+		// Tree state
+		InteriorNodeState? _topInteriorNode = null;
+		long _totalLength;
+		readonly Stack<InteriorNodeState> _freeInteriorNodes = new Stack<InteriorNodeState>();
+
+		// Leaf node state
+		uint _leafHash;
+		int _leafLength;
+
+		/// <summary>
+		/// Length of the file so far
+		/// </summary>
+		public long Length => _totalLength;
+
+		/// <summary>
+		/// Constructor
+		/// </summary>
+		/// <param name="writer">Writer for new nodes</param>
+		/// <param name="options">Chunking options</param>
+		public FileNodeWriter(TreeWriter writer, ChunkingOptions options)
+		{
+			_writer = writer;
+			_options = options;
+		}
+
+		/// <summary>
+		/// Reset the current state
+		/// </summary>
+		public void Reset()
+		{
+			FreeInteriorNodes();
+			ResetLeafState();
+			_totalLength = 0;
+		}
+
+		/// <summary>
+		/// Resets the state of the current leaf node
+		/// </summary>
+		void ResetLeafState()
+		{
+			_leafHash = 0;
+			_leafLength = 0;
+		}
+
+		/// <summary>
+		/// Creates a new interior node state
+		/// </summary>
+		/// <returns>State object</returns>
+		InteriorNodeState CreateInteriorNode()
+		{
+			InteriorNodeState? result;
+			if (!_freeInteriorNodes.TryPop(out result))
+			{
+				result = new InteriorNodeState(_options.InteriorOptions.MaxSize);
+			}
+			return result;
+		}
+
+		/// <summary>
+		/// Free all the current interior nodes
+		/// </summary>
+		void FreeInteriorNodes()
+		{
+			while (_topInteriorNode != null)
+			{
+				InteriorNodeState current = _topInteriorNode;
+				_topInteriorNode = _topInteriorNode._parent;
+				current.Reset();
+				_freeInteriorNodes.Push(current);
+			}
+		}
+
+		/// <summary>
+		/// Creates data for the given file
+		/// </summary>
+		/// <param name="fileInfo">File to append</param>
+		/// <param name="cancellationToken">Cancellation token for the operation</param>
+		public async Task<NodeHandle> CreateAsync(FileInfo fileInfo, CancellationToken cancellationToken)
+		{
+			using (FileStream stream = fileInfo.OpenRead())
+			{
+				return await CreateAsync(stream, cancellationToken);
+			}
+		}
+
+		/// <summary>
+		/// Creates data from the given stream
+		/// </summary>
+		/// <param name="stream">Stream to append</param>
+		/// <param name="cancellationToken">Cancellation token for the operation</param>
+		public async Task<NodeHandle> CreateAsync(Stream stream, CancellationToken cancellationToken)
+		{
+			Reset();
+			await AppendAsync(stream, cancellationToken);
+			return await CompleteAsync(cancellationToken);
+		}
+
+		/// <summary>
+		/// Creates data from the given data
+		/// </summary>
+		/// <param name="data">Stream to append</param>
+		/// <param name="cancellationToken">Cancellation token for the operation</param>
+		public async Task<NodeHandle> CreateAsync(ReadOnlyMemory<byte> data, CancellationToken cancellationToken)
+		{
+			Reset();
+			await AppendAsync(data, cancellationToken);
+			return await FlushAsync(cancellationToken);
+		}
+
+		/// <summary>
+		/// Appends data to the current file
+		/// </summary>
+		/// <param name="stream">Stream containing data to append</param>
+		/// <param name="cancellationToken">Cancellation token for the operation</param>
+		public async Task AppendAsync(Stream stream, CancellationToken cancellationToken)
+		{
+			const int BufferLength = 32 * 1024;
+
+			using IMemoryOwner<byte> owner = MemoryPool<byte>.Shared.Rent(BufferLength * 2);
+			Memory<byte> buffer = owner.Memory;
+
+			int readBufferOffset = 0;
+			Memory<byte> appendBuffer = Memory<byte>.Empty;
+			for (; ; )
+			{
+				// Start a read into memory
+				Memory<byte> readBuffer = buffer.Slice(readBufferOffset, BufferLength);
+				Task<int> readTask = Task.Run(async () => await stream.ReadAsync(readBuffer, cancellationToken), cancellationToken);
+
+				// In the meantime, append the last data that was read to the tree
+				if (appendBuffer.Length > 0)
+				{
+					await AppendAsync(appendBuffer, cancellationToken);
+				}
+
+				// Wait for the read to finish
+				int numBytes = await readTask;
+				if (numBytes == 0)
+				{
+					break;
+				}
+
+				// Switch the buffers around
+				appendBuffer = readBuffer.Slice(0, numBytes);
+				readBufferOffset ^= BufferLength;
+			}
+		}
+
+		/// <summary>
+		/// Appends data to the current file
+		/// </summary>
+		/// <param name="data">Data to append</param>
+		/// <param name="cancellationToken">Cancellation token for the operation</param>
+		public async Task AppendAsync(ReadOnlyMemory<byte> data, CancellationToken cancellationToken)
+		{
+			Memory<byte> buffer = _writer.GetOutputBuffer(_leafLength, _leafLength);
+			for (; ; )
+			{
+				// Append data to the current leaf node
+				int appendLength = LeafFileNode.AppendData(buffer.Span.Slice(0, _leafLength), data.Span, ref _leafHash, _options.LeafOptions);
+
+				buffer = _writer.GetOutputBuffer(_leafLength, _leafLength + appendLength);
+				data.Slice(0, appendLength).CopyTo(buffer.Slice(_leafLength));
+
+				_leafLength += appendLength;
+				data = data.Slice(appendLength);
+
+				_totalLength += appendLength;
+
+				if (data.Length == 0)
+				{
+					break;
+				}
+
+				// Flush the leaf node and any interior nodes that are full
+				NodeHandle handle = await WriteLeafNodeAsync(cancellationToken);
+				ResetLeafState();
+				await AddToInteriorNodeAsync(handle, cancellationToken);
+			}
+		}
+
+		async Task AddToInteriorNodeAsync(NodeHandle handle, CancellationToken cancellationToken)
+		{
+			_topInteriorNode ??= CreateInteriorNode();
+			await AddToInteriorNodeAsync(_topInteriorNode, handle, cancellationToken);
+		}
+
+		async Task AddToInteriorNodeAsync(InteriorNodeState interiorNode, NodeHandle handle, CancellationToken cancellationToken)
+		{
+			// If the node is already full, flush it
+			if (InteriorFileNode.IsComplete(interiorNode._data.WrittenSpan, interiorNode._rollingHash, _options.InteriorOptions))
+			{
+				NodeHandle interiorNodeHandle = await WriteInteriorNodeAndResetAsync(interiorNode, cancellationToken);
+				interiorNode._parent ??= CreateInteriorNode();
+				await AddToInteriorNodeAsync(interiorNode._parent, interiorNodeHandle, cancellationToken);
+			}
+
+			// Add this handle
+			InteriorFileNode.AppendData(interiorNode._data.WrittenSpan, handle.Hash, ref interiorNode._rollingHash, _options.InteriorOptions);
+			interiorNode.Write(handle);
+		}
+
+		/// <summary>
+		/// Complete the current file, and write all open nodes to the underlying writer
+		/// </summary>
+		/// <param name="cancellationToken">Cancellation token for the operation</param>
+		/// <returns>Handle to the root node</returns>
+		public async Task<NodeHandle> CompleteAsync(CancellationToken cancellationToken)
+		{
+			NodeHandle handle = await WriteLeafNodeAsync(cancellationToken);
+			ResetLeafState();
+
+			for (InteriorNodeState? state = _topInteriorNode; state != null; state = state._parent)
+			{
+				await AddToInteriorNodeAsync(state, handle, cancellationToken);
+				handle = await WriteInteriorNodeAndResetAsync(state, cancellationToken);
+			}
+
+			FreeInteriorNodes();
+			return handle;
+		}
+
+		/// <summary>
+		/// Flush the state of the writer
+		/// </summary>
+		/// <param name="cancellationToken">Cancellation token for the operation</param>
+		/// <returns>Handle to the root FileNode</returns>
+		public async Task<NodeHandle> FlushAsync(CancellationToken cancellationToken)
+		{
+			NodeHandle handle = await CompleteAsync(cancellationToken);
+			await _writer.FlushAsync(cancellationToken);
+			return handle;
+		}
+
+		/// <summary>
+		/// Writes the state of the given interior node to storage
+		/// </summary>
+		/// <param name="state"></param>
+		/// <param name="cancellationToken"></param>
+		/// <returns></returns>
+		async ValueTask<NodeHandle> WriteInteriorNodeAndResetAsync(InteriorNodeState state, CancellationToken cancellationToken)
+		{
+			Memory<byte> buffer = _writer.GetOutputBuffer(0, state._data.Length);
+			state._data.WrittenMemory.CopyTo(buffer);
+
+			NodeHandle handle = await _writer.WriteNodeAsync(state._data.Length, state._children, s_interiorNodeType, cancellationToken);
+			state.Reset();
+
+			return handle;
+		}
+
+		/// <summary>
+		/// Writes the contents of the current leaf node to storage
+		/// </summary>
+		/// <param name="cancellationToken">Cancellation token for the operation</param>
+		/// <returns>Handle to the written leaf node</returns>
+		async ValueTask<NodeHandle> WriteLeafNodeAsync(CancellationToken cancellationToken)
+		{
+			return await _writer.WriteNodeAsync(_leafLength, Array.Empty<NodeHandle>(), s_leafNodeType, cancellationToken);
+		}
+	}
+}
