Horde: Support for filtering files that are zipped for artifacts.

diff --git a/Engine/Source/Programs/Horde/Horde.Server/Artifacts/ArtifactsController.cs b/Engine/Source/Programs/Horde/Horde.Server/Artifacts/ArtifactsController.cs
--- a/Engine/Source/Programs/Horde/Horde.Server/Artifacts/ArtifactsController.cs
+++ b/Engine/Source/Programs/Horde/Horde.Server/Artifacts/ArtifactsController.cs
@@ -112,6 +112,17 @@ namespace Horde.Server.Artifacts
 		}
 	}
 
+	/// <summary>
+	/// Request to create a zip file with artifact data
+	/// </summary>
+	public class CreateZipRequest
+	{
+		/// <summary>
+		/// Filter lines for the zip. Uses standard <see cref="FileFilter"/> syntax.
+		/// </summary>
+		public List<string> Filter { get; set; } = new List<string>();
+	}
+
 	/// <summary>
 	/// Public interface for artifacts
 	/// </summary>
@@ -306,11 +317,31 @@ namespace Horde.Server.Artifacts
 		/// Downloads an individual file from an artifact
 		/// </summary>
 		/// <param name="id">Identifier of the artifact to retrieve</param>
+		/// <param name="filter">Paths to include in the zip file. The post version of this request allows for more parameters than can fit in a request string.</param>
 		/// <param name="cancellationToken">Cancellation token for the operation</param>
 		/// <returns>Information about all the artifacts</returns>
 		[HttpGet]
 		[Route("/api/v2/artifacts/{id}/zip")]
-		public async Task<ActionResult<object>> GetZipAsync(ArtifactId id, CancellationToken cancellationToken = default)
+		public async Task<ActionResult<object>> GetZipAsync(ArtifactId id, [FromQuery(Name = "filter")] string[]? filter, CancellationToken cancellationToken = default)
+		{
+			return await GetZipInternalAsync(id, filter, cancellationToken);
+		}
+
+		/// <summary>
+		/// Downloads an individual file from an artifact
+		/// </summary>
+		/// <param name="id">Identifier of the artifact to retrieve</param>
+		/// <param name="request">Filter for the zip file</param>
+		/// <param name="cancellationToken">Cancellation token for the operation</param>
+		/// <returns>Information about all the artifacts</returns>
+		[HttpPost]
+		[Route("/api/v2/artifacts/{id}/zip")]
+		public async Task<ActionResult<object>> CreateZipFromFilterAsync(ArtifactId id, CreateZipRequest request, CancellationToken cancellationToken = default)
+		{
+			return await GetZipInternalAsync(id, request.Filter, cancellationToken);
+		}
+
+		async Task<ActionResult> GetZipInternalAsync(ArtifactId id, IEnumerable<string>? fileFilter, CancellationToken cancellationToken)
 		{
 			IArtifact? artifact = await _artifactCollection.GetAsync(id, cancellationToken);
 			if (artifact == null)
@@ -322,11 +353,17 @@ namespace Horde.Server.Artifacts
 				return Forbid(AclAction.ReadArtifact, artifact.AclScope);
 			}
 
+			FileFilter? filter = null;
+			if (fileFilter != null && fileFilter.Any())
+			{
+				filter = new FileFilter(fileFilter);
+			}
+
 			TreeReader reader = await _storageService.GetReaderAsync(artifact.NamespaceId, cancellationToken);
 			TreeNodeRef<DirectoryNode> directoryRef = await reader.ReadNodeRefAsync<DirectoryNode>(artifact.RefName, DateTime.UtcNow.AddHours(1.0), cancellationToken);
 			DirectoryNode directory = await directoryRef.ExpandAsync(reader, cancellationToken);
 
-			Stream stream = directory.AsZipStream(reader);
+			Stream stream = directory.AsZipStream(reader, filter);
 			return new FileStreamResult(stream, "application/zip") { FileDownloadName = $"{artifact.RefName}.zip" };
 		}
 
diff --git a/Engine/Source/Programs/Shared/EpicGames.Horde/Storage/Nodes/DirectoryNode.cs b/Engine/Source/Programs/Shared/EpicGames.Horde/Storage/Nodes/DirectoryNode.cs
--- a/Engine/Source/Programs/Shared/EpicGames.Horde/Storage/Nodes/DirectoryNode.cs
+++ b/Engine/Source/Programs/Shared/EpicGames.Horde/Storage/Nodes/DirectoryNode.cs
@@ -807,8 +807,9 @@ namespace EpicGames.Horde.Storage.Nodes
 		/// Returns a stream containing the zipped contents of this directory
 		/// </summary>
 		/// <param name="reader">Reader for other nodes</param>
+		/// <param name="filter">Filter for files to include in the zip</param>
 		/// <returns>Stream containing zipped archive data</returns>
-		public Stream AsZipStream(TreeReader reader) => new DirectoryNodeZipStream(reader, this);
+		public Stream AsZipStream(TreeReader reader, FileFilter? filter = null) => new DirectoryNodeZipStream(reader, this, filter);
 	}
 
 	/// <summary>
@@ -842,10 +843,11 @@ namespace EpicGames.Horde.Storage.Nodes
 		/// </summary>
 		/// <param name="reader">Reader to read nodes through</param>
 		/// <param name="node">Root node to copy from</param>
-		public DirectoryNodeZipStream(TreeReader reader, DirectoryNode node)
+		/// <param name="filter">Filter for files to include in the zip</param>
+		public DirectoryNodeZipStream(TreeReader reader, DirectoryNode node, FileFilter? filter)
 		{
 			_pipe = new Pipe();
-			_backgroundTask = BackgroundTask.StartNew(ctx => CopyToPipeAsync(reader, node, _pipe.Writer, ctx));
+			_backgroundTask = BackgroundTask.StartNew(ctx => CopyToPipeAsync(reader, node, filter, _pipe.Writer, ctx));
 		}
 
 		/// <inheritdoc/>
@@ -889,36 +891,44 @@ namespace EpicGames.Horde.Storage.Nodes
 			return length;
 		}
 
-		static async Task CopyToPipeAsync(TreeReader reader, DirectoryNode node, PipeWriter writer, CancellationToken cancellationToken)
+		static async Task CopyToPipeAsync(TreeReader reader, DirectoryNode node, FileFilter? filter, PipeWriter writer, CancellationToken cancellationToken)
 		{
 			using Stream outputStream = writer.AsStream();
 			using ZipArchive archive = new ZipArchive(outputStream, ZipArchiveMode.Create);
-			await CopyFilesAsync(reader, node, "", archive, cancellationToken);
+			await CopyFilesAsync(reader, node, "", filter, archive, cancellationToken);
 		}
 
-		static async Task CopyFilesAsync(TreeReader reader, DirectoryNode directory, string prefix, ZipArchive archive, CancellationToken cancellationToken)
+		static async Task CopyFilesAsync(TreeReader reader, DirectoryNode directory, string prefix, FileFilter? filter, ZipArchive archive, CancellationToken cancellationToken)
 		{
 			foreach (DirectoryEntry directoryEntry in directory.Directories)
 			{
-				DirectoryNode node = await directoryEntry.ExpandAsync(reader, cancellationToken);
-				await CopyFilesAsync(reader, node, $"{prefix}{directoryEntry.Name}/", archive, cancellationToken);
+				string directoryPath = $"{prefix}{directoryEntry.Name}/";
+				if (filter == null || filter.PossiblyMatches(directoryPath))
+				{
+					DirectoryNode node = await directoryEntry.ExpandAsync(reader, cancellationToken);
+					await CopyFilesAsync(reader, node, directoryPath, filter, archive, cancellationToken);
+				}
 			}
 
 			foreach (FileEntry fileEntry in directory.Files)
 			{
-				ZipArchiveEntry entry = archive.CreateEntry($"{prefix}{fileEntry}");
-
-				if ((fileEntry.Flags & FileEntryFlags.Executable) != 0)
+				string filePath = $"{prefix}{fileEntry}";
+				if (filter == null || filter.Matches(filePath))
 				{
-					entry.ExternalAttributes |= 0b_111_111_101 << 16; // rwx rwx r-x
-				}
-				else
-				{
-					entry.ExternalAttributes |= 0b_110_110_100 << 16; // rw- rw- r--
-				}
+					ZipArchiveEntry entry = archive.CreateEntry(filePath);
 
-				using Stream entryStream = entry.Open();
-				await fileEntry.CopyToStreamAsync(reader, entryStream, cancellationToken);
+					if ((fileEntry.Flags & FileEntryFlags.Executable) != 0)
+					{
+						entry.ExternalAttributes |= 0b_111_111_101 << 16; // rwx rwx r-x
+					}
+					else
+					{
+						entry.ExternalAttributes |= 0b_110_110_100 << 16; // rw- rw- r--
+					}
+
+					using Stream entryStream = entry.Open();
+					await fileEntry.CopyToStreamAsync(reader, entryStream, cancellationToken);
+				}
 			}
 		}
 
