Updated JsonTracer to be thread safe and support creating spans from multiple threads at once.

diff --git a/Engine/Source/Programs/Shared/EpicGames.Core.Tests/JsonTracerTests.cs b/Engine/Source/Programs/Shared/EpicGames.Core.Tests/JsonTracerTests.cs
--- a/Engine/Source/Programs/Shared/EpicGames.Core.Tests/JsonTracerTests.cs
+++ b/Engine/Source/Programs/Shared/EpicGames.Core.Tests/JsonTracerTests.cs
@@ -8,6 +8,8 @@ using System.Linq;
 using System.Text;
 using System.Text.Json;
 using OpenTracing;
+using System.Threading.Tasks;
+using System.Diagnostics;
 
 namespace EpicGames.Core.Tests
 {
@@ -65,5 +67,48 @@ namespace EpicGames.Core.Tests
 				Assert.AreEqual("myResource", SpansList[1].GetProperty("Resource").GetString());
 			}
 		}
+
+		[TestMethod]
+		public void SpansWithParallelFor()
+		{
+			JsonTracer Tracer = new JsonTracer();
+
+			using (IScope ParentScope1 = Tracer.BuildSpan("Parent").StartActive()) 
+			{
+				Parallel.For(0, 100, Index =>
+				{
+					using (IScope ChildScope = Tracer.BuildSpan("Child").StartActive()) { }
+				});
+			}
+		}
+
+		private static async Task<bool> AsyncSpanTest(JsonTracer Tracer)
+		{
+			for (int Index = 0; Index < 100; Index++)
+			{
+				using (IScope AsyncChildScope = Tracer.BuildSpan("AsyncChild").StartActive())
+				{
+					await Task.Run(() => 
+					{
+						using (IScope TaskChildScope = Tracer.BuildSpan("TaskChild").StartActive())
+						{
+							return true;
+						}
+					});
+				}
+			}
+			return true;
+		}
+
+		[TestMethod]
+		public void SpansWithAsync()
+		{
+			JsonTracer Tracer = new JsonTracer();
+
+			using (IScope ParentScope1 = Tracer.BuildSpan("Parent").StartActive())
+			{
+				 AsyncSpanTest(Tracer).Wait();
+			}
+		}
 	}
 }
diff --git a/Engine/Source/Programs/Shared/EpicGames.Core/JsonTracer.cs b/Engine/Source/Programs/Shared/EpicGames.Core/JsonTracer.cs
--- a/Engine/Source/Programs/Shared/EpicGames.Core/JsonTracer.cs
+++ b/Engine/Source/Programs/Shared/EpicGames.Core/JsonTracer.cs
@@ -1,7 +1,10 @@
 // Copyright Epic Games, Inc. All Rights Reserved.
 
 using System;
+using System.Collections;
+using System.Collections.Concurrent;
 using System.Collections.Generic;
+using System.Diagnostics.Tracing;
 using System.Globalization;
 using System.IO;
 using System.Linq;
@@ -188,6 +191,7 @@ namespace EpicGames.Core
 				CheckForFinished("Tried to finish already finished span");
 				_finishTimestamp = finishTimestamp;
 				_tracer.AppendFinishedSpan(this);
+				_tracer.DeactivateActiveSpan(this);
 				_finished = true;
 			}
 		}
@@ -374,6 +378,7 @@ namespace EpicGames.Core
 		public IScope StartActive(bool finishSpanOnDispose)
 		{
 			ISpan span = Start();
+			_tracer.ActivateSpan(span);
 			return _tracer.ScopeManager.Activate(span, finishSpanOnDispose);
 		}
 
@@ -393,20 +398,197 @@ namespace EpicGames.Core
 			return new JsonTracerSpan(_tracer, _operationName, _startTimestamp, _tags, _references);
 		}
 	}
-	
+
 	public class JsonTracer : ITracer
 	{
+		private class LocalEventListener : EventListener
+		{
+			protected override void OnEventSourceCreated(EventSource eventSource)
+			{
+				base.OnEventSourceCreated(eventSource);
+				switch(eventSource.Name)
+				{
+					case "System.Threading.Tasks.Parallel.EventSource":
+					case "System.Threading.Tasks.TplEventSource":
+						{
+							EnableEvents(eventSource, EventLevel.LogAlways, EventKeywords.All);
+							break;
+						}
+					default:
+						{
+							break;
+						}
+				}
+			}
+		}
+
 		public IScopeManager ScopeManager { get; }
-		public ISpan? ActiveSpan => ScopeManager.Active?.Span;
+		public ISpan? ActiveSpan
+		{ 
+			get
+			{
+				if (ActiveSpans.TryGetValue(Environment.CurrentManagedThreadId, out Stack<ISpan>? currentThreadActiveSpan) && currentThreadActiveSpan.Any())
+				{
+					return currentThreadActiveSpan.Peek();
+				}	
+				return null;
+			}
+		}
 
 		private readonly object _lock = new object();
 		private readonly List<JsonTracerSpan> _finishedSpans = new List<JsonTracerSpan>();
 		private readonly DirectoryReference? _telemetryDir;
-		
+		private readonly LocalEventListener _eventListener = new LocalEventListener();
+
+		private readonly ConcurrentDictionary<int, Stack<ISpan>> ActiveSpans = new ConcurrentDictionary<int, Stack<ISpan>>();
+		private readonly ConcurrentDictionary<int, ISpan> ForkIdToActiveSpans = new ConcurrentDictionary<int, ISpan>();
+		private readonly ConcurrentDictionary<int, ISpan> TaskIdToActiveSpans = new ConcurrentDictionary<int, ISpan>();
+
+		// These could possible be set to 2, which is what MS sets them to internally, if this becomes a performance issue.
+		private static int? ParallelLoopBegin_ForkJoinContextIDIndex;
+		private static int? ParallelLoopEnd_ForkJoinContextIDIndex;
+		private static int? ParallelFork_ForkJoinContextIDIndex;
+		private static int? ParallelJoin_ForkJoinContextIDIndex;
+		private static int? TraceOperationBegin_TaskIDIndex;
+		private static int? TraceOperationEnd_TaskIDIndex;
+		private static int? TaskStarted_TaskIDIndex;
+		private const string ForkJoinContextIDStr = "ForkJoinContextID";
+		private const string TaskIDStr = "TaskID";
+
 		public JsonTracer(DirectoryReference? telemetryDir = null)
 		{
 			ScopeManager = new AsyncLocalScopeManager();
 			_telemetryDir = telemetryDir;
+			_eventListener.EventWritten += EventListenerEventWritten;
+		}
+
+		private static int GetEventIntPayload(EventWrittenEventArgs e, string payloadName, ref int? payloadIndex)
+		{
+			if (!payloadIndex.HasValue)
+			{
+				int? foundIndex = e.PayloadNames?.FindIndex(pn => pn == payloadName);
+				if (foundIndex == null || foundIndex == -1)
+				{
+					throw new InvalidOperationException($"{payloadName} was not found in the event name payload.");
+				}
+				payloadIndex = foundIndex;
+			}
+
+			if (e.Payload == null || e.Payload.Count <= payloadIndex.Value)
+			{
+				throw new InvalidOperationException($"{payloadName} was not in the event payload.");
+			}
+
+			int? intValue = e.Payload[payloadIndex.Value] as int?;
+			if (!intValue.HasValue)
+			{
+				throw new InvalidOperationException($"{payloadName} in not an int.");
+			}
+			return intValue.Value;
+		}
+
+		private void BeginEvent(EventWrittenEventArgs e, string PayloadName, ref int? payloadIndex, ConcurrentDictionary<int, ISpan> activeSpansDict)
+		{
+			int value = GetEventIntPayload(e, PayloadName, ref payloadIndex);
+
+			ISpan? currentActiveSpan = null;
+			if (ActiveSpans.TryGetValue(Environment.CurrentManagedThreadId, out Stack<ISpan>? currentThreadActiveSpan) && currentThreadActiveSpan.Any())
+			{
+				currentActiveSpan = currentThreadActiveSpan.Peek();
+			}
+
+			// the current active span on this thread needs to be attached to this new event
+			if (currentActiveSpan != null)
+			{
+				if (!activeSpansDict.TryAdd(value, currentActiveSpan))
+				{
+					throw new InvalidOperationException($"Event id '{value}' was already added.");
+				}
+			}
+		}
+
+		private static void EndEvent(EventWrittenEventArgs e, string payloadName, ref int? payloadIndex, ConcurrentDictionary<int, ISpan> activeSpansDict)
+		{
+			int value = GetEventIntPayload(e, payloadName, ref payloadIndex);
+			if (activeSpansDict.ContainsKey(value))
+			{
+				if (!activeSpansDict.TryRemove(value, out _))
+				{
+					throw new InvalidOperationException($"Event id '{value}' could not be removed.");
+				}
+			}
+		}
+
+		private void BeginWork(EventWrittenEventArgs e, string payloadName, ref int? payloadIndex, ConcurrentDictionary<int, ISpan> activeSpansDict)
+		{
+			// find the active span for this event and make it the active span for this thread
+			int value = GetEventIntPayload(e, payloadName, ref payloadIndex);
+			if (activeSpansDict.TryGetValue(value, out ISpan? activeSpanForFork))
+			{
+				ActivateSpan(activeSpanForFork);
+			}
+		}
+
+		private void EndWork(EventWrittenEventArgs e, string payloadName, ref int? payloadIndex, ConcurrentDictionary<int, ISpan> activeSpansDict)
+		{
+			int value = GetEventIntPayload(e, payloadName, ref payloadIndex);
+			if (activeSpansDict.TryGetValue(value, out ISpan? activeSpanForEvent))
+			{
+				if (ActiveSpans.TryGetValue(Environment.CurrentManagedThreadId, out Stack<ISpan>? spanStack))
+				{
+					ISpan threadActiveSpan = spanStack.Pop();
+					if (activeSpanForEvent != threadActiveSpan)
+					{
+						throw new InvalidOperationException($"Span should be equal. Event span op: '{((JsonTracerSpan)activeSpanForEvent).OperationName}' Thread span op: '{((JsonTracerSpan)threadActiveSpan).OperationName}'.");
+					}
+				}
+			}
+		}
+
+		private void EventListenerEventWritten(object? sender, EventWrittenEventArgs e)
+		{
+			switch (e.EventName)
+			{
+				case "ParallelLoopBegin":
+					{
+						BeginEvent(e, ForkJoinContextIDStr, ref ParallelLoopBegin_ForkJoinContextIDIndex, ForkIdToActiveSpans);
+						break;
+					}
+				case "ParallelLoopEnd":
+					{
+						EndEvent(e, ForkJoinContextIDStr, ref ParallelLoopEnd_ForkJoinContextIDIndex, ForkIdToActiveSpans);
+						break;
+					}
+				case "ParallelFork":
+					{
+						BeginWork(e, ForkJoinContextIDStr, ref ParallelFork_ForkJoinContextIDIndex, ForkIdToActiveSpans);
+						break;
+					}
+				case "ParallelJoin":
+					{
+						EndWork(e, ForkJoinContextIDStr, ref ParallelJoin_ForkJoinContextIDIndex, ForkIdToActiveSpans);
+						break;
+					}
+				case "TraceOperationBegin":
+					{
+						BeginEvent(e, TaskIDStr, ref TraceOperationBegin_TaskIDIndex, TaskIdToActiveSpans);
+						break;
+					}
+				case "TraceOperationEnd":
+					{
+						EndEvent(e, TaskIDStr, ref TraceOperationEnd_TaskIDIndex, TaskIdToActiveSpans);
+						break;
+					}
+				case "TaskStarted":
+					{
+						BeginWork(e, TaskIDStr, ref TaskStarted_TaskIDIndex, TaskIdToActiveSpans);
+						break;
+					}
+				default:
+					{
+						break;
+					}
+			}
 		}
 
 		public static JsonTracer? TryRegisterAsGlobalTracer()
@@ -485,6 +667,11 @@ namespace EpicGames.Core
 				}
 				writer.WriteValue("StartTime", span.StartTimestamp.ToString("o", CultureInfo.InvariantCulture));
 				writer.WriteValue("FinishTime", span.FinishTimestamp.ToString("o", CultureInfo.InvariantCulture));
+				writer.WriteValue("SpanId", Int32.Parse(span.Context.SpanId));
+				if (span.ParentId != null)
+				{
+					writer.WriteValue("ParentId", Int32.Parse(span.ParentId));
+				}
 				writer.WriteObjectStart("Metadata");
 				// TODO: Write tags as metadata?
 				writer.WriteObjectEnd();
@@ -501,5 +688,34 @@ namespace EpicGames.Core
 				_finishedSpans.Add(jsonTracerSpan);
 			}
 		}
+
+		internal void ActivateSpan(ISpan span)
+		{
+			ActiveSpans.AddOrUpdate(Environment.CurrentManagedThreadId, 
+			(k) =>
+			{
+				Stack<ISpan> stack = new Stack<ISpan>();
+				stack.Push(span);
+				return stack;
+			}, 
+			(k,v) =>
+			{
+				v.Push(span);
+				return v;
+			}
+			);
+		}
+
+		internal void DeactivateActiveSpan(ISpan span)
+		{
+			if (ActiveSpans.TryGetValue(Environment.CurrentManagedThreadId, out Stack<ISpan>? spanStack))
+			{
+				// It's possible for other threads to finish spans
+				if (spanStack != null && spanStack.Any() && spanStack.Peek() == span) 
+				{
+					spanStack.Pop();
+				}
+			}
+		}
 	}
 }
