Updated JsonTracer to derive from IDisposable.

diff --git a/Engine/Source/Programs/AutomationTool/AutomationUtils/Automation.cs b/Engine/Source/Programs/AutomationTool/AutomationUtils/Automation.cs
--- a/Engine/Source/Programs/AutomationTool/AutomationUtils/Automation.cs
+++ b/Engine/Source/Programs/AutomationTool/AutomationUtils/Automation.cs
@@ -178,6 +178,7 @@ namespace AutomationTool
 					if (Tracer != null)
 					{
 						Tracer.Flush();
+						Tracer.Dispose();
 					}
 				}
 			}
diff --git a/Engine/Source/Programs/Shared/EpicGames.Core.Tests/JsonTracerTests.cs b/Engine/Source/Programs/Shared/EpicGames.Core.Tests/JsonTracerTests.cs
--- a/Engine/Source/Programs/Shared/EpicGames.Core.Tests/JsonTracerTests.cs
+++ b/Engine/Source/Programs/Shared/EpicGames.Core.Tests/JsonTracerTests.cs
@@ -9,7 +9,6 @@ using System.Text;
 using System.Text.Json;
 using OpenTracing;
 using System.Threading.Tasks;
-using System.Diagnostics;
 
 namespace EpicGames.Core.Tests
 {
@@ -19,66 +18,70 @@ namespace EpicGames.Core.Tests
 		[TestMethod]
 		public void Run()
 		{
-			JsonTracer Tracer = new JsonTracer();
+			using (JsonTracer Tracer = new JsonTracer())
+			{
+				using (IScope Scope1 = Tracer.BuildSpan("bogusOp1").StartActive()) { }
+				using (IScope Scope2 = Tracer.BuildSpan("bogusOp2").StartActive()) { }
 
-			using (IScope Scope1 = Tracer.BuildSpan("bogusOp1").StartActive()) { }
-			using (IScope Scope2 = Tracer.BuildSpan("bogusOp2").StartActive()) { }
-			
-			Assert.AreEqual(2, Tracer.GetFinishedSpans().Count);
+				Assert.AreEqual(2, Tracer.GetFinishedSpans().Count);
+			}
 		}
 		
 		[TestMethod]
 		public void SpansWithTags()
 		{
-			JsonTracer Tracer = new JsonTracer();
+			using (JsonTracer Tracer = new JsonTracer())
+			{
+				using (IScope Scope1 = Tracer.BuildSpan("bogusOp1").WithTag("age", 94).StartActive()) { }
+				using (IScope Scope2 = Tracer.BuildSpan("bogusOp1").WithTag("foo", "bar").StartActive()) { }
 
-			using (IScope Scope1 = Tracer.BuildSpan("bogusOp1").WithTag("age", 94).StartActive()) { }
-			using (IScope Scope2 = Tracer.BuildSpan("bogusOp1").WithTag("foo", "bar").StartActive()) { }
+				List<JsonTracerSpan> Spans = Tracer.GetFinishedSpans();
 
-			List<JsonTracerSpan> Spans = Tracer.GetFinishedSpans();
-			
-			Assert.AreEqual(2, Spans.Count);
-			Assert.AreEqual("bogusOp1", Spans[0].OperationName);
+				Assert.AreEqual(2, Spans.Count);
+				Assert.AreEqual("bogusOp1", Spans[0].OperationName);
+			}
 		}
 
 		[TestMethod]
 		public void SerializeToJson()
 		{
-			JsonTracer Tracer = new JsonTracer();
-
-			using (IScope Scope1 = Tracer.BuildSpan("bogusOp1").WithTag("age", 94).WithTag("foo", "bar").StartActive()) { }
-			using (IScope Scope2 = Tracer.BuildSpan("bogusOp2").WithTag("Service", "myService").WithTag("Resource", "myResource").StartActive()) { }
+			using (JsonTracer Tracer = new JsonTracer())
+			{
+				using (IScope Scope1 = Tracer.BuildSpan("bogusOp1").WithTag("age", 94).WithTag("foo", "bar").StartActive()) { }
+				using (IScope Scope2 = Tracer.BuildSpan("bogusOp2").WithTag("Service", "myService").WithTag("Resource", "myResource").StartActive()) { }
 
-			using StringWriter Sw = new StringWriter();
-			using JsonWriter Jw = new JsonWriter(Sw);
-			Tracer.GetFinishedSpansAsJson(Jw);
-			byte[] Data = Encoding.ASCII.GetBytes(Sw.ToString());
+				using StringWriter Sw = new StringWriter();
+				using JsonWriter Jw = new JsonWriter(Sw);
+				Tracer.GetFinishedSpansAsJson(Jw);
+				byte[] Data = Encoding.ASCII.GetBytes(Sw.ToString());
 
-			using (JsonDocument Document = JsonDocument.Parse(Data.AsMemory()))
-			{
-				JsonElement Root = Document.RootElement;
-				JsonElement SpansElement = Root.GetProperty("Spans");
-				List<JsonElement> SpansList = SpansElement.EnumerateArray().ToList();
-				Console.WriteLine(SpansList[0]);
-				Assert.AreEqual("bogusOp1", SpansList[0].GetProperty("Name").GetString());
-				Assert.IsNotNull(SpansList[0].GetProperty("StartTime").GetDateTime());
-				Assert.IsNotNull(SpansList[0].GetProperty("FinishTime").GetDateTime());
-				Assert.AreEqual("myService", SpansList[1].GetProperty("Service").GetString());
-				Assert.AreEqual("myResource", SpansList[1].GetProperty("Resource").GetString());
+				using (JsonDocument Document = JsonDocument.Parse(Data.AsMemory()))
+				{
+					JsonElement Root = Document.RootElement;
+					JsonElement SpansElement = Root.GetProperty("Spans");
+					List<JsonElement> SpansList = SpansElement.EnumerateArray().ToList();
+					Console.WriteLine(SpansList[0]);
+					Assert.AreEqual("bogusOp1", SpansList[0].GetProperty("Name").GetString());
+					Assert.IsNotNull(SpansList[0].GetProperty("StartTime").GetDateTime());
+					Assert.IsNotNull(SpansList[0].GetProperty("FinishTime").GetDateTime());
+					Assert.AreEqual("myService", SpansList[1].GetProperty("Service").GetString());
+					Assert.AreEqual("myResource", SpansList[1].GetProperty("Resource").GetString());
+				}
 			}
 		}
 
 		[TestMethod]
 		public void SpansWithParallelFor()
 		{
-			JsonTracer Tracer = new JsonTracer();
-
-			using (IScope ParentScope1 = Tracer.BuildSpan("Parent").StartActive()) 
+			using (JsonTracer Tracer = new JsonTracer())
 			{
-				Parallel.For(0, 100, Index =>
+				using (IScope ParentScope = Tracer.BuildSpan("Parent").StartActive())
 				{
-					using (IScope ChildScope = Tracer.BuildSpan("Child").StartActive()) { }
-				});
+					Parallel.For(0, 100, Index =>
+					{
+						using (IScope ChildScope = Tracer.BuildSpan("Child").StartActive()) { }
+					});
+				}
 			}
 		}
 
@@ -103,11 +106,12 @@ namespace EpicGames.Core.Tests
 		[TestMethod]
 		public void SpansWithAsync()
 		{
-			JsonTracer Tracer = new JsonTracer();
-
-			using (IScope ParentScope1 = Tracer.BuildSpan("Parent").StartActive())
+			using (JsonTracer Tracer = new JsonTracer())
 			{
-				 AsyncSpanTest(Tracer).Wait();
+				using (IScope ParentScope = Tracer.BuildSpan("Parent").StartActive())
+				{
+					AsyncSpanTest(Tracer).Wait();
+				}
 			}
 		}
 	}
diff --git a/Engine/Source/Programs/Shared/EpicGames.Core/JsonTracer.cs b/Engine/Source/Programs/Shared/EpicGames.Core/JsonTracer.cs
--- a/Engine/Source/Programs/Shared/EpicGames.Core/JsonTracer.cs
+++ b/Engine/Source/Programs/Shared/EpicGames.Core/JsonTracer.cs
@@ -399,7 +399,7 @@ namespace EpicGames.Core
 		}
 	}
 
-	public class JsonTracer : ITracer
+	public class JsonTracer : ITracer, IDisposable
 	{
 		private class LocalEventListener : EventListener
 		{
@@ -462,6 +462,11 @@ namespace EpicGames.Core
 			_eventListener.EventWritten += EventListenerEventWritten;
 		}
 
+		public void Dispose()
+		{
+			_eventListener.Dispose();
+		}
+
 		private static int GetEventIntPayload(EventWrittenEventArgs e, string payloadName, ref int? payloadIndex)
 		{
 			if (!payloadIndex.HasValue)
diff --git a/Engine/Source/Programs/UnrealBuildTool/UnrealBuildTool.cs b/Engine/Source/Programs/UnrealBuildTool/UnrealBuildTool.cs
--- a/Engine/Source/Programs/UnrealBuildTool/UnrealBuildTool.cs
+++ b/Engine/Source/Programs/UnrealBuildTool/UnrealBuildTool.cs
@@ -716,6 +716,7 @@ namespace UnrealBuildTool
 				if (Tracer != null)
 				{
 					Tracer.Flush();
+					Tracer.Dispose();
 				}
 
 				// Delete the ubt run file
