Added more tests to the JsonTracer unit tests.
Fixed some potential issues where the parent or span id wasn't what was expected.

diff --git a/Engine/Source/Programs/Shared/EpicGames.Core.Tests/JsonTracerTests.cs b/Engine/Source/Programs/Shared/EpicGames.Core.Tests/JsonTracerTests.cs
--- a/Engine/Source/Programs/Shared/EpicGames.Core.Tests/JsonTracerTests.cs
+++ b/Engine/Source/Programs/Shared/EpicGames.Core.Tests/JsonTracerTests.cs
@@ -79,13 +79,20 @@ namespace EpicGames.Core.Tests
 				{
 					Parallel.For(0, 100, Index =>
 					{
-						using (IScope ChildScope = Tracer.BuildSpan("Child").StartActive()) { }
+						using (IScope ChildScope = Tracer.BuildSpan("Child").StartActive())
+						{
+							Assert.AreEqual(Tracer.ActiveSpan, ChildScope.Span);
+							Assert.AreEqual(((JsonTracerSpan)ChildScope.Span).ParentId, ((JsonTracerSpan)ParentScope.Span).Context.SpanId);
+						}
+						Assert.AreEqual(Tracer.ActiveSpan, ParentScope.Span);
 					});
+					Assert.AreEqual(Tracer.ActiveSpan, ParentScope.Span);
 				}
+				Assert.AreEqual(Tracer.ActiveSpan, null);
 			}
 		}
 
-		private static async Task<bool> AsyncSpanTest(JsonTracer Tracer)
+		private static async Task<bool> AsyncSpanTest(JsonTracer Tracer, IScope ParentScope)
 		{
 			for (int Index = 0; Index < 100; Index++)
 			{
@@ -95,10 +102,15 @@ namespace EpicGames.Core.Tests
 					{
 						using (IScope TaskChildScope = Tracer.BuildSpan("TaskChild").StartActive())
 						{
+							Assert.AreEqual(Tracer.ActiveSpan, TaskChildScope.Span);
+							Assert.AreEqual(((JsonTracerSpan)TaskChildScope.Span).ParentId, ((JsonTracerSpan)AsyncChildScope.Span).Context.SpanId);
 							return true;
 						}
 					});
+					Assert.AreEqual(Tracer.ActiveSpan, AsyncChildScope.Span);
+					Assert.AreEqual(((JsonTracerSpan)AsyncChildScope.Span).ParentId, ((JsonTracerSpan)ParentScope.Span).Context.SpanId);
 				}
+				Assert.AreEqual(Tracer.ActiveSpan, ParentScope.Span);
 			}
 			return true;
 		}
@@ -110,8 +122,12 @@ namespace EpicGames.Core.Tests
 			{
 				using (IScope ParentScope = Tracer.BuildSpan("Parent").StartActive())
 				{
-					AsyncSpanTest(Tracer).Wait();
+					Assert.AreEqual(Tracer.ActiveSpan, ParentScope.Span);
+					AsyncSpanTest(Tracer, ParentScope).Wait();
+					Assert.AreEqual(Tracer.ActiveSpan, ParentScope.Span);
 				}
+
+				Assert.AreEqual(Tracer.ActiveSpan, null);
 			}
 		}
 	}
diff --git a/Engine/Source/Programs/Shared/EpicGames.Core/JsonTracer.cs b/Engine/Source/Programs/Shared/EpicGames.Core/JsonTracer.cs
--- a/Engine/Source/Programs/Shared/EpicGames.Core/JsonTracer.cs
+++ b/Engine/Source/Programs/Shared/EpicGames.Core/JsonTracer.cs
@@ -427,7 +427,7 @@ namespace EpicGames.Core
 		{ 
 			get
 			{
-				if (ActiveSpans.TryGetValue(Environment.CurrentManagedThreadId, out Stack<ISpan>? currentThreadActiveSpan) && currentThreadActiveSpan.Any())
+				if (PerThreadActiveSpans.TryGetValue(Environment.CurrentManagedThreadId, out Stack<ISpan>? currentThreadActiveSpan) && currentThreadActiveSpan.Any())
 				{
 					return currentThreadActiveSpan.Peek();
 				}	
@@ -440,9 +440,9 @@ namespace EpicGames.Core
 		private readonly DirectoryReference? _telemetryDir;
 		private readonly LocalEventListener _eventListener = new LocalEventListener();
 
-		private readonly ConcurrentDictionary<int, Stack<ISpan>> ActiveSpans = new ConcurrentDictionary<int, Stack<ISpan>>();
-		private readonly ConcurrentDictionary<int, ISpan> ForkIdToActiveSpans = new ConcurrentDictionary<int, ISpan>();
-		private readonly ConcurrentDictionary<int, ISpan> TaskIdToActiveSpans = new ConcurrentDictionary<int, ISpan>();
+		private readonly ConcurrentDictionary<int, Stack<ISpan>> PerThreadActiveSpans = new ConcurrentDictionary<int, Stack<ISpan>>();
+		private readonly ConcurrentDictionary<int, Stack<ISpan>> ForkIdToActiveSpans = new ConcurrentDictionary<int, Stack<ISpan>>();
+		private readonly ConcurrentDictionary<int, Stack<ISpan>> TaskIdToActiveSpans = new ConcurrentDictionary<int, Stack<ISpan>>();
 
 		// These could possible be set to 2, which is what MS sets them to internally, if this becomes a performance issue.
 		private static int? ParallelLoopBegin_ForkJoinContextIDIndex;
@@ -492,59 +492,69 @@ namespace EpicGames.Core
 			return intValue.Value;
 		}
 
-		private void BeginEvent(EventWrittenEventArgs e, string PayloadName, ref int? payloadIndex, ConcurrentDictionary<int, ISpan> activeSpansDict)
+		private void BeginEvent(EventWrittenEventArgs e, string PayloadName, ref int? payloadIndex, ConcurrentDictionary<int, Stack<ISpan>> eventActiveSpansDict)
 		{
 			int value = GetEventIntPayload(e, PayloadName, ref payloadIndex);
 
-			ISpan? currentActiveSpan = null;
-			if (ActiveSpans.TryGetValue(Environment.CurrentManagedThreadId, out Stack<ISpan>? currentThreadActiveSpan) && currentThreadActiveSpan.Any())
+			Stack<ISpan>? currentActiveSpanStack = null;
+			if (PerThreadActiveSpans.TryGetValue(Environment.CurrentManagedThreadId, out Stack<ISpan>? currentThreadActiveSpan) && currentThreadActiveSpan.Any())
 			{
-				currentActiveSpan = currentThreadActiveSpan.Peek();
+				currentActiveSpanStack = new Stack<ISpan>(currentThreadActiveSpan.Reverse());
 			}
 
 			// the current active span on this thread needs to be attached to this new event
-			if (currentActiveSpan != null)
+			if (currentActiveSpanStack != null)
 			{
-				if (!activeSpansDict.TryAdd(value, currentActiveSpan))
+				if (!eventActiveSpansDict.TryAdd(value, currentActiveSpanStack))
 				{
 					throw new InvalidOperationException($"Event id '{value}' was already added.");
 				}
 			}
 		}
 
-		private static void EndEvent(EventWrittenEventArgs e, string payloadName, ref int? payloadIndex, ConcurrentDictionary<int, ISpan> activeSpansDict)
+		private static void EndEvent(EventWrittenEventArgs e, string payloadName, ref int? payloadIndex, ConcurrentDictionary<int, Stack<ISpan>> eventActiveSpansDict)
 		{
 			int value = GetEventIntPayload(e, payloadName, ref payloadIndex);
-			if (activeSpansDict.ContainsKey(value))
+			if (eventActiveSpansDict.ContainsKey(value))
 			{
-				if (!activeSpansDict.TryRemove(value, out _))
+				if (!eventActiveSpansDict.TryRemove(value, out _))
 				{
 					throw new InvalidOperationException($"Event id '{value}' could not be removed.");
 				}
 			}
 		}
 
-		private void BeginWork(EventWrittenEventArgs e, string payloadName, ref int? payloadIndex, ConcurrentDictionary<int, ISpan> activeSpansDict)
+		private void BeginWork(EventWrittenEventArgs e, string payloadName, ref int? payloadIndex, ConcurrentDictionary<int, Stack<ISpan>> eventActiveSpansDict)
 		{
 			// find the active span for this event and make it the active span for this thread
 			int value = GetEventIntPayload(e, payloadName, ref payloadIndex);
-			if (activeSpansDict.TryGetValue(value, out ISpan? activeSpanForFork))
+			if (eventActiveSpansDict.TryGetValue(value, out Stack<ISpan>? activeSpanForEvent))
 			{
-				ActivateSpan(activeSpanForFork);
+				PerThreadActiveSpans.AddOrUpdate(Environment.CurrentManagedThreadId,
+				(k) =>
+				{
+					return new Stack<ISpan>(activeSpanForEvent.Reverse());
+				},
+				(k, v) =>
+				{
+					return new Stack<ISpan>(activeSpanForEvent.Reverse());
+				}
+				);
 			}
 		}
 
-		private void EndWork(EventWrittenEventArgs e, string payloadName, ref int? payloadIndex, ConcurrentDictionary<int, ISpan> activeSpansDict)
+		private void EndWork(EventWrittenEventArgs e, string payloadName, ref int? payloadIndex, ConcurrentDictionary<int, Stack<ISpan>> eventActiveSpansDict)
 		{
 			int value = GetEventIntPayload(e, payloadName, ref payloadIndex);
-			if (activeSpansDict.TryGetValue(value, out ISpan? activeSpanForEvent))
+			if (eventActiveSpansDict.TryGetValue(value, out Stack<ISpan>? activeSpanForEventStack))
 			{
-				if (ActiveSpans.TryGetValue(Environment.CurrentManagedThreadId, out Stack<ISpan>? spanStack))
+				if (PerThreadActiveSpans.TryGetValue(Environment.CurrentManagedThreadId, out Stack<ISpan>? activeSpanForThreadStack))
 				{
-					ISpan threadActiveSpan = spanStack.Pop();
-					if (activeSpanForEvent != threadActiveSpan)
+					ISpan eventActiveSpan = activeSpanForEventStack.Peek();
+					ISpan threadActiveSpan = activeSpanForThreadStack.Peek();
+					if (eventActiveSpan.Context.SpanId != threadActiveSpan.Context.SpanId)
 					{
-						throw new InvalidOperationException($"Span should be equal. Event span op: '{((JsonTracerSpan)activeSpanForEvent).OperationName}' Thread span op: '{((JsonTracerSpan)threadActiveSpan).OperationName}'.");
+						throw new InvalidOperationException($"Span should be equal. Event span op: '{((JsonTracerSpan)eventActiveSpan).OperationName}' | Event span id: {eventActiveSpan.Context.SpanId} | Thread span op: '{((JsonTracerSpan)threadActiveSpan).OperationName} | | Event span id: {threadActiveSpan.Context.SpanId}'.");
 					}
 				}
 			}
@@ -696,7 +706,7 @@ namespace EpicGames.Core
 
 		internal void ActivateSpan(ISpan span)
 		{
-			ActiveSpans.AddOrUpdate(Environment.CurrentManagedThreadId, 
+			PerThreadActiveSpans.AddOrUpdate(Environment.CurrentManagedThreadId, 
 			(k) =>
 			{
 				Stack<ISpan> stack = new Stack<ISpan>();
@@ -713,12 +723,12 @@ namespace EpicGames.Core
 
 		internal void DeactivateActiveSpan(ISpan span)
 		{
-			if (ActiveSpans.TryGetValue(Environment.CurrentManagedThreadId, out Stack<ISpan>? spanStack))
+			// It's possible for other threads to finish spans
+			foreach(Stack<ISpan> threadSpanStack in PerThreadActiveSpans.Values)
 			{
-				// It's possible for other threads to finish spans
-				if (spanStack != null && spanStack.Any() && spanStack.Peek() == span) 
+				if (threadSpanStack != null && threadSpanStack.Any() && threadSpanStack.Peek() == span)
 				{
-					spanStack.Pop();
+					threadSpanStack.Pop();
 				}
 			}
 		}
