Horde: Take change views into account when enumerating filters for cached Perforce commits.

diff --git a/Engine/Source/Programs/Horde/Horde.Server/Perforce/PerforceServiceCache.cs b/Engine/Source/Programs/Horde/Horde.Server/Perforce/PerforceServiceCache.cs
--- a/Engine/Source/Programs/Horde/Horde.Server/Perforce/PerforceServiceCache.cs
+++ b/Engine/Source/Programs/Horde/Horde.Server/Perforce/PerforceServiceCache.cs
@@ -62,12 +62,15 @@ namespace Horde.Server.Perforce
 		{
 			public StreamConfig StreamConfig { get; set; }
 			public PerforceViewMap View { get; }
+			public PerforceChangeView ChangeView { get; }
+
 			public List<CommitTagInfo> CommitTags { get; set; } = new List<CommitTagInfo>();
 
-			public StreamInfo(StreamConfig streamConfig, PerforceViewMap view)
+			public StreamInfo(StreamConfig streamConfig, PerforceViewMap view, PerforceChangeView changeView)
 			{
 				StreamConfig = streamConfig;
 				View = view;
+				ChangeView = changeView;
 
 				foreach (CommitTagConfig commitTagConfig in streamConfig.GetAllCommitTags())
 				{
@@ -361,12 +364,15 @@ namespace Horde.Server.Perforce
 		{
 			using (IPooledPerforceConnection perforce = await ConnectAsync(clusterName, null, cancellationToken))
 			{
+				InfoRecord serverInfo = await perforce.GetInfoAsync(cancellationToken);
+
 				List<StreamInfo> streamInfoList = new List<StreamInfo>();
 				foreach (StreamConfig stream in streams)
 				{
 					StreamRecord record = await perforce.GetStreamAsync(stream.Name, true, cancellationToken);
 					PerforceViewMap view = PerforceViewMap.Parse(record.View);
-					streamInfoList.Add(new StreamInfo(stream, view));
+					PerforceChangeView changeView = PerforceChangeView.Parse(record.ChangeView, !serverInfo.IsCaseSensitive);
+					streamInfoList.Add(new StreamInfo(stream, view, changeView));
 				}
 				return streamInfoList;
 			}
@@ -443,7 +449,10 @@ namespace Horde.Server.Perforce
 							string relativePath;
 							if (streamInfo.View.TryMapFile(describeFile.DepotFile, info.PathComparison, out relativePath))
 							{
-								files.Add(relativePath);
+								if (streamInfo.ChangeView.IsVisible(describeFile.DepotFile, describeRecord.Number))
+								{
+									files.Add(relativePath);
+								}
 							}
 						}
 
diff --git a/Engine/Source/Programs/Shared/EpicGames.Perforce.Tests/PerforceChangeViewTest.cs b/Engine/Source/Programs/Shared/EpicGames.Perforce.Tests/PerforceChangeViewTest.cs
new file mode 100644
--- /dev/null
+++ b/Engine/Source/Programs/Shared/EpicGames.Perforce.Tests/PerforceChangeViewTest.cs
@@ -0,0 +1,53 @@
+// Copyright Epic Games, Inc. All Rights Reserved.
+
+using Microsoft.VisualStudio.TestTools.UnitTesting;
+
+namespace EpicGames.Perforce.Tests
+{
+	[TestClass]
+	public class PerforceChangeViewTest
+	{
+		[TestMethod]
+		public void TestImportAtChange()
+		{
+			string[] viewLines =
+			{
+				"//UE5/Main/Testing/...@123",
+				"//UE5/Main/Testing/Foo/...@456",
+				"//UE5/Main/Testing/Foo/Bar/...@100",
+			};
+
+			PerforceChangeView changeView = PerforceChangeView.Parse(viewLines, true);
+
+			Assert.IsTrue(changeView.IsVisible("//UE5/Main/File.txt", 1));
+			Assert.IsTrue(changeView.IsVisible("//UE5/Main/Other/File.txt", 1));
+			Assert.IsTrue(changeView.IsVisible("//UE5/Main/Testing/File.txt", 1));
+			Assert.IsTrue(changeView.IsVisible("//UE5/Main/Testing/Foo/File.txt", 1));
+			Assert.IsTrue(changeView.IsVisible("//UE5/Main/Testing/Foo/Bar/File.txt", 1));
+
+			Assert.IsTrue(changeView.IsVisible("//UE5/Main/File.txt", 100));
+			Assert.IsTrue(changeView.IsVisible("//UE5/Main/Other/File.txt", 100));
+			Assert.IsTrue(changeView.IsVisible("//UE5/Main/Testing/File.txt", 100));
+			Assert.IsTrue(changeView.IsVisible("//UE5/Main/Testing/Foo/File.txt", 100));
+			Assert.IsTrue(changeView.IsVisible("//UE5/Main/Testing/Foo/Bar/File.txt", 100));
+
+			Assert.IsTrue(changeView.IsVisible("//UE5/Main/File.txt", 101));
+			Assert.IsTrue(changeView.IsVisible("//UE5/Main/Other/File.txt", 101));
+			Assert.IsTrue(changeView.IsVisible("//UE5/Main/Testing/File.txt", 101));
+			Assert.IsTrue(changeView.IsVisible("//UE5/Main/Testing/Foo/File.txt", 101));
+			Assert.IsFalse(changeView.IsVisible("//UE5/Main/Testing/Foo/Bar/File.txt", 101));
+
+			Assert.IsTrue(changeView.IsVisible("//UE5/Main/File.txt", 124));
+			Assert.IsTrue(changeView.IsVisible("//UE5/Main/Other/File.txt", 124));
+			Assert.IsFalse(changeView.IsVisible("//UE5/Main/Testing/File.txt", 124));
+			Assert.IsTrue(changeView.IsVisible("//UE5/Main/Testing/Foo/File.txt", 124));
+			Assert.IsFalse(changeView.IsVisible("//UE5/Main/Testing/Foo/Bar/File.txt", 124));
+
+			Assert.IsTrue(changeView.IsVisible("//UE5/Main/File.txt", 457));
+			Assert.IsTrue(changeView.IsVisible("//UE5/Main/Other/File.txt", 457));
+			Assert.IsFalse(changeView.IsVisible("//UE5/Main/Testing/File.txt", 457));
+			Assert.IsFalse(changeView.IsVisible("//UE5/Main/Testing/Foo/File.txt", 457));
+			Assert.IsFalse(changeView.IsVisible("//UE5/Main/Testing/Foo/Bar/File.txt", 457));
+		}
+	}
+}
diff --git a/Engine/Source/Programs/Shared/EpicGames.Perforce/PerforceChangeView.cs b/Engine/Source/Programs/Shared/EpicGames.Perforce/PerforceChangeView.cs
new file mode 100644
--- /dev/null
+++ b/Engine/Source/Programs/Shared/EpicGames.Perforce/PerforceChangeView.cs
@@ -0,0 +1,97 @@
+// Copyright Epic Games, Inc. All Rights Reserved.
+
+using System;
+using System.Collections.Generic;
+using System.Globalization;
+using System.Text.RegularExpressions;
+
+namespace EpicGames.Perforce
+{
+	/// <summary>
+	/// Represents a change view from a clientspec or stream
+	/// </summary>
+	public class PerforceChangeView
+	{
+		/// <summary>
+		/// List of entries for this change view
+		/// </summary>
+		public List<PerforceChangeViewEntry> Entries { get; } = new List<PerforceChangeViewEntry>();
+
+		/// <summary>
+		/// Determines if a file revision is visible 
+		/// </summary>
+		/// <param name="depotFile">Path to the depot file</param>
+		/// <param name="change">Change number of the file</param>
+		public bool IsVisible(string depotFile, int change)
+		{
+			bool visible = true;
+			foreach (PerforceChangeViewEntry entry in Entries)
+			{
+				if (entry.Matches(depotFile))
+				{
+					visible = change <= entry.Change;
+				}
+			}
+			return visible;
+		}
+
+		/// <summary>
+		/// Parse a change view from a specification
+		/// </summary>
+		/// <param name="lines">Lines for the change view</param>
+		/// <param name="ignoreCase">Whether this will be evaluated in the context of a case-insensitive server</param>
+		public static PerforceChangeView Parse(IEnumerable<string> lines, bool ignoreCase)
+		{
+			PerforceChangeView changeView = new PerforceChangeView();
+			foreach (string line in lines)
+			{
+				if (!String.IsNullOrWhiteSpace(line))
+				{
+					changeView.Entries.Add(new PerforceChangeViewEntry(line, ignoreCase));
+				}
+			}
+			return changeView;
+		}
+	}
+
+	/// <summary>
+	/// Entry for a change view
+	/// </summary>
+	public class PerforceChangeViewEntry
+	{
+		/// <summary>
+		/// Pattern to match
+		/// </summary>
+		public string Path { get; }
+
+		/// <summary>
+		/// Change to be imported at
+		/// </summary>
+		public long Change { get; }
+
+		readonly Regex _pathRegex;
+
+		/// <summary>
+		/// Constructor
+		/// </summary>
+		public PerforceChangeViewEntry(string pattern, bool ignoreCase)
+		{
+			int atIdx = pattern.LastIndexOf('@');
+			Path = pattern.Substring(0, atIdx);
+			Change = Int64.Parse(pattern.AsSpan(atIdx + 1), NumberStyles.None);
+
+			string regexPattern = Regex.Escape(Path);
+			regexPattern = regexPattern.Replace(@"\?", ".", StringComparison.Ordinal);
+			regexPattern = regexPattern.Replace(@"\.\.\.", ".*", StringComparison.Ordinal);
+
+			_pathRegex = new Regex($"^{regexPattern}$", ignoreCase? RegexOptions.IgnoreCase : RegexOptions.None);
+		}
+
+		/// <summary>
+		/// Determine
+		/// </summary>
+		/// <param name="depotPath"></param>
+		/// <returns></returns>
+		public bool Matches(string depotPath) => _pathRegex.IsMatch(depotPath);
+	}
+}
diff --git a/Engine/Source/Programs/Shared/EpicGames.Perforce/Records/InfoRecord.cs b/Engine/Source/Programs/Shared/EpicGames.Perforce/Records/InfoRecord.cs
--- a/Engine/Source/Programs/Shared/EpicGames.Perforce/Records/InfoRecord.cs
+++ b/Engine/Source/Programs/Shared/EpicGames.Perforce/Records/InfoRecord.cs
@@ -58,20 +58,25 @@ namespace EpicGames.Perforce
 		[PerforceTag("caseHandling", Optional = true)]
 		public string? CaseHandling { get; set; }
 
+		/// <summary>
+		/// Whether the server is case sensitive
+		/// </summary>
+		public bool IsCaseSensitive => !String.Equals(CaseHandling, "insensitive", StringComparison.OrdinalIgnoreCase);
+
 		/// <summary>
 		/// How to compare paths on this server
 		/// </summary>
-		public StringComparer PathComparer => String.Equals(CaseHandling, "insensitive", StringComparison.OrdinalIgnoreCase) ? StringComparer.OrdinalIgnoreCase : StringComparer.Ordinal;
+		public StringComparer PathComparer => IsCaseSensitive ? StringComparer.Ordinal : StringComparer.OrdinalIgnoreCase;
 
 		/// <summary>
 		/// How to compare paths on this server
 		/// </summary>
-		public StringComparison PathComparison => String.Equals(CaseHandling, "insensitive", StringComparison.OrdinalIgnoreCase) ? StringComparison.OrdinalIgnoreCase : StringComparison.Ordinal;
+		public StringComparison PathComparison => IsCaseSensitive ? StringComparison.Ordinal : StringComparison.OrdinalIgnoreCase;
 
 		/// <summary>
 		/// How to compare paths on this server
 		/// </summary>
-		public Utf8StringComparer Utf8PathComparer => String.Equals(CaseHandling, "insensitive", StringComparison.OrdinalIgnoreCase) ? Utf8StringComparer.OrdinalIgnoreCase : Utf8StringComparer.Ordinal;
+		public Utf8StringComparer Utf8PathComparer => IsCaseSensitive ? Utf8StringComparer.Ordinal : Utf8StringComparer.OrdinalIgnoreCase;
 
 		/// <summary>
 		/// List of services provided by this server
diff --git a/Engine/Source/Programs/Shared/EpicGames.Perforce/Records/StreamRecord.cs b/Engine/Source/Programs/Shared/EpicGames.Perforce/Records/StreamRecord.cs
--- a/Engine/Source/Programs/Shared/EpicGames.Perforce/Records/StreamRecord.cs
+++ b/Engine/Source/Programs/Shared/EpicGames.Perforce/Records/StreamRecord.cs
@@ -76,6 +76,12 @@ namespace EpicGames.Perforce
 		[PerforceTag("View", Optional = true)]
 		public List<string> View { get; } = new List<string>();
 
+		/// <summary>
+		/// Specific changelists to import
+		/// </summary>
+		[PerforceTag("ChangeView", Optional = true)]
+		public List<string> ChangeView { get; } = new List<string>();
+
 		/// <summary>
 		/// Private constructor for serialization
 		/// </summary>
