Sealed JsonTracer to fix static analysis issue.

diff --git a/Engine/Source/Programs/Shared/EpicGames.Core/JsonTracer.cs b/Engine/Source/Programs/Shared/EpicGames.Core/JsonTracer.cs
--- a/Engine/Source/Programs/Shared/EpicGames.Core/JsonTracer.cs
+++ b/Engine/Source/Programs/Shared/EpicGames.Core/JsonTracer.cs
@@ -399,7 +399,7 @@ namespace EpicGames.Core
 		}
 	}
 
-	public class JsonTracer : ITracer, IDisposable
+	public sealed class JsonTracer : ITracer, IDisposable
 	{
 		private class LocalEventListener : EventListener
 		{
