Horde: Fix namespace for ClientComputeSocket to match location on disk.

diff --git a/Engine/Source/Programs/Horde/Horde.Agent/Commands/Compute/ComputeCommand.cs b/Engine/Source/Programs/Horde/Horde.Agent/Commands/Compute/ComputeCommand.cs
--- a/Engine/Source/Programs/Horde/Horde.Agent/Commands/Compute/ComputeCommand.cs
+++ b/Engine/Source/Programs/Horde/Horde.Agent/Commands/Compute/ComputeCommand.cs
@@ -12,6 +12,7 @@ using EpicGames.Core;
 using EpicGames.Horde.Common;
 using EpicGames.Horde.Compute;
 using EpicGames.Horde.Compute.Clients;
+using EpicGames.Horde.Compute.Sockets;
 using EpicGames.Horde.Compute.Transports;
 using Microsoft.Extensions.Caching.Memory;
 using Microsoft.Extensions.DependencyInjection;
diff --git a/Engine/Source/Programs/Horde/Horde.Agent/Leases/Handlers/ComputeHandler.cs b/Engine/Source/Programs/Horde/Horde.Agent/Leases/Handlers/ComputeHandler.cs
--- a/Engine/Source/Programs/Horde/Horde.Agent/Leases/Handlers/ComputeHandler.cs
+++ b/Engine/Source/Programs/Horde/Horde.Agent/Leases/Handlers/ComputeHandler.cs
@@ -6,6 +6,7 @@ using System.Threading;
 using System.Threading.Tasks;
 using EpicGames.Core;
 using EpicGames.Horde.Compute;
+using EpicGames.Horde.Compute.Sockets;
 using EpicGames.Horde.Compute.Transports;
 using Horde.Agent.Services;
 using HordeCommon.Rpc.Tasks;
diff --git a/Engine/Source/Programs/Shared/EpicGames.Horde.Tests/BufferTests.cs b/Engine/Source/Programs/Shared/EpicGames.Horde.Tests/BufferTests.cs
--- a/Engine/Source/Programs/Shared/EpicGames.Horde.Tests/BufferTests.cs
+++ b/Engine/Source/Programs/Shared/EpicGames.Horde.Tests/BufferTests.cs
@@ -11,6 +11,7 @@ using System.Threading;
 using System.Threading.Tasks;
 using EpicGames.Horde.Compute;
 using EpicGames.Horde.Compute.Buffers;
+using EpicGames.Horde.Compute.Sockets;
 using EpicGames.Horde.Compute.Transports;
 using Microsoft.Extensions.Logging.Abstractions;
 using Microsoft.VisualStudio.TestPlatform.CrossPlatEngine;
diff --git a/Engine/Source/Programs/Shared/EpicGames.Horde/Compute/Clients/AgentComputeClient.cs b/Engine/Source/Programs/Shared/EpicGames.Horde/Compute/Clients/AgentComputeClient.cs
--- a/Engine/Source/Programs/Shared/EpicGames.Horde/Compute/Clients/AgentComputeClient.cs
+++ b/Engine/Source/Programs/Shared/EpicGames.Horde/Compute/Clients/AgentComputeClient.cs
@@ -10,6 +10,7 @@ using System.Runtime.InteropServices;
 using System.Threading;
 using System.Threading.Tasks;
 using EpicGames.Core;
+using EpicGames.Horde.Compute.Sockets;
 using EpicGames.Horde.Compute.Transports;
 using Microsoft.Extensions.Logging;
 
diff --git a/Engine/Source/Programs/Shared/EpicGames.Horde/Compute/Clients/LocalComputeClient.cs b/Engine/Source/Programs/Shared/EpicGames.Horde/Compute/Clients/LocalComputeClient.cs
--- a/Engine/Source/Programs/Shared/EpicGames.Horde/Compute/Clients/LocalComputeClient.cs
+++ b/Engine/Source/Programs/Shared/EpicGames.Horde/Compute/Clients/LocalComputeClient.cs
@@ -7,6 +7,7 @@ using System.Net.Sockets;
 using System.Threading;
 using System.Threading.Tasks;
 using EpicGames.Core;
+using EpicGames.Horde.Compute.Sockets;
 using EpicGames.Horde.Compute.Transports;
 using Microsoft.Extensions.Caching.Memory;
 using Microsoft.Extensions.Logging;
diff --git a/Engine/Source/Programs/Shared/EpicGames.Horde/Compute/Clients/ServerComputeClient.cs b/Engine/Source/Programs/Shared/EpicGames.Horde/Compute/Clients/ServerComputeClient.cs
--- a/Engine/Source/Programs/Shared/EpicGames.Horde/Compute/Clients/ServerComputeClient.cs
+++ b/Engine/Source/Programs/Shared/EpicGames.Horde/Compute/Clients/ServerComputeClient.cs
@@ -1,6 +1,7 @@
 // Copyright Epic Games, Inc. All Rights Reserved.
 
 using EpicGames.Core;
+using EpicGames.Horde.Compute.Sockets;
 using EpicGames.Horde.Compute.Transports;
 using Microsoft.Extensions.Logging;
 using System;
diff --git a/Engine/Source/Programs/Shared/EpicGames.Horde/Compute/Sockets/ClientComputeSocket.cs b/Engine/Source/Programs/Shared/EpicGames.Horde/Compute/Sockets/ClientComputeSocket.cs
--- a/Engine/Source/Programs/Shared/EpicGames.Horde/Compute/Sockets/ClientComputeSocket.cs
+++ b/Engine/Source/Programs/Shared/EpicGames.Horde/Compute/Sockets/ClientComputeSocket.cs
@@ -4,7 +4,7 @@ using System;
 using EpicGames.Horde.Compute.Buffers;
 using Microsoft.Extensions.Logging;
 
-namespace EpicGames.Horde.Compute
+namespace EpicGames.Horde.Compute.Sockets
 {
 	/// <summary>
 	/// Manages a set of readers and writers to buffers across a transport layer
