Added an unique GUID to the ID to make it unique between processes/runs. This is for Horde.

diff --git a/Engine/Source/Programs/Shared/EpicGames.Core/JsonTracer.cs b/Engine/Source/Programs/Shared/EpicGames.Core/JsonTracer.cs
--- a/Engine/Source/Programs/Shared/EpicGames.Core/JsonTracer.cs
+++ b/Engine/Source/Programs/Shared/EpicGames.Core/JsonTracer.cs
@@ -1,7 +1,6 @@
 // Copyright Epic Games, Inc. All Rights Reserved.
 
 using System;
-using System.Collections;
 using System.Collections.Concurrent;
 using System.Collections.Generic;
 using System.Diagnostics.Tracing;
@@ -44,12 +43,17 @@ namespace EpicGames.Core
 		/// </summary>
 		private static long s_nextIdCounter = 0;
 
+		/// <summary>
+		/// Unique GUID for this process so IDs are always unique
+		/// </summary>
+		private static readonly string s_processGuidStr = Guid.NewGuid().ToString();
+
 		/// <summary>
 		/// A simple-as-possible (consecutive for repeatability) id generator.
 		/// </summary>
 		private static string NextId()
 		{
-			return Interlocked.Increment(ref s_nextIdCounter).ToString(CultureInfo.InvariantCulture);
+			return s_processGuidStr + "." + Interlocked.Increment(ref s_nextIdCounter).ToString(CultureInfo.InvariantCulture);
 		}
 
 		// C# doesn't have "return type covariance" so we use the trick with the explicit interface implementation
@@ -608,7 +612,7 @@ namespace EpicGames.Core
 
 		public static JsonTracer? TryRegisterAsGlobalTracer()
 		{
-			string? telemetryDir = Environment.GetEnvironmentVariable("UE_TELEMETRY_DIR");
+			string? telemetryDir = "D:\\test.log";// Environment.GetEnvironmentVariable("UE_TELEMETRY_DIR");
 			if (telemetryDir != null)
 			{
 				JsonTracer tracer = new JsonTracer(new DirectoryReference(telemetryDir));
@@ -682,10 +686,10 @@ namespace EpicGames.Core
 				}
 				writer.WriteValue("StartTime", span.StartTimestamp.ToString("o", CultureInfo.InvariantCulture));
 				writer.WriteValue("FinishTime", span.FinishTimestamp.ToString("o", CultureInfo.InvariantCulture));
-				writer.WriteValue("SpanId", Int32.Parse(span.Context.SpanId));
+				writer.WriteValue("SpanId", span.Context.SpanId);
 				if (span.ParentId != null)
 				{
-					writer.WriteValue("ParentId", Int32.Parse(span.ParentId));
+					writer.WriteValue("ParentId", span.ParentId);
 				}
 				writer.WriteObjectStart("Metadata");
 				// TODO: Write tags as metadata?
