Removed test code.

diff --git a/Engine/Source/Programs/Shared/EpicGames.Core/JsonTracer.cs b/Engine/Source/Programs/Shared/EpicGames.Core/JsonTracer.cs
--- a/Engine/Source/Programs/Shared/EpicGames.Core/JsonTracer.cs
+++ b/Engine/Source/Programs/Shared/EpicGames.Core/JsonTracer.cs
@@ -612,7 +612,7 @@ namespace EpicGames.Core
 
 		public static JsonTracer? TryRegisterAsGlobalTracer()
 		{
-			string? telemetryDir = "D:\\test.log";// Environment.GetEnvironmentVariable("UE_TELEMETRY_DIR");
+			string? telemetryDir = Environment.GetEnvironmentVariable("UE_TELEMETRY_DIR");
 			if (telemetryDir != null)
 			{
 				JsonTracer tracer = new JsonTracer(new DirectoryReference(telemetryDir));
