Horde: Remote execution.

* Added core support for multiple readers on buffers.
* Move more functionality into ComputeBufferBase implementation.
* Added channel abstraction for socket-like communication.

diff --git a/Engine/Source/Programs/Horde/Horde.Agent/Commands/Compute/CppComputeCommand.cs b/Engine/Source/Programs/Horde/Horde.Agent/Commands/Compute/CppComputeCommand.cs
--- a/Engine/Source/Programs/Horde/Horde.Agent/Commands/Compute/CppComputeCommand.cs
+++ b/Engine/Source/Programs/Horde/Horde.Agent/Commands/Compute/CppComputeCommand.cs
@@ -56,7 +56,7 @@ namespace Horde.Agent.Commands.Compute
 			NodeLocator sandbox = await CreateSandboxAsync(TaskFile, storage, cancellationToken);
 
 			// Open a socket and upload the sandbox
-			await using (IComputeMessageChannel channel = await lease.Socket.CreateMessageChannelAsync(0, 4 * 1024 * 1024, _logger, cancellationToken))
+			using (IComputeMessageChannel channel = lease.Socket.CreateMessageChannel(0, 4 * 1024 * 1024, _logger, cancellationToken))
 			{
 				await channel.UploadFilesAsync("", sandbox, storage, cancellationToken);
 
diff --git a/Engine/Source/Programs/Horde/Horde.Agent/Commands/Compute/DdcComputeCommand.cs b/Engine/Source/Programs/Horde/Horde.Agent/Commands/Compute/DdcComputeCommand.cs
--- a/Engine/Source/Programs/Horde/Horde.Agent/Commands/Compute/DdcComputeCommand.cs
+++ b/Engine/Source/Programs/Horde/Horde.Agent/Commands/Compute/DdcComputeCommand.cs
@@ -31,7 +31,7 @@ namespace Horde.Agent.Commands.Compute
 		protected override async Task<bool> HandleRequestAsync(IComputeLease lease, CancellationToken cancellationToken)
 		{
 			IComputeSocket socket = lease.Socket;
-			await using (IComputeMessageChannel channel = await socket.CreateMessageChannelAsync(0, _logger, cancellationToken))
+			using (IComputeMessageChannel channel = socket.CreateMessageChannel(0, _logger, cancellationToken))
 			{
 				_logger.LogInformation("Sending XOR request");
 				await channel.SendXorRequestAsync(new byte[] { 1, 2, 3, 4, 5 }, (byte)123, cancellationToken);
diff --git a/Engine/Source/Programs/Horde/Horde.Agent/Commands/Compute/XorComputeCommand.cs b/Engine/Source/Programs/Horde/Horde.Agent/Commands/Compute/XorComputeCommand.cs
--- a/Engine/Source/Programs/Horde/Horde.Agent/Commands/Compute/XorComputeCommand.cs
+++ b/Engine/Source/Programs/Horde/Horde.Agent/Commands/Compute/XorComputeCommand.cs
@@ -31,7 +31,7 @@ namespace Horde.Agent.Commands.Compute
 		protected override async Task<bool> HandleRequestAsync(IComputeLease lease, CancellationToken cancellationToken)
 		{
 			IComputeSocket socket = lease.Socket;
-			await using (IComputeMessageChannel mainChannel = await socket.CreateMessageChannelAsync(0, _logger, cancellationToken))
+			using (IComputeMessageChannel mainChannel = socket.CreateMessageChannel(0, _logger, cancellationToken))
 			{
 				_logger.LogInformation("Forking compute channel...");
 
diff --git a/Engine/Source/Programs/Horde/Samples/RemoteClient/RemoteClientApp.cs b/Engine/Source/Programs/Horde/Samples/RemoteClient/RemoteClientApp.cs
--- a/Engine/Source/Programs/Horde/Samples/RemoteClient/RemoteClientApp.cs
+++ b/Engine/Source/Programs/Horde/Samples/RemoteClient/RemoteClientApp.cs
@@ -1,5 +1,6 @@
 // Copyright Epic Games, Inc. All Rights Reserved.
 
+using System.Buffers.Binary;
 using System.Net.Http.Headers;
 using System.Reflection;
 using EpicGames.Core;
@@ -79,7 +80,7 @@ namespace RemoteClient
 		{
 			// Create a message channel on channel id 0. The Horde Agent always listens on this channel for requests.
 			const int ControlChannelId = 0;
-			await using IComputeMessageChannel channel = await lease.Socket.CreateMessageChannelAsync(ControlChannelId, 4 * 1024 * 1024, logger);
+			using IComputeMessageChannel channel = lease.Socket.CreateMessageChannel(ControlChannelId, 4 * 1024 * 1024, logger);
 			await channel.WaitForAttachAsync();
 
 			// Upload the sandbox
@@ -111,16 +112,15 @@ namespace RemoteClient
 		{
 			// Generate data into a buffer attached to channel 1. The remote server will echo them back to us as it receives them, then exit when the channel is complete/closed.
 			const int DataChannelId = 1;
-			using (IComputeBufferWriter writer = await socket.AttachSendBufferAsync(DataChannelId, 1024 * 1024, cancellationToken))
+
+			byte[] buffer = new byte[4];
+			for (int idx = 0; idx < 100; idx++)
 			{
-				for (int idx = 0; idx < 100; idx++)
-				{
-					cancellationToken.ThrowIfCancellationRequested();
-					logger.LogInformation("Writing value: {Value}", idx);
-					writer.GetMemory().Span[0] = (byte)idx;
-					writer.Advance(1);
-					await Task.Delay(1000, cancellationToken);
-				}
+				cancellationToken.ThrowIfCancellationRequested();
+				logger.LogInformation("Writing value: {Value}", idx);
+				BinaryPrimitives.WriteInt32LittleEndian(buffer, idx);
+				await socket.SendAsync(DataChannelId, buffer, cancellationToken);
+				await Task.Delay(1000, cancellationToken);
 			}
 		}
 
diff --git a/Engine/Source/Programs/Horde/Samples/RemoteWorker/RemoteWorkerApp.cs b/Engine/Source/Programs/Horde/Samples/RemoteWorker/RemoteWorkerApp.cs
--- a/Engine/Source/Programs/Horde/Samples/RemoteWorker/RemoteWorkerApp.cs
+++ b/Engine/Source/Programs/Horde/Samples/RemoteWorker/RemoteWorkerApp.cs
@@ -1,5 +1,6 @@
 // Copyright Epic Games, Inc. All Rights Reserved.
 
+using System.Buffers.Binary;
 using EpicGames.Horde.Compute;
 using Microsoft.Extensions.Logging.Abstractions;
 
@@ -9,19 +10,14 @@ namespace RemoteServer
 	{
 		static async Task Main()
 		{
-			await using IComputeSocket socket = ComputeSocket.ConnectAsWorker(NullLogger.Instance);
-			using IComputeBufferReader reader = await socket.AttachRecvBufferAsync(1, 1024 * 1024);
+			using IComputeChannel channel = ComputeChannel.ConnectAsWorker();
 			Console.WriteLine("Connected to client"); // Client will wait for output before sending data on channel 1
 
-			while (!reader.IsComplete)
+			byte[] data = new byte[4];
+			while(await channel.TryReceiveMessageAsync(data))
 			{
-				ReadOnlyMemory<byte> memory = reader.GetMemory();
-				if (memory.Length > 0)
-				{
-					Console.WriteLine("Read value {0}", memory.Span[0]);
-					reader.Advance(1);
-				}
-				await reader.WaitForDataAsync(0);
+				int value = BinaryPrimitives.ReadInt32LittleEndian(data);
+				Console.WriteLine("Read value {0}", value);
 			}
 		}
 	}
diff --git a/Engine/Source/Programs/Shared/EpicGames.Horde.Tests/BufferTests.cs b/Engine/Source/Programs/Shared/EpicGames.Horde.Tests/BufferTests.cs
--- a/Engine/Source/Programs/Shared/EpicGames.Horde.Tests/BufferTests.cs
+++ b/Engine/Source/Programs/Shared/EpicGames.Horde.Tests/BufferTests.cs
@@ -22,6 +22,8 @@ namespace EpicGames.Horde.Tests
 	[TestClass]
 	public class BufferTests
 	{
+		const int ChannelId = 0;
+
 		[TestMethod]
 		public async Task TestPooledBuffer()
 		{
@@ -33,7 +35,7 @@ namespace EpicGames.Horde.Tests
 		{
 			if (RuntimeInformation.IsOSPlatform(OSPlatform.Windows))
 			{
-				await TestProducerConsumerAsync(length => new SharedMemoryBuffer(length), CancellationToken.None);
+				await TestProducerConsumerAsync(length => SharedMemoryBuffer.CreateNew(null, length), CancellationToken.None);
 			}
 		}
 
@@ -44,34 +46,33 @@ namespace EpicGames.Horde.Tests
 			Pipe sourceToTargetPipe = new Pipe();
 			Pipe targetToSourcePipe = new Pipe();
 
-			await using IComputeSocket sourceSocket = new ClientComputeSocket(new PipeTransport(targetToSourcePipe.Reader, sourceToTargetPipe.Writer), NullLogger.Instance);
-			await using IComputeSocket targetSocket = new ClientComputeSocket(new PipeTransport(sourceToTargetPipe.Reader, targetToSourcePipe.Writer), NullLogger.Instance);
+			await using IComputeSocket producerSocket = new ClientComputeSocket(new PipeTransport(targetToSourcePipe.Reader, sourceToTargetPipe.Writer), NullLogger.Instance);
+			await using IComputeSocket consumerSocket = new ClientComputeSocket(new PipeTransport(sourceToTargetPipe.Reader, targetToSourcePipe.Writer), NullLogger.Instance);
 
-			using IComputeBufferWriter sourceWriter = await sourceSocket.AttachSendBufferAsync(0, createBuffer(Length), cancellationToken);
-			using IComputeBufferReader targetReader = await targetSocket.AttachRecvBufferAsync(0, createBuffer(Length), cancellationToken);
+			using IComputeBuffer consumerBuffer = createBuffer(Length);
+			consumerSocket.AttachRecvBuffer(ChannelId, consumerBuffer.Writer);
 
 			byte[] input = RandomNumberGenerator.GetBytes(Length);
-			Task producerTask = RunProducerAsync(sourceWriter, input);
+			Task producerTask = RunProducerAsync(producerSocket, input);
 
 			byte[] output = new byte[Length];
-			await RunConsumerAsync(targetReader, output);
+			await RunConsumerAsync(consumerBuffer.Reader, output);
 
 			await producerTask;
 			Assert.IsTrue(input.SequenceEqual(output));
 		}
 
-		static async Task RunProducerAsync(IComputeBufferWriter writer, ReadOnlyMemory<byte> input)
+		static async Task RunProducerAsync(IComputeSocket socket, ReadOnlyMemory<byte> input)
 		{
 			int offset = 0;
 			while (offset < input.Length)
 			{
 				int length = Math.Min(input.Length - offset, 100);
-				input.Slice(offset, length).CopyTo(writer.GetMemory());
-				writer.Advance(length);
+				await socket.SendAsync(ChannelId, input.Slice(offset, length));
 				await Task.Delay(10);
 				offset += length;
 			}
-			writer.MarkComplete();
+			await socket.MarkCompleteAsync(ChannelId);
 		}
 
 		static async Task RunConsumerAsync(IComputeBufferReader reader, Memory<byte> output)
@@ -80,10 +81,15 @@ namespace EpicGames.Horde.Tests
 			while (!reader.IsComplete)
 			{
 				ReadOnlyMemory<byte> memory = reader.GetMemory();
+				if (memory.Length == 0)
+				{
+					await reader.WaitToReadAsync(0, CancellationToken.None);
+					continue;
+				}
+
 				int length = Math.Min(memory.Length, 7);
 				memory.Slice(0, length).CopyTo(output.Slice(offset));
 				reader.Advance(length);
-				await reader.WaitForDataAsync(memory.Length - length, CancellationToken.None);
 				offset += length;
 			}
 		}
diff --git a/Engine/Source/Programs/Shared/EpicGames.Horde/Compute/Buffers/PooledBuffer.cs b/Engine/Source/Programs/Shared/EpicGames.Horde/Compute/Buffers/PooledBuffer.cs
--- a/Engine/Source/Programs/Shared/EpicGames.Horde/Compute/Buffers/PooledBuffer.cs
+++ b/Engine/Source/Programs/Shared/EpicGames.Horde/Compute/Buffers/PooledBuffer.cs
@@ -4,214 +4,105 @@ using System;
 using System.Buffers;
 using System.Threading;
 using System.Threading.Tasks;
+using EpicGames.Core;
 
 namespace EpicGames.Horde.Compute.Buffers
 {
 	/// <summary>
 	/// In-process buffer used to store compute messages
 	/// </summary>
-	class PooledBuffer : ComputeBufferBase
+	public sealed class PooledBuffer : ComputeBufferBase
 	{
-		readonly IMemoryOwner<byte> _memoryOwner;
-		readonly Memory<byte> _memory;
-
-		enum State
+		class Chunk : ChunkBase, IDisposable
 		{
-			Normal,
-			ReaderWaitingForData, // Once new data is available, set _readerTcs
-			WriterWaitingToFlush, // Once flush has been complete, set _writerTcs
-			Finished
-		}
+			readonly IMemoryOwner<byte> _memoryOwner;
+			ulong _stateValue;
 
-		int _readPosition;
-		int _readLength;
-		int _writePosition;
-		int _state;
-		TaskCompletionSource _readerTcs = new TaskCompletionSource(); // Used when reader is waiting for writer. Reset by the reader.
-		TaskCompletionSource _writerTcs = new TaskCompletionSource(); // Used when writer is waiting for reader. Reset by the writer.
+			/// <inheritdoc/>
+			protected override ref ulong StateValue => ref _stateValue;
 
-		/// <inheritdoc/>
-		public override long Length => _memory.Length;
+			public Chunk(IMemoryOwner<byte> memoryOwner)
+				: base(memoryOwner.Memory)
+			{
+				_memoryOwner = memoryOwner;
+			}
+
+			public void Dispose() => _memoryOwner.Dispose();
+		}
 
-		// Accessor for the value of read length. 
-		int GetInterlockedReadLength() => Interlocked.CompareExchange(ref _readLength, 0, 0);
+		readonly Chunk[] _chunks;
+		readonly AsyncEvent _writerEvent = new AsyncEvent();
+		readonly AsyncEvent _readerEvent = new AsyncEvent();
 
 		/// <summary>
-		/// Creates a local buffer with the given capacity
+		/// Constructor
 		/// </summary>
-		/// <param name="capacity">Capacity of the buffer</param>
+		/// <param name="capacity">Total capacity of the buffer</param>
 		public PooledBuffer(int capacity)
+			: this(2, capacity / 2)
 		{
-			_memoryOwner = MemoryPool<byte>.Shared.Rent(capacity);
-			_memory = _memoryOwner.Memory.Slice(0, capacity);
-		}
-
-		/// <inheritdoc/>
-		protected override void Dispose(bool disposing)
-		{
-			base.Dispose(disposing);
-
-			if (disposing)
-			{
-				_memoryOwner.Dispose();
-			}
 		}
 
-		/// <inheritdoc/>
-		public override Memory<byte> GetMemory(long offset, int length)
+		/// <summary>
+		/// Constructor
+		/// </summary>
+		/// <param name="numChunks">Number of chunks in the buffer</param>
+		/// <param name="chunkLength">Length of each chunk</param>
+		public PooledBuffer(int numChunks, int chunkLength)
+			: this(CreateChunks(numChunks, chunkLength), 1)
 		{
-			if (offset > Int32.MaxValue)
-			{
-				throw new ArgumentException("Offset is out of range", nameof(offset));
-			}
-			return _memory.Slice((int)offset, length);
 		}
 
-		#region Reader
-
-		/// <inheritdoc/>
-		public override bool FinishedReading() => (State)_state == State.Finished && GetInterlockedReadLength() == 0;
-
-		/// <inheritdoc/>
-		public override void AdvanceReadPosition(int size)
+		/// <summary>
+		/// Constructor
+		/// </summary>
+		private PooledBuffer(Chunk[] chunks, int readerIdx)
+			: base(chunks, readerIdx)
 		{
-			_readPosition += size;
-			Interlocked.Add(ref _readLength, -size);
+			_chunks = chunks;
 		}
 
-		/// <inheritdoc/>
-		public override ReadOnlyMemory<byte> GetReadMemory() => _memory.Slice(_readPosition, GetInterlockedReadLength());
-
-		/// <inheritdoc/>
-		public override async ValueTask WaitForDataAsync(int currentLength, CancellationToken cancellationToken)
+		static Chunk[] CreateChunks(int numChunks, int chunkLength)
 		{
-			try
+			Chunk[] chunks = new Chunk[numChunks];
+			for (int chunkIdx = 0; chunkIdx < numChunks; chunkIdx++)
 			{
-				while (GetInterlockedReadLength() == currentLength)
-				{
-					int state = _state;
-					if (state == (int)State.Finished)
-					{
-						break;
-					}
-					else if (state == (int)State.Normal)
-					{
-						TaskCompletionSource readerTcs = _readerTcs;
-						if (readerTcs.Task.IsCompleted)
-						{
-							Interlocked.CompareExchange(ref _readerTcs, new TaskCompletionSource(), readerTcs);
-						}
-						Interlocked.CompareExchange(ref _state, (int)State.ReaderWaitingForData, state);
-					}
-					else if (state == (int)State.ReaderWaitingForData)
-					{
-						await _readerTcs.Task.WaitAsync(cancellationToken);
-					}
-					else if (state == (int)State.WriterWaitingToFlush)
-					{
-						Compact();
-						Interlocked.CompareExchange(ref _state, (int)State.Normal, state);
-						_writerTcs.TrySetResult();
-					}
-					else
-					{
-						throw new NotImplementedException();
-					}
-					cancellationToken.ThrowIfCancellationRequested();
-				}
-			}
-			finally
-			{
-				Interlocked.CompareExchange(ref _state, (int)State.Normal, (int)State.ReaderWaitingForData);
+				Chunk chunk = new Chunk(MemoryPool<byte>.Shared.Rent(chunkLength));
+				chunks[chunkIdx] = chunk;
 			}
+			return chunks;
 		}
 
-		#endregion
-
-		#region Writer
-
 		/// <inheritdoc/>
-		public override void FinishWriting()
+		protected override void Dispose(bool disposing)
 		{
-			for (; ; )
+			base.Dispose(disposing);
+
+			if (disposing)
 			{
-				int state = _state;
-				if (Interlocked.CompareExchange(ref _state, (int)State.Finished, state) == state)
+				for (int idx = 0; idx < _chunks.Length; idx++)
 				{
-					_readerTcs.TrySetResult();
-					break;
+					_chunks[idx].Dispose();
 				}
 			}
 		}
 
 		/// <inheritdoc/>
-		public override void AdvanceWritePosition(int size)
-		{
-			if (_state == (int)State.Finished)
-			{
-				throw new InvalidOperationException("Cannot update write position after marking as complete");
-			}
+		protected override void SetReadEvent(int readerIdx) => _readerEvent.Set();
 
-			_writePosition += size;
-			Interlocked.Add(ref _readLength, size);
-
-			if (Interlocked.CompareExchange(ref _state, (int)State.Normal, (int)State.ReaderWaitingForData) == (int)State.ReaderWaitingForData)
-			{
-				_readerTcs.TrySetResult();
-			}
-		}
+		/// <inheritdoc/>
+		protected override void ResetReadEvent(int readerIdx) => _readerEvent.Reset();
 
 		/// <inheritdoc/>
-		public override Memory<byte> GetWriteMemory()
-		{
-			return _memory.Slice(_writePosition);
-		}
+		protected override Task WaitForReadEvent(int readerIdx, CancellationToken cancellationToken) => _readerEvent.Task.WaitAsync(cancellationToken);
 
 		/// <inheritdoc/>
-		public override async ValueTask FlushAsync(CancellationToken cancellationToken)
-		{
-			try
-			{
-				while (_readPosition > 0)
-				{
-					int state = _state;
-					if (state == (int)State.Finished)
-					{
-						break;
-					}
-					else if (state == (int)State.Normal)
-					{
-						TaskCompletionSource writerTcs = _writerTcs;
-						if (writerTcs.Task.IsCompleted)
-						{
-							Interlocked.CompareExchange(ref _writerTcs, new TaskCompletionSource(), writerTcs);
-						}
-						Interlocked.CompareExchange(ref _state, (int)State.WriterWaitingToFlush, state);
-					}
-					else if (state == (int)State.ReaderWaitingForData)
-					{
-						Compact();
-						break;
-					}
-					else if (state == (int)State.WriterWaitingToFlush)
-					{
-						await _writerTcs.Task.WaitAsync(cancellationToken);
-					}
-				}
-			}
-			finally
-			{
-				Interlocked.CompareExchange(ref _state, (int)State.Normal, (int)State.WriterWaitingToFlush);
-			}
-		}
+		protected override void SetWriteEvent() => _writerEvent.Set();
 
-		void Compact()
-		{
-			_memory.Slice(_readPosition, _readLength).CopyTo(_memory);
-			_readPosition = 0;
-			_writePosition = _readLength;
-		}
+		/// <inheritdoc/>
+		protected override void ResetWriteEvent() => _writerEvent.Reset();
 
-		#endregion
+		/// <inheritdoc/>
+		protected override Task WaitForWriteEvent(CancellationToken cancellationToken) => _writerEvent.Task.WaitAsync(cancellationToken);
 	}
 }
diff --git a/Engine/Source/Programs/Shared/EpicGames.Horde/Compute/Buffers/SharedMemoryBuffer.cs b/Engine/Source/Programs/Shared/EpicGames.Horde/Compute/Buffers/SharedMemoryBuffer.cs
--- a/Engine/Source/Programs/Shared/EpicGames.Horde/Compute/Buffers/SharedMemoryBuffer.cs
+++ b/Engine/Source/Programs/Shared/EpicGames.Horde/Compute/Buffers/SharedMemoryBuffer.cs
@@ -2,15 +2,12 @@
 
 using System;
 using EpicGames.Core;
-using Microsoft.Win32.SafeHandles;
 using System.IO.MemoryMappedFiles;
-using System.Reflection;
 using System.IO;
 using System.Threading.Tasks;
 using System.Threading;
-using System.Linq;
-using System.Globalization;
-using System.ComponentModel;
+using System.Buffers.Binary;
+using System.Runtime.InteropServices;
 
 namespace EpicGames.Horde.Compute.Buffers
 {
@@ -19,190 +16,123 @@ namespace EpicGames.Horde.Compute.Buffers
 	/// </summary>
 	public sealed class SharedMemoryBuffer : ComputeBufferBase
 	{
-		const int HeaderLength = 32;
+		const int HeaderLength = sizeof(int) + sizeof(int); // num chunks, chunk length
 
-		enum State
+		sealed unsafe class Chunk : ChunkBase
 		{
-			Normal,
-			ReaderWaitingForData, // Once new data is available, set _readerTcs
-			WriterWaitingToFlush, // Once flush has been complete, set _writerTcs
-			Finished
-		}
-
-		sealed unsafe class Header
-		{
-			readonly int* _statePtr;
-			readonly long* _readPositionPtr;
-			readonly long* _readLengthPtr;
-			readonly long* _writePositionPtr;
-
-			public Header(MemoryMappedView view)
-			{
-				byte* basePtr = view.GetPointer();
-				_statePtr = (int*)basePtr;
-
-				long* nextPtr = (long*)basePtr + 1;
-				_readPositionPtr = nextPtr++;
-				_readLengthPtr = nextPtr++;
-				_writePositionPtr = nextPtr++;
-			}
-
-			public State State => (State)Interlocked.CompareExchange(ref *_statePtr, 0, 0);
-			public long ReadPosition => Interlocked.Read(ref *_readPositionPtr);
-			public long ReadLength => Interlocked.Read(ref *_readLengthPtr);
-			public long WritePosition => Interlocked.Read(ref *_writePositionPtr);
-
-			public void Compact()
-			{
-				Interlocked.Exchange(ref *_readPositionPtr, 0);
-				Interlocked.Exchange(ref *_writePositionPtr, ReadLength);
-			}
+			readonly ulong* _statePtr;
 
-			public bool UpdateState(State oldState, State newState) => Interlocked.CompareExchange(ref *_statePtr, (int)newState, (int)oldState) == (int)oldState;
+			protected override ref ulong StateValue => ref *_statePtr;
 
-			public void FinishWriting()
+			public Chunk(ulong* statePtr, Memory<byte> memory)
+				: base(memory)
 			{
-				*_statePtr = (int)State.Finished;
-			}
-
-			public long AdvanceReadPosition(long value)
-			{
-				Interlocked.Add(ref *_readPositionPtr, value);
-				return Interlocked.Add(ref *_readLengthPtr, -value);
-			}
-
-			public long AdvanceWritePosition(long value)
-			{
-				Interlocked.Add(ref *_writePositionPtr, value);
-				return Interlocked.Add(ref *_readLengthPtr, value);
+				_statePtr = statePtr;
 			}
 		}
 
-		const int ChunkOffsetPow2 = 10;
-
 		readonly MemoryMappedFile _memoryMappedFile;
 		readonly MemoryMappedViewAccessor _memoryMappedViewAccessor;
 		readonly MemoryMappedView _memoryMappedView;
-		readonly long _length;
-		readonly Header _header;
-		readonly Memory<byte> _memory; // TODO: DEPRECATE
-		readonly Memory<byte>[] _chunks;
 
 		readonly Native.EventHandle _readerEvent;
 		readonly Native.EventHandle _writerEvent;
 
 		/// <inheritdoc/>
-		public override long Length => _length;
+		public string Name { get; }
 
 		/// <summary>
-		/// Creates a new shared memory buffer with the given capacity
+		/// Constructor
 		/// </summary>
-		/// <param name="capacity">Capacity of the buffer</param>
-		public SharedMemoryBuffer(long capacity)
-			: this(MemoryMappedFile.CreateNew(null, capacity, MemoryMappedFileAccess.ReadWrite, MemoryMappedFileOptions.None, HandleInheritability.Inheritable))
+		private unsafe SharedMemoryBuffer(string name, MemoryMappedFile memoryMappedFile, MemoryMappedViewAccessor memoryMappedViewAccessor, MemoryMappedView memoryMappedView, Chunk[] chunks, Native.EventHandle readerEvent, Native.EventHandle writerEvent)
+			: base(chunks, 1)
 		{
-		}
+			Name = name;
 
-		/// <summary>
-		/// 
-		/// </summary>
-		/// <param name="buffer"></param>
-		public SharedMemoryBuffer(MemoryMappedFile buffer)
-			: this(buffer, new Native.EventHandle(HandleInheritability.Inheritable), new Native.EventHandle(HandleInheritability.Inheritable))
-		{
+			_memoryMappedFile = memoryMappedFile;
+			_memoryMappedViewAccessor = memoryMappedViewAccessor;
+			_memoryMappedView = memoryMappedView;
+
+			_readerEvent = readerEvent;
+			_writerEvent = writerEvent;
 		}
 
 		/// <summary>
-		/// Opens a <see cref="SharedMemoryBuffer"/> from handles
+		/// Create a new shared memory buffer
 		/// </summary>
-		public SharedMemoryBuffer(IntPtr memoryHandle, IntPtr readerEventHandle, IntPtr writerEventHandle)
-			: this(OpenMemoryMappedFileFromHandle(memoryHandle), new Native.EventHandle(readerEventHandle, true), new Native.EventHandle(writerEventHandle, true))
+		/// <param name="name">Name of the buffer</param>
+		/// <param name="capacity">Capacity of the buffer</param>
+		public static SharedMemoryBuffer CreateNew(string? name, long capacity)
 		{
+			int numChunks = (int)Math.Max(4, capacity / Int32.MaxValue);
+			return CreateNew(name, numChunks, (int)(capacity / numChunks));
 		}
 
 		/// <summary>
-		/// Creates a shared memory buffer from a memory mapped file
+		/// Create a new shared memory buffer
 		/// </summary>
-		internal SharedMemoryBuffer(MemoryMappedFile memoryMappedFile, Native.EventHandle writtenEvent, Native.EventHandle flushedEvent)
+		/// <param name="name">Name of the buffer</param>
+		/// <param name="numChunks">Number of chunks in the buffer</param>
+		/// <param name="chunkLength">Length of each chunk</param>
+		public static unsafe SharedMemoryBuffer CreateNew(string? name, int numChunks, int chunkLength)
 		{
-			_memoryMappedFile = memoryMappedFile;
-			_memoryMappedViewAccessor = memoryMappedFile.CreateViewAccessor();
-			_memoryMappedView = new MemoryMappedView(_memoryMappedViewAccessor);
+			long capacity = HeaderLength + (numChunks * sizeof(ulong)) + (numChunks * chunkLength);
 
-			_readerEvent = writtenEvent;
-			_writerEvent = flushedEvent;
+			name ??= $"Local\\COMPUTE_{Guid.NewGuid()}";
 
-			_header = new Header(_memoryMappedView);
-			_length = _memoryMappedViewAccessor.Capacity;
+			MemoryMappedFile memoryMappedFile = MemoryMappedFile.CreateNew($"{name}_M", capacity, MemoryMappedFileAccess.ReadWrite, MemoryMappedFileOptions.None, HandleInheritability.Inheritable);
+			MemoryMappedViewAccessor memoryMappedViewAccessor = memoryMappedFile.CreateViewAccessor();
+			MemoryMappedView memoryMappedView = new MemoryMappedView(memoryMappedViewAccessor);
 
-			_memory = _memoryMappedView.GetMemory(HeaderLength, (int)(_length - HeaderLength)); // TODO: REMOVE
+			Span<byte> header = memoryMappedView.GetMemory(0, 8).Span;
+			BinaryPrimitives.WriteInt32LittleEndian(header, numChunks);
+			BinaryPrimitives.WriteInt32LittleEndian(header.Slice(4), chunkLength);
 
-			int chunkCount = (int)((_length + ((1 << ChunkOffsetPow2) - 1)) >> ChunkOffsetPow2);
-			_chunks = new Memory<byte>[chunkCount];
+			Native.EventHandle writerEvent = Native.EventHandle.CreateNew($"{name}_W", EventResetMode.ManualReset, true, HandleInheritability.Inheritable);
+			Native.EventHandle readerEvent = Native.EventHandle.CreateNew($"{name}_R", EventResetMode.ManualReset, true, HandleInheritability.Inheritable);
 
-			for (int chunkIdx = 0; chunkIdx < chunkCount; chunkIdx++)
-			{
-				long offset = (chunkIdx << ChunkOffsetPow2) + HeaderLength;
-				int length = (int)Math.Min(Length - offset, Int32.MaxValue);
-				_chunks[chunkIdx] = _memoryMappedView.GetMemory(offset, length);
-			}
+			Chunk[] chunks = CreateChunks(numChunks, chunkLength, memoryMappedView);
+			return new SharedMemoryBuffer(name, memoryMappedFile, memoryMappedViewAccessor, memoryMappedView, chunks, readerEvent, writerEvent);
 		}
 
-		/// <summary>
-		/// Opens an IpcBuffer from a string passed in from another process
-		/// </summary>
-		/// <param name="handle">Descriptor for the buffer to open</param>
-		static MemoryMappedFile OpenMemoryMappedFileFromHandle(IntPtr handle)
+		static unsafe Chunk[] CreateChunks(int numChunks, int chunkLength, MemoryMappedView memoryMappedView)
 		{
-			MethodInfo? setHandleInfo = typeof(SafeMemoryMappedFileHandle).GetMethod("SetHandle", BindingFlags.Instance | BindingFlags.NonPublic, new[] { typeof(IntPtr) });
-			if (setHandleInfo == null)
+			Chunk[] chunks = new Chunk[numChunks];
+			for (int chunkIdx = 0; chunkIdx < numChunks; chunkIdx++)
 			{
-				throw new InvalidOperationException("Cannot find SetHandle method for SafeMemoryMappedFileHandle");
-			}
+				byte* statePtr = memoryMappedView.GetPointer() + HeaderLength + (chunkIdx * sizeof(ulong));
+				int chunkOffset = HeaderLength + (numChunks * sizeof(ulong)) + (chunkLength * chunkIdx);
 
-			ConstructorInfo? constructorInfo = typeof(MemoryMappedFile).GetConstructor(BindingFlags.Instance | BindingFlags.NonPublic, null, new[] { typeof(SafeMemoryMappedFileHandle) }, null);
-			if (constructorInfo == null)
-			{
-				throw new InvalidOperationException("Cannot find private constructor for memory mapped file");
+				Chunk chunk = new Chunk((ulong*)statePtr, memoryMappedView.GetMemory(chunkOffset, chunkLength));
+				chunks[chunkIdx] = chunk;
 			}
-
-			SafeMemoryMappedFileHandle safeHandle = new SafeMemoryMappedFileHandle();
-			setHandleInfo.Invoke(safeHandle, new object[] { handle });
-
-			return (MemoryMappedFile)constructorInfo.Invoke(new object[] { safeHandle });
-		}
-
-		/// <summary>
-		/// Gets a string that can be used to open the same buffer in another process
-		/// </summary>
-		public void GetHandles(out IntPtr memoryHandle, out IntPtr readerEventHandle, out IntPtr writerEventHandle)
-		{
-			memoryHandle = _memoryMappedFile.SafeMemoryMappedFileHandle.DangerousGetHandle();
-			readerEventHandle = _readerEvent.SafeWaitHandle.DangerousGetHandle();
-			writerEventHandle = _writerEvent.SafeWaitHandle.DangerousGetHandle();
+			return chunks;
 		}
 
 		/// <summary>
-		/// Gets a string that can be used to open the same buffer in another process
+		/// Open an existing buffer by name
 		/// </summary>
-		public void DuplicateHandles(IntPtr targetProcessHandle, out IntPtr targetMemoryHandle, out IntPtr targetReaderEventHandle, out IntPtr targetWriterEventHandle)
+		/// <param name="name">Name of the buffer to open</param>
+		public static SharedMemoryBuffer OpenExisting(string name)
 		{
-			IntPtr sourceProcessHandle = Native.GetCurrentProcess();
-
-			if (!Native.DuplicateHandle(sourceProcessHandle, _memoryMappedFile.SafeMemoryMappedFileHandle.DangerousGetHandle(), targetProcessHandle, out targetMemoryHandle, 0, false, Native.DUPLICATE_SAME_ACCESS))
+			if (!RuntimeInformation.IsOSPlatform(OSPlatform.Windows))
 			{
-				throw new Win32Exception();
+				throw new NotSupportedException();
 			}
 
-			if (!Native.DuplicateHandle(sourceProcessHandle, _readerEvent.SafeWaitHandle.DangerousGetHandle(), targetProcessHandle, out targetReaderEventHandle, 0, false, Native.DUPLICATE_SAME_ACCESS))
-			{
-				throw new Win32Exception();
-			}
+			MemoryMappedFile memoryMappedFile = MemoryMappedFile.OpenExisting($"{name}_M");
+			MemoryMappedViewAccessor memoryMappedViewAccessor = memoryMappedFile.CreateViewAccessor();
+			MemoryMappedView memoryMappedView = new MemoryMappedView(memoryMappedViewAccessor);
 
-			if (!Native.DuplicateHandle(sourceProcessHandle, _writerEvent.SafeWaitHandle.DangerousGetHandle(), targetProcessHandle, out targetWriterEventHandle, 0, false, Native.DUPLICATE_SAME_ACCESS))
-			{
-				throw new Win32Exception();
-			}
+			ReadOnlySpan<byte> header = memoryMappedView.GetMemory(0, HeaderLength).Span;
+			int numChunks = BinaryPrimitives.ReadInt32LittleEndian(header);
+			int chunkLength = BinaryPrimitives.ReadInt32LittleEndian(header.Slice(4));
+
+			Native.EventHandle readerEvent = Native.EventHandle.OpenExisting($"{name}_R");
+			Native.EventHandle writerEvent = Native.EventHandle.OpenExisting($"{name}_W");
+
+			Chunk[] chunks = CreateChunks(numChunks, chunkLength, memoryMappedView);
+			return new SharedMemoryBuffer(name, memoryMappedFile, memoryMappedViewAccessor, memoryMappedView, chunks, readerEvent, writerEvent);
 		}
 
 		/// <inheritdoc/>
@@ -212,8 +142,9 @@ namespace EpicGames.Horde.Compute.Buffers
 
 			if (disposing)
 			{
-				_writerEvent.Dispose();
 				_readerEvent.Dispose();
+				_writerEvent.Dispose();
+
 				_memoryMappedView.Dispose();
 				_memoryMappedViewAccessor.Dispose();
 				_memoryMappedFile.Dispose();
@@ -221,137 +152,21 @@ namespace EpicGames.Horde.Compute.Buffers
 		}
 
 		/// <inheritdoc/>
-		public override Memory<byte> GetMemory(long offset, int length)
-		{
-			int chunkIdx = (int)(offset >> ChunkOffsetPow2);
-			long baseOffset = offset - (chunkIdx << ChunkOffsetPow2);
-			return _chunks[chunkIdx].Slice((int)(offset - baseOffset), length);
-		}
-
-		#region Reader
+		protected override void SetReadEvent(int readerIdx) => _readerEvent.Set();
 
 		/// <inheritdoc/>
-		public override bool FinishedReading() => _header.State == State.Finished && _header.ReadLength == 0;
+		protected override void ResetReadEvent(int readerIdx) => _readerEvent.Reset();
 
 		/// <inheritdoc/>
-		public override void AdvanceReadPosition(int size)
-		{
-			if (_header.AdvanceReadPosition(size) == 0)
-			{
-				_writerEvent.Set();
-			}
-		}
+		protected override Task WaitForReadEvent(int readerIdx, CancellationToken cancellationToken) => _readerEvent.WaitOneAsync(cancellationToken);
 
 		/// <inheritdoc/>
-		public override ReadOnlyMemory<byte> GetReadMemory()
-		{
-			return _memory.Slice((int)_header.ReadPosition, (int)_header.ReadLength);
-		}
+		protected override void SetWriteEvent() => _writerEvent.Set();
 
 		/// <inheritdoc/>
-		public override async ValueTask WaitForDataAsync(int currentLength, CancellationToken cancellationToken)
-		{
-			try
-			{
-				while (_header.ReadLength == currentLength)
-				{
-					State state = _header.State;
-					if (state == State.Finished)
-					{
-						break;
-					}
-					else if (state == State.Normal)
-					{
-						_readerEvent.Reset();
-						_header.UpdateState(state, State.ReaderWaitingForData);
-					}
-					else if (state == State.ReaderWaitingForData)
-					{
-						await _readerEvent.WaitOneAsync(cancellationToken);
-					}
-					else if (state == State.WriterWaitingToFlush)
-					{
-						_header.Compact();
-						_header.UpdateState(state, State.Normal);
-						_writerEvent.Set();
-					}
-					else
-					{
-						throw new NotImplementedException();
-					}
-					cancellationToken.ThrowIfCancellationRequested();
-				}
-			}
-			finally
-			{
-				_header.UpdateState(State.ReaderWaitingForData, State.Normal);
-			}
-		}
-
-		#endregion
-
-		#region Writer
+		protected override void ResetWriteEvent() => _writerEvent.Reset();
 
 		/// <inheritdoc/>
-		public override void FinishWriting()
-		{
-			_header.FinishWriting();
-			_readerEvent.Set();
-		}
-
-		/// <inheritdoc/>
-		public override void AdvanceWritePosition(int size)
-		{
-			if (_header.State == State.Finished)
-			{
-				throw new InvalidOperationException("Cannot update write position after marking as complete");
-			}
-
-			_header.AdvanceWritePosition(size);
-
-			if (_header.UpdateState(State.ReaderWaitingForData, State.Normal))
-			{
-				_readerEvent.Set();
-			}
-		}
-
-		/// <inheritdoc/>
-		public override Memory<byte> GetWriteMemory() => _memory.Slice((int)_header.WritePosition);
-
-		/// <inheritdoc/>
-		public override async ValueTask FlushAsync(CancellationToken cancellationToken)
-		{
-			try
-			{
-				while (_header.ReadPosition > 0)
-				{
-					State state = _header.State;
-					if (state == State.Finished)
-					{
-						break;
-					}
-					else if (state == State.Normal)
-					{
-						_writerEvent.Reset();
-						_header.UpdateState(state, State.WriterWaitingToFlush);
-					}
-					else if (state == State.ReaderWaitingForData)
-					{
-						_header.Compact();
-						break;
-					}
-					else if (state == State.WriterWaitingToFlush)
-					{
-						await _writerEvent.WaitOneAsync(cancellationToken);
-					}
-				}
-			}
-			finally
-			{
-				_header.UpdateState(State.WriterWaitingToFlush, State.Normal);
-			}
-		}
-
-		#endregion
+		protected override Task WaitForWriteEvent(CancellationToken cancellationToken) => _writerEvent.WaitOneAsync(cancellationToken);
 	}
 }
diff --git a/Engine/Source/Programs/Shared/EpicGames.Horde/Compute/ComputeBufferBase.cs b/Engine/Source/Programs/Shared/EpicGames.Horde/Compute/ComputeBufferBase.cs
--- a/Engine/Source/Programs/Shared/EpicGames.Horde/Compute/ComputeBufferBase.cs
+++ b/Engine/Source/Programs/Shared/EpicGames.Horde/Compute/ComputeBufferBase.cs
@@ -7,149 +7,369 @@ using System.Threading.Tasks;
 namespace EpicGames.Horde.Compute
 {
 	/// <summary>
-	/// Base implementation for <see cref="IComputeBuffer"/>, which offers reader and writer objects that redirect to the owner
+	/// In-process buffer used to store compute messages
 	/// </summary>
 	public abstract class ComputeBufferBase : IComputeBuffer
 	{
-		class ReaderImpl : IComputeBufferReader
+		/// <summary>
+		/// Write state for a chunk
+		/// </summary>
+		protected enum WriteState
 		{
-			readonly ComputeBufferBase _buffer;
+			/// <summary>
+			/// Chunk is still being appended to
+			/// </summary>
+			Writing = 0,
 
-			/// <inheritdoc/>
-			public IComputeBuffer Buffer => _buffer;
+			/// <summary>
+			/// Writer has moved to the next chunk
+			/// </summary>
+			MovedToNext = 2,
 
-			public ReaderImpl(ComputeBufferBase outer) => _buffer = outer;
+			/// <summary>
+			/// This chunk marks the end of the stream
+			/// </summary>
+			Complete = 3,
+		}
 
-			/// <inheritdoc/>
-			public void Dispose()
-			{
-				Dispose(true);
-				GC.SuppressFinalize(this);
-			}
+		/// <summary>
+		/// Stores the state of a chunk in a 64-bit value, which can be updated atomically
+		/// </summary>
+		protected record struct ChunkState(ulong Value)
+		{
+			/// <summary>
+			/// Written length of this chunk
+			/// </summary>
+			public int Length => (int)(Value & 0x7fffffff);
 
 			/// <summary>
-			/// 
+			/// Set of flags which are set for each reader that still has to read from a chunk
 			/// </summary>
-			/// <param name="disposing"></param>
-			protected virtual void Dispose(bool disposing)
-			{
-			}
+			public int ReaderFlags => (int)((Value >> 31) & 0x7fffffff);
 
-			/// <inheritdoc/>
-			public bool IsComplete => _buffer.FinishedReading();
+			/// <summary>
+			/// State of the writer
+			/// </summary>
+			public WriteState WriteState => (WriteState)(Value >> 62);
 
-			/// <inheritdoc/>
-			public void Advance(int size) => _buffer.AdvanceReadPosition(size);
+			/// <summary>
+			/// Constructor
+			/// </summary>
+			public ChunkState(WriteState writerState, int readerFlags, int length) : this(((ulong)writerState << 62) | ((ulong)readerFlags << 31) | (uint)length) { }
 
-			/// <inheritdoc/>
-			public ReadOnlyMemory<byte> GetMemory() => _buffer.GetReadMemory();
+			/// <summary>
+			/// Test whether a particular reader is still referencing the chunk
+			/// </summary>
+			public bool HasReaderFlag(int readerIdx) => (Value & (1UL << (readerIdx + 31))) != 0;
 
 			/// <inheritdoc/>
-			public ValueTask WaitForDataAsync(int currentLength, CancellationToken cancellationToken) => _buffer.WaitForDataAsync(currentLength, cancellationToken);
+			public override string ToString() => $"{WriteState}, Length: {Length}, Readers: {ReaderFlags}";
 		}
 
-		class WriterImpl : IComputeBufferWriter
+		/// <summary>
+		/// Base class for chunks of data in the buffer
+		/// </summary>
+		protected abstract class ChunkBase
 		{
-			readonly ComputeBufferBase _buffer;
+			/// <summary>
+			/// Accessor for the current state
+			/// </summary>
+			public ChunkState State => new ChunkState(Interlocked.CompareExchange(ref StateValue, 0, 0));
 
-			/// <inheritdoc/>
-			public IComputeBuffer Buffer => _buffer;
+			/// <summary>
+			/// Data underlying the chunk
+			/// </summary>
+			public Memory<byte> Memory { get; }
 
-			public WriterImpl(ComputeBufferBase buffer) => _buffer = buffer;
+			/// <summary>
+			/// Current state of this chunk
+			/// </summary>
+			protected abstract ref ulong StateValue { get; }
 
-			/// <inheritdoc/>
-			public void Dispose()
+			/// <summary>
+			/// Constructor
+			/// </summary>
+			protected ChunkBase(Memory<byte> memory)
 			{
-				Dispose(true);
-				GC.SuppressFinalize(this);
+				Memory = memory;
 			}
 
 			/// <summary>
-			/// 
+			/// Append data to the chunk
+			/// </summary>
+			public void Append(int length) => Interlocked.Add(ref StateValue, (uint)length);
+
+			/// <summary>
+			/// Mark the chunk as being written to
+			/// </summary>
+			public void StartWriting(int numReaders) => Interlocked.Exchange(ref StateValue, new ChunkState(WriteState.Writing, (1 << numReaders) - 1, 0).Value);
+
+			/// <summary>
+			/// Move to the next chunk
+			/// </summary>
+			public void MoveToNext() => Interlocked.Or(ref StateValue, (ulong)WriteState.MovedToNext << 62);
+
+			/// <summary>
+			/// Move to the next chunk
 			/// </summary>
-			/// <param name="disposing"></param>
-			protected virtual void Dispose(bool disposing)
+			public void MarkComplete() => Interlocked.Or(ref StateValue, (ulong)WriteState.Complete << 62);
+
+			/// <summary>
+			/// Clear the reader flag
+			/// </summary>
+			public void FinishReading(int readerIdx) => Interlocked.And(ref StateValue, ~(1UL << (readerIdx + 31)));
+
+			/// <inheritdoc/>
+			public override string ToString() => State.ToString();
+		}
+
+		sealed class ReaderImpl : IComputeBufferReader
+		{
+			readonly ComputeBufferBase _buffer;
+			readonly ChunkBase[] _chunks;
+			readonly int _readerIdx;
+
+			int _chunkIdx;
+			int _offset;
+
+			public bool IsComplete
 			{
-				_buffer.FinishWriting();
+				get
+				{
+					ChunkState state = _chunks[_chunkIdx].State;
+					return state.WriteState == WriteState.Complete && _offset == state.Length;
+				}
 			}
 
-			/// <inheritdoc/>
-			public void MarkComplete() => _buffer.FinishWriting();
+			public ReaderImpl(ComputeBufferBase buffer, ChunkBase[] chunks, int readerIdx)
+			{
+				_buffer = buffer;
+				_chunks = chunks;
+				_readerIdx = readerIdx;
+			}
 
 			/// <inheritdoc/>
-			public void Advance(int size) => _buffer.AdvanceWritePosition(size);
+			public void Advance(int length) => _offset += length;
 
 			/// <inheritdoc/>
-			public Memory<byte> GetMemory() => _buffer.GetWriteMemory();
+			public ReadOnlyMemory<byte> GetMemory()
+			{
+				ChunkBase chunk = _chunks[_chunkIdx];
+				ChunkState state = chunk.State;
+				return state.HasReaderFlag(_readerIdx) ? chunk.Memory.Slice(_offset, state.Length - _offset) : ReadOnlyMemory<byte>.Empty;
+			}
 
 			/// <inheritdoc/>
-			public ValueTask FlushAsync(CancellationToken cancellationToken) => _buffer.FlushAsync(cancellationToken);
+			public async ValueTask WaitToReadAsync(int currentLength, CancellationToken cancellationToken = default)
+			{
+				for (; ; )
+				{
+					ChunkBase chunk = _chunks[_chunkIdx];
+					ChunkState state = chunk.State;
+
+					if (!state.HasReaderFlag(_readerIdx))
+					{
+						// Wait until the current chunk is readable
+						_buffer.ResetReadEvent(_readerIdx);
+						if (!chunk.State.HasReaderFlag(_readerIdx))
+						{
+							await _buffer.WaitForReadEvent(_readerIdx, cancellationToken);
+						}
+					}
+					else if (_offset + currentLength < state.Length || state.WriteState == WriteState.Complete)
+					{
+						// Still have data to read from this chunk
+						break;
+					}
+					else if (_offset + currentLength >= state.Length && state.WriteState == WriteState.Writing)
+					{
+						// Wait until there is more data in the chunk
+						_buffer.ResetReadEvent(_readerIdx);
+						if (chunk.State == state)
+						{
+							await _buffer.WaitForReadEvent(_readerIdx, cancellationToken);
+						}
+					}
+					else if (state.WriteState == WriteState.MovedToNext)
+					{
+						// Move to the next chunk
+						chunk.FinishReading(_readerIdx);
+						_buffer.SetWriteEvent();
+
+						_chunkIdx++;
+						if (_chunkIdx == _chunks.Length)
+						{
+							_chunkIdx = 0;
+						}
+						_offset = 0;
+					}
+					else
+					{
+						// Still need to read data from the current buffer
+						break;
+					}
+				}
+			}
 		}
 
-		/// <inheritdoc/>
-		public void Dispose()
+		sealed class WriterImpl : IComputeBufferWriter
 		{
-			Dispose(true);
-			GC.SuppressFinalize(this);
+			readonly ComputeBufferBase _buffer;
+			readonly ChunkBase[] _chunks;
+
+			int _chunkIdx;
+
+			public WriterImpl(ComputeBufferBase buffer, ChunkBase[] chunks)
+			{
+				_buffer = buffer;
+				_chunks = chunks;
+			}
+
+			public void Advance(int size)
+			{
+				ChunkBase chunk = _chunks[_chunkIdx];
+				chunk.Append(size);
+				_buffer.SetAllReadEvents();
+			}
+
+			public Memory<byte> GetMemory()
+			{
+				ChunkBase chunk = _chunks[_chunkIdx];
+
+				ChunkState state = chunk.State;
+				if (state.WriteState != WriteState.Writing)
+				{
+					throw new InvalidOperationException();
+				}
+
+				return chunk.Memory.Slice(state.Length);
+			}
+
+			public void MarkComplete()
+			{
+				ChunkBase chunk = _chunks[_chunkIdx];
+				chunk.MarkComplete();
+				_buffer.SetAllReadEvents();
+			}
+
+			public async ValueTask WaitToWriteAsync(int currentLength, CancellationToken cancellationToken)
+			{
+				if (currentLength >= _chunks[_chunkIdx].Memory.Length)
+				{
+					throw new NotImplementedException("Will never succeed");
+				}
+
+				for (; ; )
+				{
+					Memory<byte> memory = GetMemory();
+					if (memory.Length != currentLength)
+					{
+						break;
+					}
+					await MoveNextWriteChunk(cancellationToken);
+				}
+			}
+
+			async ValueTask MoveNextWriteChunk(CancellationToken cancellationToken)
+			{
+				ChunkBase chunk = _chunks[_chunkIdx];
+				chunk.MoveToNext();
+
+				int nextWriteChunkIdx = _chunkIdx + 1;
+				if (nextWriteChunkIdx == _chunks.Length)
+				{
+					nextWriteChunkIdx = 0;
+				}
+
+				ChunkBase nextChunk = _chunks[nextWriteChunkIdx];
+				while (nextChunk.State.ReaderFlags != 0)
+				{
+					await _buffer.WaitForWriteEvent(cancellationToken);
+					_buffer.ResetWriteEvent();
+				}
+
+				nextChunk.StartWriting(_buffer._numReaders);
+				_chunkIdx = nextWriteChunkIdx;
+			}
 		}
 
-		/// <inheritdoc/>
-		public abstract long Length { get; }
+		readonly ChunkBase[] _chunks;
+		readonly ReaderImpl _reader;
+		readonly WriterImpl _writer;
+		readonly int _numReaders;
 
 		/// <inheritdoc/>
-		public IComputeBufferReader Reader { get; }
+		public IComputeBufferReader Reader => _reader;
 
 		/// <inheritdoc/>
-		public IComputeBufferWriter Writer { get; }
+		public IComputeBufferWriter Writer => _writer;
 
 		/// <summary>
-		/// Default constructor
+		/// Constructor
 		/// </summary>
-		protected ComputeBufferBase()
+		/// <param name="chunks">Data for the buffer</param>
+		/// <param name="numReaders">Number of readers for this buffer</param>
+		protected ComputeBufferBase(ChunkBase[] chunks, int numReaders)
+		{
+			_chunks = chunks;
+			_chunks[0].StartWriting(numReaders);
+
+			_numReaders = numReaders;
+
+			_reader = new ReaderImpl(this, _chunks, 0);
+			_writer = new WriterImpl(this, _chunks);
+		}
+
+		/// <inheritdoc/>
+		public void Dispose()
 		{
-			Reader = new ReaderImpl(this);
-			Writer = new WriterImpl(this);
+			Dispose(true);
+			GC.SuppressFinalize(this);
 		}
 
 		/// <summary>
-		/// Standard dispose pattern
+		/// Overriable dispose method
 		/// </summary>
 		protected virtual void Dispose(bool disposing)
 		{
 		}
 
-		/// <inheritdoc/>
-		public abstract Memory<byte> GetMemory(long offset, int length);
-
-		#region Reader
-
-		/// <inheritdoc cref="IComputeBufferReader.IsComplete"/>
-		public abstract bool FinishedReading();
-
-		/// <inheritdoc cref="IComputeBufferReader.Advance(Int32)"/>
-		public abstract void AdvanceReadPosition(int size);
-
-		/// <inheritdoc cref="IComputeBufferReader.GetMemory"/>
-		public abstract ReadOnlyMemory<byte> GetReadMemory();
-
-		/// <inheritdoc cref="IComputeBufferReader.WaitForDataAsync(Int32, CancellationToken)"/>
-		public abstract ValueTask WaitForDataAsync(int currentLength, CancellationToken cancellationToken);
+		private void SetAllReadEvents()
+		{
+			for (int readerIdx = 0; readerIdx < _numReaders; readerIdx++)
+			{
+				SetReadEvent(readerIdx);
+			}
+		}
 
-		#endregion
-		#region Writer
+		/// <summary>
+		/// Signals a read event
+		/// </summary>
+		protected abstract void SetReadEvent(int readerIdx);
 
-		/// <inheritdoc cref="IComputeBufferWriter.MarkComplete"/>
-		public abstract void FinishWriting();
+		/// <summary>
+		/// Resets a read event
+		/// </summary>
+		protected abstract void ResetReadEvent(int readerIdx);
 
-		/// <inheritdoc cref="IComputeBufferWriter.Advance(Int32)"/>
-		public abstract void AdvanceWritePosition(int size);
+		/// <summary>
+		/// Waits for a read event to be signalled
+		/// </summary>
+		protected abstract Task WaitForReadEvent(int readerIdx, CancellationToken cancellationToken);
 
-		/// <inheritdoc cref="IComputeBufferWriter.GetMemory()"/>
-		public abstract Memory<byte> GetWriteMemory();
+		/// <summary>
+		/// Signals the write event
+		/// </summary>
+		protected abstract void SetWriteEvent();
 
-		/// <inheritdoc cref="IComputeBufferWriter.FlushAsync(CancellationToken)"/>
-		public abstract ValueTask FlushAsync(CancellationToken cancellationToken);
+		/// <summary>
+		/// Resets the write event
+		/// </summary>
+		protected abstract void ResetWriteEvent();
 
-		#endregion
+		/// <summary>
+		/// Waits for the write event to be signalled
+		/// </summary>
+		protected abstract Task WaitForWriteEvent(CancellationToken cancellationToken);
 	}
 }
diff --git a/Engine/Source/Programs/Shared/EpicGames.Horde/Compute/ComputeMessageChannel.cs b/Engine/Source/Programs/Shared/EpicGames.Horde/Compute/ComputeMessageChannel.cs
--- a/Engine/Source/Programs/Shared/EpicGames.Horde/Compute/ComputeMessageChannel.cs
+++ b/Engine/Source/Programs/Shared/EpicGames.Horde/Compute/ComputeMessageChannel.cs
@@ -13,7 +13,7 @@ namespace EpicGames.Horde.Compute
 	/// <summary>
 	/// Implementation of a compute channel
 	/// </summary>
-	public class ComputeMessageChannel : IComputeMessageChannel, IAsyncDisposable
+	public sealed class ComputeMessageChannel : IComputeMessageChannel
 	{
 		// Length of a message header. Consists of a 1 byte type field, followed by 4 byte length field.
 		const int HeaderLength = 5;
@@ -108,8 +108,10 @@ namespace EpicGames.Horde.Compute
 			public Span<byte> GetSpan(int sizeHint = 0) => GetMemory(sizeHint).Span;
 		}
 
-		readonly IComputeBufferReader _receiveBufferReader;
-		readonly IComputeBufferWriter _sendBufferWriter;
+		readonly IComputeSocket _socket;
+		readonly int _channelId;
+		readonly IComputeBuffer _recvBuffer;
+		readonly IComputeBuffer _sendBuffer;
 
 		// Can lock chunked memory writer to acuqire pointer
 		readonly ILogger _logger;
@@ -121,40 +123,47 @@ namespace EpicGames.Horde.Compute
 		/// <summary>
 		/// Constructor
 		/// </summary>
-		/// <param name="receiveBufferReader">Reader for incoming messages</param>
-		/// <param name="sendBufferWriter">Writer for outgoing messages</param>
+		/// <param name="socket"></param>
+		/// <param name="channelId"></param>
+		/// <param name="recvBuffer"></param>
+		/// <param name="sendBuffer"></param>
 		/// <param name="logger">Logger for diagnostic output</param>
-		public ComputeMessageChannel(IComputeBufferReader receiveBufferReader, IComputeBufferWriter sendBufferWriter, ILogger logger)
+		public ComputeMessageChannel(IComputeSocket socket, int channelId, IComputeBuffer recvBuffer, IComputeBuffer sendBuffer, ILogger logger)
 		{
-			_receiveBufferReader = receiveBufferReader;
-			_sendBufferWriter = sendBufferWriter;
+			_socket = socket;
+			_channelId = channelId;
+			_socket.AttachRecvBuffer(channelId, recvBuffer.Writer);
+			_socket.AttachSendBuffer(channelId, sendBuffer.Reader);
+			_recvBuffer = recvBuffer;
+			_sendBuffer = sendBuffer;
 			_logger = logger;
 		}
 
-		/// <inheritdoc/>
-		public async ValueTask DisposeAsync()
+		/// <summary>
+		/// Overridable dispose method
+		/// </summary>
+		public void Dispose()
 		{
-			await DisposeAsync(true);
-			GC.SuppressFinalize(this);
+			_currentBuilder?.Dispose();
+			_sendBuffer.Dispose();
+			_recvBuffer.Dispose();
 		}
 
 		/// <summary>
-		/// Overridable dispose method
+		/// Mark the send buffer as complete
 		/// </summary>
-		protected virtual async ValueTask DisposeAsync(bool disposing)
+		public void MarkComplete()
 		{
-			_currentBuilder?.Dispose();
-			_sendBufferWriter.MarkComplete();
-			await _sendBufferWriter.FlushAsync(CancellationToken.None);
-			_sendBufferWriter.Dispose();
+			_sendBuffer.Writer.MarkComplete();
 		}
 
 		/// <inheritdoc/>
 		public async ValueTask<IComputeMessage> ReceiveAsync(CancellationToken cancellationToken)
 		{
-			while (!_receiveBufferReader.IsComplete)
+			IComputeBufferReader recvBufferReader = _recvBuffer.Reader;
+			while (!recvBufferReader.IsComplete)
 			{
-				ReadOnlyMemory<byte> memory = _receiveBufferReader.GetMemory();
+				ReadOnlyMemory<byte> memory = recvBufferReader.GetMemory();
 				if (memory.Length >= HeaderLength)
 				{
 					int length = BinaryPrimitives.ReadInt32LittleEndian(memory.Span.Slice(1, 4));
@@ -166,15 +175,11 @@ namespace EpicGames.Horde.Compute
 						{
 							LogMessageInfo("RECV", message.Type, message.Data.Span);
 						}
-						_receiveBufferReader.Advance(HeaderLength + length);
+						recvBufferReader.Advance(HeaderLength + length);
 						return message;
 					}
-					if (_receiveBufferReader.Buffer.Length < HeaderLength + length)
-					{
-						throw new Exception($"Receive buffer is too small for message; size is {_receiveBufferReader.Buffer.Length}, need {HeaderLength + length}");
-					}
 				}
-				await _receiveBufferReader.WaitForDataAsync(memory.Length, cancellationToken);
+				await recvBufferReader.WaitToReadAsync(memory.Length, cancellationToken);
 			}
 			return new Message(ComputeMessageType.None, ReadOnlyMemory<byte>.Empty);
 		}
@@ -201,24 +206,16 @@ namespace EpicGames.Horde.Compute
 				throw new InvalidOperationException("Only one writer can be active at a time. Dispose of the previous writer first.");
 			}
 
+			IComputeBufferWriter sendBufferWriter = _sendBuffer.Writer;
 			for(; ;)
 			{
-				Memory<byte> memory = _sendBufferWriter.GetMemory();
+				Memory<byte> memory = sendBufferWriter.GetMemory();
 				if (memory.Length >= maxSize)
 				{
-					_currentBuilder = new MessageBuilder(this, _sendBufferWriter, type);
+					_currentBuilder = new MessageBuilder(this, sendBufferWriter, type);
 					return _currentBuilder;
 				}
-
-				long totalLength = _sendBufferWriter.Buffer.Length;
-				if (totalLength < maxSize)
-				{
-					throw new ComputeInternalException($"Insufficient space in buffer to write message of {maxSize} bytes; buffer only has {totalLength} remaining");
-				}
-				else
-				{
-					await _sendBufferWriter.FlushAsync(cancellationToken);
-				}
+				await sendBufferWriter.WaitToWriteAsync(memory.Length, cancellationToken);
 			}
 		}
 	}
diff --git a/Engine/Source/Programs/Shared/EpicGames.Horde/Compute/ComputeProcess.cs b/Engine/Source/Programs/Shared/EpicGames.Horde/Compute/ComputeProcess.cs
--- a/Engine/Source/Programs/Shared/EpicGames.Horde/Compute/ComputeProcess.cs
+++ b/Engine/Source/Programs/Shared/EpicGames.Horde/Compute/ComputeProcess.cs
@@ -109,7 +109,7 @@ namespace EpicGames.Horde.Compute
 				else
 				{
 					AppendToBuffer(data.Slice(0, lineLen));
-					str = Encoding.UTF8.GetString(data.Slice(0, _bufferLength + lineLen));
+					str = Encoding.UTF8.GetString(_buffer.AsSpan(0, _bufferLength + lineLen));
 					_bufferLength = 0;
 				}
 
diff --git a/Engine/Source/Programs/Shared/EpicGames.Horde/Compute/ComputeSocketBase.cs b/Engine/Source/Programs/Shared/EpicGames.Horde/Compute/ComputeSocketBase.cs
--- a/Engine/Source/Programs/Shared/EpicGames.Horde/Compute/ComputeSocketBase.cs
+++ b/Engine/Source/Programs/Shared/EpicGames.Horde/Compute/ComputeSocketBase.cs
@@ -23,13 +23,13 @@ namespace EpicGames.Horde.Compute
 
 		readonly IComputeTransport _transport;
 		readonly ILogger _logger;
-		readonly Task _readTask;
-		CancellationTokenSource _cancellationSource = new CancellationTokenSource();
 
-		readonly Dictionary<int, IComputeBufferWriter> _receiveBuffers = new Dictionary<int, IComputeBufferWriter>();
+		readonly CancellationTokenSource _cancellationSource = new CancellationTokenSource();
+
+		readonly BackgroundTask _recvTask;
+		readonly Dictionary<int, IComputeBufferWriter> _recvBuffers = new Dictionary<int, IComputeBufferWriter>();
 
 		readonly SemaphoreSlim _sendSemaphore = new SemaphoreSlim(1, 1);
-		readonly Dictionary<int, IComputeBufferReader> _sendBuffers = new Dictionary<int, IComputeBufferReader>();
 		readonly Dictionary<int, Task> _sendTasks = new Dictionary<int, Task>();
 
 		/// <summary>
@@ -41,7 +41,7 @@ namespace EpicGames.Horde.Compute
 		{
 			_transport = transport;
 			_logger = logger;
-			_readTask = RunReaderAsync(transport, _cancellationSource.Token);
+			_recvTask = BackgroundTask.StartNew(ctx => RunReaderAsync(transport, ctx));
 		}
 
 		/// <summary>
@@ -49,8 +49,6 @@ namespace EpicGames.Horde.Compute
 		/// </summary>
 		public async ValueTask CloseAsync(CancellationToken cancellationToken)
 		{
-			Task cancellationTask = Task.Delay(-1, cancellationToken);
-
 			// Wait for all the individual send tasks to complete
 			await Task.WhenAll(_sendTasks.Values);
 			cancellationToken.ThrowIfCancellationRequested();
@@ -62,8 +60,8 @@ namespace EpicGames.Horde.Compute
 			_logger.LogTrace("Sent shutdown packet");
 
 			// Wait for the reader to stop
-			await Task.WhenAny(_readTask, cancellationTask);
-			cancellationToken.ThrowIfCancellationRequested();
+			await _recvTask.Task.WaitAsync(TimeSpan.FromSeconds(5), cancellationToken);
+			await _recvTask.StopAsync();
 
 			_logger.LogTrace("Closed");
 		}
@@ -71,18 +69,9 @@ namespace EpicGames.Horde.Compute
 		/// <inheritdoc/>
 		public virtual async ValueTask DisposeAsync()
 		{
-			if (_cancellationSource != null)
-			{
-				_cancellationSource.Cancel();
-
-				await Task.WhenAll(_sendTasks.Values);
-				await _readTask;
-
-				_cancellationSource.Dispose();
-				_cancellationSource = null!;
-			}
-
+			await _recvTask.DisposeAsync();
 			_sendSemaphore.Dispose();
+			_cancellationSource.Dispose();
 			GC.SuppressFinalize(this);
 		}
 
@@ -123,7 +112,7 @@ namespace EpicGames.Horde.Compute
 					if (size < 0)
 					{
 						_logger.LogTrace("Detaching buffer {Id}", id);
-						DetachRecvBufferWriter(cachedWriters, id);
+						DetachRecvBuffer(cachedWriters, id);
 					}
 					else
 					{
@@ -140,7 +129,7 @@ namespace EpicGames.Horde.Compute
 			lock (_lockObject)
 			{
 				_complete = true;
-				foreach (IComputeBufferWriter writer in _receiveBuffers.Values)
+				foreach (IComputeBufferWriter writer in _recvBuffers.Values)
 				{
 					writer.MarkComplete();
 				}
@@ -151,23 +140,21 @@ namespace EpicGames.Horde.Compute
 
 		async Task ReadPacketAsync(IComputeTransport transport, int id, int size, IComputeBufferWriter writer, CancellationToken cancellationToken)
 		{
-			for (int offset = 0; offset < size;)
+			Memory<byte> memory = writer.GetMemory();
+			while (memory.Length < size)
 			{
-				Memory<byte> memory = writer.GetMemory();
-				if (memory.Length == 0)
-				{
-					_logger.LogTrace("No space in buffer {Id}, flushing", id);
-					await writer.FlushAsync(cancellationToken);
-				}
-				else
-				{
-					int maxRead = Math.Min(size - offset, memory.Length);
-					int read = await transport.ReadPartialAsync(memory.Slice(0, maxRead), cancellationToken);
-					writer.Advance(read);
+				_logger.LogTrace("No space in buffer {Id}, flushing", id);
+				await writer.WaitToWriteAsync(memory.Length, cancellationToken);
+				memory = writer.GetMemory();
+			}
 
-					offset += read;
-				}
+			for (int offset = 0; offset < size;)
+			{
+				int read = await transport.ReadPartialAsync(memory.Slice(offset, size - offset), cancellationToken);
+				offset += read;
 			}
+
+			writer.Advance(size);
 		}
 
 		IComputeBufferWriter GetReceiveBuffer(Dictionary<int, IComputeBufferWriter> cachedWriters, int id)
@@ -180,7 +167,7 @@ namespace EpicGames.Horde.Compute
 
 			lock (_lockObject)
 			{
-				if (_receiveBuffers.TryGetValue(id, out writer))
+				if (_recvBuffers.TryGetValue(id, out writer))
 				{
 					cachedWriters.Add(id, writer);
 					return writer;
@@ -200,54 +187,38 @@ namespace EpicGames.Horde.Compute
 			}
 		}
 
-		async Task SendFromBufferAsync(int id, IComputeBufferReader reader, CancellationToken cancellationToken)
-		{
-			byte[] header = new byte[8];
-
-			BinaryPrimitives.WriteInt32LittleEndian(header, id);
-
-			SendSegment headerSegment = new SendSegment();
-			SendSegment bodySegment = new SendSegment();
-			headerSegment.Set(header, bodySegment, 0);
+		readonly byte[] _header = new byte[8];
+		readonly SendSegment _headerSegment = new SendSegment();
+		readonly SendSegment _bodySegment = new SendSegment();
 
-			for (; ; )
+		/// <inheritdoc/>
+		public async ValueTask SendAsync(int id, ReadOnlyMemory<byte> memory, CancellationToken cancellationToken)
+		{
+			if (memory.Length > 0)
 			{
-				// Wait for something to read
-				await reader.WaitForDataAsync(0, cancellationToken);
+				await SendInternalAsync(id, memory.Length, memory, cancellationToken);
+			}
+		}
 
-				// Acquire the semaphore
-				await _sendSemaphore.WaitAsync(cancellationToken);
-				try
-				{
-					ReadOnlyMemory<byte> memory = reader.GetMemory();
-					if (memory.Length > 0)
-					{
-						BinaryPrimitives.WriteInt32LittleEndian(header.AsSpan(4), memory.Length);
+		/// <inheritdoc/>
+		public ValueTask MarkCompleteAsync(int id, CancellationToken cancellationToken) => SendInternalAsync(id, -1, ReadOnlyMemory<byte>.Empty, cancellationToken);
 
-						bodySegment.Set(memory, null, header.Length);
-						ReadOnlySequence<byte> sequence = new ReadOnlySequence<byte>(headerSegment, 0, bodySegment, memory.Length);
-						await _transport.WriteAsync(sequence, cancellationToken);
+		async ValueTask SendInternalAsync(int id, int size, ReadOnlyMemory<byte> memory, CancellationToken cancellationToken)
+		{
+			await _sendSemaphore.WaitAsync(cancellationToken);
+			try
+			{
+				BinaryPrimitives.WriteInt32LittleEndian(_header, id);
+				BinaryPrimitives.WriteInt32LittleEndian(_header.AsSpan(4), size);
+				_headerSegment.Set(_header, _bodySegment, 0);
+				_bodySegment.Set(memory, null, _header.Length);
 
-						reader.Advance(memory.Length);
-					}
-					else if (reader.IsComplete)
-					{
-						_logger.LogTrace("Sending complete packet for channel {ChannelId}", id);
-						BinaryPrimitives.WriteInt32LittleEndian(header.AsSpan(4), -1);
-						await _transport.WriteAsync(header, cancellationToken);
-						break;
-					}
-				}
-				finally
-				{
-					_sendSemaphore.Release();
-				}
+				ReadOnlySequence<byte> sequence = new ReadOnlySequence<byte>(_headerSegment, 0, _bodySegment, memory.Length);
+				await _transport.WriteAsync(sequence, cancellationToken);
 			}
-
-			lock (_lockObject)
+			finally
 			{
-				_sendBuffers.Remove(id);
-				_sendTasks.Remove(id);
+				_sendSemaphore.Release();
 			}
 		}
 
@@ -255,7 +226,7 @@ namespace EpicGames.Horde.Compute
 		public abstract IComputeBuffer CreateBuffer(long capacity);
 
 		/// <inheritdoc/>
-		public ValueTask AttachRecvBufferAsync(int channelId, IComputeBufferWriter recvBufferWriter, CancellationToken cancellationToken)
+		public void AttachRecvBuffer(int channelId, IComputeBufferWriter recvBuffer)
 		{
 			bool complete;
 			lock (_lockObject)
@@ -263,19 +234,17 @@ namespace EpicGames.Horde.Compute
 				complete = _complete;
 				if (!complete)
 				{
-					_receiveBuffers.Add(channelId, recvBufferWriter);
+					_recvBuffers.Add(channelId, recvBuffer);
 				}
 			}
 
-			if (recvBufferWriter != null && complete)
+			if (recvBuffer != null && complete)
 			{
-				recvBufferWriter.MarkComplete();
+				recvBuffer.MarkComplete();
 			}
-
-			return new ValueTask();
 		}
 
-		void DetachRecvBufferWriter(Dictionary<int, IComputeBufferWriter> cachedWriters, int id)
+		void DetachRecvBuffer(Dictionary<int, IComputeBufferWriter> cachedWriters, int id)
 		{
 			cachedWriters.Remove(id);
 
@@ -283,30 +252,41 @@ namespace EpicGames.Horde.Compute
 			lock (_lockObject)
 			{
 #pragma warning disable CA2000 // Dispose objects before losing scope
-				_receiveBuffers.Remove(id, out writer);
+				_recvBuffers.Remove(id, out writer);
 #pragma warning restore CA2000 // Dispose objects before losing scope
 			}
 			if (writer != null)
 			{
 				writer.MarkComplete();
-				writer.Dispose();
 			}
 		}
 
-		/// <summary>
-		/// Registers the read end of an attached send buffer
-		/// </summary>
-		/// <param name="channelId">Channel id</param>
-		/// <param name="sendBufferReader">Writer for the receive buffer</param>
-		/// <param name="cancellationToken">Cancellation token for the operation</param>
-		public ValueTask AttachSendBufferAsync(int channelId, IComputeBufferReader sendBufferReader, CancellationToken cancellationToken)
+		/// <inheritdoc/>
+		public void AttachSendBuffer(int channelId, IComputeBufferReader reader)
 		{
 			lock (_lockObject)
 			{
-				_sendBuffers.Add(channelId, sendBufferReader);
-				_sendTasks.Add(channelId, Task.Run(() => SendFromBufferAsync(channelId, sendBufferReader, _cancellationSource.Token), _cancellationSource.Token));
+				_sendTasks.Add(channelId, Task.Run(() => SendFromBufferAsync(channelId, reader, _cancellationSource.Token), _cancellationSource.Token));
+			}
+		}
+
+		async Task SendFromBufferAsync(int channelId, IComputeBufferReader reader, CancellationToken cancellationToken)
+		{
+			while (!cancellationToken.IsCancellationRequested)
+			{
+				ReadOnlyMemory<byte> memory = reader.GetMemory();
+				if (memory.Length > 0)
+				{
+					await SendAsync(channelId, memory, cancellationToken);
+					reader.Advance(memory.Length);
+				}
+				if (reader.IsComplete)
+				{
+					await MarkCompleteAsync(channelId, cancellationToken);
+					break;
+				}
+				await reader.WaitToReadAsync(0, cancellationToken);
 			}
-			return new ValueTask();
 		}
 	}
 }
diff --git a/Engine/Source/Programs/Shared/EpicGames.Horde/Compute/ComputeWorker.cs b/Engine/Source/Programs/Shared/EpicGames.Horde/Compute/ComputeWorker.cs
--- a/Engine/Source/Programs/Shared/EpicGames.Horde/Compute/ComputeWorker.cs
+++ b/Engine/Source/Programs/Shared/EpicGames.Horde/Compute/ComputeWorker.cs
@@ -51,7 +51,7 @@ namespace EpicGames.Horde.Compute
 
 		async Task RunAsync(IComputeSocket socket, int channelId, CancellationToken cancellationToken)
 		{
-			await using (IComputeMessageChannel channel = await socket.CreateMessageChannelAsync(channelId, 4 * 1024 * 1024, _logger, cancellationToken))
+			using (IComputeMessageChannel channel = socket.CreateMessageChannel(channelId, 4 * 1024 * 1024, _logger, cancellationToken))
 			{
 				await channel.SendReadyAsync(cancellationToken);
 
@@ -66,12 +66,6 @@ namespace EpicGames.Horde.Compute
 						case ComputeMessageType.None:
 							await Task.WhenAll(childTasks);
 							return;
-						case ComputeMessageType.Fork:
-							{
-								ForkMessage fork = message.ParseForkMessage();
-								childTasks.Add(RunAsync(socket, fork.ChannelId, cancellationToken));
-							}
-							break;
 						case ComputeMessageType.WriteFiles:
 							{
 								UploadFilesMessage writeFiles = message.ParseUploadFilesMessage();
@@ -171,94 +165,31 @@ namespace EpicGames.Horde.Compute
 
 		async Task ExecuteProcessWindowsAsync(IComputeSocket socket, IComputeMessageChannel channel, string executable, IReadOnlyList<string> arguments, string? workingDir, IReadOnlyDictionary<string, string?>? envVars, CancellationToken cancellationToken)
 		{
-			using SharedMemoryBuffer ipcBuffer = new SharedMemoryBuffer(4096);
-			await using (BackgroundTask ipcTask = BackgroundTask.StartNew(ctx => RunIpcChannel(socket, ipcBuffer, ctx)))
-			{
-				Dictionary<string, string?> newEnvVars = new Dictionary<string, string?>();
-				if (envVars != null)
-				{
-					foreach ((string name, string? value) in envVars)
-					{
-						newEnvVars.Add(name, value);
-					}
-				}
-
-				newEnvVars[WorkerComputeSocket.IpcEnvVar] = WorkerComputeSocket.GetIpcEnvVarValue(ipcBuffer);
-
-				_logger.LogInformation("Launching {Executable} {Arguments}", CommandLineArguments.Quote(executable), CommandLineArguments.Join(arguments));
-				await ExecuteProcessInternalAsync(channel, executable, arguments, workingDir, newEnvVars, cancellationToken);
-				_logger.LogInformation("Finished executing process");
+			const int ChannelId = 1;
+			string baseBufferName = $"Local\\Compute-{Guid.NewGuid()}";
 
-				Memory<byte> buffer = ipcBuffer.GetWriteMemory();
-				buffer.Span[0] = (byte)IpcMessageType.Finish;
-				ipcBuffer.AdvanceWritePosition(1);
+			using SharedMemoryBuffer sendBuffer = SharedMemoryBuffer.CreateNew(ComputeChannel.GetSendBufferName(baseBufferName), 4096);
+			socket.AttachSendBuffer(ChannelId, sendBuffer.Reader);
 
-				await ipcTask.Task.WaitAsync(TimeSpan.FromSeconds(5.0), cancellationToken);
-			}
-			_logger.LogInformation("Child process has shut down");
-		}
+			using SharedMemoryBuffer recvBuffer = SharedMemoryBuffer.CreateNew(ComputeChannel.GetRecvBufferName(baseBufferName), 4096);
+			socket.AttachRecvBuffer(ChannelId, recvBuffer.Writer);
 
-		async Task RunIpcChannel(IComputeSocket socket, SharedMemoryBuffer ipcBuffer, CancellationToken cancellationToken)
-		{
-			List<IDisposable> buffers = new List<IDisposable>();
-			try
+			Dictionary<string, string?> newEnvVars = new Dictionary<string, string?>();
+			if (envVars != null)
 			{
-				for (; ; )
+				foreach ((string name, string? value) in envVars)
 				{
-					ReadOnlyMemory<byte> memory = ipcBuffer.GetReadMemory();
-					if (memory.Length == 0)
-					{
-						await ipcBuffer.WaitForDataAsync(0, cancellationToken);
-						continue;
-					}
+					newEnvVars.Add(name, value);
+				}
+			}
 
-					IpcMessageType messageType = (IpcMessageType)memory.Span[0];
-					switch (messageType)
-					{
-						case IpcMessageType.Finish:
-							return;
-						case IpcMessageType.AttachRecvBuffer:
-							{
-								IpcAttachBufferRequest attachRecvBuffer = IpcAttachBufferRequest.Read(ipcBuffer);
+			newEnvVars[ComputeChannel.IpcEnvVar] = baseBufferName;
 
-								_logger.LogInformation("Attaching receive buffer to channel {Id}", attachRecvBuffer.ChannelId);
-#pragma warning disable CA2000 // Dispose objects before losing scope
-								SharedMemoryBuffer buffer = new SharedMemoryBuffer(attachRecvBuffer.MemoryHandle, attachRecvBuffer.ReaderEventHandle, attachRecvBuffer.WriterEventHandle);
-								(IComputeBufferReader reader, IComputeBufferWriter writer) = buffer.ToShared();
-								reader.Dispose();
-								await socket.AttachRecvBufferAsync(attachRecvBuffer.ChannelId, writer, cancellationToken);
-#pragma warning restore CA2000 // Dispose objects before losing scope
-							}
-							break;
-						case IpcMessageType.AttachSendBuffer:
-							{
-								IpcAttachBufferRequest attachSendBuffer = IpcAttachBufferRequest.Read(ipcBuffer);
+			_logger.LogInformation("Launching {Executable} {Arguments}", CommandLineArguments.Quote(executable), CommandLineArguments.Join(arguments));
+			await ExecuteProcessInternalAsync(channel, executable, arguments, workingDir, newEnvVars, cancellationToken);
+			_logger.LogInformation("Finished executing process");
 
-								_logger.LogInformation("Attaching send buffer to channel {Id}", attachSendBuffer.ChannelId);
-#pragma warning disable CA2000 // Dispose objects before losing scope
-								SharedMemoryBuffer buffer = new SharedMemoryBuffer(attachSendBuffer.MemoryHandle, attachSendBuffer.ReaderEventHandle, attachSendBuffer.WriterEventHandle);
-								(IComputeBufferReader reader, IComputeBufferWriter writer) = buffer.ToShared();
-								writer.Dispose();
-								await socket.AttachSendBufferAsync(attachSendBuffer.ChannelId, reader, cancellationToken);
-#pragma warning restore CA2000 // Dispose objects before losing scope
-							}
-							break;
-						default:
-							throw new ComputeInternalException($"Invalid IPC message type: {messageType}");
-					}
-				}
-			}
-			catch (OperationCanceledException)
-			{
-				_logger.LogInformation("Cancelled during process exec");
-			}
-			finally
-			{
-				foreach (SharedMemoryBuffer buffer in buffers)
-				{
-					buffer.Dispose();
-				}
-			}
+			_logger.LogInformation("Child process has shut down");
 		}
 
 		async Task ExecuteProcessInternalAsync(IComputeMessageChannel channel, string executable, IReadOnlyList<string> arguments, string? workingDir, IReadOnlyDictionary<string, string?>? envVars, CancellationToken cancellationToken)
diff --git a/Engine/Source/Programs/Shared/EpicGames.Horde/Compute/IComputeBuffer.cs b/Engine/Source/Programs/Shared/EpicGames.Horde/Compute/IComputeBuffer.cs
--- a/Engine/Source/Programs/Shared/EpicGames.Horde/Compute/IComputeBuffer.cs
+++ b/Engine/Source/Programs/Shared/EpicGames.Horde/Compute/IComputeBuffer.cs
@@ -1,7 +1,6 @@
 // Copyright Epic Games, Inc. All Rights Reserved.
 
 using System;
-using System.IO;
 using System.Threading;
 using System.Threading.Tasks;
 
@@ -12,11 +11,6 @@ namespace EpicGames.Horde.Compute
 	/// </summary>
 	public interface IComputeBuffer : IDisposable
 	{
-		/// <summary>
-		/// Total length of this buffer
-		/// </summary>
-		long Length { get; }
-
 		/// <summary>
 		/// Reader for this buffer
 		/// </summary>
@@ -26,23 +20,13 @@ namespace EpicGames.Horde.Compute
 		/// Writer for this buffer
 		/// </summary>
 		IComputeBufferWriter Writer { get; }
-
-		/// <summary>
-		/// Access the underlying memory for the buffer
-		/// </summary>
-		Memory<byte> GetMemory(long offset, int length);
 	}
 
 	/// <summary>
 	/// Read interface for a compute buffer
 	/// </summary>
-	public interface IComputeBufferReader : IDisposable
+	public interface IComputeBufferReader
 	{
-		/// <summary>
-		/// Accessor for the buffer this is reading from
-		/// </summary>
-		IComputeBuffer Buffer { get; }
-
 		/// <summary>
 		/// Whether this buffer is complete (no more data will be added)
 		/// </summary>
@@ -65,19 +49,14 @@ namespace EpicGames.Horde.Compute
 		/// </summary>
 		/// <param name="currentLength">Current length of the buffer</param>
 		/// <param name="cancellationToken">Cancellation token for the operation</param>
-		ValueTask WaitForDataAsync(int currentLength, CancellationToken cancellationToken = default);
+		ValueTask WaitToReadAsync(int currentLength, CancellationToken cancellationToken = default);
 	}
 
 	/// <summary>
 	/// Buffer that can receive data from a remote machine.
 	/// </summary>
-	public interface IComputeBufferWriter : IDisposable
+	public interface IComputeBufferWriter
 	{
-		/// <summary>
-		/// Accessor for the buffer we're writing to
-		/// </summary>
-		IComputeBuffer Buffer { get; }
-
 		/// <summary>
 		/// Mark the output to this buffer as complete
 		/// </summary>
@@ -95,10 +74,10 @@ namespace EpicGames.Horde.Compute
 		Memory<byte> GetMemory();
 
 		/// <summary>
-		/// Waits until all data in this buffer has been written
+		/// Gets memory to write to
 		/// </summary>
-		/// <param name="cancellationToken">Cancellation token for the operation</param>
-		ValueTask FlushAsync(CancellationToken cancellationToken);
+		/// <returns>Memory to be written to</returns>
+		ValueTask WaitToWriteAsync(int currentLength, CancellationToken cancellationToken = default);
 	}
 
 	/// <summary>
@@ -106,106 +85,47 @@ namespace EpicGames.Horde.Compute
 	/// </summary>
 	public static class ComputeBufferExtensions
 	{
-		sealed class RefCountedBuffer
+		/// <summary>
+		/// Read from a buffer into another buffer
+		/// </summary>
+		/// <param name="reader">Buffer to read from</param>
+		/// <param name="buffer">Memory to receive the read data</param>
+		/// <param name="cancellationToken">Cancellation token for the operation</param>
+		/// <returns>Number of bytes read</returns>
+		public static async ValueTask<int> ReadAsync(this IComputeBufferReader reader, Memory<byte> buffer, CancellationToken cancellationToken = default)
 		{
-			int _refFlags = 1 | 2;
-
-			public IComputeBuffer Buffer { get; }
-
-			public RefCountedBuffer(IComputeBuffer buffer) => Buffer = buffer;
-
-			public void Release(int flag)
+			for (; ; )
 			{
-				for (; ; )
+				ReadOnlyMemory<byte> readMemory = reader.GetMemory();
+				if (readMemory.Length > 0)
 				{
-					int initialRefFlags = _refFlags;
-					if (Interlocked.CompareExchange(ref _refFlags, initialRefFlags & ~flag, initialRefFlags) != initialRefFlags)
-					{
-						continue;
-					}
-					if (initialRefFlags == flag)
-					{
-						Buffer.Dispose();
-					}
-					break;
+					int length = Math.Min(readMemory.Length, buffer.Length);
+					readMemory.Slice(0, length).CopyTo(buffer);
+					reader.Advance(length);
+					return length;
 				}
+				await reader.WaitToReadAsync(0, cancellationToken);
 			}
 		}
 
-		class RefCountedReader : IComputeBufferReader
-		{
-			public readonly RefCountedBuffer _buffer;
-			public readonly IComputeBufferReader _reader;
-
-			public IComputeBuffer Buffer => _reader.Buffer;
-
-			public RefCountedReader(RefCountedBuffer refCountedBuffer)
-			{
-				_buffer = refCountedBuffer;
-				_reader = refCountedBuffer.Buffer.Reader;
-			}
-
-			public void Dispose() => _buffer.Release(1);
-			public bool IsComplete => _reader.IsComplete;
-			public void Advance(int size) => _reader.Advance(size);
-			public ReadOnlyMemory<byte> GetMemory() => _reader.GetMemory();
-			public ValueTask WaitForDataAsync(int currentLength, CancellationToken cancellationToken) => _reader.WaitForDataAsync(currentLength, cancellationToken);
-		}
-
-		class RefCountedWriter : IComputeBufferWriter
-		{
-			public readonly RefCountedBuffer _buffer;
-			public readonly IComputeBufferWriter _writer;
-
-			public IComputeBuffer Buffer => _writer.Buffer;
-
-			public RefCountedWriter(RefCountedBuffer refCountedBuffer)
-			{
-				_buffer = refCountedBuffer;
-				_writer = refCountedBuffer.Buffer.Writer;
-			}
-
-			public void Dispose() => _buffer.Release(2);
-			public void Advance(int size) => _writer.Advance(size);
-			public ValueTask FlushAsync(CancellationToken cancellationToken) => _writer.FlushAsync(cancellationToken);
-			public Memory<byte> GetMemory() => _writer.GetMemory();
-			public void MarkComplete() => _writer.MarkComplete();
-		}
-
-		/// <summary>
-		/// Converts a compute buffer to be disposable through its reader/writer 
-		/// </summary>
-		/// <param name="buffer"></param>
-		/// <returns></returns>
-		public static (IComputeBufferReader, IComputeBufferWriter) ToShared(this IComputeBuffer buffer)
-		{
-			RefCountedBuffer refCountedBuffer = new RefCountedBuffer(buffer);
-			return (new RefCountedReader(refCountedBuffer), new RefCountedWriter(refCountedBuffer));
-		}
-
 		/// <summary>
-		/// Waits until there is a block of the certain size in the buffer, and returns it
+		/// Writes data into a buffer from a memory block
 		/// </summary>
-		/// <param name="reader">Instance to read from</param>
-		/// <param name="minLength">Minimum length of the memory</param>
-		/// <param name="cancellationToken"></param>
-		/// <returns></returns>
-		public static async ValueTask<ReadOnlyMemory<byte>> GetMemoryAsync(this IComputeBufferReader reader, int minLength, CancellationToken cancellationToken)
+		/// <param name="writer">Writer to output the data to</param>
+		/// <param name="buffer">The data to write</param>
+		/// <param name="cancellationToken">Cancellation token for the operation</param>
+		public static async ValueTask WriteAsync(this IComputeBufferWriter writer, ReadOnlyMemory<byte> buffer, CancellationToken cancellationToken = default)
 		{
-			for (; ; )
+			while (buffer.Length > 0)
 			{
-				bool complete = reader.IsComplete;
-
-				ReadOnlyMemory<byte> memory = reader.GetMemory();
-				if (memory.Length >= minLength)
+				Memory<byte> writeMemory = writer.GetMemory();
+				if (writeMemory.Length >= buffer.Length)
 				{
-					return memory;
-				}
-				else if (complete)
-				{
-					throw new EndOfStreamException();
+					buffer.CopyTo(writeMemory);
+					writer.Advance(buffer.Length);
+					break;
 				}
-				await reader.WaitForDataAsync(memory.Length, cancellationToken);
+				await writer.WaitToWriteAsync(writeMemory.Length, cancellationToken);
 			}
 		}
 	}
diff --git a/Engine/Source/Programs/Shared/EpicGames.Horde/Compute/IComputeChannel.cs b/Engine/Source/Programs/Shared/EpicGames.Horde/Compute/IComputeChannel.cs
new file mode 100644
--- /dev/null
+++ b/Engine/Source/Programs/Shared/EpicGames.Horde/Compute/IComputeChannel.cs
@@ -0,0 +1,191 @@
+// Copyright Epic Games, Inc. All Rights Reserved.
+
+using System;
+using System.IO;
+using System.Threading;
+using System.Threading.Tasks;
+using EpicGames.Horde.Compute.Buffers;
+using Microsoft.Extensions.Logging;
+
+namespace EpicGames.Horde.Compute
+{
+	/// <summary>
+	/// Conventional TCP-like interface for writing data to a socket. Sends are "push", receives are "pull".
+	/// </summary>
+	public interface IComputeChannel : IDisposable
+	{
+		/// <summary>
+		/// Sends data to a remote channel
+		/// </summary>
+		/// <param name="memory">Memory to write</param>
+		/// <param name="cancellationToken">Cancellation token for the operation</param>
+		ValueTask SendAsync(ReadOnlyMemory<byte> memory, CancellationToken cancellationToken = default);
+
+		/// <summary>
+		/// Marks a channel as complete
+		/// </summary>
+		/// <param name="buffer">Buffer to receive the data</param>
+		/// <param name="cancellationToken">Cancellation token for the operation</param>
+		ValueTask<int> ReceiveAsync(Memory<byte> buffer, CancellationToken cancellationToken = default);
+	}
+
+	/// <summary>
+	/// Opens a channel for compute workers
+	/// </summary>
+	public static class ComputeChannel
+	{
+		class BufferedReaderChannel : IComputeChannel
+		{
+			readonly IComputeSocket _socket;
+			readonly int _channelId;
+			readonly IComputeBuffer _recvBuffer;
+
+			public BufferedReaderChannel(IComputeSocket socket, int channelId, IComputeBuffer recvBuffer)
+			{
+				_socket = socket;
+				_channelId = channelId;
+				_recvBuffer = recvBuffer;
+
+				_socket.AttachRecvBuffer(_channelId, recvBuffer.Writer);
+			}
+
+			public void Dispose()
+			{
+				_recvBuffer.Dispose();
+			}
+
+			public ValueTask<int> ReceiveAsync(Memory<byte> buffer, CancellationToken cancellationToken = default) => _recvBuffer.Reader.ReadAsync(buffer, cancellationToken);
+
+			public ValueTask SendAsync(ReadOnlyMemory<byte> memory, CancellationToken cancellationToken = default) => _socket.SendAsync(_channelId, memory, cancellationToken);
+		}
+
+		class BufferedReaderWriterChannel : IComputeChannel
+		{
+			readonly IComputeSocket _socket;
+			readonly int _channelId;
+			readonly IComputeBuffer _recvBuffer;
+			readonly IComputeBuffer _sendBuffer;
+
+			public BufferedReaderWriterChannel(IComputeSocket socket, int channelId, IComputeBuffer recvBuffer, IComputeBuffer sendBuffer)
+			{
+				_socket = socket;
+				_channelId = channelId;
+				_recvBuffer = recvBuffer;
+				_sendBuffer = sendBuffer;
+
+				_socket.AttachRecvBuffer(_channelId, recvBuffer.Writer);
+				_socket.AttachSendBuffer(_channelId, sendBuffer.Reader);
+			}
+
+			public void Dispose()
+			{
+				_recvBuffer.Dispose();
+				_sendBuffer.Dispose();
+			}
+
+			public ValueTask<int> ReceiveAsync(Memory<byte> buffer, CancellationToken cancellationToken = default) => _recvBuffer.Reader.ReadAsync(buffer, cancellationToken);
+
+			public ValueTask SendAsync(ReadOnlyMemory<byte> memory, CancellationToken cancellationToken = default) => _sendBuffer.Writer.WriteAsync(memory, cancellationToken);
+		}
+
+		class WorkerChannel : IComputeChannel
+		{
+			readonly IComputeBuffer _recvBuffer;
+			readonly IComputeBuffer _sendBuffer;
+
+			public WorkerChannel(IComputeBuffer recvBuffer, IComputeBuffer sendBuffer)
+			{
+				_recvBuffer = recvBuffer;
+				_sendBuffer = sendBuffer;
+			}
+
+			public void Dispose() { }
+
+			public ValueTask<int> ReceiveAsync(Memory<byte> buffer, CancellationToken cancellationToken = default) => _recvBuffer.Reader.ReadAsync(buffer, cancellationToken);
+
+			public ValueTask SendAsync(ReadOnlyMemory<byte> memory, CancellationToken cancellationToken = default) => _sendBuffer.Writer.WriteAsync(memory, cancellationToken);
+		}
+
+		/// <summary>
+		/// Name of the environment variable for passing the name of the compute channel
+		/// </summary>
+		public const string IpcEnvVar = "UE_HORDE_COMPUTE_IPC";
+
+		internal static string GetSendBufferName(string baseName) => $"{baseName}_SEND";
+		internal static string GetRecvBufferName(string baseName) => $"{baseName}_RECV";
+
+
+		/// <summary>
+		/// Creates a socket for a worker
+		/// </summary>
+		public static IComputeChannel ConnectAsWorker()
+		{
+			string? baseName = Environment.GetEnvironmentVariable(IpcEnvVar);
+			if (baseName == null)
+			{
+				throw new InvalidOperationException($"Environment variable {IpcEnvVar} is not defined; cannot connect as worker.");
+			}
+
+			SharedMemoryBuffer recvBuffer = SharedMemoryBuffer.OpenExisting(GetRecvBufferName(baseName));
+			SharedMemoryBuffer sendBuffer = SharedMemoryBuffer.OpenExisting(GetSendBufferName(baseName));
+			return new WorkerChannel(recvBuffer, sendBuffer);
+		}
+
+		/// <summary>
+		/// Creates a channel using a socket and receive buffer
+		/// </summary>
+		/// <param name="socket">Socket to use for sending data</param>
+		/// <param name="channelId">Channel id to send and receive data</param>
+		/// <param name="recvBuffer">Buffer for receiving data</param>
+		public static IComputeChannel CreateChannel(this IComputeSocket socket, int channelId, IComputeBuffer recvBuffer) => new BufferedReaderChannel(socket, channelId, recvBuffer);
+
+		/// <summary>
+		/// Creates a channel using a socket and receive buffer
+		/// </summary>
+		/// <param name="socket">Socket to use for sending data</param>
+		/// <param name="channelId">Channel id to send and receive data</param>
+		/// <param name="recvBuffer">Buffer for receiving data</param>
+		/// <param name="sendBuffer">Buffer for sending data</param>
+		public static IComputeChannel CreateChannel(this IComputeSocket socket, int channelId, IComputeBuffer recvBuffer, IComputeBuffer sendBuffer) => new BufferedReaderWriterChannel(socket, channelId, recvBuffer, sendBuffer);
+	}
+
+	/// <summary>
+	/// Extension methods for <see cref="IComputeChannel"/>
+	/// </summary>
+	public static class ComputeChannelExtensions
+	{
+		/// <summary>
+		/// Reads a complete message from the given socket, retrying reads until the buffer is full.
+		/// </summary>
+		/// <param name="channel">Channel to read from</param>
+		/// <param name="buffer">Buffer to store the data</param>
+		/// <param name="cancellationToken">Cancellation token for the operation</param>
+		public static async ValueTask ReceiveMessageAsync(this IComputeChannel channel, Memory<byte> buffer, CancellationToken cancellationToken = default)
+		{
+			if (!await TryReceiveMessageAsync(channel, buffer, cancellationToken))
+			{
+				throw new EndOfStreamException();
+			}
+		}
+
+		/// <summary>
+		/// Reads either a full message or end of stream from the channel
+		/// </summary>
+		/// <param name="channel">Channel to read from</param>
+		/// <param name="buffer">Buffer to store the data</param>
+		/// <param name="cancellationToken">Cancellation token for the operation</param>
+		public static async ValueTask<bool> TryReceiveMessageAsync(this IComputeChannel channel, Memory<byte> buffer, CancellationToken cancellationToken = default)
+		{
+			for (int offset = 0; offset < buffer.Length;)
+			{
+				int read = await channel.ReceiveAsync(buffer.Slice(offset), cancellationToken);
+				if (read == 0)
+				{
+					return false;
+				}
+				offset += read;
+			}
+			return true;
+		}
+	}
+}
diff --git a/Engine/Source/Programs/Shared/EpicGames.Horde/Compute/IComputeMessage.cs b/Engine/Source/Programs/Shared/EpicGames.Horde/Compute/IComputeMessage.cs
--- a/Engine/Source/Programs/Shared/EpicGames.Horde/Compute/IComputeMessage.cs
+++ b/Engine/Source/Programs/Shared/EpicGames.Horde/Compute/IComputeMessage.cs
@@ -4,11 +4,9 @@ using System;
 using System.Buffers.Binary;
 using System.Collections.Generic;
 using System.IO;
-using System.Text;
 using System.Threading;
 using System.Threading.Tasks;
 using EpicGames.Core;
-using EpicGames.Horde.Compute.Buffers;
 using EpicGames.Horde.Storage;
 using EpicGames.Horde.Storage.Nodes;
 
@@ -34,16 +32,6 @@ namespace EpicGames.Horde.Compute
 		/// </summary>
 		Exception = 0x02,
 
-		/// <summary>
-		/// Fork a new request channel
-		/// </summary>
-		Fork = 0x03,
-
-		/// <summary>
-		/// Untyped user data
-		/// </summary>
-		UserData = 0x04,
-
 		#region Process Management
 
 		/// <summary>
@@ -128,11 +116,6 @@ namespace EpicGames.Horde.Compute
 	/// </summary>
 	public record struct ExceptionMessage(string Message, string Description);
 
-	/// <summary>
-	/// Message to request forking the message loop
-	/// </summary>
-	public record struct ForkMessage(int ChannelId);
-
 	/// <summary>
 	/// Extract files from a bundle to a path in the remote sandbox
 	/// </summary>
@@ -212,35 +195,6 @@ namespace EpicGames.Horde.Compute
 			string description = message.ReadString();
 			return new ExceptionMessage(msg, description);
 		}
-/*
-		/// <inheritdoc/>
-		public static void Fork(this IComputeMessageChannel channel, int newChannelId)
-		{
-			using (IComputeMessageBuilder builder = channel.CreateMessage(ComputeMessageType.Fork))
-			{
-				builder.WriteInt32(newChannelId);
-				builder.Send();
-			}
-		}
-*/
-		/// <summary>
-		/// Parses a fork message from the given compute message
-		/// </summary>
-		public static ForkMessage ParseForkMessage(this IComputeMessage message)
-		{
-			int channelId = message.ReadInt32();
-			return new ForkMessage(channelId);
-		}
-
-		/// <summary>
-		/// Sends untyped user data to the remote
-		/// </summary>
-		public static async ValueTask SendUserDataAsync(this IComputeMessageChannel channel, ReadOnlyMemory<byte> memory, CancellationToken cancellationToken)
-		{
-			using IComputeMessageBuilder message = await channel.CreateMessageAsync(ComputeMessageType.UserData, memory.Length, cancellationToken);
-			message.WriteFixedLengthBytes(memory.Span);
-			message.Send();
-		}
 
 		#region Process
 
@@ -286,7 +240,7 @@ namespace EpicGames.Horde.Compute
 			using IComputeMessage response = await RunStorageServer(channel, storage, cancellationToken);
 			if (response.Type != ComputeMessageType.WriteFilesResponse)
 			{
-				throw new NotSupportedException();
+				throw new ComputeInvalidMessageException(response);
 			}
 		}
 
@@ -413,7 +367,7 @@ namespace EpicGames.Horde.Compute
 			readonly IComputeMessage _message;
 
 			public BlobDataStream(IComputeMessage message)
-				: base(message.Data)
+				: base(message.Data.Slice(8))
 			{
 				_message = message;
 			}
@@ -448,14 +402,45 @@ namespace EpicGames.Horde.Compute
 				request.Send();
 			}
 
-			IComputeMessage response = await channel.ReceiveAsync(cancellationToken);
-			if (response.Type != ComputeMessageType.ReadBlobResponse)
+			byte[]? buffer = null;
+			for(; ;)
 			{
-				response.Dispose();
-				throw new ComputeInvalidMessageException(response);
+				IComputeMessage? response = null;
+				try
+				{
+					response = await channel.ReceiveAsync(cancellationToken);
+					if (response.Type != ComputeMessageType.ReadBlobResponse)
+					{
+						throw new ComputeInvalidMessageException(response);
+					}
+
+					int chunkOffset = BinaryPrimitives.ReadInt32LittleEndian(response.Data.Span.Slice(0, 4));
+					int chunkLength = response.Data.Length - 8;
+					int totalLength = BinaryPrimitives.ReadInt32LittleEndian(response.Data.Span.Slice(4, 4));
+
+					if (chunkOffset == 0 && chunkLength == totalLength)
+					{
+						BlobDataStream stream = new BlobDataStream(response);
+						response = null;
+						return stream;
+					}
+
+					buffer ??= new byte[totalLength];
+					response.Data.Slice(8).CopyTo(buffer.AsMemory(chunkOffset));
+
+					if (chunkOffset + chunkLength == totalLength)
+					{
+						break;
+					}
+				}
+				catch
+				{
+					response?.Dispose();
+					throw;
+				}
 			}
 
-			return new BlobDataStream(response);
+			return new ReadOnlyMemoryStream(buffer);
 		}
 
 		/// <summary>
@@ -501,10 +486,18 @@ namespace EpicGames.Horde.Compute
 				}
 			}
 
-			using (IComputeMessageBuilder response = await channel.CreateMessageAsync(ComputeMessageType.ReadBlobResponse, data.Length + 128, cancellationToken))
+			const int MaxChunkSize = 512 * 1024;
+			for (int chunkOffset = 0; chunkOffset < data.Length;)
 			{
-				response.WriteFixedLengthBytes(data);
-				response.Send();
+				int chunkLength = Math.Min(data.Length - chunkOffset, MaxChunkSize);
+				using (IComputeMessageBuilder response = await channel.CreateMessageAsync(ComputeMessageType.ReadBlobResponse, chunkLength + 128, cancellationToken))
+				{
+					response.WriteInt32(chunkOffset);
+					response.WriteInt32(data.Length);
+					response.WriteFixedLengthBytes(data.AsSpan(chunkOffset, chunkLength));
+					response.Send();
+				}
+				chunkOffset += chunkLength;
 			}
 		}
 
diff --git a/Engine/Source/Programs/Shared/EpicGames.Horde/Compute/IComputeMessageChannel.cs b/Engine/Source/Programs/Shared/EpicGames.Horde/Compute/IComputeMessageChannel.cs
--- a/Engine/Source/Programs/Shared/EpicGames.Horde/Compute/IComputeMessageChannel.cs
+++ b/Engine/Source/Programs/Shared/EpicGames.Horde/Compute/IComputeMessageChannel.cs
@@ -10,7 +10,7 @@ namespace EpicGames.Horde.Compute
 	/// <summary>
 	/// Full-duplex channel for sending and receiving messages
 	/// </summary>
-	public interface IComputeMessageChannel : IAsyncDisposable
+	public interface IComputeMessageChannel : IDisposable
 	{
 		/// <summary>
 		/// Reads a message from the channel
@@ -43,7 +43,7 @@ namespace EpicGames.Horde.Compute
 	/// <summary>
 	/// Extension methods to allow creating channels from leases
 	/// </summary>
-	public static class ComputeChannelExtensions
+	public static class ComputeMessageChannelExtensions
 	{
 		/// <summary>
 		/// Waits for the remote machine to send a 'ready' response
diff --git a/Engine/Source/Programs/Shared/EpicGames.Horde/Compute/IComputeSocket.cs b/Engine/Source/Programs/Shared/EpicGames.Horde/Compute/IComputeSocket.cs
--- a/Engine/Source/Programs/Shared/EpicGames.Horde/Compute/IComputeSocket.cs
+++ b/Engine/Source/Programs/Shared/EpicGames.Horde/Compute/IComputeSocket.cs
@@ -1,52 +1,46 @@
 // Copyright Epic Games, Inc. All Rights Reserved.
 
 using System;
-using System.Globalization;
 using System.Threading;
 using System.Threading.Tasks;
-using EpicGames.Horde.Compute.Sockets;
+using EpicGames.Horde.Compute.Buffers;
 using Microsoft.Extensions.Logging;
 
 namespace EpicGames.Horde.Compute
 {
 	/// <summary>
-	/// Socket for sending and reciving messages
+	/// Socket for sending and reciving data using a "push" model. The application can attach multiple writers to accept received data.
 	/// </summary>
 	public interface IComputeSocket : IAsyncDisposable
 	{
 		/// <summary>
-		/// Creates a new buffer of the given capacity, suitable for this socket type
+		/// Sends data to a remote channel
 		/// </summary>
-		/// <param name="capacity">Capacity of the buffer</param>
-		IComputeBuffer CreateBuffer(long capacity);
+		/// <param name="channelId">Channel to write to</param>
+		/// <param name="memory">Memory to write</param>
+		/// <param name="cancellationToken">Cancellation token for the operation</param>
+		ValueTask SendAsync(int channelId, ReadOnlyMemory<byte> memory, CancellationToken cancellationToken = default);
 
 		/// <summary>
-		/// Attaches a buffer to receive data.
+		/// Marks a channel as complete
 		/// </summary>
-		/// <param name="channelId">Channel to receive data on</param>
-		/// <param name="writer">Writer for the received data</param>
+		/// <param name="channelId">Channel to mark as complete</param>
 		/// <param name="cancellationToken">Cancellation token for the operation</param>
-		ValueTask AttachRecvBufferAsync(int channelId, IComputeBufferWriter writer, CancellationToken cancellationToken);
+		ValueTask MarkCompleteAsync(int channelId, CancellationToken cancellationToken = default);
 
 		/// <summary>
 		/// Attaches a buffer to receive data.
 		/// </summary>
 		/// <param name="channelId">Channel to receive data on</param>
-		/// <param name="reader">Reader for incoming data</param>
-		/// <param name="cancellationToken">Cancellation token for the operation</param>
-		ValueTask AttachSendBufferAsync(int channelId, IComputeBufferReader reader, CancellationToken cancellationToken);
-	}
+		/// <param name="writer">Writer for the received data</param>
+		void AttachRecvBuffer(int channelId, IComputeBufferWriter writer);
 
-	/// <summary>
-	/// Utility methods for creating sockets
-	/// </summary>
-	public static class ComputeSocket
-	{
 		/// <summary>
-		/// Creates a socket for a worker
+		/// Attaches a buffer to send data.
 		/// </summary>
-		/// <param name="logger">Logger for diagnostic messages</param>
-		public static IComputeSocket ConnectAsWorker(ILogger logger) => WorkerComputeSocket.FromEnvironment(logger);
+		/// <param name="channelId">Channel to receive data on</param>
+		/// <param name="reader">Reader for the data to send</param>
+		void AttachSendBuffer(int channelId, IComputeBufferReader reader);
 	}
 
 	/// <summary>
@@ -54,56 +48,6 @@ namespace EpicGames.Horde.Compute
 	/// </summary>
 	public static class ComputeSocketExtensions
 	{
-		/// <summary>
-		/// Attaches a buffer to receive data.
-		/// </summary>
-		/// <param name="socket">Socket to attach to</param>
-		/// <param name="channelId">Channel to receive data on</param>
-		/// <param name="capacity">Capacity for the buffer</param>
-		/// <param name="cancellationToken">Cancellation token for the operation</param>
-		public static ValueTask<IComputeBufferReader> AttachRecvBufferAsync(this IComputeSocket socket, int channelId, long capacity, CancellationToken cancellationToken = default)
-		{
-			return AttachRecvBufferAsync(socket, channelId, socket.CreateBuffer(capacity), cancellationToken);
-		}
-
-		/// <summary>
-		/// Attaches a buffer to receive data.
-		/// </summary>
-		/// <param name="socket">Socket to attach to</param>
-		/// <param name="channelId">Channel to receive data on</param>
-		/// <param name="buffer">Buffer to attach to the socket</param>
-		/// <param name="cancellationToken">Cancellation token for the operation</param>
-		public static async ValueTask<IComputeBufferReader> AttachRecvBufferAsync(this IComputeSocket socket, int channelId, IComputeBuffer buffer, CancellationToken cancellationToken = default)
-		{
-			(IComputeBufferReader reader, IComputeBufferWriter writer) = buffer.ToShared();
-			await socket.AttachRecvBufferAsync(channelId, writer, cancellationToken);
-			return reader;
-		}
-
-		/// <summary>
-		/// Attaches a buffer from which to send data.
-		/// </summary>
-		/// <param name="socket">Socket to attach to</param>
-		/// <param name="channelId">Channel to receive data on</param>
-		/// <param name="capacity">Capacity for the buffer</param>
-		/// <param name="cancellationToken">Cancellation token for the operation</param>
-		public static ValueTask<IComputeBufferWriter> AttachSendBufferAsync(this IComputeSocket socket, int channelId, long capacity, CancellationToken cancellationToken = default) 
-			=> AttachSendBufferAsync(socket, channelId, socket.CreateBuffer(capacity), cancellationToken);
-
-		/// <summary>
-		/// Attaches a buffer from which to send data.
-		/// </summary>
-		/// <param name="socket">Socket to attach to</param>
-		/// <param name="channelId">Channel to receive data on</param>
-		/// <param name="buffer">Buffer for the data to send</param>
-		/// <param name="cancellationToken">Cancellation token for the operation</param>
-		public static async ValueTask<IComputeBufferWriter> AttachSendBufferAsync(this IComputeSocket socket, int channelId, IComputeBuffer buffer, CancellationToken cancellationToken = default)
-		{
-			(IComputeBufferReader reader, IComputeBufferWriter writer) = buffer.ToShared();
-			await socket.AttachSendBufferAsync(channelId, reader, cancellationToken);
-			return writer;
-		}
-
 		/// <summary>
 		/// Creates a message channel with the given identifier
 		/// </summary>
@@ -111,8 +55,8 @@ namespace EpicGames.Horde.Compute
 		/// <param name="channelId">Identifier for the channel</param>
 		/// <param name="logger">Logger for the channel</param>
 		/// <param name="cancellationToken">Cancellation token for the operation</param>
-		public static ValueTask<IComputeMessageChannel> CreateMessageChannelAsync(this IComputeSocket socket, int channelId, ILogger logger, CancellationToken cancellationToken = default) 
-			=> socket.CreateMessageChannelAsync(channelId, 65536, logger, cancellationToken);
+		public static IComputeMessageChannel CreateMessageChannel(this IComputeSocket socket, int channelId, ILogger logger, CancellationToken cancellationToken = default) 
+			=> socket.CreateMessageChannel(channelId, 65536, logger, cancellationToken);
 
 		/// <summary>
 		/// Creates a message channel with the given identifier
@@ -122,8 +66,8 @@ namespace EpicGames.Horde.Compute
 		/// <param name="bufferSize">Size of the send and receive buffer</param>
 		/// <param name="logger">Logger for the channel</param>
 		/// <param name="cancellationToken">Cancellation token for the operation</param>
-		public static ValueTask<IComputeMessageChannel> CreateMessageChannelAsync(this IComputeSocket socket, int channelId, long bufferSize, ILogger logger, CancellationToken cancellationToken = default)
-			=> socket.CreateMessageChannelAsync(channelId, bufferSize, bufferSize, logger, cancellationToken);
+		public static IComputeMessageChannel CreateMessageChannel(this IComputeSocket socket, int channelId, int bufferSize, ILogger logger, CancellationToken cancellationToken = default)
+			=> socket.CreateMessageChannel(channelId, bufferSize, bufferSize, logger, cancellationToken);
 
 		/// <summary>
 		/// Creates a message channel with the given identifier
@@ -134,11 +78,11 @@ namespace EpicGames.Horde.Compute
 		/// <param name="recvBufferSize">Size of the recieve buffer</param>
 		/// <param name="logger">Logger for the channel</param>
 		/// <param name="cancellationToken">Cancellation token for the operation</param>
-		public static async ValueTask<IComputeMessageChannel> CreateMessageChannelAsync(this IComputeSocket socket, int channelId, long sendBufferSize, long recvBufferSize, ILogger logger, CancellationToken cancellationToken = default)
+		public static IComputeMessageChannel CreateMessageChannel(this IComputeSocket socket, int channelId, int sendBufferSize, int recvBufferSize, ILogger logger, CancellationToken cancellationToken = default)
 		{
-			IComputeBufferReader reader = await socket.AttachRecvBufferAsync(channelId, recvBufferSize, cancellationToken);
-			IComputeBufferWriter writer = await socket.AttachSendBufferAsync(channelId, sendBufferSize, cancellationToken);
-			return new ComputeMessageChannel(reader, writer, logger);
+			PooledBuffer sendBuffer = new PooledBuffer(sendBufferSize);
+			PooledBuffer recvBuffer = new PooledBuffer(recvBufferSize);
+			return new ComputeMessageChannel(socket, channelId, sendBuffer, recvBuffer, logger);
 		}
 	}
 }
diff --git a/Engine/Source/Programs/Shared/EpicGames.Horde/Compute/IpcMessages.cs b/Engine/Source/Programs/Shared/EpicGames.Horde/Compute/IpcMessages.cs
deleted file mode 100644
diff --git a/Engine/Source/Programs/Shared/EpicGames.Horde/Compute/Native.cs b/Engine/Source/Programs/Shared/EpicGames.Horde/Compute/Native.cs
--- a/Engine/Source/Programs/Shared/EpicGames.Horde/Compute/Native.cs
+++ b/Engine/Source/Programs/Shared/EpicGames.Horde/Compute/Native.cs
@@ -1,7 +1,9 @@
 // Copyright Epic Games, Inc. All Rights Reserved.
 
 using System;
+using System.ComponentModel;
 using System.IO;
+using System.Net.Http.Headers;
 using System.Runtime.InteropServices;
 using System.Threading;
 using Microsoft.Win32.SafeHandles;
@@ -21,6 +23,9 @@ namespace EpicGames.Horde.Compute
 		[DllImport("kernel32.dll")]
 		public static extern SafeWaitHandle CreateEvent(SECURITY_ATTRIBUTES lpEventAttributes, bool bManualReset, bool bInitialState, string? lpName);
 
+		[DllImport("Kernel32.dll", SetLastError = true)]
+		public static extern SafeWaitHandle OpenEvent(uint dwDesiredAccess, bool bInheritHandle, string lpName);
+
 		[DllImport("kernel32.dll", SetLastError = true)]
 		public static extern bool ResetEvent(SafeWaitHandle hEvent);
 
@@ -29,6 +34,9 @@ namespace EpicGames.Horde.Compute
 
 		public const uint DUPLICATE_SAME_ACCESS = 2;
 
+		public const uint SYNCHRONIZE = 0x00100000;
+		public const uint EVENT_MODIFY_STATE = 2;
+
 		[DllImport("kernel32.dll", SetLastError = true)]
 		public static extern IntPtr GetCurrentProcess();
 
@@ -38,12 +46,23 @@ namespace EpicGames.Horde.Compute
 
 		public class EventHandle : WaitHandle
 		{
-			public EventHandle(HandleInheritability handleInheritability)
+			public EventHandle(SafeWaitHandle handle)
+			{
+				SafeWaitHandle = handle;
+			}
+
+			public static EventHandle CreateNew(string? name, EventResetMode resetMode, bool initialState, HandleInheritability handleInheritability)
 			{
 				Native.SECURITY_ATTRIBUTES securityAttributes = new Native.SECURITY_ATTRIBUTES();
 				securityAttributes.nLength = Marshal.SizeOf(securityAttributes);
 				securityAttributes.bInheritHandle = (handleInheritability == HandleInheritability.Inheritable) ? 1 : 0;
-				SafeWaitHandle = CreateEvent(securityAttributes, true, false, null);
+
+				SafeWaitHandle handle = CreateEvent(securityAttributes, resetMode == EventResetMode.ManualReset, initialState, name);
+				if (handle.IsInvalid)
+				{
+					throw new Win32Exception();
+				}
+				return new EventHandle(handle);
 			}
 
 			public EventHandle(IntPtr handle, bool ownsHandle)
@@ -51,6 +70,16 @@ namespace EpicGames.Horde.Compute
 				SafeWaitHandle = new SafeWaitHandle(handle, ownsHandle);
 			}
 
+			public static EventHandle OpenExisting(string name)
+			{
+				SafeWaitHandle handle = OpenEvent(SYNCHRONIZE | EVENT_MODIFY_STATE, true, name);
+				if (handle.IsInvalid)
+				{
+					throw new Win32Exception();
+				}
+				return new EventHandle(handle);
+			}
+
 			public void Set() => SetEvent(SafeWaitHandle);
 
 			public void Reset() => ResetEvent(SafeWaitHandle);
diff --git a/Engine/Source/Programs/Shared/EpicGames.Horde/Compute/Sockets/ClientComputeSocket.cs b/Engine/Source/Programs/Shared/EpicGames.Horde/Compute/Sockets/ClientComputeSocket.cs
--- a/Engine/Source/Programs/Shared/EpicGames.Horde/Compute/Sockets/ClientComputeSocket.cs
+++ b/Engine/Source/Programs/Shared/EpicGames.Horde/Compute/Sockets/ClientComputeSocket.cs
@@ -26,7 +26,7 @@ namespace EpicGames.Horde.Compute.Sockets
 		{
 			if (capacity > Int32.MaxValue)
 			{
-				return new SharedMemoryBuffer(capacity);
+				return SharedMemoryBuffer.CreateNew(null, capacity);
 			}
 			else
 			{
diff --git a/Engine/Source/Programs/Shared/EpicGames.Horde/Compute/Sockets/WorkerComputeSocket.cs b/Engine/Source/Programs/Shared/EpicGames.Horde/Compute/Sockets/WorkerComputeSocket.cs
deleted file mode 100644
