Horde: Always use the same ComputeSocket implementation.

diff --git a/Engine/Source/Programs/Horde/Horde.Agent/Commands/Compute/ComputeCommand.cs b/Engine/Source/Programs/Horde/Horde.Agent/Commands/Compute/ComputeCommand.cs
--- a/Engine/Source/Programs/Horde/Horde.Agent/Commands/Compute/ComputeCommand.cs
+++ b/Engine/Source/Programs/Horde/Horde.Agent/Commands/Compute/ComputeCommand.cs
@@ -12,7 +12,6 @@ using EpicGames.Core;
 using EpicGames.Horde.Common;
 using EpicGames.Horde.Compute;
 using EpicGames.Horde.Compute.Clients;
-using EpicGames.Horde.Compute.Sockets;
 using EpicGames.Horde.Compute.Transports;
 using Microsoft.Extensions.Caching.Memory;
 using Microsoft.Extensions.DependencyInjection;
@@ -128,7 +127,7 @@ namespace Horde.Agent.Commands.Compute
 			using Socket tcpSocket = new Socket(SocketType.Stream, ProtocolType.IP);
 			await tcpSocket.ConnectAsync(IPAddress.Loopback, Port);
 
-			await using (ClientComputeSocket socket = new ClientComputeSocket(new TcpTransport(tcpSocket), logger))
+			await using (IComputeSocket socket = ComputeSocket.Create(new TcpTransport(tcpSocket), logger))
 			{
 				logger.LogInformation("Running worker...");
 				await RunWorkerAsync(socket, _memoryCache, logger, CancellationToken.None);
diff --git a/Engine/Source/Programs/Horde/Horde.Agent/Leases/Handlers/ComputeHandler.cs b/Engine/Source/Programs/Horde/Horde.Agent/Leases/Handlers/ComputeHandler.cs
--- a/Engine/Source/Programs/Horde/Horde.Agent/Leases/Handlers/ComputeHandler.cs
+++ b/Engine/Source/Programs/Horde/Horde.Agent/Leases/Handlers/ComputeHandler.cs
@@ -6,7 +6,6 @@ using System.Threading;
 using System.Threading.Tasks;
 using EpicGames.Core;
 using EpicGames.Horde.Compute;
-using EpicGames.Horde.Compute.Sockets;
 using EpicGames.Horde.Compute.Transports;
 using Horde.Agent.Services;
 using HordeCommon.Rpc.Tasks;
@@ -55,7 +54,7 @@ namespace Horde.Agent.Leases.Handlers
 
 				_logger.LogInformation("Matched connection for {Nonce}", StringUtils.FormatHexString(computeTask.Nonce.Span));
 
-				await using (ClientComputeSocket socket = new ClientComputeSocket(new TcpTransport(tcpClient.Client), _logger))
+				await using (IComputeSocket socket = ComputeSocket.Create(new TcpTransport(tcpClient.Client), _logger))
 				{
 					ComputeWorker worker = new ComputeWorker(_sandboxDir, _memoryCache, _logger);
 					await worker.RunAsync(socket, cancellationToken);
diff --git a/Engine/Source/Programs/Shared/EpicGames.Horde.Tests/BufferTests.cs b/Engine/Source/Programs/Shared/EpicGames.Horde.Tests/BufferTests.cs
--- a/Engine/Source/Programs/Shared/EpicGames.Horde.Tests/BufferTests.cs
+++ b/Engine/Source/Programs/Shared/EpicGames.Horde.Tests/BufferTests.cs
@@ -11,7 +11,6 @@ using System.Threading;
 using System.Threading.Tasks;
 using EpicGames.Horde.Compute;
 using EpicGames.Horde.Compute.Buffers;
-using EpicGames.Horde.Compute.Sockets;
 using EpicGames.Horde.Compute.Transports;
 using Microsoft.Extensions.Logging.Abstractions;
 using Microsoft.VisualStudio.TestPlatform.CrossPlatEngine;
@@ -46,8 +45,8 @@ namespace EpicGames.Horde.Tests
 			Pipe sourceToTargetPipe = new Pipe();
 			Pipe targetToSourcePipe = new Pipe();
 
-			await using IComputeSocket producerSocket = new ClientComputeSocket(new PipeTransport(targetToSourcePipe.Reader, sourceToTargetPipe.Writer), NullLogger.Instance);
-			await using IComputeSocket consumerSocket = new ClientComputeSocket(new PipeTransport(sourceToTargetPipe.Reader, targetToSourcePipe.Writer), NullLogger.Instance);
+			await using IComputeSocket producerSocket = ComputeSocket.Create(new PipeTransport(targetToSourcePipe.Reader, sourceToTargetPipe.Writer), NullLogger.Instance);
+			await using IComputeSocket consumerSocket = ComputeSocket.Create(new PipeTransport(sourceToTargetPipe.Reader, targetToSourcePipe.Writer), NullLogger.Instance);
 
 			using IComputeBuffer consumerBuffer = createBuffer(Length);
 			consumerSocket.AttachRecvBuffer(ChannelId, consumerBuffer.Writer);
diff --git a/Engine/Source/Programs/Shared/EpicGames.Horde/Compute/Clients/AgentComputeClient.cs b/Engine/Source/Programs/Shared/EpicGames.Horde/Compute/Clients/AgentComputeClient.cs
--- a/Engine/Source/Programs/Shared/EpicGames.Horde/Compute/Clients/AgentComputeClient.cs
+++ b/Engine/Source/Programs/Shared/EpicGames.Horde/Compute/Clients/AgentComputeClient.cs
@@ -10,7 +10,6 @@ using System.Runtime.InteropServices;
 using System.Threading;
 using System.Threading.Tasks;
 using EpicGames.Core;
-using EpicGames.Horde.Compute.Sockets;
 using EpicGames.Horde.Compute.Transports;
 using Microsoft.Extensions.Logging;
 
@@ -90,7 +89,7 @@ namespace EpicGames.Horde.Compute.Clients
 			using BackgroundTask agentTask = BackgroundTask.StartNew(ctx => RunAgentAsync(_hordeAgentAssembly, _port, _logger, ctx));
 			using Socket tcpSocket = await listener.AcceptAsync(cancellationToken);
 
-			await using ClientComputeSocket socket = new ClientComputeSocket(new TcpTransport(tcpSocket), _logger);
+			await using IComputeSocket socket = ComputeSocket.Create(new TcpTransport(tcpSocket), _logger);
 			yield return socket;
 
 			await socket.CloseAsync(cancellationToken);
diff --git a/Engine/Source/Programs/Shared/EpicGames.Horde/Compute/Clients/LocalComputeClient.cs b/Engine/Source/Programs/Shared/EpicGames.Horde/Compute/Clients/LocalComputeClient.cs
--- a/Engine/Source/Programs/Shared/EpicGames.Horde/Compute/Clients/LocalComputeClient.cs
+++ b/Engine/Source/Programs/Shared/EpicGames.Horde/Compute/Clients/LocalComputeClient.cs
@@ -7,7 +7,6 @@ using System.Net.Sockets;
 using System.Threading;
 using System.Threading.Tasks;
 using EpicGames.Core;
-using EpicGames.Horde.Compute.Sockets;
 using EpicGames.Horde.Compute.Transports;
 using Microsoft.Extensions.Caching.Memory;
 using Microsoft.Extensions.Logging;
@@ -25,9 +24,9 @@ namespace EpicGames.Horde.Compute.Clients
 			public IReadOnlyDictionary<string, int> AssignedResources => new Dictionary<string, int>();
 			public IComputeSocket Socket => _socket;
 
-			readonly ClientComputeSocket _socket;
+			readonly IComputeSocket _socket;
 
-			public LeaseImpl(ClientComputeSocket socket) => _socket = socket;
+			public LeaseImpl(IComputeSocket socket) => _socket = socket;
 			public ValueTask DisposeAsync() => _socket.DisposeAsync();
 		}
 
@@ -73,7 +72,7 @@ namespace EpicGames.Horde.Compute.Clients
 
 			using MemoryCache memoryCache = new MemoryCache(new MemoryCacheOptions { SizeLimit = 10 * 1024 * 1024 });
 
-			await using (ClientComputeSocket socket = new ClientComputeSocket(new TcpTransport(tcpSocket), logger))
+			await using (IComputeSocket socket = ComputeSocket.Create(new TcpTransport(tcpSocket), logger))
 			{
 				ComputeWorker worker = new ComputeWorker(sandboxDir, memoryCache, logger);
 				await worker.RunAsync(socket, cancellationToken);
@@ -85,7 +84,7 @@ namespace EpicGames.Horde.Compute.Clients
 		public Task<IComputeLease?> TryAssignWorkerAsync(ClusterId clusterId, Requirements? requirements, CancellationToken cancellationToken)
 		{
 #pragma warning disable CA2000 // Dispose objects before losing scope
-			ClientComputeSocket socket = new ClientComputeSocket(new TcpTransport(_socket), _logger);
+			IComputeSocket socket = ComputeSocket.Create(new TcpTransport(_socket), _logger);
 			return Task.FromResult<IComputeLease?>(new LeaseImpl(socket));
 #pragma warning restore CA2000 // Dispose objects before losing scope
 		}
diff --git a/Engine/Source/Programs/Shared/EpicGames.Horde/Compute/Clients/ServerComputeClient.cs b/Engine/Source/Programs/Shared/EpicGames.Horde/Compute/Clients/ServerComputeClient.cs
--- a/Engine/Source/Programs/Shared/EpicGames.Horde/Compute/Clients/ServerComputeClient.cs
+++ b/Engine/Source/Programs/Shared/EpicGames.Horde/Compute/Clients/ServerComputeClient.cs
@@ -1,7 +1,6 @@
 // Copyright Epic Games, Inc. All Rights Reserved.
 
 using EpicGames.Core;
-using EpicGames.Horde.Compute.Sockets;
 using EpicGames.Horde.Compute.Transports;
 using Microsoft.Extensions.Logging;
 using System;
@@ -179,7 +178,7 @@ namespace EpicGames.Horde.Compute.Clients
 			// Pass the rest of the call over to the handler
 			byte[] key = StringUtils.ParseHexString(responseMessage.Key);
 
-			await using ClientComputeSocket computeSocket = new ClientComputeSocket(new TcpTransport(socket), _logger);
+			await using IComputeSocket computeSocket = ComputeSocket.Create(new TcpTransport(socket), _logger);
 			yield return new LeaseInfo(responseMessage.Properties, responseMessage.AssignedResources, computeSocket);
 		}
 	}
diff --git a/Engine/Source/Programs/Shared/EpicGames.Horde/Compute/ComputeSocketBase.cs b/Engine/Source/Programs/Shared/EpicGames.Horde/Compute/ComputeSocketBase.cs
deleted file mode 100644
diff --git a/Engine/Source/Programs/Shared/EpicGames.Horde/Compute/ComputeWorker.cs b/Engine/Source/Programs/Shared/EpicGames.Horde/Compute/ComputeWorker.cs
--- a/Engine/Source/Programs/Shared/EpicGames.Horde/Compute/ComputeWorker.cs
+++ b/Engine/Source/Programs/Shared/EpicGames.Horde/Compute/ComputeWorker.cs
@@ -1,7 +1,6 @@
 // Copyright Epic Games, Inc. All Rights Reserved.
 
 using System;
-using System.Buffers.Binary;
 using System.Collections.Generic;
 using System.Diagnostics;
 using System.Runtime.InteropServices;
@@ -9,7 +8,6 @@ using System.Threading;
 using System.Threading.Tasks;
 using EpicGames.Core;
 using EpicGames.Horde.Compute.Buffers;
-using EpicGames.Horde.Compute.Sockets;
 using EpicGames.Horde.Storage;
 using EpicGames.Horde.Storage.Nodes;
 using Microsoft.Extensions.Caching.Memory;
diff --git a/Engine/Source/Programs/Shared/EpicGames.Horde/Compute/DefaultComputeSocket.cs b/Engine/Source/Programs/Shared/EpicGames.Horde/Compute/DefaultComputeSocket.cs
new file mode 100644
--- /dev/null
+++ b/Engine/Source/Programs/Shared/EpicGames.Horde/Compute/DefaultComputeSocket.cs
@@ -0,0 +1,288 @@
+// Copyright Epic Games, Inc. All Rights Reserved.
+
+using System;
+using System.Buffers;
+using System.Buffers.Binary;
+using System.Collections.Generic;
+using System.Threading;
+using System.Threading.Tasks;
+using EpicGames.Core;
+using Microsoft.Extensions.Logging;
+
+namespace EpicGames.Horde.Compute
+{
+	/// <summary>
+	/// Manages a set of readers and writers to buffers across a transport layer
+	/// </summary>
+	class DefaultComputeSocket : IComputeSocket, IAsyncDisposable
+	{
+		readonly object _lockObject = new object();
+
+		bool _complete;
+
+		readonly IComputeTransport _transport;
+		readonly ILogger _logger;
+
+		readonly CancellationTokenSource _cancellationSource = new CancellationTokenSource();
+
+		readonly BackgroundTask _recvTask;
+		readonly Dictionary<int, IComputeBufferWriter> _recvBuffers = new Dictionary<int, IComputeBufferWriter>();
+
+		readonly SemaphoreSlim _sendSemaphore = new SemaphoreSlim(1, 1);
+		readonly Dictionary<int, Task> _sendTasks = new Dictionary<int, Task>();
+
+		/// <summary>
+		/// Constructor
+		/// </summary>
+		/// <param name="transport">Transport to communicate with the remote</param>
+		/// <param name="logger">Logger for trace output</param>
+		public DefaultComputeSocket(IComputeTransport transport, ILogger logger)
+		{
+			_transport = transport;
+			_logger = logger;
+			_recvTask = BackgroundTask.StartNew(ctx => RunReaderAsync(transport, ctx));
+		}
+
+		/// <summary>
+		/// Attempt to gracefully close the current connection and shutdown both ends of the transport
+		/// </summary>
+		public async ValueTask CloseAsync(CancellationToken cancellationToken)
+		{
+			// Wait for all the individual send tasks to complete
+			await Task.WhenAll(_sendTasks.Values);
+			cancellationToken.ThrowIfCancellationRequested();
+
+			// Send a final message indicating that the lease is done. This will allow the senders on the remote end to terminate, and trigger
+			// a message to be sent back to our read task allowing it to shut down gracefully.
+			byte[] header = new byte[8];
+			await _transport.WriteAsync(header, cancellationToken);
+			_logger.LogTrace("Sent shutdown packet");
+
+			// Wait for the reader to stop
+			await _recvTask.Task.WaitAsync(TimeSpan.FromSeconds(5), cancellationToken);
+			await _recvTask.StopAsync();
+
+			_logger.LogTrace("Closed");
+		}
+
+		/// <inheritdoc/>
+		public async ValueTask DisposeAsync()
+		{
+			await _recvTask.DisposeAsync();
+			_sendSemaphore.Dispose();
+			_cancellationSource.Dispose();
+			GC.SuppressFinalize(this);
+		}
+
+		async Task RunReaderAsync(IComputeTransport transport, CancellationToken cancellationToken)
+		{
+			_logger.LogTrace("Started socket reader");
+
+			byte[] header = new byte[8];
+			try
+			{
+				// Maintain a local cache of buffers to be able to query for them without having to acquire a global lock
+				Dictionary<int, IComputeBufferWriter> cachedWriters = new Dictionary<int, IComputeBufferWriter>();
+
+				Memory<byte> last = Memory<byte>.Empty;
+
+				// Process messages from the remote
+				for (; ; )
+				{
+					// Read the next packet header
+					if (!await transport.ReadOptionalAsync(header, cancellationToken))
+					{
+						_logger.LogTrace("End of socket");
+						break;
+					}
+
+					// Parse the target buffer and packet size
+					int id = BinaryPrimitives.ReadInt32LittleEndian(header);
+					int size = BinaryPrimitives.ReadInt32LittleEndian(header.AsSpan(4));
+
+					// If the size if negative, we're closing the entire connection
+					if (size == 0)
+					{
+						_logger.LogTrace("Received shutdown packet");
+						break;
+					}
+
+					// Dispatch it to the correct place
+					if (size < 0)
+					{
+						_logger.LogTrace("Detaching buffer {Id}", id);
+						DetachRecvBuffer(cachedWriters, id);
+					}
+					else
+					{
+						IComputeBufferWriter writer = GetReceiveBuffer(cachedWriters, id);
+						await ReadPacketAsync(transport, id, size, writer, cancellationToken);
+					}
+				}
+			}
+			catch (OperationCanceledException)
+			{
+			}
+
+			// Mark all buffers as complete
+			lock (_lockObject)
+			{
+				_complete = true;
+				foreach (IComputeBufferWriter writer in _recvBuffers.Values)
+				{
+					writer.MarkComplete();
+				}
+			}
+
+			_logger.LogTrace("Closing reader");
+		}
+
+		async Task ReadPacketAsync(IComputeTransport transport, int id, int size, IComputeBufferWriter writer, CancellationToken cancellationToken)
+		{
+			Memory<byte> memory = writer.GetMemory();
+			while (memory.Length < size)
+			{
+				_logger.LogTrace("No space in buffer {Id}, flushing", id);
+				await writer.WaitToWriteAsync(memory.Length, cancellationToken);
+				memory = writer.GetMemory();
+			}
+
+			for (int offset = 0; offset < size;)
+			{
+				int read = await transport.ReadPartialAsync(memory.Slice(offset, size - offset), cancellationToken);
+				offset += read;
+			}
+
+			writer.Advance(size);
+		}
+
+		IComputeBufferWriter GetReceiveBuffer(Dictionary<int, IComputeBufferWriter> cachedWriters, int id)
+		{
+			IComputeBufferWriter? writer;
+			if (cachedWriters.TryGetValue(id, out writer))
+			{
+				return writer;
+			}
+
+			lock (_lockObject)
+			{
+				if (_recvBuffers.TryGetValue(id, out writer))
+				{
+					cachedWriters.Add(id, writer);
+					return writer;
+				}
+			}
+
+			throw new ComputeInternalException($"No buffer is attached to channel {id}");
+		}
+
+		class SendSegment : ReadOnlySequenceSegment<byte>
+		{
+			public void Set(ReadOnlyMemory<byte> memory, ReadOnlySequenceSegment<byte>? next, long runningIndex)
+			{
+				Memory = memory;
+				Next = next;
+				RunningIndex = runningIndex;
+			}
+		}
+
+		readonly byte[] _header = new byte[8];
+		readonly SendSegment _headerSegment = new SendSegment();
+		readonly SendSegment _bodySegment = new SendSegment();
+
+		/// <inheritdoc/>
+		public async ValueTask SendAsync(int id, ReadOnlyMemory<byte> memory, CancellationToken cancellationToken)
+		{
+			if (memory.Length > 0)
+			{
+				await SendInternalAsync(id, memory.Length, memory, cancellationToken);
+			}
+		}
+
+		/// <inheritdoc/>
+		public ValueTask MarkCompleteAsync(int id, CancellationToken cancellationToken) => SendInternalAsync(id, -1, ReadOnlyMemory<byte>.Empty, cancellationToken);
+
+		async ValueTask SendInternalAsync(int id, int size, ReadOnlyMemory<byte> memory, CancellationToken cancellationToken)
+		{
+			await _sendSemaphore.WaitAsync(cancellationToken);
+			try
+			{
+				BinaryPrimitives.WriteInt32LittleEndian(_header, id);
+				BinaryPrimitives.WriteInt32LittleEndian(_header.AsSpan(4), size);
+				_headerSegment.Set(_header, _bodySegment, 0);
+				_bodySegment.Set(memory, null, _header.Length);
+
+				ReadOnlySequence<byte> sequence = new ReadOnlySequence<byte>(_headerSegment, 0, _bodySegment, memory.Length);
+				await _transport.WriteAsync(sequence, cancellationToken);
+			}
+			finally
+			{
+				_sendSemaphore.Release();
+			}
+		}
+
+		/// <inheritdoc/>
+		public void AttachRecvBuffer(int channelId, IComputeBufferWriter recvBuffer)
+		{
+			bool complete;
+			lock (_lockObject)
+			{
+				complete = _complete;
+				if (!complete)
+				{
+					_recvBuffers.Add(channelId, recvBuffer);
+				}
+			}
+
+			if (recvBuffer != null && complete)
+			{
+				recvBuffer.MarkComplete();
+			}
+		}
+
+		void DetachRecvBuffer(Dictionary<int, IComputeBufferWriter> cachedWriters, int id)
+		{
+			cachedWriters.Remove(id);
+
+			IComputeBufferWriter? writer;
+			lock (_lockObject)
+			{
+#pragma warning disable CA2000 // Dispose objects before losing scope
+				_recvBuffers.Remove(id, out writer);
+#pragma warning restore CA2000 // Dispose objects before losing scope
+			}
+			if (writer != null)
+			{
+				writer.MarkComplete();
+			}
+		}
+
+		/// <inheritdoc/>
+		public void AttachSendBuffer(int channelId, IComputeBufferReader reader)
+		{
+			lock (_lockObject)
+			{
+				_sendTasks.Add(channelId, Task.Run(() => SendFromBufferAsync(channelId, reader, _cancellationSource.Token), _cancellationSource.Token));
+			}
+		}
+
+		async Task SendFromBufferAsync(int channelId, IComputeBufferReader reader, CancellationToken cancellationToken)
+		{
+			while (!cancellationToken.IsCancellationRequested)
+			{
+				ReadOnlyMemory<byte> memory = reader.GetMemory();
+				if (memory.Length > 0)
+				{
+					await SendAsync(channelId, memory, cancellationToken);
+					reader.Advance(memory.Length);
+				}
+				if (reader.IsComplete)
+				{
+					await MarkCompleteAsync(channelId, cancellationToken);
+					break;
+				}
+				await reader.WaitToReadAsync(0, cancellationToken);
+			}
+		}
+	}
+}
diff --git a/Engine/Source/Programs/Shared/EpicGames.Horde/Compute/IComputeSocket.cs b/Engine/Source/Programs/Shared/EpicGames.Horde/Compute/IComputeSocket.cs
--- a/Engine/Source/Programs/Shared/EpicGames.Horde/Compute/IComputeSocket.cs
+++ b/Engine/Source/Programs/Shared/EpicGames.Horde/Compute/IComputeSocket.cs
@@ -41,6 +41,25 @@ namespace EpicGames.Horde.Compute
 		/// <param name="channelId">Channel to receive data on</param>
 		/// <param name="reader">Reader for the data to send</param>
 		void AttachSendBuffer(int channelId, IComputeBufferReader reader);
+
+		/// <summary>
+		/// Attempt to gracefully close the current connection and shutdown both ends of the transport
+		/// </summary>
+		/// <param name="cancellationToken">Cancellation token for the operation</param>
+		ValueTask CloseAsync(CancellationToken cancellationToken);
+	}
+
+	/// <summary>
+	/// Static methods for sockets
+	/// </summary>
+	public static class ComputeSocket
+	{
+		/// <summary>
+		/// Create a new compute socket
+		/// </summary>
+		/// <param name="transport">Transport to use for the socket</param>
+		/// <param name="logger">Logger for diagnostic output</param>
+		public static IComputeSocket Create(IComputeTransport transport, ILogger logger) => new DefaultComputeSocket(transport, logger);
 	}
 
 	/// <summary>
diff --git a/Engine/Source/Programs/Shared/EpicGames.Horde/Compute/Sockets/ClientComputeSocket.cs b/Engine/Source/Programs/Shared/EpicGames.Horde/Compute/Sockets/ClientComputeSocket.cs
deleted file mode 100644
