Horde: Remove unused CancellationToken from IComputeSocket.CreateMessageChannel() overloads.

diff --git a/Engine/Source/Programs/Horde/Horde.Agent/Commands/Compute/CppComputeCommand.cs b/Engine/Source/Programs/Horde/Horde.Agent/Commands/Compute/CppComputeCommand.cs
--- a/Engine/Source/Programs/Horde/Horde.Agent/Commands/Compute/CppComputeCommand.cs
+++ b/Engine/Source/Programs/Horde/Horde.Agent/Commands/Compute/CppComputeCommand.cs
@@ -56,7 +56,7 @@ namespace Horde.Agent.Commands.Compute
 			NodeLocator sandbox = await CreateSandboxAsync(TaskFile, storage, cancellationToken);
 
 			// Open a socket and upload the sandbox
-			using (IComputeMessageChannel channel = lease.Socket.CreateMessageChannel(0, 4 * 1024 * 1024, _logger, cancellationToken))
+			using (IComputeMessageChannel channel = lease.Socket.CreateMessageChannel(0, 4 * 1024 * 1024, _logger))
 			{
 				await channel.UploadFilesAsync("", sandbox, storage, cancellationToken);
 
diff --git a/Engine/Source/Programs/Horde/Horde.Agent/Commands/Compute/DdcComputeCommand.cs b/Engine/Source/Programs/Horde/Horde.Agent/Commands/Compute/DdcComputeCommand.cs
--- a/Engine/Source/Programs/Horde/Horde.Agent/Commands/Compute/DdcComputeCommand.cs
+++ b/Engine/Source/Programs/Horde/Horde.Agent/Commands/Compute/DdcComputeCommand.cs
@@ -31,7 +31,7 @@ namespace Horde.Agent.Commands.Compute
 		protected override async Task<bool> HandleRequestAsync(IComputeLease lease, CancellationToken cancellationToken)
 		{
 			IComputeSocket socket = lease.Socket;
-			using (IComputeMessageChannel channel = socket.CreateMessageChannel(0, _logger, cancellationToken))
+			using (IComputeMessageChannel channel = socket.CreateMessageChannel(0, _logger))
 			{
 				_logger.LogInformation("Sending XOR request");
 				await channel.SendXorRequestAsync(new byte[] { 1, 2, 3, 4, 5 }, (byte)123, cancellationToken);
diff --git a/Engine/Source/Programs/Horde/Horde.Agent/Commands/Compute/XorComputeCommand.cs b/Engine/Source/Programs/Horde/Horde.Agent/Commands/Compute/XorComputeCommand.cs
--- a/Engine/Source/Programs/Horde/Horde.Agent/Commands/Compute/XorComputeCommand.cs
+++ b/Engine/Source/Programs/Horde/Horde.Agent/Commands/Compute/XorComputeCommand.cs
@@ -31,7 +31,7 @@ namespace Horde.Agent.Commands.Compute
 		protected override async Task<bool> HandleRequestAsync(IComputeLease lease, CancellationToken cancellationToken)
 		{
 			IComputeSocket socket = lease.Socket;
-			using (IComputeMessageChannel mainChannel = socket.CreateMessageChannel(0, _logger, cancellationToken))
+			using (IComputeMessageChannel mainChannel = socket.CreateMessageChannel(0, _logger))
 			{
 				_logger.LogInformation("Forking compute channel...");
 
diff --git a/Engine/Source/Programs/Shared/EpicGames.Horde/Compute/ComputeWorker.cs b/Engine/Source/Programs/Shared/EpicGames.Horde/Compute/ComputeWorker.cs
--- a/Engine/Source/Programs/Shared/EpicGames.Horde/Compute/ComputeWorker.cs
+++ b/Engine/Source/Programs/Shared/EpicGames.Horde/Compute/ComputeWorker.cs
@@ -49,7 +49,7 @@ namespace EpicGames.Horde.Compute
 
 		async Task RunAsync(IComputeSocket socket, int channelId, CancellationToken cancellationToken)
 		{
-			using (IComputeMessageChannel channel = socket.CreateMessageChannel(channelId, 4 * 1024 * 1024, _logger, cancellationToken))
+			using (IComputeMessageChannel channel = socket.CreateMessageChannel(channelId, 4 * 1024 * 1024, _logger))
 			{
 				await channel.SendReadyAsync(cancellationToken);
 
diff --git a/Engine/Source/Programs/Shared/EpicGames.Horde/Compute/IComputeSocket.cs b/Engine/Source/Programs/Shared/EpicGames.Horde/Compute/IComputeSocket.cs
--- a/Engine/Source/Programs/Shared/EpicGames.Horde/Compute/IComputeSocket.cs
+++ b/Engine/Source/Programs/Shared/EpicGames.Horde/Compute/IComputeSocket.cs
@@ -73,9 +73,8 @@ namespace EpicGames.Horde.Compute
 		/// <param name="socket">Socket to create a channel for</param>
 		/// <param name="channelId">Identifier for the channel</param>
 		/// <param name="logger">Logger for the channel</param>
-		/// <param name="cancellationToken">Cancellation token for the operation</param>
-		public static IComputeMessageChannel CreateMessageChannel(this IComputeSocket socket, int channelId, ILogger logger, CancellationToken cancellationToken = default) 
-			=> socket.CreateMessageChannel(channelId, 65536, logger, cancellationToken);
+		public static IComputeMessageChannel CreateMessageChannel(this IComputeSocket socket, int channelId, ILogger logger) 
+			=> socket.CreateMessageChannel(channelId, 65536, logger);
 
 		/// <summary>
 		/// Creates a message channel with the given identifier
@@ -84,9 +83,8 @@ namespace EpicGames.Horde.Compute
 		/// <param name="channelId">Identifier for the channel</param>
 		/// <param name="bufferSize">Size of the send and receive buffer</param>
 		/// <param name="logger">Logger for the channel</param>
-		/// <param name="cancellationToken">Cancellation token for the operation</param>
-		public static IComputeMessageChannel CreateMessageChannel(this IComputeSocket socket, int channelId, int bufferSize, ILogger logger, CancellationToken cancellationToken = default)
-			=> socket.CreateMessageChannel(channelId, bufferSize, bufferSize, logger, cancellationToken);
+		public static IComputeMessageChannel CreateMessageChannel(this IComputeSocket socket, int channelId, int bufferSize, ILogger logger)
+			=> socket.CreateMessageChannel(channelId, bufferSize, bufferSize, logger);
 
 		/// <summary>
 		/// Creates a message channel with the given identifier
@@ -96,8 +94,7 @@ namespace EpicGames.Horde.Compute
 		/// <param name="sendBufferSize">Size of the send buffer</param>
 		/// <param name="recvBufferSize">Size of the recieve buffer</param>
 		/// <param name="logger">Logger for the channel</param>
-		/// <param name="cancellationToken">Cancellation token for the operation</param>
-		public static IComputeMessageChannel CreateMessageChannel(this IComputeSocket socket, int channelId, int sendBufferSize, int recvBufferSize, ILogger logger, CancellationToken cancellationToken = default)
+		public static IComputeMessageChannel CreateMessageChannel(this IComputeSocket socket, int channelId, int sendBufferSize, int recvBufferSize, ILogger logger)
 		{
 			PooledBuffer sendBuffer = new PooledBuffer(sendBufferSize);
 			PooledBuffer recvBuffer = new PooledBuffer(recvBufferSize);
