Horde: More resilient handling of graceful and forced shutdown scenarios.

* Graceful termination now shuts down the write end of the underlying socket, causing the remote to do the same first. We can then wait for background send/recv tasks to complete normally rather than requiring cancellation.
* IComputeBuffer now functions as a single handle to a shared reference-counted implementation, allowing the socket to acquire its own reference to the buffer and release it once done.

diff --git a/Engine/Source/Programs/Horde/Samples/RemoteClient/RemoteClientApp.cs b/Engine/Source/Programs/Horde/Samples/RemoteClient/RemoteClientApp.cs
--- a/Engine/Source/Programs/Horde/Samples/RemoteClient/RemoteClientApp.cs
+++ b/Engine/Source/Programs/Horde/Samples/RemoteClient/RemoteClientApp.cs
@@ -80,32 +80,35 @@ namespace RemoteClient
 		{
 			// Create a message channel on channel id 0. The Horde Agent always listens on this channel for requests.
 			const int ControlChannelId = 0;
-			using IComputeMessageChannel channel = lease.Socket.CreateMessageChannel(ControlChannelId, 4 * 1024 * 1024, logger);
-			await channel.WaitForAttachAsync();
-
-			// Upload the sandbox
-			MemoryStorageClient storage = new MemoryStorageClient();
-			using (TreeWriter treeWriter = new TreeWriter(storage))
-			{
-				DirectoryNode sandbox = new DirectoryNode();
-				await sandbox.CopyFromDirectoryAsync(uploadDir.ToDirectoryInfo(), new ChunkingOptions(), treeWriter, null);
-				NodeHandle handle = await treeWriter.FlushAsync(sandbox);
-				await channel.UploadFilesAsync("", handle.Locator, storage);
-			}
-
-			// Run the task remotely in the background and echo the output to the console
-			await using (IComputeProcess process = await channel.ExecuteAsync(executable, arguments, null, null))
+			using (IComputeMessageChannel channel = lease.Socket.CreateMessageChannel(ControlChannelId, 4 * 1024 * 1024, logger))
 			{
-				string? line = await process.ReadLineAsync();
-				logger.LogInformation("[REMOTE] {Line}", line);
+				await channel.WaitForAttachAsync();
 
-				await using BackgroundTask tickTask = BackgroundTask.StartNew(ctx => WriteNumbersAsync(lease.Socket, logger, ctx));
+				// Upload the sandbox
+				MemoryStorageClient storage = new MemoryStorageClient();
+				using (TreeWriter treeWriter = new TreeWriter(storage))
+				{
+					DirectoryNode sandbox = new DirectoryNode();
+					await sandbox.CopyFromDirectoryAsync(uploadDir.ToDirectoryInfo(), new ChunkingOptions(), treeWriter, null);
+					NodeHandle handle = await treeWriter.FlushAsync(sandbox);
+					await channel.UploadFilesAsync("", handle.Locator, storage);
+				}
 
-				while ((line = await process.ReadLineAsync()) != null)
+				// Run the task remotely in the background and echo the output to the console
+				await using (IComputeProcess process = await channel.ExecuteAsync(executable, arguments, null, null))
 				{
+					string? line = await process.ReadLineAsync();
 					logger.LogInformation("[REMOTE] {Line}", line);
+
+					await using BackgroundTask tickTask = BackgroundTask.StartNew(ctx => WriteNumbersAsync(lease.Socket, logger, ctx));
+
+					while ((line = await process.ReadLineAsync()) != null)
+					{
+						logger.LogInformation("[REMOTE] {Line}", line);
+					}
 				}
 			}
+			await lease.CloseAsync();
 		}
 
 		static async Task WriteNumbersAsync(IComputeSocket socket, ILogger logger, CancellationToken cancellationToken)
diff --git a/Engine/Source/Programs/Shared/EpicGames.Horde.Tests/BufferTests.cs b/Engine/Source/Programs/Shared/EpicGames.Horde.Tests/BufferTests.cs
--- a/Engine/Source/Programs/Shared/EpicGames.Horde.Tests/BufferTests.cs
+++ b/Engine/Source/Programs/Shared/EpicGames.Horde.Tests/BufferTests.cs
@@ -26,7 +26,7 @@ namespace EpicGames.Horde.Tests
 		[TestMethod]
 		public async Task TestPooledBuffer()
 		{
-			await TestProducerConsumerAsync(length => new PooledBuffer(length), CancellationToken.None);
+			await TestProducerConsumerAsync(length => new PooledBuffer(length).ToSharedInstance(), CancellationToken.None);
 		}
 
 		[TestMethod]
@@ -34,7 +34,7 @@ namespace EpicGames.Horde.Tests
 		{
 			if (RuntimeInformation.IsOSPlatform(OSPlatform.Windows))
 			{
-				await TestProducerConsumerAsync(length => SharedMemoryBuffer.CreateNew(null, length), CancellationToken.None);
+				await TestProducerConsumerAsync(length => SharedMemoryBuffer.CreateNew(null, length).ToSharedInstance(), CancellationToken.None);
 			}
 		}
 
@@ -49,7 +49,7 @@ namespace EpicGames.Horde.Tests
 			await using IComputeSocket consumerSocket = ComputeSocket.Create(new PipeTransport(sourceToTargetPipe.Reader, targetToSourcePipe.Writer), NullLogger.Instance);
 
 			using IComputeBuffer consumerBuffer = createBuffer(Length);
-			consumerSocket.AttachRecvBuffer(ChannelId, consumerBuffer.Writer);
+			consumerSocket.AttachRecvBuffer(ChannelId, consumerBuffer);
 
 			byte[] input = RandomNumberGenerator.GetBytes(Length);
 			Task producerTask = RunProducerAsync(producerSocket, input);
diff --git a/Engine/Source/Programs/Shared/EpicGames.Horde/Compute/Clients/AgentComputeClient.cs b/Engine/Source/Programs/Shared/EpicGames.Horde/Compute/Clients/AgentComputeClient.cs
--- a/Engine/Source/Programs/Shared/EpicGames.Horde/Compute/Clients/AgentComputeClient.cs
+++ b/Engine/Source/Programs/Shared/EpicGames.Horde/Compute/Clients/AgentComputeClient.cs
@@ -41,6 +41,9 @@ namespace EpicGames.Horde.Compute.Clients
 				await _source.MoveNextAsync();
 				await _source.DisposeAsync();
 			}
+
+			/// <inheritdoc/>
+			public ValueTask CloseAsync(CancellationToken cancellationToken) => Socket.CloseAsync(cancellationToken);
 		}
 
 		readonly string _hordeAgentAssembly;
diff --git a/Engine/Source/Programs/Shared/EpicGames.Horde/Compute/Clients/LocalComputeClient.cs b/Engine/Source/Programs/Shared/EpicGames.Horde/Compute/Clients/LocalComputeClient.cs
--- a/Engine/Source/Programs/Shared/EpicGames.Horde/Compute/Clients/LocalComputeClient.cs
+++ b/Engine/Source/Programs/Shared/EpicGames.Horde/Compute/Clients/LocalComputeClient.cs
@@ -27,7 +27,12 @@ namespace EpicGames.Horde.Compute.Clients
 			readonly IComputeSocket _socket;
 
 			public LeaseImpl(IComputeSocket socket) => _socket = socket;
+
+			/// <inheritdoc/>
 			public ValueTask DisposeAsync() => _socket.DisposeAsync();
+
+			/// <inheritdoc/>
+			public ValueTask CloseAsync(CancellationToken cancellationToken) => _socket.CloseAsync(cancellationToken);
 		}
 
 		readonly BackgroundTask _listenerTask;
diff --git a/Engine/Source/Programs/Shared/EpicGames.Horde/Compute/Clients/ServerComputeClient.cs b/Engine/Source/Programs/Shared/EpicGames.Horde/Compute/Clients/ServerComputeClient.cs
--- a/Engine/Source/Programs/Shared/EpicGames.Horde/Compute/Clients/ServerComputeClient.cs
+++ b/Engine/Source/Programs/Shared/EpicGames.Horde/Compute/Clients/ServerComputeClient.cs
@@ -63,6 +63,9 @@ namespace EpicGames.Horde.Compute.Clients
 				await _source.MoveNextAsync();
 				await _source.DisposeAsync();
 			}
+
+			/// <inheritdoc/>
+			public ValueTask CloseAsync(CancellationToken cancellationToken) => Socket.CloseAsync(cancellationToken);
 		}
 
 		readonly HttpClient? _defaultHttpClient;
diff --git a/Engine/Source/Programs/Shared/EpicGames.Horde/Compute/ComputeBufferBase.cs b/Engine/Source/Programs/Shared/EpicGames.Horde/Compute/ComputeBufferBase.cs
--- a/Engine/Source/Programs/Shared/EpicGames.Horde/Compute/ComputeBufferBase.cs
+++ b/Engine/Source/Programs/Shared/EpicGames.Horde/Compute/ComputeBufferBase.cs
@@ -9,7 +9,7 @@ namespace EpicGames.Horde.Compute
 	/// <summary>
 	/// In-process buffer used to store compute messages
 	/// </summary>
-	public abstract class ComputeBufferBase : IComputeBuffer
+	public abstract class ComputeBufferBase : IDisposable
 	{
 		/// <summary>
 		/// Write state for a chunk
@@ -245,11 +245,16 @@ namespace EpicGames.Horde.Compute
 				return chunk.Memory.Slice(state.Length);
 			}
 
-			public void MarkComplete()
+			public bool MarkComplete()
 			{
 				ChunkBase chunk = _chunks[_chunkIdx];
-				chunk.MarkComplete();
-				_buffer.SetAllReadEvents();
+				if (chunk.State.WriteState != WriteState.Complete)
+				{
+					chunk.MarkComplete();
+					_buffer.SetAllReadEvents();
+					return true;
+				}
+				return false;
 			}
 
 			public async ValueTask WaitToWriteAsync(int currentLength, CancellationToken cancellationToken)
@@ -293,11 +298,53 @@ namespace EpicGames.Horde.Compute
 			}
 		}
 
+		/// <summary>
+		/// Implementation of <see cref="IComputeBuffer"/> which handles reference counting 
+		/// </summary>
+		sealed class RefCountedImpl : IComputeBuffer
+		{
+			ComputeBufferBase _inner;
+
+			/// <summary>
+			/// Constructor
+			/// </summary>
+			/// <param name="buffer"></param>
+			public RefCountedImpl(ComputeBufferBase buffer) => _inner = buffer;
+
+			/// <inheritdoc/>
+			public IComputeBufferReader Reader => _inner.Reader;
+
+			/// <inheritdoc/>
+			public IComputeBufferWriter Writer => _inner.Writer;
+
+			/// <inheritdoc/>
+			public IComputeBuffer AddRef()
+			{
+				Interlocked.Increment(ref _inner._refCount);
+				return new RefCountedImpl(_inner);
+			}
+
+			/// <inheritdoc/>
+			public void Dispose()
+			{
+				if (_inner != null)
+				{
+					if (Interlocked.Decrement(ref _inner._refCount) == 0)
+					{
+						_inner.Dispose();
+					}
+					_inner = null!;
+				}
+			}
+		}
+
 		readonly ChunkBase[] _chunks;
 		readonly ReaderImpl _reader;
 		readonly WriterImpl _writer;
 		readonly int _numReaders;
 
+		int _refCount = 1;
+
 		/// <inheritdoc/>
 		public IComputeBufferReader Reader => _reader;
 
@@ -342,6 +389,11 @@ namespace EpicGames.Horde.Compute
 			}
 		}
 
+		/// <summary>
+		/// Creates a sharable <see cref="IComputeBuffer"/> implementation from this object.
+		/// </summary>
+		public IComputeBuffer ToSharedInstance() => new RefCountedImpl(this);
+
 		/// <summary>
 		/// Signals a read event
 		/// </summary>
diff --git a/Engine/Source/Programs/Shared/EpicGames.Horde/Compute/ComputeMessageChannel.cs b/Engine/Source/Programs/Shared/EpicGames.Horde/Compute/ComputeMessageChannel.cs
--- a/Engine/Source/Programs/Shared/EpicGames.Horde/Compute/ComputeMessageChannel.cs
+++ b/Engine/Source/Programs/Shared/EpicGames.Horde/Compute/ComputeMessageChannel.cs
@@ -132,8 +132,8 @@ namespace EpicGames.Horde.Compute
 		{
 			_socket = socket;
 			_channelId = channelId;
-			_socket.AttachRecvBuffer(channelId, recvBuffer.Writer);
-			_socket.AttachSendBuffer(channelId, sendBuffer.Reader);
+			_socket.AttachRecvBuffer(channelId, recvBuffer);
+			_socket.AttachSendBuffer(channelId, sendBuffer);
 			_recvBuffer = recvBuffer;
 			_sendBuffer = sendBuffer;
 			_logger = logger;
@@ -145,6 +145,7 @@ namespace EpicGames.Horde.Compute
 		public void Dispose()
 		{
 			_currentBuilder?.Dispose();
+			_sendBuffer.Writer.MarkComplete();
 			_sendBuffer.Dispose();
 			_recvBuffer.Dispose();
 		}
diff --git a/Engine/Source/Programs/Shared/EpicGames.Horde/Compute/ComputeProcess.cs b/Engine/Source/Programs/Shared/EpicGames.Horde/Compute/ComputeProcess.cs
--- a/Engine/Source/Programs/Shared/EpicGames.Horde/Compute/ComputeProcess.cs
+++ b/Engine/Source/Programs/Shared/EpicGames.Horde/Compute/ComputeProcess.cs
@@ -54,10 +54,12 @@ namespace EpicGames.Horde.Compute
 			try
 			{
 				await RunInternalAsync(channel, cancellationToken);
+				_output.Writer.TryComplete();
 			}
 			catch (Exception ex)
 			{
 				_result.TrySetException(ex);
+				_output.Writer.TryComplete(ex);
 			}
 		}
 
diff --git a/Engine/Source/Programs/Shared/EpicGames.Horde/Compute/ComputeWorker.cs b/Engine/Source/Programs/Shared/EpicGames.Horde/Compute/ComputeWorker.cs
--- a/Engine/Source/Programs/Shared/EpicGames.Horde/Compute/ComputeWorker.cs
+++ b/Engine/Source/Programs/Shared/EpicGames.Horde/Compute/ComputeWorker.cs
@@ -166,11 +166,11 @@ namespace EpicGames.Horde.Compute
 			const int ChannelId = 1;
 			string baseBufferName = $"Local\\Compute-{Guid.NewGuid()}";
 
-			using SharedMemoryBuffer sendBuffer = SharedMemoryBuffer.CreateNew(ComputeChannel.GetSendBufferName(baseBufferName), 4096);
-			socket.AttachSendBuffer(ChannelId, sendBuffer.Reader);
+			using IComputeBuffer sendBuffer = SharedMemoryBuffer.CreateNew(ComputeChannel.GetSendBufferName(baseBufferName), 4096).ToSharedInstance();
+			socket.AttachSendBuffer(ChannelId, sendBuffer);
 
-			using SharedMemoryBuffer recvBuffer = SharedMemoryBuffer.CreateNew(ComputeChannel.GetRecvBufferName(baseBufferName), 4096);
-			socket.AttachRecvBuffer(ChannelId, recvBuffer.Writer);
+			using IComputeBuffer recvBuffer = SharedMemoryBuffer.CreateNew(ComputeChannel.GetRecvBufferName(baseBufferName), 4096).ToSharedInstance();
+			socket.AttachRecvBuffer(ChannelId, recvBuffer);
 
 			Dictionary<string, string?> newEnvVars = new Dictionary<string, string?>();
 			if (envVars != null)
@@ -188,6 +188,9 @@ namespace EpicGames.Horde.Compute
 			_logger.LogInformation("Finished executing process");
 
 			_logger.LogInformation("Child process has shut down");
+
+			recvBuffer.Writer.MarkComplete();
+			sendBuffer.Writer.MarkComplete();
 		}
 
 		async Task ExecuteProcessInternalAsync(IComputeMessageChannel channel, string executable, IReadOnlyList<string> arguments, string? workingDir, IReadOnlyDictionary<string, string?>? envVars, CancellationToken cancellationToken)
diff --git a/Engine/Source/Programs/Shared/EpicGames.Horde/Compute/DefaultComputeSocket.cs b/Engine/Source/Programs/Shared/EpicGames.Horde/Compute/DefaultComputeSocket.cs
--- a/Engine/Source/Programs/Shared/EpicGames.Horde/Compute/DefaultComputeSocket.cs
+++ b/Engine/Source/Programs/Shared/EpicGames.Horde/Compute/DefaultComputeSocket.cs
@@ -4,6 +4,7 @@ using System;
 using System.Buffers;
 using System.Buffers.Binary;
 using System.Collections.Generic;
+using System.Linq;
 using System.Threading;
 using System.Threading.Tasks;
 using EpicGames.Core;
@@ -26,10 +27,11 @@ namespace EpicGames.Horde.Compute
 		readonly CancellationTokenSource _cancellationSource = new CancellationTokenSource();
 
 		readonly BackgroundTask _recvTask;
-		readonly Dictionary<int, IComputeBufferWriter> _recvBuffers = new Dictionary<int, IComputeBufferWriter>();
+		readonly Dictionary<int, IComputeBuffer> _recvBuffers = new Dictionary<int, IComputeBuffer>();
 
 		readonly SemaphoreSlim _sendSemaphore = new SemaphoreSlim(1, 1);
 		readonly Dictionary<int, Task> _sendTasks = new Dictionary<int, Task>();
+		readonly Dictionary<IComputeBuffer, int> _sendBuffers = new Dictionary<IComputeBuffer, int>();
 
 		/// <summary>
 		/// Constructor
@@ -48,26 +50,39 @@ namespace EpicGames.Horde.Compute
 		/// </summary>
 		public async ValueTask CloseAsync(CancellationToken cancellationToken)
 		{
+			// Make sure we close all buffers that are attached, otherwise we'll lock up waiting for send tasks to complete
+			Task[] sendTasks;
+			lock (_lockObject)
+			{
+				sendTasks = _sendTasks.Values.ToArray();
+				foreach ((IComputeBuffer sendBuffer, int channelId) in _sendBuffers)
+				{
+					if (sendBuffer.Writer.MarkComplete())
+					{
+						_logger.LogWarning("Buffer attached to channel {ChannelId} was not marked as complete before socket was closed.", channelId);
+					}
+				}
+				_sendBuffers.Clear();
+			}
+
 			// Wait for all the individual send tasks to complete
-			await Task.WhenAll(_sendTasks.Values);
+			await Task.WhenAll(sendTasks);
 			cancellationToken.ThrowIfCancellationRequested();
 
 			// Send a final message indicating that the lease is done. This will allow the senders on the remote end to terminate, and trigger
-			// a message to be sent back to our read task allowing it to shut down gracefully.
-			byte[] header = new byte[8];
-			await _transport.WriteAsync(header, cancellationToken);
-			_logger.LogTrace("Sent shutdown packet");
+			// a shutdown event to be sent back to our read task allowing it to shut down gracefully.
+			await _transport.MarkCompleteAsync(cancellationToken);
 
 			// Wait for the reader to stop
 			await _recvTask.Task.WaitAsync(TimeSpan.FromSeconds(5), cancellationToken);
 			await _recvTask.StopAsync();
-
-			_logger.LogTrace("Closed");
 		}
 
 		/// <inheritdoc/>
 		public async ValueTask DisposeAsync()
 		{
+			_cancellationSource.Cancel();
+			await Task.WhenAll(_sendTasks.Values);
 			await _recvTask.DisposeAsync();
 			_sendSemaphore.Dispose();
 			_cancellationSource.Dispose();
@@ -100,13 +115,6 @@ namespace EpicGames.Horde.Compute
 					int id = BinaryPrimitives.ReadInt32LittleEndian(header);
 					int size = BinaryPrimitives.ReadInt32LittleEndian(header.AsSpan(4));
 
-					// If the size if negative, we're closing the entire connection
-					if (size == 0)
-					{
-						_logger.LogTrace("Received shutdown packet");
-						break;
-					}
-
 					// Dispatch it to the correct place
 					if (size < 0)
 					{
@@ -166,8 +174,10 @@ namespace EpicGames.Horde.Compute
 
 			lock (_lockObject)
 			{
-				if (_recvBuffers.TryGetValue(id, out writer))
+				IComputeBuffer? recvBuffer;
+				if (_recvBuffers.TryGetValue(id, out recvBuffer))
 				{
+					writer = recvBuffer.Writer;
 					cachedWriters.Add(id, writer);
 					return writer;
 				}
@@ -222,21 +232,24 @@ namespace EpicGames.Horde.Compute
 		}
 
 		/// <inheritdoc/>
-		public void AttachRecvBuffer(int channelId, IComputeBufferWriter recvBuffer)
+		public void AttachRecvBuffer(int channelId, IComputeBuffer recvBuffer)
 		{
+			IComputeBuffer shared = recvBuffer.AddRef();
+
 			bool complete;
 			lock (_lockObject)
 			{
 				complete = _complete;
 				if (!complete)
 				{
-					_recvBuffers.Add(channelId, recvBuffer);
+					_recvBuffers.Add(channelId, shared);
 				}
 			}
 
-			if (recvBuffer != null && complete)
+			if (shared != null && complete)
 			{
-				recvBuffer.MarkComplete();
+				shared.Writer.MarkComplete();
+				shared.Dispose();
 			}
 		}
 
@@ -244,30 +257,36 @@ namespace EpicGames.Horde.Compute
 		{
 			cachedWriters.Remove(id);
 
-			IComputeBufferWriter? writer;
+			IComputeBuffer? buffer;
 			lock (_lockObject)
 			{
 #pragma warning disable CA2000 // Dispose objects before losing scope
-				_recvBuffers.Remove(id, out writer);
+				_recvBuffers.Remove(id, out buffer);
 #pragma warning restore CA2000 // Dispose objects before losing scope
 			}
-			if (writer != null)
+			if (buffer != null)
 			{
-				writer.MarkComplete();
+				buffer.Writer.MarkComplete();
+				buffer.Dispose();
 			}
 		}
 
 		/// <inheritdoc/>
-		public void AttachSendBuffer(int channelId, IComputeBufferReader reader)
+		public void AttachSendBuffer(int channelId, IComputeBuffer buffer)
 		{
+			IComputeBuffer shared = buffer.AddRef();
 			lock (_lockObject)
 			{
-				_sendTasks.Add(channelId, Task.Run(() => SendFromBufferAsync(channelId, reader, _cancellationSource.Token), _cancellationSource.Token));
+				_sendTasks.Add(channelId, Task.Run(() => SendFromBufferAsync(channelId, shared, _cancellationSource.Token), _cancellationSource.Token));
+				_sendBuffers.Add(shared, channelId);
 			}
 		}
 
-		async Task SendFromBufferAsync(int channelId, IComputeBufferReader reader, CancellationToken cancellationToken)
+		async Task SendFromBufferAsync(int channelId, IComputeBuffer buffer, CancellationToken cancellationToken)
 		{
+			using IComputeBuffer _ = buffer;
+
+			IComputeBufferReader reader = buffer.Reader;
 			while (!cancellationToken.IsCancellationRequested)
 			{
 				ReadOnlyMemory<byte> memory = reader.GetMemory();
@@ -283,6 +302,12 @@ namespace EpicGames.Horde.Compute
 				}
 				await reader.WaitToReadAsync(0, cancellationToken);
 			}
+
+			lock (_lockObject)
+			{
+				_sendTasks.Remove(channelId);
+				_sendBuffers.Remove(buffer);
+			}
 		}
 	}
 }
diff --git a/Engine/Source/Programs/Shared/EpicGames.Horde/Compute/IComputeBuffer.cs b/Engine/Source/Programs/Shared/EpicGames.Horde/Compute/IComputeBuffer.cs
--- a/Engine/Source/Programs/Shared/EpicGames.Horde/Compute/IComputeBuffer.cs
+++ b/Engine/Source/Programs/Shared/EpicGames.Horde/Compute/IComputeBuffer.cs
@@ -20,6 +20,11 @@ namespace EpicGames.Horde.Compute
 		/// Writer for this buffer
 		/// </summary>
 		IComputeBufferWriter Writer { get; }
+
+		/// <summary>
+		/// Creates a new reference to the underlying buffer. The underlying resources will only be destroyed once all instances are disposed of.
+		/// </summary>
+		IComputeBuffer AddRef();
 	}
 
 	/// <summary>
@@ -60,7 +65,8 @@ namespace EpicGames.Horde.Compute
 		/// <summary>
 		/// Mark the output to this buffer as complete
 		/// </summary>
-		void MarkComplete();
+		/// <returns>Whether the writer was marked as complete. False if the writer has already been marked as complete.</returns>
+		bool MarkComplete();
 
 		/// <summary>
 		/// Updates the current write position within the buffer
diff --git a/Engine/Source/Programs/Shared/EpicGames.Horde/Compute/IComputeChannel.cs b/Engine/Source/Programs/Shared/EpicGames.Horde/Compute/IComputeChannel.cs
--- a/Engine/Source/Programs/Shared/EpicGames.Horde/Compute/IComputeChannel.cs
+++ b/Engine/Source/Programs/Shared/EpicGames.Horde/Compute/IComputeChannel.cs
@@ -46,7 +46,7 @@ namespace EpicGames.Horde.Compute
 				_channelId = channelId;
 				_recvBuffer = recvBuffer;
 
-				_socket.AttachRecvBuffer(_channelId, recvBuffer.Writer);
+				_socket.AttachRecvBuffer(_channelId, recvBuffer);
 			}
 
 			public void Dispose()
@@ -73,8 +73,8 @@ namespace EpicGames.Horde.Compute
 				_recvBuffer = recvBuffer;
 				_sendBuffer = sendBuffer;
 
-				_socket.AttachRecvBuffer(_channelId, recvBuffer.Writer);
-				_socket.AttachSendBuffer(_channelId, sendBuffer.Reader);
+				_socket.AttachRecvBuffer(_channelId, recvBuffer);
+				_socket.AttachSendBuffer(_channelId, sendBuffer);
 			}
 
 			public void Dispose()
@@ -125,8 +125,8 @@ namespace EpicGames.Horde.Compute
 				throw new InvalidOperationException($"Environment variable {IpcEnvVar} is not defined; cannot connect as worker.");
 			}
 
-			SharedMemoryBuffer recvBuffer = SharedMemoryBuffer.OpenExisting(GetRecvBufferName(baseName));
-			SharedMemoryBuffer sendBuffer = SharedMemoryBuffer.OpenExisting(GetSendBufferName(baseName));
+			IComputeBuffer recvBuffer = SharedMemoryBuffer.OpenExisting(GetRecvBufferName(baseName)).ToSharedInstance();
+			IComputeBuffer sendBuffer = SharedMemoryBuffer.OpenExisting(GetSendBufferName(baseName)).ToSharedInstance();
 			return new WorkerChannel(recvBuffer, sendBuffer);
 		}
 
diff --git a/Engine/Source/Programs/Shared/EpicGames.Horde/Compute/IComputeLease.cs b/Engine/Source/Programs/Shared/EpicGames.Horde/Compute/IComputeLease.cs
--- a/Engine/Source/Programs/Shared/EpicGames.Horde/Compute/IComputeLease.cs
+++ b/Engine/Source/Programs/Shared/EpicGames.Horde/Compute/IComputeLease.cs
@@ -2,6 +2,8 @@
 
 using System;
 using System.Collections.Generic;
+using System.Threading;
+using System.Threading.Tasks;
 
 namespace EpicGames.Horde.Compute
 {
@@ -24,5 +26,11 @@ namespace EpicGames.Horde.Compute
 		/// Socket to communicate with the remote
 		/// </summary>
 		IComputeSocket Socket { get; }
+
+		/// <summary>
+		/// Relinquish the lease gracefully
+		/// </summary>
+		/// <param name="cancellationToken">Cancellation token for the operation</param>
+		ValueTask CloseAsync(CancellationToken cancellationToken = default);
 	}
 }
diff --git a/Engine/Source/Programs/Shared/EpicGames.Horde/Compute/IComputeSocket.cs b/Engine/Source/Programs/Shared/EpicGames.Horde/Compute/IComputeSocket.cs
--- a/Engine/Source/Programs/Shared/EpicGames.Horde/Compute/IComputeSocket.cs
+++ b/Engine/Source/Programs/Shared/EpicGames.Horde/Compute/IComputeSocket.cs
@@ -32,21 +32,21 @@ namespace EpicGames.Horde.Compute
 		/// Attaches a buffer to receive data.
 		/// </summary>
 		/// <param name="channelId">Channel to receive data on</param>
-		/// <param name="writer">Writer for the received data</param>
-		void AttachRecvBuffer(int channelId, IComputeBufferWriter writer);
+		/// <param name="buffer">Buffer to receive data</param>
+		void AttachRecvBuffer(int channelId, IComputeBuffer buffer);
 
 		/// <summary>
 		/// Attaches a buffer to send data.
 		/// </summary>
 		/// <param name="channelId">Channel to receive data on</param>
-		/// <param name="reader">Reader for the data to send</param>
-		void AttachSendBuffer(int channelId, IComputeBufferReader reader);
+		/// <param name="buffer">Buffer to send data from</param>
+		void AttachSendBuffer(int channelId, IComputeBuffer buffer);
 
 		/// <summary>
 		/// Attempt to gracefully close the current connection and shutdown both ends of the transport
 		/// </summary>
 		/// <param name="cancellationToken">Cancellation token for the operation</param>
-		ValueTask CloseAsync(CancellationToken cancellationToken);
+		ValueTask CloseAsync(CancellationToken cancellationToken = default);
 	}
 
 	/// <summary>
@@ -96,8 +96,8 @@ namespace EpicGames.Horde.Compute
 		/// <param name="logger">Logger for the channel</param>
 		public static IComputeMessageChannel CreateMessageChannel(this IComputeSocket socket, int channelId, int sendBufferSize, int recvBufferSize, ILogger logger)
 		{
-			PooledBuffer sendBuffer = new PooledBuffer(sendBufferSize);
-			PooledBuffer recvBuffer = new PooledBuffer(recvBufferSize);
+			IComputeBuffer sendBuffer = new PooledBuffer(sendBufferSize).ToSharedInstance();
+			IComputeBuffer recvBuffer = new PooledBuffer(recvBufferSize).ToSharedInstance();
 			return new ComputeMessageChannel(socket, channelId, sendBuffer, recvBuffer, logger);
 		}
 	}
diff --git a/Engine/Source/Programs/Shared/EpicGames.Horde/Compute/IComputeTransport.cs b/Engine/Source/Programs/Shared/EpicGames.Horde/Compute/IComputeTransport.cs
--- a/Engine/Source/Programs/Shared/EpicGames.Horde/Compute/IComputeTransport.cs
+++ b/Engine/Source/Programs/Shared/EpicGames.Horde/Compute/IComputeTransport.cs
@@ -31,6 +31,12 @@ namespace EpicGames.Horde.Compute
 		/// <param name="buffer">Buffer to be written</param>
 		/// <param name="cancellationToken">Cancellation token for the operation</param>
 		ValueTask WriteAsync(ReadOnlySequence<byte> buffer, CancellationToken cancellationToken);
+
+		/// <summary>
+		/// Indicate that all data has been written to the transport layer, and that there will be no more calls to <see cref="WriteAsync(ReadOnlySequence{Byte}, CancellationToken)"/>
+		/// </summary>
+		/// <param name="cancellationToken">Cancellation token for the operation</param>
+		ValueTask MarkCompleteAsync(CancellationToken cancellationToken);
 	}
 
 	/// <summary>
diff --git a/Engine/Source/Programs/Shared/EpicGames.Horde/Compute/Transports/AesTransport.cs b/Engine/Source/Programs/Shared/EpicGames.Horde/Compute/Transports/AesTransport.cs
--- a/Engine/Source/Programs/Shared/EpicGames.Horde/Compute/Transports/AesTransport.cs
+++ b/Engine/Source/Programs/Shared/EpicGames.Horde/Compute/Transports/AesTransport.cs
@@ -210,5 +210,8 @@ namespace EpicGames.Horde.Compute.Transports
 		{
 			BinaryPrimitives.WriteInt64LittleEndian(nonce, BinaryPrimitives.ReadInt64LittleEndian(nonce) + 1);
 		}
+
+		/// <inheritdoc/>
+		public ValueTask MarkCompleteAsync(CancellationToken cancellationToken) => _inner.MarkCompleteAsync(cancellationToken);
 	}
 }
diff --git a/Engine/Source/Programs/Shared/EpicGames.Horde/Compute/Transports/PipeTransport.cs b/Engine/Source/Programs/Shared/EpicGames.Horde/Compute/Transports/PipeTransport.cs
--- a/Engine/Source/Programs/Shared/EpicGames.Horde/Compute/Transports/PipeTransport.cs
+++ b/Engine/Source/Programs/Shared/EpicGames.Horde/Compute/Transports/PipeTransport.cs
@@ -86,5 +86,8 @@ namespace EpicGames.Horde.Compute.Transports
 				Position += memory.Length;
 			}
 		}
+
+		/// <inheritdoc/>
+		public ValueTask MarkCompleteAsync(CancellationToken cancellationToken) => _writer.CompleteAsync();
 	}
 }
diff --git a/Engine/Source/Programs/Shared/EpicGames.Horde/Compute/Transports/StreamTransport.cs b/Engine/Source/Programs/Shared/EpicGames.Horde/Compute/Transports/StreamTransport.cs
--- a/Engine/Source/Programs/Shared/EpicGames.Horde/Compute/Transports/StreamTransport.cs
+++ b/Engine/Source/Programs/Shared/EpicGames.Horde/Compute/Transports/StreamTransport.cs
@@ -41,5 +41,8 @@ namespace EpicGames.Horde.Compute.Transports
 				Position += memory.Length;
 			}
 		}
+
+		/// <inheritdoc/>
+		public ValueTask MarkCompleteAsync(CancellationToken cancellationToken) => new ValueTask();
 	}
 }
diff --git a/Engine/Source/Programs/Shared/EpicGames.Horde/Compute/Transports/TcpTransport.cs b/Engine/Source/Programs/Shared/EpicGames.Horde/Compute/Transports/TcpTransport.cs
--- a/Engine/Source/Programs/Shared/EpicGames.Horde/Compute/Transports/TcpTransport.cs
+++ b/Engine/Source/Programs/Shared/EpicGames.Horde/Compute/Transports/TcpTransport.cs
@@ -42,5 +42,12 @@ namespace EpicGames.Horde.Compute.Transports
 				Position += memory.Length;
 			}
 		}
+
+		/// <inheritdoc/>
+		public ValueTask MarkCompleteAsync(CancellationToken cancelationToken)
+		{
+			_socket.Shutdown(SocketShutdown.Send);
+			return new ValueTask();
+		}
 	}
 }
