Reverted back to using the AsyncLocalScopeManager to manage the active span. This appears to be passing the tests I added.
Added another test for Parallel.ForEachAsync.

diff --git a/Engine/Source/Programs/AutomationTool/AutomationUtils/Automation.cs b/Engine/Source/Programs/AutomationTool/AutomationUtils/Automation.cs
--- a/Engine/Source/Programs/AutomationTool/AutomationUtils/Automation.cs
+++ b/Engine/Source/Programs/AutomationTool/AutomationUtils/Automation.cs
@@ -178,7 +178,6 @@ namespace AutomationTool
 					if (Tracer != null)
 					{
 						Tracer.Flush();
-						Tracer.Dispose();
 					}
 				}
 			}
diff --git a/Engine/Source/Programs/Shared/EpicGames.Core.Tests/JsonTracerTests.cs b/Engine/Source/Programs/Shared/EpicGames.Core.Tests/JsonTracerTests.cs
--- a/Engine/Source/Programs/Shared/EpicGames.Core.Tests/JsonTracerTests.cs
+++ b/Engine/Source/Programs/Shared/EpicGames.Core.Tests/JsonTracerTests.cs
@@ -18,78 +18,70 @@ namespace EpicGames.Core.Tests
 		[TestMethod]
 		public void Run()
 		{
-			using (JsonTracer Tracer = new JsonTracer())
-			{
-				using (IScope Scope1 = Tracer.BuildSpan("bogusOp1").StartActive()) { }
-				using (IScope Scope2 = Tracer.BuildSpan("bogusOp2").StartActive()) { }
+			JsonTracer Tracer = new JsonTracer();
+			using (IScope Scope1 = Tracer.BuildSpan("bogusOp1").StartActive()) { }
+			using (IScope Scope2 = Tracer.BuildSpan("bogusOp2").StartActive()) { }
 
-				Assert.AreEqual(2, Tracer.GetFinishedSpans().Count);
-			}
+			Assert.AreEqual(2, Tracer.GetFinishedSpans().Count);
 		}
 		
 		[TestMethod]
 		public void SpansWithTags()
 		{
-			using (JsonTracer Tracer = new JsonTracer())
-			{
-				using (IScope Scope1 = Tracer.BuildSpan("bogusOp1").WithTag("age", 94).StartActive()) { }
-				using (IScope Scope2 = Tracer.BuildSpan("bogusOp1").WithTag("foo", "bar").StartActive()) { }
+			JsonTracer Tracer = new JsonTracer();
+			using (IScope Scope1 = Tracer.BuildSpan("bogusOp1").WithTag("age", 94).StartActive()) { }
+			using (IScope Scope2 = Tracer.BuildSpan("bogusOp1").WithTag("foo", "bar").StartActive()) { }
 
-				List<JsonTracerSpan> Spans = Tracer.GetFinishedSpans();
+			List<JsonTracerSpan> Spans = Tracer.GetFinishedSpans();
 
-				Assert.AreEqual(2, Spans.Count);
-				Assert.AreEqual("bogusOp1", Spans[0].OperationName);
-			}
+			Assert.AreEqual(2, Spans.Count);
+			Assert.AreEqual("bogusOp1", Spans[0].OperationName);
 		}
 
 		[TestMethod]
 		public void SerializeToJson()
 		{
-			using (JsonTracer Tracer = new JsonTracer())
-			{
-				using (IScope Scope1 = Tracer.BuildSpan("bogusOp1").WithTag("age", 94).WithTag("foo", "bar").StartActive()) { }
-				using (IScope Scope2 = Tracer.BuildSpan("bogusOp2").WithTag("Service", "myService").WithTag("Resource", "myResource").StartActive()) { }
+			JsonTracer Tracer = new JsonTracer();
+			using (IScope Scope1 = Tracer.BuildSpan("bogusOp1").WithTag("age", 94).WithTag("foo", "bar").StartActive()) { }
+			using (IScope Scope2 = Tracer.BuildSpan("bogusOp2").WithTag("Service", "myService").WithTag("Resource", "myResource").StartActive()) { }
 
-				using StringWriter Sw = new StringWriter();
-				using JsonWriter Jw = new JsonWriter(Sw);
-				Tracer.GetFinishedSpansAsJson(Jw);
-				byte[] Data = Encoding.ASCII.GetBytes(Sw.ToString());
+			using StringWriter Sw = new StringWriter();
+			using JsonWriter Jw = new JsonWriter(Sw);
+			Tracer.GetFinishedSpansAsJson(Jw);
+			byte[] Data = Encoding.ASCII.GetBytes(Sw.ToString());
 
-				using (JsonDocument Document = JsonDocument.Parse(Data.AsMemory()))
-				{
-					JsonElement Root = Document.RootElement;
-					JsonElement SpansElement = Root.GetProperty("Spans");
-					List<JsonElement> SpansList = SpansElement.EnumerateArray().ToList();
-					Console.WriteLine(SpansList[0]);
-					Assert.AreEqual("bogusOp1", SpansList[0].GetProperty("Name").GetString());
-					Assert.IsNotNull(SpansList[0].GetProperty("StartTime").GetDateTime());
-					Assert.IsNotNull(SpansList[0].GetProperty("FinishTime").GetDateTime());
-					Assert.AreEqual("myService", SpansList[1].GetProperty("Service").GetString());
-					Assert.AreEqual("myResource", SpansList[1].GetProperty("Resource").GetString());
-				}
+			using (JsonDocument Document = JsonDocument.Parse(Data.AsMemory()))
+			{
+				JsonElement Root = Document.RootElement;
+				JsonElement SpansElement = Root.GetProperty("Spans");
+				List<JsonElement> SpansList = SpansElement.EnumerateArray().ToList();
+				Console.WriteLine(SpansList[0]);
+				Assert.AreEqual("bogusOp1", SpansList[0].GetProperty("Name").GetString());
+				Assert.IsNotNull(SpansList[0].GetProperty("StartTime").GetDateTime());
+				Assert.IsNotNull(SpansList[0].GetProperty("FinishTime").GetDateTime());
+				Assert.AreEqual("myService", SpansList[1].GetProperty("Service").GetString());
+				Assert.AreEqual("myResource", SpansList[1].GetProperty("Resource").GetString());
 			}
 		}
 
 		[TestMethod]
 		public void SpansWithParallelFor()
 		{
-			using (JsonTracer Tracer = new JsonTracer())
+			JsonTracer Tracer = new JsonTracer();
+			using (IScope ParentScope = Tracer.BuildSpan("Parent").StartActive())
 			{
-				using (IScope ParentScope = Tracer.BuildSpan("Parent").StartActive())
+				Parallel.For(0, 100, Index =>
 				{
-					Parallel.For(0, 100, Index =>
+					using (IScope ChildScope = Tracer.BuildSpan("Child").StartActive())
 					{
-						using (IScope ChildScope = Tracer.BuildSpan("Child").StartActive())
-						{
-							Assert.AreEqual(Tracer.ActiveSpan, ChildScope.Span);
-							Assert.AreEqual(((JsonTracerSpan)ChildScope.Span).ParentId, ((JsonTracerSpan)ParentScope.Span).Context.SpanId);
-						}
-						Assert.AreEqual(Tracer.ActiveSpan, ParentScope.Span);
-					});
+						Assert.AreEqual(Tracer.ActiveSpan, ChildScope.Span);
+						Assert.AreEqual(((JsonTracerSpan)ChildScope.Span).ParentId, ((JsonTracerSpan)ParentScope.Span).Context.SpanId);
+					}
 					Assert.AreEqual(Tracer.ActiveSpan, ParentScope.Span);
-				}
-				Assert.AreEqual(Tracer.ActiveSpan, null);
+				});
+				Assert.AreEqual(Tracer.ActiveSpan, ParentScope.Span);
 			}
+			Assert.AreEqual(Tracer.ActiveSpan, null);
 		}
 
 		private static async Task<bool> AsyncSpanTest(JsonTracer Tracer, IScope ParentScope)
@@ -118,17 +110,44 @@ namespace EpicGames.Core.Tests
 		[TestMethod]
 		public void SpansWithAsync()
 		{
-			using (JsonTracer Tracer = new JsonTracer())
+			JsonTracer Tracer = new JsonTracer();
+			using (IScope ParentScope = Tracer.BuildSpan("Parent").StartActive())
 			{
-				using (IScope ParentScope = Tracer.BuildSpan("Parent").StartActive())
+				Assert.AreEqual(Tracer.ActiveSpan, ParentScope.Span);
+				AsyncSpanTest(Tracer, ParentScope).Wait();
+				Assert.AreEqual(Tracer.ActiveSpan, ParentScope.Span);
+			}
+
+			Assert.AreEqual(Tracer.ActiveSpan, null);
+		}
+
+		private static async Task ForEachAsyncSpanTest(JsonTracer Tracer, IScope ParentScope)
+		{
+			int[] Values = Enumerable.Range(0, 10000).ToArray();
+			await Parallel.ForEachAsync(Values, async (i, token) =>
+			{
+				using (IScope ChildScope = Tracer.BuildSpan("Child").StartActive())
 				{
-					Assert.AreEqual(Tracer.ActiveSpan, ParentScope.Span);
-					AsyncSpanTest(Tracer, ParentScope).Wait();
-					Assert.AreEqual(Tracer.ActiveSpan, ParentScope.Span);
+					Assert.AreEqual(Tracer.ActiveSpan, ChildScope.Span);
+					Assert.AreEqual(((JsonTracerSpan)ChildScope.Span).ParentId, ((JsonTracerSpan)ParentScope.Span).Context.SpanId);
+					await Task.Delay(10, token);
 				}
+				Assert.AreEqual(Tracer.ActiveSpan, ParentScope.Span);
+			});
+		}
 
-				Assert.AreEqual(Tracer.ActiveSpan, null);
+		[TestMethod]
+		public void SpansWithParallelForEachAsync()
+		{
+			JsonTracer Tracer = new JsonTracer();
+			using (IScope ParentScope = Tracer.BuildSpan("Parent").StartActive())
+			{
+				Assert.AreEqual(Tracer.ActiveSpan, ParentScope.Span);
+				ForEachAsyncSpanTest(Tracer, ParentScope).Wait();
+				Assert.AreEqual(Tracer.ActiveSpan, ParentScope.Span);
 			}
+
+			Assert.AreEqual(Tracer.ActiveSpan, null);
 		}
 	}
 }
diff --git a/Engine/Source/Programs/Shared/EpicGames.Core/JsonTracer.cs b/Engine/Source/Programs/Shared/EpicGames.Core/JsonTracer.cs
--- a/Engine/Source/Programs/Shared/EpicGames.Core/JsonTracer.cs
+++ b/Engine/Source/Programs/Shared/EpicGames.Core/JsonTracer.cs
@@ -1,9 +1,7 @@
 // Copyright Epic Games, Inc. All Rights Reserved.
 
 using System;
-using System.Collections.Concurrent;
 using System.Collections.Generic;
-using System.Diagnostics.Tracing;
 using System.Globalization;
 using System.IO;
 using System.Linq;
@@ -195,7 +193,6 @@ namespace EpicGames.Core
 				CheckForFinished("Tried to finish already finished span");
 				_finishTimestamp = finishTimestamp;
 				_tracer.AppendFinishedSpan(this);
-				_tracer.DeactivateActiveSpan(this);
 				_finished = true;
 			}
 		}
@@ -382,7 +379,6 @@ namespace EpicGames.Core
 		public IScope StartActive(bool finishSpanOnDispose)
 		{
 			ISpan span = Start();
-			_tracer.ActivateSpan(span);
 			return _tracer.ScopeManager.Activate(span, finishSpanOnDispose);
 		}
 
@@ -402,212 +398,20 @@ namespace EpicGames.Core
 			return new JsonTracerSpan(_tracer, _operationName, _startTimestamp, _tags, _references);
 		}
 	}
-
-	public sealed class JsonTracer : ITracer, IDisposable
+	
+	public class JsonTracer : ITracer
 	{
-		private class LocalEventListener : EventListener
-		{
-			protected override void OnEventSourceCreated(EventSource eventSource)
-			{
-				base.OnEventSourceCreated(eventSource);
-				switch(eventSource.Name)
-				{
-					case "System.Threading.Tasks.Parallel.EventSource":
-					case "System.Threading.Tasks.TplEventSource":
-						{
-							EnableEvents(eventSource, EventLevel.LogAlways, EventKeywords.All);
-							break;
-						}
-					default:
-						{
-							break;
-						}
-				}
-			}
-		}
-
 		public IScopeManager ScopeManager { get; }
-		public ISpan? ActiveSpan
-		{ 
-			get
-			{
-				if (PerThreadActiveSpans.TryGetValue(Environment.CurrentManagedThreadId, out Stack<ISpan>? currentThreadActiveSpan) && currentThreadActiveSpan.Any())
-				{
-					return currentThreadActiveSpan.Peek();
-				}	
-				return null;
-			}
-		}
+		public ISpan? ActiveSpan => ScopeManager.Active?.Span;
 
 		private readonly object _lock = new object();
 		private readonly List<JsonTracerSpan> _finishedSpans = new List<JsonTracerSpan>();
 		private readonly DirectoryReference? _telemetryDir;
-		private readonly LocalEventListener _eventListener = new LocalEventListener();
-
-		private readonly ConcurrentDictionary<int, Stack<ISpan>> PerThreadActiveSpans = new ConcurrentDictionary<int, Stack<ISpan>>();
-		private readonly ConcurrentDictionary<int, Stack<ISpan>> ForkIdToActiveSpans = new ConcurrentDictionary<int, Stack<ISpan>>();
-		private readonly ConcurrentDictionary<int, Stack<ISpan>> TaskIdToActiveSpans = new ConcurrentDictionary<int, Stack<ISpan>>();
-
-		// These could possible be set to 2, which is what MS sets them to internally, if this becomes a performance issue.
-		private static int? ParallelLoopBegin_ForkJoinContextIDIndex;
-		private static int? ParallelLoopEnd_ForkJoinContextIDIndex;
-		private static int? ParallelFork_ForkJoinContextIDIndex;
-		private static int? ParallelJoin_ForkJoinContextIDIndex;
-		private static int? TraceOperationBegin_TaskIDIndex;
-		private static int? TraceOperationEnd_TaskIDIndex;
-		private static int? TaskStarted_TaskIDIndex;
-		private const string ForkJoinContextIDStr = "ForkJoinContextID";
-		private const string TaskIDStr = "TaskID";
-
+		
 		public JsonTracer(DirectoryReference? telemetryDir = null)
 		{
 			ScopeManager = new AsyncLocalScopeManager();
 			_telemetryDir = telemetryDir;
-			_eventListener.EventWritten += EventListenerEventWritten;
-		}
-
-		public void Dispose()
-		{
-			_eventListener.Dispose();
-		}
-
-		private static int GetEventIntPayload(EventWrittenEventArgs e, string payloadName, ref int? payloadIndex)
-		{
-			if (!payloadIndex.HasValue)
-			{
-				int? foundIndex = e.PayloadNames?.FindIndex(pn => pn == payloadName);
-				if (foundIndex == null || foundIndex == -1)
-				{
-					throw new InvalidOperationException($"{payloadName} was not found in the event name payload.");
-				}
-				payloadIndex = foundIndex;
-			}
-
-			if (e.Payload == null || e.Payload.Count <= payloadIndex.Value)
-			{
-				throw new InvalidOperationException($"{payloadName} was not in the event payload.");
-			}
-
-			int? intValue = e.Payload[payloadIndex.Value] as int?;
-			if (!intValue.HasValue)
-			{
-				throw new InvalidOperationException($"{payloadName} in not an int.");
-			}
-			return intValue.Value;
-		}
-
-		private void BeginEvent(EventWrittenEventArgs e, string PayloadName, ref int? payloadIndex, ConcurrentDictionary<int, Stack<ISpan>> eventActiveSpansDict)
-		{
-			int value = GetEventIntPayload(e, PayloadName, ref payloadIndex);
-
-			Stack<ISpan>? currentActiveSpanStack = null;
-			if (PerThreadActiveSpans.TryGetValue(Environment.CurrentManagedThreadId, out Stack<ISpan>? currentThreadActiveSpan) && currentThreadActiveSpan.Any())
-			{
-				currentActiveSpanStack = new Stack<ISpan>(currentThreadActiveSpan.Reverse());
-			}
-
-			// the current active span on this thread needs to be attached to this new event
-			if (currentActiveSpanStack != null)
-			{
-				if (!eventActiveSpansDict.TryAdd(value, currentActiveSpanStack))
-				{
-					throw new InvalidOperationException($"Event id '{value}' was already added.");
-				}
-			}
-		}
-
-		private static void EndEvent(EventWrittenEventArgs e, string payloadName, ref int? payloadIndex, ConcurrentDictionary<int, Stack<ISpan>> eventActiveSpansDict)
-		{
-			int value = GetEventIntPayload(e, payloadName, ref payloadIndex);
-			if (eventActiveSpansDict.ContainsKey(value))
-			{
-				if (!eventActiveSpansDict.TryRemove(value, out _))
-				{
-					throw new InvalidOperationException($"Event id '{value}' could not be removed.");
-				}
-			}
-		}
-
-		private void BeginWork(EventWrittenEventArgs e, string payloadName, ref int? payloadIndex, ConcurrentDictionary<int, Stack<ISpan>> eventActiveSpansDict)
-		{
-			// find the active span for this event and make it the active span for this thread
-			int value = GetEventIntPayload(e, payloadName, ref payloadIndex);
-			if (eventActiveSpansDict.TryGetValue(value, out Stack<ISpan>? activeSpanForEvent))
-			{
-				PerThreadActiveSpans.AddOrUpdate(Environment.CurrentManagedThreadId,
-				(k) =>
-				{
-					return new Stack<ISpan>(activeSpanForEvent.Reverse());
-				},
-				(k, v) =>
-				{
-					return new Stack<ISpan>(activeSpanForEvent.Reverse());
-				}
-				);
-			}
-		}
-
-		private void EndWork(EventWrittenEventArgs e, string payloadName, ref int? payloadIndex, ConcurrentDictionary<int, Stack<ISpan>> eventActiveSpansDict)
-		{
-			int value = GetEventIntPayload(e, payloadName, ref payloadIndex);
-			if (eventActiveSpansDict.TryGetValue(value, out Stack<ISpan>? activeSpanForEventStack))
-			{
-				if (PerThreadActiveSpans.TryGetValue(Environment.CurrentManagedThreadId, out Stack<ISpan>? activeSpanForThreadStack))
-				{
-					ISpan eventActiveSpan = activeSpanForEventStack.Peek();
-					ISpan threadActiveSpan = activeSpanForThreadStack.Peek();
-					if (eventActiveSpan.Context.SpanId != threadActiveSpan.Context.SpanId)
-					{
-						throw new InvalidOperationException($"Span should be equal. Event span op: '{((JsonTracerSpan)eventActiveSpan).OperationName}' | Event span id: {eventActiveSpan.Context.SpanId} | Thread span op: '{((JsonTracerSpan)threadActiveSpan).OperationName} | | Event span id: {threadActiveSpan.Context.SpanId}'.");
-					}
-				}
-			}
-		}
-
-		private void EventListenerEventWritten(object? sender, EventWrittenEventArgs e)
-		{
-			switch (e.EventName)
-			{
-				case "ParallelLoopBegin":
-					{
-						BeginEvent(e, ForkJoinContextIDStr, ref ParallelLoopBegin_ForkJoinContextIDIndex, ForkIdToActiveSpans);
-						break;
-					}
-				case "ParallelLoopEnd":
-					{
-						EndEvent(e, ForkJoinContextIDStr, ref ParallelLoopEnd_ForkJoinContextIDIndex, ForkIdToActiveSpans);
-						break;
-					}
-				case "ParallelFork":
-					{
-						BeginWork(e, ForkJoinContextIDStr, ref ParallelFork_ForkJoinContextIDIndex, ForkIdToActiveSpans);
-						break;
-					}
-				case "ParallelJoin":
-					{
-						EndWork(e, ForkJoinContextIDStr, ref ParallelJoin_ForkJoinContextIDIndex, ForkIdToActiveSpans);
-						break;
-					}
-				case "TraceOperationBegin":
-					{
-						BeginEvent(e, TaskIDStr, ref TraceOperationBegin_TaskIDIndex, TaskIdToActiveSpans);
-						break;
-					}
-				case "TraceOperationEnd":
-					{
-						EndEvent(e, TaskIDStr, ref TraceOperationEnd_TaskIDIndex, TaskIdToActiveSpans);
-						break;
-					}
-				case "TaskStarted":
-					{
-						BeginWork(e, TaskIDStr, ref TaskStarted_TaskIDIndex, TaskIdToActiveSpans);
-						break;
-					}
-				default:
-					{
-						break;
-					}
-			}
 		}
 
 		public static JsonTracer? TryRegisterAsGlobalTracer()
@@ -707,34 +511,5 @@ namespace EpicGames.Core
 				_finishedSpans.Add(jsonTracerSpan);
 			}
 		}
-
-		internal void ActivateSpan(ISpan span)
-		{
-			PerThreadActiveSpans.AddOrUpdate(Environment.CurrentManagedThreadId, 
-			(k) =>
-			{
-				Stack<ISpan> stack = new Stack<ISpan>();
-				stack.Push(span);
-				return stack;
-			}, 
-			(k,v) =>
-			{
-				v.Push(span);
-				return v;
-			}
-			);
-		}
-
-		internal void DeactivateActiveSpan(ISpan span)
-		{
-			// It's possible for other threads to finish spans
-			foreach(Stack<ISpan> threadSpanStack in PerThreadActiveSpans.Values)
-			{
-				if (threadSpanStack != null && threadSpanStack.Any() && threadSpanStack.Peek() == span)
-				{
-					threadSpanStack.Pop();
-				}
-			}
-		}
 	}
 }
diff --git a/Engine/Source/Programs/UnrealBuildTool/UnrealBuildTool.cs b/Engine/Source/Programs/UnrealBuildTool/UnrealBuildTool.cs
--- a/Engine/Source/Programs/UnrealBuildTool/UnrealBuildTool.cs
+++ b/Engine/Source/Programs/UnrealBuildTool/UnrealBuildTool.cs
@@ -716,7 +716,6 @@ namespace UnrealBuildTool
 				if (Tracer != null)
 				{
 					Tracer.Flush();
-					Tracer.Dispose();
 				}
 
 				// Delete the ubt run file
