Fixed new coding standard issues in UHT

diff --git a/Engine/Source/Programs/Shared/EpicGames.UHT/Exporters/CodeGen/UhtCodeGenerator.cs b/Engine/Source/Programs/Shared/EpicGames.UHT/Exporters/CodeGen/UhtCodeGenerator.cs
--- a/Engine/Source/Programs/Shared/EpicGames.UHT/Exporters/CodeGen/UhtCodeGenerator.cs
+++ b/Engine/Source/Programs/Shared/EpicGames.UHT/Exporters/CodeGen/UhtCodeGenerator.cs
@@ -124,21 +124,24 @@ namespace EpicGames.UHT.Exporters.CodeGen
 
 				bool writeHeader = false;
 				prereqs.Clear();
-				foreach (UhtHeaderFile headerFile in package.Children)
+				foreach (UhtType packageChild in package.Children)
 				{
-					prereqs.Add(HeaderInfos[headerFile.HeaderFileTypeIndex].Task);
-					if (!writeHeader)
+					if (packageChild is UhtHeaderFile headerFile)
 					{
-						foreach (UhtType type in headerFile.Children)
+						prereqs.Add(HeaderInfos[headerFile.HeaderFileTypeIndex].Task);
+						if (!writeHeader)
 						{
-							if (type is UhtClass classObj)
+							foreach (UhtType type in headerFile.Children)
 							{
-								if (classObj.ClassType != UhtClassType.NativeInterface &&
-									classObj.ClassFlags.HasExactFlags(EClassFlags.Native | EClassFlags.Intrinsic, EClassFlags.Native) &&
-									!classObj.ClassExportFlags.HasAllFlags(UhtClassExportFlags.NoExport))
+								if (type is UhtClass classObj)
 								{
-									writeHeader = true;
-									break;
+									if (classObj.ClassType != UhtClassType.NativeInterface &&
+										classObj.ClassFlags.HasExactFlags(EClassFlags.Native | EClassFlags.Intrinsic, EClassFlags.Native) &&
+										!classObj.ClassExportFlags.HasAllFlags(UhtClassExportFlags.NoExport))
+									{
+										writeHeader = true;
+										break;
+									}
 								}
 							}
 						}
@@ -247,9 +250,12 @@ namespace EpicGames.UHT.Exporters.CodeGen
 			objectInfo.UnregisteredExternalDecl = objectInfo.RegsiteredExternalDecl = $"\t{packageInfo.Api}_API UPackage* {objectInfo.RegisteredSingletonName}();\r\n"; //COMPATIBILITY-TODO remove the extra _API
 			objectInfo.UnregsiteredCrossReference = objectInfo.RegsiteredCrossReference = $"\tUPackage* {objectInfo.RegisteredSingletonName}();\r\n";
 
-			foreach (UhtHeaderFile headerFile in package.Children)
+			foreach (UhtType packageChild in package.Children)
 			{
-				InitHeaderInfo(builder, package, ref packageInfo, headerFile);
+				if (packageChild is UhtHeaderFile headerFile)
+				{
+					InitHeaderInfo(builder, package, ref packageInfo, headerFile);
+				}
 			}
 		}
 
@@ -300,9 +306,12 @@ namespace EpicGames.UHT.Exporters.CodeGen
 			}
 			headerInfo.FileId = new string(outFilePath);
 
-			foreach (UhtObject obj in headerFile.Children)
+			foreach (UhtType headerFileChild in headerFile.Children)
 			{
-				InitObjectInfo(builder, package, ref packageInfo, ref headerInfo, obj);
+				if (headerFileChild is UhtObject obj)
+				{
+					InitObjectInfo(builder, package, ref packageInfo, ref headerInfo, obj);
+				}
 			}
 		}
 
@@ -407,11 +416,14 @@ namespace EpicGames.UHT.Exporters.CodeGen
 		private static List<UhtHeaderFile> GetSortedHeaderFiles(UhtPackage package)
 		{
 			List<UhtHeaderFile> sortedHeaders = new(package.Children.Count);
-			foreach (UhtHeaderFile headerFile in package.Children)
+			foreach (UhtType packageChild in package.Children)
 			{
-				if (headerFile.ShouldExport)
+				if (packageChild is UhtHeaderFile headerFile)
 				{
-					sortedHeaders.Add(headerFile);
+					if (headerFile.ShouldExport)
+					{
+						sortedHeaders.Add(headerFile);
+					}
 				}
 			}
 			sortedHeaders.Sort((lhs, rhs) => { return StringComparerUE.OrdinalIgnoreCase.Compare(lhs.FilePath, rhs.FilePath); });
diff --git a/Engine/Source/Programs/Shared/EpicGames.UHT/Exporters/CodeGen/UhtHeaderCodeGenerator.cs b/Engine/Source/Programs/Shared/EpicGames.UHT/Exporters/CodeGen/UhtHeaderCodeGenerator.cs
--- a/Engine/Source/Programs/Shared/EpicGames.UHT/Exporters/CodeGen/UhtHeaderCodeGenerator.cs
+++ b/Engine/Source/Programs/Shared/EpicGames.UHT/Exporters/CodeGen/UhtHeaderCodeGenerator.cs
@@ -57,35 +57,38 @@ namespace EpicGames.UHT.Exporters.CodeGen
 			builder.AppendTabs(tabs).Append('{').Append(endl);
 
 			++tabs;
-			foreach (UhtProperty property in function.Children)
+			foreach (UhtType functionChild in function.Children)
 			{
-				bool emitConst = property.PropertyFlags.HasAnyFlags(EPropertyFlags.ConstParm) && property is UhtObjectProperty;
-
-				//@TODO: UCREMOVAL: This is awful code duplication to avoid a double-const
+				if (functionChild is UhtProperty property)
 				{
-					//@TODO: bEmitConst will only be true if we have an object, so checking interface here doesn't do anything.
-					// export 'const' for parameters
-					bool isConstParam = property is UhtInterfaceProperty && !property.PropertyFlags.HasAnyFlags(EPropertyFlags.OutParm); //@TODO: This should be const once that flag exists
-					bool isOnConstClass = false;
-					if (property is UhtObjectProperty objectProperty)
+					bool emitConst = property.PropertyFlags.HasAnyFlags(EPropertyFlags.ConstParm) && property is UhtObjectProperty;
+
+					//@TODO: UCREMOVAL: This is awful code duplication to avoid a double-const
 					{
-						isOnConstClass = objectProperty.Class.ClassFlags.HasAnyFlags(EClassFlags.Const);
+						//@TODO: bEmitConst will only be true if we have an object, so checking interface here doesn't do anything.
+						// export 'const' for parameters
+						bool isConstParam = property is UhtInterfaceProperty && !property.PropertyFlags.HasAnyFlags(EPropertyFlags.OutParm); //@TODO: This should be const once that flag exists
+						bool isOnConstClass = false;
+						if (property is UhtObjectProperty objectProperty)
+						{
+							isOnConstClass = objectProperty.Class.ClassFlags.HasAnyFlags(EClassFlags.Const);
+						}
+
+						if (isConstParam || isOnConstClass)
+						{
+							emitConst = false; // ExportCppDeclaration will do it for us
+						}
 					}
 
-					if (isConstParam || isOnConstClass)
+					builder.AppendTabs(tabs);
+					if (emitConst)
 					{
-						emitConst = false; // ExportCppDeclaration will do it for us
+						builder.Append("const ");
 					}
-				}
 
-				builder.AppendTabs(tabs);
-				if (emitConst)
-				{
-					builder.Append("const ");
+					builder.AppendFullDecl(property, textType, false);
+					builder.Append(';').Append(endl);
 				}
-
-				builder.AppendFullDecl(property, textType, false);
-				builder.Append(';').Append(endl);
 			}
 
 			if (outputConstructor)
@@ -295,23 +298,26 @@ namespace EpicGames.UHT.Exporters.CodeGen
 			builder.AppendTabs(tabs).Append(eventParameterStructName).Append(" Parms;").Append(endl);
 
 			// Declare a parameter struct for this event/delegate and assign the struct members using the values passed into the event/delegate call.
-			foreach (UhtProperty property in function.ParameterProperties.Span)
+			foreach (UhtType parameter in function.ParameterProperties.Span)
 			{
-				if (property.IsStaticArray)
-				{
-					builder.AppendTabs(tabs).Append("FMemory::Memcpy(Parm.")
-						.Append(property.SourceName).Append(',')
-						.Append(property.SourceName).Append(",sizeof(Parms.")
-						.Append(property.SourceName).Append(");").Append(endl);
-				}
-				else
+				if (parameter is UhtProperty property)
 				{
-					builder.AppendTabs(tabs).Append("Parms.").Append(property.SourceName).Append('=').Append(property.SourceName);
-					if (property is UhtBoolProperty)
+					if (property.IsStaticArray)
 					{
-						builder.Append(" ? true : false");
+						builder.AppendTabs(tabs).Append("FMemory::Memcpy(Parm.")
+							.Append(property.SourceName).Append(',')
+							.Append(property.SourceName).Append(",sizeof(Parms.")
+							.Append(property.SourceName).Append(");").Append(endl);
+					}
+					else
+					{
+						builder.AppendTabs(tabs).Append("Parms.").Append(property.SourceName).Append('=').Append(property.SourceName);
+						if (property is UhtBoolProperty)
+						{
+							builder.Append(" ? true : false");
+						}
+						builder.Append(';').Append(endl);
 					}
-					builder.Append(';').Append(endl);
 				}
 			}
 			return builder;
@@ -322,30 +328,33 @@ namespace EpicGames.UHT.Exporters.CodeGen
 			++tabs;
 
 			// Out parm copying.
-			foreach (UhtProperty property in function.ParameterProperties.Span)
+			foreach (UhtType parameter in function.ParameterProperties.Span)
 			{
-				if (property.PropertyFlags.HasExactFlags(EPropertyFlags.ConstParm | EPropertyFlags.OutParm, EPropertyFlags.OutParm))
+				if (parameter is UhtProperty property)
 				{
-					if (property.IsStaticArray)
-					{
-						builder
-							.AppendTabs(tabs)
-							.Append("FMemory::Memcpy(&")
-							.Append(property.SourceName)
-							.Append(",&Parms.")
-							.Append(property.SourceName)
-							.Append(",sizeof(")
-							.Append(property.SourceName)
-							.Append("));").Append(endl);
-					}
-					else
+					if (property.PropertyFlags.HasExactFlags(EPropertyFlags.ConstParm | EPropertyFlags.OutParm, EPropertyFlags.OutParm))
 					{
-						builder
-							.AppendTabs(tabs)
-							.Append(property.SourceName)
-							.Append("=Parms.")
-							.Append(property.SourceName)
-							.Append(';').Append(endl);
+						if (property.IsStaticArray)
+						{
+							builder
+								.AppendTabs(tabs)
+								.Append("FMemory::Memcpy(&")
+								.Append(property.SourceName)
+								.Append(",&Parms.")
+								.Append(property.SourceName)
+								.Append(",sizeof(")
+								.Append(property.SourceName)
+								.Append("));").Append(endl);
+						}
+						else
+						{
+							builder
+								.AppendTabs(tabs)
+								.Append(property.SourceName)
+								.Append("=Parms.")
+								.Append(property.SourceName)
+								.Append(';').Append(endl);
+						}
 					}
 				}
 			}
@@ -390,17 +399,20 @@ namespace EpicGames.UHT.Exporters.CodeGen
 				needsSeperator = true;
 			}
 
-			foreach (UhtProperty parameter in function.ParameterProperties.Span)
+			foreach (UhtType parameter in function.ParameterProperties.Span)
 			{
-				if (needsSeperator)
+				if (parameter is UhtProperty property)
 				{
-					builder.Append(", ");
-				}
-				else
-				{
-					needsSeperator = true;
+					if (needsSeperator)
+					{
+						builder.Append(", ");
+					}
+					else
+					{
+						needsSeperator = true;
+					}
+					builder.AppendFullDecl(property, textType, skipParameterName);
 				}
-				builder.AppendFullDecl(parameter, textType, skipParameterName);
 			}
 
 			builder.Append(')');
@@ -518,7 +530,6 @@ namespace EpicGames.UHT.Exporters.CodeGen
 			}
 		}
 
-
 		protected static string GetDelegateFunctionExportName(UhtFunction function)
 		{
 			const string DelegatePrefix = "delegate";
diff --git a/Engine/Source/Programs/Shared/EpicGames.UHT/Exporters/CodeGen/UhtHeaderCodeGeneratorCppFile.cs b/Engine/Source/Programs/Shared/EpicGames.UHT/Exporters/CodeGen/UhtHeaderCodeGeneratorCppFile.cs
--- a/Engine/Source/Programs/Shared/EpicGames.UHT/Exporters/CodeGen/UhtHeaderCodeGeneratorCppFile.cs
+++ b/Engine/Source/Programs/Shared/EpicGames.UHT/Exporters/CodeGen/UhtHeaderCodeGeneratorCppFile.cs
@@ -726,11 +726,11 @@ namespace EpicGames.UHT.Exporters.CodeGen
 					builder.Append("{\r\n");
 					if (String.IsNullOrEmpty(scriptStruct.RigVMStructInfo.ExecuteContextMember))
 					{
-						builder.Append("\t").Append(scriptStruct.RigVMStructInfo.ExecuteContextType).Append(" TemporaryExecuteContext;\r\n");
+						builder.Append('\t').Append(scriptStruct.RigVMStructInfo.ExecuteContextType).Append(" TemporaryExecuteContext;\r\n");
 					}
 					else
 					{
-						builder.Append("\t").Append(scriptStruct.RigVMStructInfo.ExecuteContextType).Append("& TemporaryExecuteContext = ").Append(scriptStruct.RigVMStructInfo.ExecuteContextMember).Append(";\r\n");
+						builder.Append('\t').Append(scriptStruct.RigVMStructInfo.ExecuteContextType).Append("& TemporaryExecuteContext = ").Append(scriptStruct.RigVMStructInfo.ExecuteContextMember).Append(";\r\n");
 					}
 
 					builder.Append("\tTemporaryExecuteContext.Initialize();\r\n");
@@ -902,7 +902,7 @@ namespace EpicGames.UHT.Exporters.CodeGen
 
 				builder.AppendPropertiesParamsCountStaticAssert(function, staticsName);
 
-				if (!string.IsNullOrEmpty(sizeOfStatic))
+				if (!String.IsNullOrEmpty(sizeOfStatic))
 				{
 					builder.Append("\tstatic_assert(").Append(sizeOfStatic).Append(" < MAX_uint16);\r\n");
 				}
@@ -1088,7 +1088,7 @@ namespace EpicGames.UHT.Exporters.CodeGen
 			callbackFunctions.Sort((x, y) => StringComparerUE.OrdinalIgnoreCase.Compare(x.EngineName, y.EngineName));
 
 			// Generate the callback parameter structures
-			AppendCallbackParametersDecls(builder, classObj, callbackFunctions);
+			AppendCallbackParametersDecls(builder, callbackFunctions);
 
 			// VM -> C++ proxies (events and delegates).
 			if (!classObj.ClassExportFlags.HasAnyFlags(UhtClassExportFlags.NoExport))
@@ -1296,7 +1296,7 @@ namespace EpicGames.UHT.Exporters.CodeGen
 			return builder;
 		}
 
-		private StringBuilder AppendSparseAccessors(StringBuilder builder, UhtClass classObj)
+		private static StringBuilder AppendSparseAccessors(StringBuilder builder, UhtClass classObj)
 		{
 			string[]? sparseDataTypes = classObj.MetaData.GetStringArray(UhtNames.SparseClassDataTypes);
 			if (sparseDataTypes == null)
@@ -1352,9 +1352,12 @@ namespace EpicGames.UHT.Exporters.CodeGen
 				builder.Append("\t\tUFunction* const Func = O->FindFunction(NAME_").Append(function.Outer?.SourceName).Append('_').Append(function.SourceName).Append(");\r\n");
 				builder.Append("\t\tif (Func)\r\n");
 				builder.Append("\t\t{\r\n");
-				foreach (UhtProperty property in function.ParameterProperties.Span)
+				foreach (UhtType parameter in function.ParameterProperties.Span)
 				{
-					builder.Append("\t\t\tParms.").Append(property.SourceName).Append('=').Append(property.SourceName).Append(";\r\n");
+					if (parameter is UhtProperty property)
+					{
+						builder.Append("\t\t\tParms.").Append(property.SourceName).Append('=').Append(property.SourceName).Append(";\r\n");
+					}
 				}
 				builder
 					.Append("\t\t\t")
@@ -1362,11 +1365,14 @@ namespace EpicGames.UHT.Exporters.CodeGen
 					.Append("->ProcessEvent(Func, ")
 					.Append(function.Children.Count > 0 ? "&Parms" : "NULL")
 					.Append(");\r\n");
-				foreach (UhtProperty property in function.ParameterProperties.Span)
+				foreach (UhtType parameter in function.ParameterProperties.Span)
 				{
-					if (property.PropertyFlags.HasExactFlags(EPropertyFlags.OutParm | EPropertyFlags.ConstParm, EPropertyFlags.OutParm))
+					if (parameter is UhtProperty property)
 					{
-						builder.Append("\t\t\t").Append(property.SourceName).Append("=Parms.").Append(property.SourceName).Append(";\r\n");
+						if (property.PropertyFlags.HasExactFlags(EPropertyFlags.OutParm | EPropertyFlags.ConstParm, EPropertyFlags.OutParm))
+						{
+							builder.Append("\t\t\t").Append(property.SourceName).Append("=Parms.").Append(property.SourceName).Append(";\r\n");
+						}
 					}
 				}
 				builder.Append("\t\t}\r\n");
@@ -1390,14 +1396,17 @@ namespace EpicGames.UHT.Exporters.CodeGen
 					builder.Append("I->").Append(function.SourceName).Append("_Implementation(");
 
 					bool first = true;
-					foreach (UhtProperty property in function.ParameterProperties.Span)
+					foreach (UhtType parameter in function.ParameterProperties.Span)
 					{
-						if (!first)
+						if (parameter is UhtProperty property)
 						{
-							builder.Append(',');
+							if (!first)
+							{
+								builder.Append(',');
+							}
+							first = false;
+							builder.Append(property.SourceName);
 						}
-						first = false;
-						builder.Append(property.SourceName);
 					}
 					builder.Append(");\r\n");
 					builder.Append("\t\t}\r\n");
@@ -1806,7 +1815,7 @@ namespace EpicGames.UHT.Exporters.CodeGen
 			return builder;
 		}
 
-		private StringBuilder AppendCallbackParametersDecls(StringBuilder builder, UhtClass classObj, List<UhtFunction> callbackFunctions)
+		private static StringBuilder AppendCallbackParametersDecls(StringBuilder builder, List<UhtFunction> callbackFunctions)
 		{
 			foreach (UhtFunction function in callbackFunctions)
 			{
@@ -1814,6 +1823,7 @@ namespace EpicGames.UHT.Exporters.CodeGen
 			}
 			return builder;
 		}
+
 		private StringBuilder AppendCallbackFunctions(StringBuilder builder, UhtClass classObj, List<UhtFunction> callbackFunctions)
 		{
 			if (callbackFunctions.Count > 0)
@@ -1903,9 +1913,12 @@ namespace EpicGames.UHT.Exporters.CodeGen
 		private static StringBuilder AppendFunctionThunk(StringBuilder builder, UhtFunction function)
 		{
 			// Export the GET macro for the parameters
-			foreach (UhtProperty parameter in function.ParameterProperties.Span)
+			foreach (UhtType parameter in function.ParameterProperties.Span)
 			{
-				builder.Append("\t\t").AppendFunctionThunkParameterGet(parameter).Append(";\r\n");
+				if (parameter is UhtProperty property)
+				{
+					builder.Append("\t\t").AppendFunctionThunkParameterGet(property).Append(";\r\n");
+				}
 			}
 
 			builder.Append("\t\tP_FINISH;\r\n");
@@ -2034,14 +2047,17 @@ namespace EpicGames.UHT.Exporters.CodeGen
 		public static StringBuilder AppendFunctionThunkParameterNames(this StringBuilder builder, UhtFunction function)
 		{
 			bool first = true;
-			foreach (UhtProperty parameter in function.ParameterProperties.Span)
+			foreach (UhtType parameter in function.ParameterProperties.Span)
 			{
-				if (!first)
+				if (parameter is UhtProperty property)
 				{
-					builder.Append(',');
+					if (!first)
+					{
+						builder.Append(',');
+					}
+					builder.AppendFunctionThunkParameterArg(property);
+					first = false;
 				}
-				builder.AppendFunctionThunkParameterArg(parameter);
-				first = false;
 			}
 			return builder;
 		}
diff --git a/Engine/Source/Programs/Shared/EpicGames.UHT/Exporters/CodeGen/UhtHeaderCodeGeneratorHFile.cs b/Engine/Source/Programs/Shared/EpicGames.UHT/Exporters/CodeGen/UhtHeaderCodeGeneratorHFile.cs
--- a/Engine/Source/Programs/Shared/EpicGames.UHT/Exporters/CodeGen/UhtHeaderCodeGeneratorHFile.cs
+++ b/Engine/Source/Programs/Shared/EpicGames.UHT/Exporters/CodeGen/UhtHeaderCodeGeneratorHFile.cs
@@ -4,7 +4,6 @@ using System;
 using System.Collections.Generic;
 using System.IO;
 using System.Linq;
-using System.Net.NetworkInformation;
 using System.Text;
 using EpicGames.Core;
 using EpicGames.UHT.Types;
@@ -41,7 +40,7 @@ namespace EpicGames.UHT.Exporters.CodeGen
 				StringBuilder builder = borrower.StringBuilder;
 
 				builder.Append(HeaderCopyright);
-				builder.Append("// IWYU pragma: private, include \"").Append(this.HeaderFile.IncludeFilePath).Append("\"\r\n");
+				builder.Append("// IWYU pragma: private, include \"").Append(HeaderFile.IncludeFilePath).Append("\"\r\n");
 
 				// Attempt to limit the headers included. This is needed in the lower level engine code
 				// to get around circular header include issues.
@@ -308,14 +307,14 @@ namespace EpicGames.UHT.Exporters.CodeGen
 					{
 						builder.Append("\ttypedef ").Append(scriptStruct.SuperScriptStruct.SourceName).Append(" Super; \\\r\n");
 					}
-					UhtProperty? FastArrayProperty = ObjectInfos[scriptStruct.ObjectTypeIndex].FastArrayProperty;
-					if (FastArrayProperty != null)
+					UhtProperty? fastArrayProperty = ObjectInfos[scriptStruct.ObjectTypeIndex].FastArrayProperty;
+					if (fastArrayProperty != null)
 					{
 						builder
 							.Append("\tUE_NET_DECLARE_FASTARRAY(")
 							.Append(scriptStruct.SourceName)
 							.Append(", ")
-							.Append(FastArrayProperty.SourceName)
+							.Append(fastArrayProperty.SourceName)
 							.Append(", ");
 						if (!scriptStruct.ScriptStructFlags.HasAnyFlags(EStructFlags.RequiredAPI))
 						{
@@ -369,9 +368,9 @@ namespace EpicGames.UHT.Exporters.CodeGen
 			}
 			else if (enumObj.CppForm == UhtEnumCppForm.Namespaced && enumObj.UnderlyingType != UhtEnumUnderlyingType.Unspecified)
 			{
-				string[] SplitName = enumObj.CppType.Split("::");
+				string[] splitName = enumObj.CppType.Split("::");
 				builder.Append("\r\n");
-				builder.Append("namespace ").Append(SplitName[0]).Append(" { enum ").Append(SplitName[1]);
+				builder.Append("namespace ").Append(splitName[0]).Append(" { enum ").Append(splitName[1]);
 				builder.Append(" : ").Append(enumObj.UnderlyingType.ToString().ToLower());
 				builder.Append("; }\r\n");
 			}
@@ -486,42 +485,34 @@ namespace EpicGames.UHT.Exporters.CodeGen
 
 				if (usesLegacy)
 				{
-					using (UhtMacroCreator macro = new(builder, this, classObj, GeneratedBodyLegacyMacroSuffix))
-					{
-						builder.Append('\t');
-						AppendGeneratedMacroDeprecationWarning(builder, "GENERATED_UINTERFACE_BODY");
-						builder.Append('\t').Append(DisableDeprecationWarnings).Append(" \\\r\n");
-						builder.Append('\t').AppendMacroName(this, classObj, GeneratedUInterfaceBodyMacroSuffix).Append(" \\\r\n");
-						builder.Append('\t').AppendMacroName(this, classObj, StandardConstructorsMacroSuffix).Append(" \\\r\n");
-						builder.Append('\t').Append(EnableDeprecationWarnings).Append(" \\\r\n");
-					}
+					using UhtMacroCreator macro = new(builder, this, classObj, GeneratedBodyLegacyMacroSuffix);
+					builder.Append('\t');
+					AppendGeneratedMacroDeprecationWarning(builder, "GENERATED_UINTERFACE_BODY");
+					builder.Append('\t').Append(DisableDeprecationWarnings).Append(" \\\r\n");
+					builder.Append('\t').AppendMacroName(this, classObj, GeneratedUInterfaceBodyMacroSuffix).Append(" \\\r\n");
+					builder.Append('\t').AppendMacroName(this, classObj, StandardConstructorsMacroSuffix).Append(" \\\r\n");
+					builder.Append('\t').Append(EnableDeprecationWarnings).Append(" \\\r\n");
 				}
 				else
 				{
-					using (UhtMacroCreator macro = new(builder, this, classObj, GeneratedBodyMacroSuffix))
-					{
-						builder.Append('\t').Append(DisableDeprecationWarnings).Append(" \\\r\n");
-						builder.Append('\t').AppendMacroName(this, classObj, GeneratedUInterfaceBodyMacroSuffix).Append(" \\\r\n");
-						builder.Append('\t').AppendMacroName(this, classObj, EnchancedConstructorsMacroSuffix).Append(" \\\r\n");
-						AppendAccessSpecifier(builder, classObj);
-						builder.Append(" \\\r\n");
-						builder.Append('\t').Append(EnableDeprecationWarnings).Append(" \\\r\n");
-					}
+					using UhtMacroCreator macro = new(builder, this, classObj, GeneratedBodyMacroSuffix);
+					builder.Append('\t').Append(DisableDeprecationWarnings).Append(" \\\r\n");
+					builder.Append('\t').AppendMacroName(this, classObj, GeneratedUInterfaceBodyMacroSuffix).Append(" \\\r\n");
+					builder.Append('\t').AppendMacroName(this, classObj, EnchancedConstructorsMacroSuffix).Append(" \\\r\n");
+					AppendAccessSpecifier(builder, classObj);
+					builder.Append(" \\\r\n");
+					builder.Append('\t').Append(EnableDeprecationWarnings).Append(" \\\r\n");
 				}
 
 				if (alternateUsesLegacy)
 				{
-					using (UhtMacroCreator macro = new(builder, this, classObj, InClassIInterfaceMacroSuffix))
-					{
-						AppendInClassIInterface(builder, classObj, callbackFunctions, api);
-					}
+					using UhtMacroCreator macro = new(builder, this, classObj, InClassIInterfaceMacroSuffix);
+					AppendInClassIInterface(builder, classObj, callbackFunctions, api);
 				}
 				else
 				{
-					using (UhtMacroCreator macro = new(builder, this, classObj, InClassIInterfaceNoPureDeclsMacroSuffix))
-					{
-						AppendInClassIInterface(builder, classObj, callbackFunctions, api);
-					}
+					using UhtMacroCreator macro = new(builder, this, classObj, InClassIInterfaceNoPureDeclsMacroSuffix);
+					AppendInClassIInterface(builder, classObj, callbackFunctions, api);
 				}
 
 				AppendProlog(builder, classObj);
@@ -716,22 +707,18 @@ namespace EpicGames.UHT.Exporters.CodeGen
 
 			if (usesLegacy)
 			{
-				using (UhtMacroCreator macro = new(builder, this, classObj, editorOnly ? EditorOnlyRpcWrappersMacroSuffix : RpcWrappersMacroSuffix))
-				{
-					AppendAutogeneratedBlueprintFunctionDeclarations(builder, classObj, reversedFunctions, editorOnly);
-					AppendRpcWrappers(builder, reversedFunctions, editorOnly);
-				}
+				using UhtMacroCreator macro = new(builder, this, classObj, editorOnly ? EditorOnlyRpcWrappersMacroSuffix : RpcWrappersMacroSuffix);
+				AppendAutogeneratedBlueprintFunctionDeclarations(builder, classObj, reversedFunctions, editorOnly);
+				AppendRpcWrappers(builder, reversedFunctions, editorOnly);
 			}
 			else
 			{
-				using (UhtMacroCreator macro = new(builder, this, classObj, editorOnly ? EditorOnlyRpcWrappersNoPureDeclsMacroSuffix : RpcWrappersNoPureDeclsMacroSuffix))
+				using UhtMacroCreator macro = new(builder, this, classObj, editorOnly ? EditorOnlyRpcWrappersNoPureDeclsMacroSuffix : RpcWrappersNoPureDeclsMacroSuffix);
+				if (classObj.GeneratedCodeVersion <= EGeneratedCodeVersion.V1)
 				{
-					if (classObj.GeneratedCodeVersion <= EGeneratedCodeVersion.V1)
-					{
-						AppendAutogeneratedBlueprintFunctionDeclarationsOnlyNotDeclared(builder, classObj, reversedFunctions, editorOnly);
-					}
-					AppendRpcWrappers(builder, reversedFunctions, editorOnly);
+					AppendAutogeneratedBlueprintFunctionDeclarationsOnlyNotDeclared(builder, classObj, reversedFunctions, editorOnly);
 				}
+				AppendRpcWrappers(builder, reversedFunctions, editorOnly);
 			}
 
 			if (editorOnly)
@@ -739,15 +726,11 @@ namespace EpicGames.UHT.Exporters.CodeGen
 				builder.Append("#else\r\n");
 				if (usesLegacy)
 				{
-					using (UhtMacroCreator macro = new(builder, this, classObj, editorOnly ? EditorOnlyRpcWrappersMacroSuffix : RpcWrappersMacroSuffix))
-					{
-					}
+					using UhtMacroCreator macro = new(builder, this, classObj, editorOnly ? EditorOnlyRpcWrappersMacroSuffix : RpcWrappersMacroSuffix);
 				}
 				else
 				{
-					using (UhtMacroCreator macro = new(builder, this, classObj, editorOnly ? EditorOnlyRpcWrappersNoPureDeclsMacroSuffix : RpcWrappersNoPureDeclsMacroSuffix))
-					{
-					}
+					using UhtMacroCreator macro = new(builder, this, classObj, editorOnly ? EditorOnlyRpcWrappersNoPureDeclsMacroSuffix : RpcWrappersNoPureDeclsMacroSuffix);
 				}
 				builder.AppendEndEditorOnlyGuard();
 			}
@@ -856,19 +839,6 @@ namespace EpicGames.UHT.Exporters.CodeGen
 			return builder;
 		}
 
-		private StringBuilder AppendCallbackParametersDecls(StringBuilder builder, UhtClass classObj, List<UhtFunction> callbackFunctions)
-		{
-			if (callbackFunctions.Count > 0)
-			{
-				using UhtMacroCreator macro = new(builder, this, classObj, EventParamsMacroSuffix);
-				foreach (UhtFunction function in callbackFunctions)
-				{
-					AppendEventParameter(builder, function, function.StrippedFunctionName, UhtPropertyTextType.EventParameterMember, true, 1, " \\\r\n");
-				}
-			}
-			return builder;
-		}
-
 		private StringBuilder AppendCallbackRpcWrapperDecls(StringBuilder builder, UhtClass classObj, List<UhtFunction> callbackFunctions)
 		{
 			if (callbackFunctions.Count > 0)
diff --git a/Engine/Source/Programs/Shared/EpicGames.UHT/Parsers/UhtClassParser.cs b/Engine/Source/Programs/Shared/EpicGames.UHT/Parsers/UhtClassParser.cs
--- a/Engine/Source/Programs/Shared/EpicGames.UHT/Parsers/UhtClassParser.cs
+++ b/Engine/Source/Programs/Shared/EpicGames.UHT/Parsers/UhtClassParser.cs
@@ -110,10 +110,7 @@ namespace EpicGames.UHT.Parsers
 
 					specifiers.ParseDeferred();
 
-					if (classObj.Outer != null)
-					{
-						classObj.Outer.AddChild(classObj);
-					}
+					classObj.Outer?.AddChild(classObj);
 
 					topScope.AddModuleRelativePathToMetaData();
 
diff --git a/Engine/Source/Programs/Shared/EpicGames.UHT/Parsers/UhtInterfaceClassParser.cs b/Engine/Source/Programs/Shared/EpicGames.UHT/Parsers/UhtInterfaceClassParser.cs
--- a/Engine/Source/Programs/Shared/EpicGames.UHT/Parsers/UhtInterfaceClassParser.cs
+++ b/Engine/Source/Programs/Shared/EpicGames.UHT/Parsers/UhtInterfaceClassParser.cs
@@ -106,10 +106,7 @@ namespace EpicGames.UHT.Parsers
 
 					specifiers.ParseDeferred();
 
-					if (classObj.Outer != null)
-					{
-						classObj.Outer.AddChild(classObj);
-					}
+					classObj.Outer?.AddChild(classObj);
 
 					topScope.AddModuleRelativePathToMetaData();
 
diff --git a/Engine/Source/Programs/Shared/EpicGames.UHT/Parsers/UhtSpecifierParser.cs b/Engine/Source/Programs/Shared/EpicGames.UHT/Parsers/UhtSpecifierParser.cs
--- a/Engine/Source/Programs/Shared/EpicGames.UHT/Parsers/UhtSpecifierParser.cs
+++ b/Engine/Source/Programs/Shared/EpicGames.UHT/Parsers/UhtSpecifierParser.cs
@@ -183,10 +183,7 @@ namespace EpicGames.UHT.Parsers
 				{
 					if (specifier.When == UhtSpecifierWhen.Deferred)
 					{
-						if (_deferredSpecifiers == null)
-						{
-							_deferredSpecifiers = new List<DeferredSpecifier>();
-						}
+						_deferredSpecifiers ??= new List<DeferredSpecifier>();
 						_deferredSpecifiers.Add(new DeferredSpecifier { _specifier = specifier, _value = value });
 					}
 					else
@@ -223,19 +220,13 @@ namespace EpicGames.UHT.Parsers
 
 		private void ParseKVPValueInternal()
 		{
-			if (_currentKVPValues == null)
-			{
-				_currentKVPValues = new List<KeyValuePair<StringView, StringView>>();
-			}
+			_currentKVPValues ??= new List<KeyValuePair<StringView, StringView>>();
 			_currentKVPValues.Add(ReadKVP());
 		}
 
 		private void ParseStringViewListInternal()
 		{
-			if (_currentStringValues == null)
-			{
-				_currentStringValues = new List<StringView>();
-			}
+			_currentStringValues ??= new List<StringView>();
 			_currentStringValues.Add(ReadValue());
 		}
 
diff --git a/Engine/Source/Programs/Shared/EpicGames.UHT/Types/Properties/UhtNumericProperty.cs b/Engine/Source/Programs/Shared/EpicGames.UHT/Types/Properties/UhtNumericProperty.cs
--- a/Engine/Source/Programs/Shared/EpicGames.UHT/Types/Properties/UhtNumericProperty.cs
+++ b/Engine/Source/Programs/Shared/EpicGames.UHT/Types/Properties/UhtNumericProperty.cs
@@ -1,7 +1,6 @@
 // Copyright Epic Games, Inc. All Rights Reserved.
 
 using System.Text;
-using System.Text.Json.Serialization;
 
 namespace EpicGames.UHT.Types
 {
diff --git a/Engine/Source/Programs/Shared/EpicGames.UHT/Types/UhtScriptStruct.cs b/Engine/Source/Programs/Shared/EpicGames.UHT/Types/UhtScriptStruct.cs
--- a/Engine/Source/Programs/Shared/EpicGames.UHT/Types/UhtScriptStruct.cs
+++ b/Engine/Source/Programs/Shared/EpicGames.UHT/Types/UhtScriptStruct.cs
@@ -923,16 +923,16 @@ namespace EpicGames.UHT.Types
 		{
 			if (RigVMStructInfo != null)
 			{
-				if (MetaData.TryGetValueHierarchical("ExecuteContext", out string? ExecuteContextMetadata))
+				if (MetaData.TryGetValueHierarchical("ExecuteContext", out string? executeContextMetadata))
 				{
-					RigVMStructInfo.ExecuteContextType = ExecuteContextMetadata;
+					RigVMStructInfo.ExecuteContextType = executeContextMetadata;
 				}
 
 				for (UhtStruct? current = this; current != null; current = current.SuperStruct)
 				{
 					foreach (UhtProperty property in current.Properties)
 					{
-						UhtRigVMParameter parameter = new UhtRigVMParameter(property, RigVMStructInfo.Members.Count);
+						UhtRigVMParameter parameter = new(property, RigVMStructInfo.Members.Count);
 						if (parameter.IsExecuteContext())
 						{
 							if (String.IsNullOrEmpty(RigVMStructInfo.ExecuteContextMember))
@@ -945,14 +945,14 @@ namespace EpicGames.UHT.Types
 							}
 							else if (RigVMStructInfo.ExecuteContextType != parameter.Type)
 							{
-								this.LogError($"RigVM Struct {this.SourceName} contains properties of varying execute context type {RigVMStructInfo.ExecuteContextType} vs {parameter.Type}.");
+								this.LogError($"RigVM Struct {SourceName} contains properties of varying execute context type {RigVMStructInfo.ExecuteContextType} vs {parameter.Type}.");
 							}
 						}
 						else
 						{
 							if (parameter.IsLazy && parameter.Output)
 							{
-								this.LogError($"RigVM Struct {this.SourceName} - Member {parameter.Name} is both an output and a lazy input.");
+								this.LogError($"RigVM Struct {SourceName} - Member {parameter.Name} is both an output and a lazy input.");
 							}
 							RigVMStructInfo.Members.Add(parameter);
 						}
diff --git a/Engine/Source/Programs/Shared/EpicGames.UHT/Types/UhtStruct.cs b/Engine/Source/Programs/Shared/EpicGames.UHT/Types/UhtStruct.cs
--- a/Engine/Source/Programs/Shared/EpicGames.UHT/Types/UhtStruct.cs
+++ b/Engine/Source/Programs/Shared/EpicGames.UHT/Types/UhtStruct.cs
@@ -114,7 +114,7 @@ namespace EpicGames.UHT.Types
 		/// <summary>
 		/// Stack used to test to see if we are recursing in a ScanForInstanceReferenced call
 		/// </summary>
-		private readonly ThreadLocal<List<UhtType>> _scanForInstanceReferencedStack = new ThreadLocal<List<UhtType>>(() =>
+		private static readonly ThreadLocal<List<UhtType>> s_scanForInstanceReferencedStack = new(() =>
 		{
 			return new List<UhtType>();
 		});
@@ -156,10 +156,7 @@ namespace EpicGames.UHT.Types
 		/// <inheritdoc/>
 		protected override void ResolveSuper(UhtResolvePhase resolvePhase)
 		{
-			if (Super != null)
-			{
-				Super.Resolve(resolvePhase);
-			}
+			Super?.Resolve(resolvePhase);
 
 			foreach (UhtStruct baseStruct in Bases)
 			{
@@ -177,7 +174,7 @@ namespace EpicGames.UHT.Types
 		/// <returns></returns>
 		public bool ScanForInstancedReferenced(bool deepScan)
 		{
-			List<UhtType> scanForInstanceReferencedStack = _scanForInstanceReferencedStack.Value!;
+			List<UhtType> scanForInstanceReferencedStack = s_scanForInstanceReferencedStack.Value!;
 			if (scanForInstanceReferencedStack.Contains(this))
 			{
 				return false;
@@ -257,10 +254,7 @@ namespace EpicGames.UHT.Types
 					UhtStruct? baseStruct = (UhtStruct?)FindType(findOptions | UhtFindOptions.Class | UhtFindOptions.ScriptStruct | UhtFindOptions.SourceName | UhtFindOptions.NoSelf, baseIdentifier);
 					if (baseStruct != null)
 					{
-						if (_bases == null)
-						{
-							_bases = new List<UhtStruct>();
-						}
+						_bases ??= new List<UhtStruct>();
 						_bases.Add(baseStruct);
 						HeaderFile.AddReferencedHeader(baseStruct);
 					}
diff --git a/Engine/Source/Programs/Shared/EpicGames.UHT/Utils/UhtUtilities.cs b/Engine/Source/Programs/Shared/EpicGames.UHT/Utils/UhtUtilities.cs
--- a/Engine/Source/Programs/Shared/EpicGames.UHT/Utils/UhtUtilities.cs
+++ b/Engine/Source/Programs/Shared/EpicGames.UHT/Utils/UhtUtilities.cs
@@ -297,10 +297,7 @@ namespace EpicGames.UHT.Utils
 		{
 			if (!_useStringBuilder)
 			{
-				if (_stringBuilder == null)
-				{
-					_stringBuilder = new StringBuilder();
-				}
+				_stringBuilder ??= new StringBuilder();
 				_useStringBuilder = true;
 				_stringBuilder.Append(_stringView.Span);
 			}
