EpicGames.Core: Add a more descriptive error if unable to parse .csproj files.

diff --git a/Engine/Source/Programs/Shared/EpicGames.Core/CsProjectInfo.cs b/Engine/Source/Programs/Shared/EpicGames.Core/CsProjectInfo.cs
--- a/Engine/Source/Programs/Shared/EpicGames.Core/CsProjectInfo.cs
+++ b/Engine/Source/Programs/Shared/EpicGames.Core/CsProjectInfo.cs
@@ -10,6 +10,16 @@ using System.Xml;
 
 namespace EpicGames.Core
 {
+	/// <summary>
+	/// Exception parsing a csproj file
+	/// </summary>
+	public sealed class CsProjectParseException : Exception
+	{
+		internal CsProjectParseException(string? message, Exception? exception) : base(message, exception)
+		{
+		}
+	}
+
 	/// <summary>
 	/// Basic information from a preprocessed C# project file. Supports reading a project file, expanding simple conditions in it, parsing property values, assembly references and references to other projects.
 	/// </summary>
@@ -390,7 +400,14 @@ namespace EpicGames.Core
 			CsProjectInfo projectInfo = new CsProjectInfo(properties, file);
 
 			// Parse elements in the root node
-			ParseNode(document.DocumentElement, projectInfo);
+			try
+			{
+				ParseNode(document.DocumentElement, projectInfo);
+			}
+			catch (Exception ex)
+			{
+				throw new CsProjectParseException($"Error parsing {file}: {ex.Message}", ex);
+			}
 
 			// Return the complete project
 			outProjectInfo = projectInfo;
