1) Added generic double buffer async file reader with the ability to specify buffer sizes.
2) Added the ability to directly set the IoHash in the file hasher
3) Modified the ComputeAsync hash method to provide a file size hint so it can default to a larger buffer size to improve performance.
4) Modified the ComputeAsync hash method to use the double buffer reader above.

diff --git a/Engine/Source/Programs/Shared/EpicGames.Build/System/FileHasher.cs b/Engine/Source/Programs/Shared/EpicGames.Build/System/FileHasher.cs
--- a/Engine/Source/Programs/Shared/EpicGames.Build/System/FileHasher.cs
+++ b/Engine/Source/Programs/Shared/EpicGames.Build/System/FileHasher.cs
@@ -2,6 +2,7 @@
 
 using System.Collections.Concurrent;
 using System.Collections.Generic;
+using System.Diagnostics;
 using System.IO;
 using System.Linq;
 using System.Text.Json;
@@ -122,6 +123,18 @@ namespace UnrealBuildBase
 			return CachedDigests.AddOrUpdate(item.FullName, digest, (key, oldValue) => digest).Digest;
 		}
 
+		/// <summary>
+		/// In systems where the IoHash is already available, this allows it to be updated directly and
+		/// not recomputed.
+		/// </summary>
+		/// <param name="item">The updated file item information</param>
+		/// <param name="hash">The updated hash</param>
+		public void SetDigest(FileItem item, IoHash hash)
+		{
+			CachedDigest digest = new CachedDigest(item, hash);
+			CachedDigests.AddOrUpdate(item.FullName, digest, (key, oldValue) => digest);
+		}
+
 		/// <summary>
 		/// Get the IoHash digest of a file's contents.
 		/// </summary>
@@ -137,7 +150,7 @@ namespace UnrealBuildBase
 		static async Task<CachedDigest> ComputeDigest(FileItem item, ILogger? logger, CancellationToken CancellationToken = default)
 		{
 			using FileStream stream = FileReference.Open(item.Location, FileMode.Open, FileAccess.Read, FileShare.Read);
-			CachedDigest digest = new CachedDigest(item, await IoHash.ComputeAsync(stream, CancellationToken));
+			CachedDigest digest = new CachedDigest(item, await IoHash.ComputeAsync(stream, item.Length, CancellationToken));
 			logger?.LogDebug("Computed IoHash {Digest} File {Location} Size {Size} LastWrite {LastWrite}", digest.Digest, item.FullName, item.Length, item.LastWriteTimeUtc);
 			return digest;
 		}
diff --git a/Engine/Source/Programs/Shared/EpicGames.Core/StreamUtils.cs b/Engine/Source/Programs/Shared/EpicGames.Core/StreamUtils.cs
--- a/Engine/Source/Programs/Shared/EpicGames.Core/StreamUtils.cs
+++ b/Engine/Source/Programs/Shared/EpicGames.Core/StreamUtils.cs
@@ -1,6 +1,7 @@
 // Copyright Epic Games, Inc. All Rights Reserved.
 
 using System;
+using System.Buffers;
 using System.IO;
 using System.Threading;
 using System.Threading.Tasks;
@@ -51,5 +52,66 @@ namespace EpicGames.Core
 				offset += readLength;
 			}
 		}
+
+		/// <summary>
+		/// Read the contents of a file async using double buffering
+		/// </summary>
+		/// <param name="stream">Data to compute the hash for</param>
+		/// <param name="fileSizeHint">If available, the file size so an appropriate buffer size can be used</param>
+		/// <param name="minBufferSize">Minimum size of the buffer</param>
+		/// <param name="maxBufferSize">Maximum size of the buffer</param>
+		/// <param name="callback">Callback used to send read data back to the caller</param>
+		/// <param name="cancellationToken">Cancellation token used to terminate processing</param>
+		/// <returns>New content hash instance containing the hash of the data</returns>
+		public static async Task ReadAllBytesAsync(this Stream stream, long fileSizeHint, int minBufferSize, int maxBufferSize, Func<ReadOnlyMemory<byte>, Task> callback, CancellationToken cancellationToken = default)
+		{
+			int bufferLength = minBufferSize;
+			if (fileSizeHint > minBufferSize)
+			{
+				bufferLength = (int)(fileSizeHint <= maxBufferSize ? fileSizeHint : maxBufferSize);
+			}
+
+			await ReadAllBytesAsync(stream, bufferLength, callback, cancellationToken);
+		}
+
+		/// <summary>
+		/// Read the contents of a file async using double buffering
+		/// </summary>
+		/// <param name="stream">Data to compute the hash for</param>
+		/// <param name="bufferLength">Size of the internal buffer</param>
+		/// <param name="callback">Callback used to send read data back to the caller</param>
+		/// <param name="cancellationToken">Cancellation token used to terminate processing</param>
+		/// <returns>New content hash instance containing the hash of the data</returns>
+		public static async Task ReadAllBytesAsync(this Stream stream, int bufferLength, Func<ReadOnlyMemory<byte>, Task> callback, CancellationToken cancellationToken = default)
+		{
+			using IMemoryOwner<byte> owner = MemoryPool<byte>.Shared.Rent(bufferLength * 2);
+			Memory<byte> buffer = owner.Memory;
+
+			int readBufferOffset = 0;
+			Memory<byte> appendBuffer = Memory<byte>.Empty;
+			for (; ; )
+			{
+				// Start a read into memory
+				Memory<byte> readBuffer = buffer[readBufferOffset..(readBufferOffset + bufferLength)];
+				Task<int> readTask = Task.Run(async () => await stream.ReadAsync(readBuffer, cancellationToken), cancellationToken);
+
+				// In the meantime, append the last data that was read to the tree
+				if (appendBuffer.Length > 0)
+				{
+					await callback(appendBuffer[0..appendBuffer.Length]);
+				}
+
+				// Wait for the read to finish
+				int numBytes = await readTask;
+				if (numBytes == 0)
+				{
+					break;
+				}
+
+				// Switch the buffers around
+				appendBuffer = readBuffer[0..numBytes];
+				readBufferOffset ^= bufferLength;
+			}
+		}
 	}
 }
diff --git a/Engine/Source/Programs/Shared/EpicGames.IoHash/IoHash.cs b/Engine/Source/Programs/Shared/EpicGames.IoHash/IoHash.cs
--- a/Engine/Source/Programs/Shared/EpicGames.IoHash/IoHash.cs
+++ b/Engine/Source/Programs/Shared/EpicGames.IoHash/IoHash.cs
@@ -127,16 +127,28 @@ namespace EpicGames.Core
 		/// </summary>
 		/// <param name="stream">Data to compute the hash for</param>
 		/// <returns>New content hash instance containing the hash of the data</returns>
-		public static async Task<IoHash> ComputeAsync(Stream stream, CancellationToken cancellationToken = default)
+		public static async Task<IoHash> ComputeAsync(Stream stream, CancellationToken cancellationToken = default) => await ComputeAsync(stream, -1, cancellationToken);
+
+		/// <summary>
+		/// Creates the IoHash for a stream asynchronously. 
+		/// </summary>
+		/// <param name="stream">Data to compute the hash for</param>
+		/// <param name="fileSizeHint">If available, the file size so an appropriate buffer size can be used</param>
+		/// <param name="cancellationToken">Cancellation token used to terminate processing</param>
+		/// <returns>New content hash instance containing the hash of the data</returns>
+		public static async Task<IoHash> ComputeAsync(Stream stream, long fileSizeHint, CancellationToken cancellationToken = default)
 		{
+			const int MaxBufferSize = 1 * 1024 * 1024;
+			const int MinBufferSize = 16 * 1024;
 			using (Blake3.Hasher hasher = Blake3.Hasher.New())
 			{
-				Memory<byte> buffer = new byte[16384];
-				int length;
-				while ((length = await stream.ReadAsync(buffer, cancellationToken)) > 0)
+				Task Callback(ReadOnlyMemory<byte> data)
 				{
-					hasher.Update(buffer.Slice(0, length).Span);
+					hasher.Update(data.Span);
+					return Task.CompletedTask;
 				}
+
+				await stream.ReadAllBytesAsync(fileSizeHint, MinBufferSize, MaxBufferSize, Callback, cancellationToken);
 				return FromBlake3(hasher);
 			}
 		}
