Optimizations to Horde storage when dealing with larger data.

1) Ability to specify buffer size with FileNodeWriter methods CreateAsync and AppendAsync.
2) Ability to copy a FileNode tree directly into a stream without creating the nodes themselves.
3) When reading data, avoid creating byte[]s for the cache if the cache is not enabled.
4) Reduced default read size from 15MB to 1MB when caching isn't enabled
5) New ReadNodeAsync method that invokes a callback with the data instead of creating the node instance.
6) TreeWriter uses the current packet's full memory buffer size to detect the need to resize indead of the max packat size.
7) New bundle test for the streaming of FileNodes directly to a stream.

diff --git a/Engine/Source/Programs/Shared/EpicGames.Horde.Tests/BundleTests.cs b/Engine/Source/Programs/Shared/EpicGames.Horde.Tests/BundleTests.cs
--- a/Engine/Source/Programs/Shared/EpicGames.Horde.Tests/BundleTests.cs
+++ b/Engine/Source/Programs/Shared/EpicGames.Horde.Tests/BundleTests.cs
@@ -9,7 +9,9 @@ using Microsoft.Extensions.Logging.Abstractions;
 using Microsoft.VisualStudio.TestTools.UnitTesting;
 using System;
 using System.Collections.Generic;
+using System.IO;
 using System.Linq;
+using System.Reflection.Metadata;
 using System.Text;
 using System.Threading;
 using System.Threading.Tasks;
@@ -76,31 +78,58 @@ namespace EpicGames.Horde.Tests
 
 			NodeHandle handle;
 			FileNode node;
+			byte[] data = CreateBuffer(1024);
 
-			handle = await fileNodeWriter.CreateAsync(new byte[7], CancellationToken.None);
+			handle = await fileNodeWriter.CreateAsync(data.AsMemory(0, 7), CancellationToken.None);
 			node = await reader.ReadNodeAsync<FileNode>(handle.Locator);
 			Assert.IsTrue(node is LeafFileNode);
 			Assert.AreEqual(7, ((LeafFileNode)node).Data.Length);
+			await TestBufferlessReadsAsync(reader, handle.Locator, data.AsMemory(0, 7));
 
-			handle = await fileNodeWriter.CreateAsync(new byte[8], CancellationToken.None);
+			handle = await fileNodeWriter.CreateAsync(data.AsMemory(0, 8), CancellationToken.None);
 			node = await reader.ReadNodeAsync<FileNode>(handle.Locator);
 			Assert.IsTrue(node is LeafFileNode);
 			Assert.AreEqual(8, ((LeafFileNode)node).Data.Length);
+			await TestBufferlessReadsAsync(reader, handle.Locator, data.AsMemory(0, 8));
 
-			handle = await fileNodeWriter.CreateAsync(new byte[9], CancellationToken.None);
+			handle = await fileNodeWriter.CreateAsync(data.AsMemory(0, 9), CancellationToken.None);
 			node = await reader.ReadNodeAsync<FileNode>(handle.Locator);
 			Assert.IsTrue(node is InteriorFileNode);
 			Assert.AreEqual(2, ((InteriorFileNode)node).Children.Count);
+			await TestBufferlessReadsAsync(reader, handle.Locator, data.AsMemory(0, 9));
 
-			FileNode? childNode1 = await ((TreeNodeRef<FileNode>)((InteriorFileNode)node).Children[0]).ExpandAsync(reader);
+			FileNode? childNode1 = await ((InteriorFileNode)node).Children[0].ExpandAsync(reader);
 			Assert.IsNotNull(childNode1);
 			Assert.IsTrue(childNode1 is LeafFileNode);
 			Assert.AreEqual(8, ((LeafFileNode)childNode1!).Data.Length);
 
-			FileNode? childNode2 = await ((TreeNodeRef<FileNode>)((InteriorFileNode)node).Children[1]).ExpandAsync(reader);
+			FileNode? childNode2 = await ((InteriorFileNode)node).Children[1].ExpandAsync(reader);
 			Assert.IsNotNull(childNode2);
 			Assert.IsTrue(childNode2 is LeafFileNode);
 			Assert.AreEqual(1, ((LeafFileNode)childNode2!).Data.Length);
+
+			handle = await fileNodeWriter.CreateAsync(data, CancellationToken.None);
+			node = await reader.ReadNodeAsync<FileNode>(handle.Locator);
+			Assert.IsTrue(node is InteriorFileNode);
+			await TestBufferlessReadsAsync(reader, handle.Locator, data);
+		}
+
+		private static byte[] CreateBuffer(int length)
+		{
+			byte[] output = GC.AllocateUninitializedArray<byte>(length);
+			for (int i = 0; i < length; i++)
+			{
+				output[i] = (byte)i;
+			}
+			return output;
+		}
+
+		private static async Task TestBufferlessReadsAsync(TreeReader reader, NodeLocator locator, ReadOnlyMemory<byte> expected)
+		{
+			MemoryStream memoryStream = new MemoryStream();
+			await FileNode.CopyToStreamAsync(reader, locator, memoryStream, default);
+			ReadOnlyMemory<byte> read = memoryStream.ToArray().AsMemory();
+			Assert.IsTrue(read.Span.SequenceEqual(expected.Span));			
 		}
 
 		[TestMethod]
diff --git a/Engine/Source/Programs/Shared/EpicGames.Horde/Storage/Nodes/FileNode.cs b/Engine/Source/Programs/Shared/EpicGames.Horde/Storage/Nodes/FileNode.cs
--- a/Engine/Source/Programs/Shared/EpicGames.Horde/Storage/Nodes/FileNode.cs
+++ b/Engine/Source/Programs/Shared/EpicGames.Horde/Storage/Nodes/FileNode.cs
@@ -1,21 +1,13 @@
 // Copyright Epic Games, Inc. All Rights Reserved.
 
-using EpicGames.Core;
-using EpicGames.Horde.Storage.Nodes;
-using EpicGames.Serialization;
-using Microsoft.Extensions.Options;
 using System;
-using System.Buffers;
-using System.Buffers.Binary;
 using System.Collections.Generic;
 using System.Diagnostics;
 using System.IO;
 using System.Linq;
-using System.Reflection;
-using System.Reflection.Metadata;
 using System.Threading;
 using System.Threading.Tasks;
-using System.Xml.Linq;
+using EpicGames.Core;
 
 namespace EpicGames.Horde.Storage.Nodes
 {
@@ -99,6 +91,34 @@ namespace EpicGames.Horde.Storage.Nodes
 		/// <param name="cancellationToken">Cancellation token for the operation</param>
 		public abstract Task CopyToStreamAsync(TreeReader reader, Stream outputStream, CancellationToken cancellationToken);
 
+		/// <summary>
+		/// Copy the contents of the node to the output stream without creating the intermediate FileNodes
+		/// </summary>
+		/// <param name="reader">Reader for nodes in the tree</param>
+		/// <param name="locator">File node to be copied</param>
+		/// <param name="outputStream">The output stream to receive the data</param>
+		/// <param name="cancellationToken">Cancellation token for the operation</param>
+		public static async Task CopyToStreamAsync(TreeReader reader, NodeLocator locator, Stream outputStream, CancellationToken cancellationToken)
+		{
+			async Task CopyFunc(ReadNodeAsyncCallbackParams parms, ITreeNodeReader nodeReader)
+			{
+				if (parms.Type == typeof(LeafFileNode))
+				{
+					await LeafFileNode.CopyToStreamAsync(outputStream, nodeReader, cancellationToken);
+				}
+				else if (parms.Type == typeof(InteriorFileNode))
+				{
+					await InteriorFileNode.CopyToStreamAsync(reader, outputStream, nodeReader, cancellationToken);
+				}
+				else
+				{
+					throw new ArgumentException("Unexpected FileNode type found.");
+				}
+			}
+
+			await reader.ReadNodeAsync(locator, CopyFunc, cancellationToken);
+		}
+
 		/// <summary>
 		/// Extracts the contents of this node to a file
 		/// </summary>
@@ -157,6 +177,7 @@ namespace EpicGames.Horde.Storage.Nodes
 		/// </summary>
 		public LeafFileNode(ITreeNodeReader reader)
 		{
+			// Keep this code in sync with CopyToStreamAsync
 			Data = reader.ReadFixedLengthBytes(reader.Length);
 		}
 
@@ -251,6 +272,18 @@ namespace EpicGames.Horde.Storage.Nodes
 		{
 			await outputStream.WriteAsync(Data, cancellationToken);
 		}
+
+		/// <summary>
+		/// Copy the contents of the node to the output stream without creating the intermediate FileNodes
+		/// </summary>
+		/// <param name="reader">Reader for nodes in the tree</param>
+		/// <param name="outputStream">The output stream to receive the data</param>
+		/// <param name="cancellationToken">Cancellation token for the operation</param>
+		public static async Task CopyToStreamAsync(Stream outputStream, ITreeNodeReader reader, CancellationToken cancellationToken)
+		{
+			// Keep this code in sync with the constructor
+			await outputStream.WriteAsync(reader.ReadFixedLengthBytes(reader.Length), cancellationToken);
+		}
 	}
 
 	/// <summary>
@@ -278,6 +311,7 @@ namespace EpicGames.Horde.Storage.Nodes
 		/// </summary>
 		public InteriorFileNode(ITreeNodeReader reader)
 		{
+			// Keep this code in sync with CopyToStreamAsync
 			TreeNodeRef<FileNode>[] children = new TreeNodeRef<FileNode>[reader.Length / IoHash.NumBytes];
 			for (int idx = 0; idx < children.Length; idx++)
 			{
@@ -356,5 +390,38 @@ namespace EpicGames.Horde.Storage.Nodes
 				await childNode.CopyToStreamAsync(reader, outputStream, cancellationToken);
 			}
 		}
+
+		/// <summary>
+		/// Copy the contents of the node to the output stream without creating the intermediate FileNodes
+		/// </summary>
+		/// <param name="reader">Reader for nodes in the tree</param>
+		/// <param name="outputStream">The output stream to receive the data</param>
+		/// <param name="nodeReader">Source data</param>
+		/// <param name="cancellationToken">Cancellation token for the operation</param>
+		public static async Task CopyToStreamAsync(TreeReader reader, Stream outputStream, ITreeNodeReader nodeReader, CancellationToken cancellationToken)
+		{
+			async Task CopyFunc(ReadNodeAsyncCallbackParams parms, ITreeNodeReader nodeReader)
+			{
+				if (parms.Type == typeof(LeafFileNode))
+				{
+					await LeafFileNode.CopyToStreamAsync(outputStream, nodeReader, cancellationToken);
+				}
+				else
+				{
+					throw new ArgumentException("Unexpected FileNode type found.");
+				}
+			}
+
+			// Keep this code in sync with the constructor
+			TreeNodeRef<FileNode>[] children = new TreeNodeRef<FileNode>[nodeReader.Length / IoHash.NumBytes];
+			for (int idx = 0; idx < children.Length; idx++)
+			{
+				children[idx] = nodeReader.ReadRef<FileNode>();
+			}
+			foreach (TreeNodeRef<FileNode> node in children)
+			{
+				await reader.ReadNodeAsync(node.Handle!.Locator, CopyFunc, cancellationToken);
+			}
+		}
 	}
 }
diff --git a/Engine/Source/Programs/Shared/EpicGames.Horde/Storage/Nodes/FileNodeWriter.cs b/Engine/Source/Programs/Shared/EpicGames.Horde/Storage/Nodes/FileNodeWriter.cs
--- a/Engine/Source/Programs/Shared/EpicGames.Horde/Storage/Nodes/FileNodeWriter.cs
+++ b/Engine/Source/Programs/Shared/EpicGames.Horde/Storage/Nodes/FileNodeWriter.cs
@@ -15,6 +15,12 @@ namespace EpicGames.Horde.Storage.Nodes
 	/// </summary>
 	public class FileNodeWriter
 	{
+
+		/// <summary>
+		/// Default buffer length when calling CreateAsync/AppendAsync
+		/// </summary>
+		public const int DefaultBufferLength = 32 * 1024;
+
 		class InteriorNodeState
 		{
 			public InteriorNodeState? _parent;
@@ -126,10 +132,21 @@ namespace EpicGames.Horde.Storage.Nodes
 		/// <param name="fileInfo">File to append</param>
 		/// <param name="cancellationToken">Cancellation token for the operation</param>
 		public async Task<NodeHandle> CreateAsync(FileInfo fileInfo, CancellationToken cancellationToken)
+		{
+			return await CreateAsync(fileInfo, DefaultBufferLength, cancellationToken);
+		}
+
+		/// <summary>
+		/// Creates data for the given file
+		/// </summary>
+		/// <param name="fileInfo">File to append</param>
+		/// <param name="bufferLength">Size of the read buffer</param>
+		/// <param name="cancellationToken">Cancellation token for the operation</param>
+		public async Task<NodeHandle> CreateAsync(FileInfo fileInfo, int bufferLength, CancellationToken cancellationToken)
 		{
 			using (FileStream stream = fileInfo.OpenRead())
 			{
-				return await CreateAsync(stream, cancellationToken);
+				return await CreateAsync(stream, bufferLength, cancellationToken);
 			}
 		}
 
@@ -139,9 +156,20 @@ namespace EpicGames.Horde.Storage.Nodes
 		/// <param name="stream">Stream to append</param>
 		/// <param name="cancellationToken">Cancellation token for the operation</param>
 		public async Task<NodeHandle> CreateAsync(Stream stream, CancellationToken cancellationToken)
+		{
+			return await CreateAsync(stream, DefaultBufferLength, cancellationToken);
+		}
+
+		/// <summary>
+		/// Creates data from the given stream
+		/// </summary>
+		/// <param name="stream">Stream to append</param>
+		/// <param name="bufferLength">Size of the read buffer</param>
+		/// <param name="cancellationToken">Cancellation token for the operation</param>
+		public async Task<NodeHandle> CreateAsync(Stream stream, int bufferLength, CancellationToken cancellationToken)
 		{
 			Reset();
-			await AppendAsync(stream, cancellationToken);
+			await AppendAsync(stream, bufferLength, cancellationToken);
 			return await CompleteAsync(cancellationToken);
 		}
 
@@ -164,36 +192,18 @@ namespace EpicGames.Horde.Storage.Nodes
 		/// <param name="cancellationToken">Cancellation token for the operation</param>
 		public async Task AppendAsync(Stream stream, CancellationToken cancellationToken)
 		{
-			const int BufferLength = 32 * 1024;
-
-			using IMemoryOwner<byte> owner = MemoryPool<byte>.Shared.Rent(BufferLength * 2);
-			Memory<byte> buffer = owner.Memory;
-
-			int readBufferOffset = 0;
-			Memory<byte> appendBuffer = Memory<byte>.Empty;
-			for (; ; )
-			{
-				// Start a read into memory
-				Memory<byte> readBuffer = buffer.Slice(readBufferOffset, BufferLength);
-				Task<int> readTask = Task.Run(async () => await stream.ReadAsync(readBuffer, cancellationToken), cancellationToken);
-
-				// In the meantime, append the last data that was read to the tree
-				if (appendBuffer.Length > 0)
-				{
-					await AppendAsync(appendBuffer, cancellationToken);
-				}
-
-				// Wait for the read to finish
-				int numBytes = await readTask;
-				if (numBytes == 0)
-				{
-					break;
-				}
+			await AppendAsync(stream, DefaultBufferLength, cancellationToken);
+		}
 
-				// Switch the buffers around
-				appendBuffer = readBuffer.Slice(0, numBytes);
-				readBufferOffset ^= BufferLength;
-			}
+		/// <summary>
+		/// Appends data to the current file
+		/// </summary>
+		/// <param name="stream">Stream containing data to append</param>
+		/// <param name="bufferLength">Size of the read buffer</param>
+		/// <param name="cancellationToken">Cancellation token for the operation</param>
+		public async Task AppendAsync(Stream stream, int bufferLength, CancellationToken cancellationToken)
+		{
+			await stream.ReadAllBytesAsync(bufferLength, async (x) => await AppendAsync(x, cancellationToken), cancellationToken);
 		}
 
 		/// <summary>
diff --git a/Engine/Source/Programs/Shared/EpicGames.Horde/Storage/TreeReader.cs b/Engine/Source/Programs/Shared/EpicGames.Horde/Storage/TreeReader.cs
--- a/Engine/Source/Programs/Shared/EpicGames.Horde/Storage/TreeReader.cs
+++ b/Engine/Source/Programs/Shared/EpicGames.Horde/Storage/TreeReader.cs
@@ -13,6 +13,7 @@ using Microsoft.Extensions.Caching.Memory;
 using Microsoft.Extensions.Logging;
 using System.Collections.Concurrent;
 using Microsoft.CodeAnalysis;
+using Blake3;
 
 namespace EpicGames.Horde.Storage
 {
@@ -76,6 +77,13 @@ namespace EpicGames.Horde.Storage
 		}
 	}
 
+	/// <summary>
+	/// Callback parameters for ReadNodeAsync calls. 
+	/// </summary>
+	/// <param name="Type">TreeNode type being read</param>
+	/// <param name="Deserialize">Deserialization function commonly </param>
+	public record struct ReadNodeAsyncCallbackParams(Type Type, Func<ITreeNodeReader, TreeNode> Deserialize);
+
 	/// <summary>
 	/// Writes nodes from bundles in an <see cref="IStorageClient"/> instance.
 	/// </summary>
@@ -266,6 +274,9 @@ namespace EpicGames.Horde.Storage
 		// Size of data to fetch by default. This is larger than the minimum request size to reduce number of reads.
 		const int DefaultFetchSize = 15 * 1024 * 1024;
 
+		// When reader is uncached, use a smaller default fetch size
+		const int DefaultUncachedFetchSize = 1 * 1024 * 1024;
+
 		readonly IStorageClient _store;
 		readonly IMemoryCache? _cache;
 		readonly Dictionary<Guid, TypeInfo> _types;
@@ -432,7 +443,7 @@ namespace EpicGames.Horde.Storage
 		/// <param name="cancellationToken">Cancellation token for the operation</param>
 		async Task PerformHeaderReadAsync(QueuedHeader queuedHeader, CancellationToken cancellationToken)
 		{
-			int prefetchSize = DefaultFetchSize;
+			int prefetchSize = _cache != null ? DefaultFetchSize : DefaultUncachedFetchSize;
 			for (; ; )
 			{
 				using (IMemoryOwner<byte> owner = MemoryPool<byte>.Shared.Rent(prefetchSize))
@@ -462,28 +473,31 @@ namespace EpicGames.Horde.Storage
 					// Construct the bundle info
 					BundleInfo bundleInfo = new BundleInfo(queuedHeader.Blob, header, headerSize, types);
 
-					// Also add any encoded packets we prefetched
-					List<ReadOnlyMemory<byte>> packets = new List<ReadOnlyMemory<byte>>();
-					for (int packetOffset = headerSize; packets.Count < header.Packets.Count;)
+					if (_cache != null)
 					{
-						int packetIdx = packets.Count;
-
-						int packetLength = header.Packets[packetIdx].EncodedLength;
-						if (packetOffset + packetLength > memory.Length)
+						// Also add any encoded packets we prefetched
+						List<ReadOnlyMemory<byte>> packets = new List<ReadOnlyMemory<byte>>();
+						for (int packetOffset = headerSize; packets.Count < header.Packets.Count;)
 						{
-							break;
+							int packetIdx = packets.Count;
+
+							int packetLength = header.Packets[packetIdx].EncodedLength;
+							if (packetOffset + packetLength > memory.Length)
+							{
+								break;
+							}
+
+							ReadOnlyMemory<byte> packetData = memory.Slice(packetOffset, packetLength).ToArray();
+							AddToCache(GetEncodedPacketCacheKey(queuedHeader.Blob.BlobId, packetIdx), packetData, packetData.Length);
+							packets.Add(packetData);
+							packetOffset += packetLength;
 						}
 
-						ReadOnlyMemory<byte> packetData = memory.Slice(packetOffset, packetLength).ToArray();
-						AddToCache(GetEncodedPacketCacheKey(queuedHeader.Blob.BlobId, packetIdx), packetData, packetData.Length);
-						packets.Add(packetData);
-						packetOffset += packetLength;
+						// Add the info to the cache
+						string cacheKey = GetBundleInfoCacheKey(queuedHeader.Blob.BlobId);
+						AddToCache(cacheKey, bundleInfo, headerSize);
 					}
 
-					// Add the info to the cache
-					string cacheKey = GetBundleInfoCacheKey(queuedHeader.Blob.BlobId);
-					AddToCache(cacheKey, bundleInfo, headerSize);
-
 					// Update the task
 					queuedHeader.CompletionSource.SetResult(bundleInfo);
 
@@ -544,13 +558,16 @@ namespace EpicGames.Horde.Storage
 				buffer = await _store.ReadBlobRangeAsync(bundleInfo.Locator, bundleInfo.PacketOffsets[minPacketIdx], buffer, cancellationToken);
 
 				// Copy all the packets that have been read into separate buffers, so we can cache them individually.
-				List<ReadOnlyMemory<byte>> packets = new List<ReadOnlyMemory<byte>>();
-				for (int idx = minPacketIdx; idx < maxPacketIdx; idx++)
+				ReadOnlyMemory<byte>?[] packets = new ReadOnlyMemory<byte>?[maxPacketIdx - minPacketIdx];
+				if (_cache != null)
 				{
-					string cacheKey = GetEncodedPacketCacheKey(bundleInfo.Locator.BlobId, idx);
-					ReadOnlyMemory<byte> data = buffer.Slice(bundleInfo.PacketOffsets[idx] - bundleInfo.PacketOffsets[minPacketIdx], bundleInfo.Header.Packets[idx].EncodedLength).ToArray();
-					packets.Add(data);
-					AddToCache(cacheKey, data, data.Length);
+					for (int idx = minPacketIdx; idx < maxPacketIdx; idx++)
+					{
+						string cacheKey = GetEncodedPacketCacheKey(bundleInfo.Locator.BlobId, idx);
+						ReadOnlyMemory<byte> data = buffer.Slice(bundleInfo.PacketOffsets[idx] - bundleInfo.PacketOffsets[minPacketIdx], bundleInfo.Header.Packets[idx].EncodedLength).ToArray();
+						packets[idx - minPacketIdx] = data;
+						AddToCache(cacheKey, data, data.Length);
+					}
 				}
 
 				// Find all the packets we can mark as complete
@@ -563,8 +580,12 @@ namespace EpicGames.Horde.Storage
 				// Mark them all as complete
 				foreach (QueuedPacket updatePacket in updatePackets)
 				{
-					ReadOnlyMemory<byte> data = packets[updatePacket.PacketIdx - minPacketIdx];
-					updatePacket.CompletionSource.SetResult(data);
+					ReadOnlyMemory<byte>? data = packets[updatePacket.PacketIdx - minPacketIdx];
+					if (data == null)
+					{
+						data = buffer.Slice(bundleInfo.PacketOffsets[updatePacket.PacketIdx] - bundleInfo.PacketOffsets[minPacketIdx], bundleInfo.Header.Packets[updatePacket.PacketIdx].EncodedLength).ToArray();
+					}
+					updatePacket.CompletionSource.SetResult((ReadOnlyMemory<byte>)data);
 				}
 
 				// Remove all the completed packets from the queue
@@ -736,6 +757,24 @@ namespace EpicGames.Horde.Storage
 		/// <param name="cancellationToken">Cancellation token for the operation</param>
 		/// <returns>Node data read from the given bundle</returns>
 		public async ValueTask<TreeNode> ReadNodeAsync(NodeLocator locator, CancellationToken cancellationToken = default)
+		{
+			TreeNode? treeNode = null;
+			await ReadNodeAsync(locator, (ReadNodeAsyncCallbackParams parms, ITreeNodeReader reader) =>
+			{
+				treeNode = parms.Deserialize(reader);
+				return Task.CompletedTask;
+			}, cancellationToken);
+			return treeNode!;
+		}
+
+		/// <summary>
+		/// Reads a node from a bundle
+		/// </summary>
+		/// <param name="locator">Locator for the node</param>
+		/// <param name="readFunc">Action to be invoked</param>
+		/// <param name="cancellationToken">Cancellation token for the operation</param>
+		/// <returns>Node data read from the given bundle</returns>
+		public async ValueTask ReadNodeAsync(NodeLocator locator, Func<ReadNodeAsyncCallbackParams /*parms*/, ITreeNodeReader /*reader*/, Task> readFunc, CancellationToken cancellationToken = default)
 		{
 			BundleInfo bundleInfo = await GetBundleInfoAsync(locator.Blob, cancellationToken);
 			BundleExport export = bundleInfo.Header.Exports[locator.ExportIdx];
@@ -762,7 +801,7 @@ namespace EpicGames.Horde.Storage
 				throw new InvalidOperationException($"No registered serializer for type {bundleInfo.Header.Types[export.TypeIdx].Guid}");
 			}
 
-			return typeInfo.Deserialize(new NodeReader(nodeData, export.Hash, refs, typeInfo.BundleType));
+			await readFunc(new ReadNodeAsyncCallbackParams(typeInfo.Type, typeInfo.Deserialize), new NodeReader(nodeData, export.Hash, refs, typeInfo.BundleType));
 		}
 
 		/// <summary>
diff --git a/Engine/Source/Programs/Shared/EpicGames.Horde/Storage/TreeWriter.cs b/Engine/Source/Programs/Shared/EpicGames.Horde/Storage/TreeWriter.cs
--- a/Engine/Source/Programs/Shared/EpicGames.Horde/Storage/TreeWriter.cs
+++ b/Engine/Source/Programs/Shared/EpicGames.Horde/Storage/TreeWriter.cs
@@ -398,7 +398,7 @@ namespace EpicGames.Horde.Storage
 				{
 					_currentPacket = MemoryPool<byte>.Shared.Rent(Math.Max(_maxPacketSize, desiredSize));
 				}
-				else if (_currentPacketLength + desiredSize > _maxPacketSize)
+				else if (_currentPacketLength + desiredSize > _currentPacket.Memory.Length)
 				{
 					IMemoryOwner<byte> nextPacket = MemoryPool<byte>.Shared.Rent(Math.Max(_maxPacketSize, desiredSize));
 					if (usedSize > 0)
@@ -407,7 +407,6 @@ namespace EpicGames.Horde.Storage
 					}
 					FlushPacket();
 					_currentPacket = nextPacket;
-					return _currentPacket.Memory.Slice(_currentPacketLength, desiredSize);
 				}
 				return _currentPacket.Memory.Slice(_currentPacketLength, desiredSize);
 			}
