Horde: Fix warning in BundleTests.

diff --git a/Engine/Source/Programs/Shared/EpicGames.Horde.Tests/BundleTests.cs b/Engine/Source/Programs/Shared/EpicGames.Horde.Tests/BundleTests.cs
--- a/Engine/Source/Programs/Shared/EpicGames.Horde.Tests/BundleTests.cs
+++ b/Engine/Source/Programs/Shared/EpicGames.Horde.Tests/BundleTests.cs
@@ -126,7 +126,7 @@ namespace EpicGames.Horde.Tests
 
 		private static async Task TestBufferlessReadsAsync(TreeReader reader, NodeLocator locator, ReadOnlyMemory<byte> expected)
 		{
-			MemoryStream memoryStream = new MemoryStream();
+			using MemoryStream memoryStream = new MemoryStream();
 			await FileNode.CopyToStreamAsync(reader, locator, memoryStream, default);
 			ReadOnlyMemory<byte> read = memoryStream.ToArray().AsMemory();
 			Assert.IsTrue(read.Span.SequenceEqual(expected.Span));			
