Horde: Add support for forking the remote message loop during remote execution.

diff --git a/Engine/Source/Programs/Horde/Samples/RemoteClient/RemoteClientApp.cs b/Engine/Source/Programs/Horde/Samples/RemoteClient/RemoteClientApp.cs
--- a/Engine/Source/Programs/Horde/Samples/RemoteClient/RemoteClientApp.cs
+++ b/Engine/Source/Programs/Horde/Samples/RemoteClient/RemoteClientApp.cs
@@ -79,12 +79,18 @@ namespace RemoteClient
 		static async Task RunRemoteAsync(IComputeLease lease, DirectoryReference uploadDir, string executable, List<string> arguments, ILogger logger)
 		{
 			// Create a message channel on channel id 0. The Horde Agent always listens on this channel for requests.
-			const int ControlChannelId = 0;
-			using (IComputeMessageChannel channel = lease.Socket.CreateMessageChannel(ControlChannelId, 4 * 1024 * 1024, logger))
+			const int PrimaryChannelId = 0;
+			using (IComputeMessageChannel channel = lease.Socket.CreateMessageChannel(PrimaryChannelId, 4 * 1024 * 1024, logger))
 			{
 				await channel.WaitForAttachAsync();
 
-				// Upload the sandbox
+				// Fork another message loop on channel id 2. We'll use this to run an XOR task in the background.
+				const int BackgroundChannelId = 2;
+				using IComputeMessageChannel backgroundChannel = lease.Socket.CreateMessageChannel(BackgroundChannelId, 4 * 1024 * 1024, logger);
+				await using BackgroundTask otherChannelTask = BackgroundTask.StartNew(ctx => RunBackgroundXorAsync(backgroundChannel));
+				await channel.ForkAsync(BackgroundChannelId, 4 * 1024 * 1024, default);
+
+				// Upload the sandbox to the primary channel.
 				MemoryStorageClient storage = new MemoryStorageClient();
 				using (TreeWriter treeWriter = new TreeWriter(storage))
 				{
@@ -110,7 +116,7 @@ namespace RemoteClient
 			}
 			await lease.CloseAsync();
 		}
-
+		
 		static async Task WriteNumbersAsync(IComputeSocket socket, ILogger logger, CancellationToken cancellationToken)
 		{
 			// Generate data into a buffer attached to channel 1. The remote server will echo them back to us as it receives them, then exit when the channel is complete/closed.
@@ -129,6 +135,25 @@ namespace RemoteClient
 			await socket.MarkCompleteAsync(DataChannelId, cancellationToken);
 		}
 
+		static async Task RunBackgroundXorAsync(IComputeMessageChannel channel)
+		{
+			await channel.WaitForAttachAsync();
+
+			byte[] dataToXor = new byte[] { 1, 2, 3, 4, 5 };
+			await channel.SendXorRequestAsync(dataToXor, 123);
+
+			using IComputeMessage response = await channel.ReceiveAsync(ComputeMessageType.XorResponse);
+			for (int idx = 0; idx < dataToXor.Length; idx++)
+			{
+				if (response.Data.Span[idx] != (byte)(dataToXor[idx] ^ 123))
+				{
+					throw new InvalidOperationException();
+				}
+			}
+
+			await channel.CloseAsync();
+		}
+
 		static async Task<IComputeClient> CreateClientAsync(ClientAppOptions options, ILogger logger)
 		{
 			if (options.Server == null)
diff --git a/Engine/Source/Programs/Shared/EpicGames.Horde/Compute/ComputeWorker.cs b/Engine/Source/Programs/Shared/EpicGames.Horde/Compute/ComputeWorker.cs
--- a/Engine/Source/Programs/Shared/EpicGames.Horde/Compute/ComputeWorker.cs
+++ b/Engine/Source/Programs/Shared/EpicGames.Horde/Compute/ComputeWorker.cs
@@ -44,12 +44,12 @@ namespace EpicGames.Horde.Compute
 		/// <param name="cancellationToken">Cancellation token for the operation</param>
 		public async Task RunAsync(IComputeSocket socket, CancellationToken cancellationToken)
 		{
-			await RunAsync(socket, 0, cancellationToken);
+			await RunAsync(socket, 0, 4 * 1024 * 1024, cancellationToken);
 		}
 
-		async Task RunAsync(IComputeSocket socket, int channelId, CancellationToken cancellationToken)
+		async Task RunAsync(IComputeSocket socket, int channelId, int bufferSize, CancellationToken cancellationToken)
 		{
-			using (IComputeMessageChannel channel = socket.CreateMessageChannel(channelId, 4 * 1024 * 1024, _logger))
+			using (IComputeMessageChannel channel = socket.CreateMessageChannel(channelId, bufferSize, _logger))
 			{
 				await channel.SendReadyAsync(cancellationToken);
 
@@ -64,6 +64,12 @@ namespace EpicGames.Horde.Compute
 						case ComputeMessageType.None:
 							await Task.WhenAll(childTasks);
 							return;
+						case ComputeMessageType.Fork:
+							{
+								ForkMessage fork = message.ParseForkMessage();
+								childTasks.Add(Task.Run(() => RunAsync(socket, fork.channelId, fork.bufferSize, cancellationToken), cancellationToken));
+							}
+							break;
 						case ComputeMessageType.WriteFiles:
 							{
 								UploadFilesMessage writeFiles = message.ParseUploadFilesMessage();
diff --git a/Engine/Source/Programs/Shared/EpicGames.Horde/Compute/IComputeMessage.cs b/Engine/Source/Programs/Shared/EpicGames.Horde/Compute/IComputeMessage.cs
--- a/Engine/Source/Programs/Shared/EpicGames.Horde/Compute/IComputeMessage.cs
+++ b/Engine/Source/Programs/Shared/EpicGames.Horde/Compute/IComputeMessage.cs
@@ -32,6 +32,11 @@ namespace EpicGames.Horde.Compute
 		/// </summary>
 		Exception = 0x02,
 
+		/// <summary>
+		/// Fork the message loop into a new channel
+		/// </summary>
+		Fork = 0x03,
+
 		#region Process Management
 
 		/// <summary>
@@ -116,6 +121,13 @@ namespace EpicGames.Horde.Compute
 	/// </summary>
 	public record struct ExceptionMessage(string Message, string Description);
 
+	/// <summary>
+	/// Message requesting that the message loop be forked
+	/// </summary>
+	/// <param name="channelId">New channel to communicate on</param>
+	/// <param name="bufferSize">Size of the buffer</param>
+	public record struct ForkMessage(int channelId, int bufferSize);
+
 	/// <summary>
 	/// Extract files from a bundle to a path in the remote sandbox
 	/// </summary>
@@ -161,10 +173,19 @@ namespace EpicGames.Horde.Compute
 	/// </summary>
 	public static class ComputeMessageExtensions
 	{
+		/// <summary>
+		/// Closes the remote message loop
+		/// </summary>
+		public static async ValueTask CloseAsync(this IComputeMessageChannel channel, CancellationToken cancellationToken = default)
+		{
+			using IComputeMessageBuilder message = await channel.CreateMessageAsync(ComputeMessageType.None, cancellationToken);
+			message.Send();
+		}
+
 		/// <summary>
 		/// Sends a message to the remote indicating that we're ready to receive messages on the given channel
 		/// </summary>
-		public static async ValueTask SendReadyAsync(this IComputeMessageChannel channel, CancellationToken cancellationToken)
+		public static async ValueTask SendReadyAsync(this IComputeMessageChannel channel, CancellationToken cancellationToken = default)
 		{
 			using IComputeMessageBuilder message = await channel.CreateMessageAsync(ComputeMessageType.Ready, cancellationToken);
 			message.Send();
@@ -173,12 +194,12 @@ namespace EpicGames.Horde.Compute
 		/// <summary>
 		/// Sends an exception response to the remote
 		/// </summary>
-		public static ValueTask SendExceptionAsync(this IComputeMessageChannel channel, Exception ex, CancellationToken cancellationToken) => SendExceptionAsync(channel, ex.Message, ex.ToString(), cancellationToken);
+		public static ValueTask SendExceptionAsync(this IComputeMessageChannel channel, Exception ex, CancellationToken cancellationToken = default) => SendExceptionAsync(channel, ex.Message, ex.ToString(), cancellationToken);
 
 		/// <summary>
 		/// Sends an exception response to the remote
 		/// </summary>
-		public static async ValueTask SendExceptionAsync(this IComputeMessageChannel channel, string description, string trace, CancellationToken cancellationToken)
+		public static async ValueTask SendExceptionAsync(this IComputeMessageChannel channel, string description, string trace, CancellationToken cancellationToken = default)
 		{
 			using IComputeMessageBuilder message = await channel.CreateMessageAsync(ComputeMessageType.Exception, cancellationToken);
 			message.WriteString(description);
@@ -196,9 +217,30 @@ namespace EpicGames.Horde.Compute
 			return new ExceptionMessage(msg, description);
 		}
 
+		/// <summary>
+		/// Sends an exception response to the remote
+		/// </summary>
+		public static async ValueTask ForkAsync(this IComputeMessageChannel channel, int channelId, int bufferSize, CancellationToken cancellationToken = default)
+		{
+			using IComputeMessageBuilder message = await channel.CreateMessageAsync(ComputeMessageType.Fork, cancellationToken);
+			message.WriteInt32(channelId);
+			message.WriteInt32(bufferSize);
+			message.Send();
+		}
+
+		/// <summary>
+		/// Parses a message as an <see cref="ExceptionMessage"/>
+		/// </summary>
+		public static ForkMessage ParseForkMessage(this IComputeMessage message)
+		{
+			int channelId = message.ReadInt32();
+			int bufferSize = message.ReadInt32();
+			return new ForkMessage(channelId, bufferSize);
+		}
+
 		#region Process
 
-		static async Task<IComputeMessage> RunStorageServer(this IComputeMessageChannel channel, IStorageClient storage, CancellationToken cancellationToken)
+		static async Task<IComputeMessage> RunStorageServer(this IComputeMessageChannel channel, IStorageClient storage, CancellationToken cancellationToken = default)
 		{
 			for (; ; )
 			{
diff --git a/Engine/Source/Programs/Shared/EpicGames.Horde/Compute/IComputeMessageChannel.cs b/Engine/Source/Programs/Shared/EpicGames.Horde/Compute/IComputeMessageChannel.cs
--- a/Engine/Source/Programs/Shared/EpicGames.Horde/Compute/IComputeMessageChannel.cs
+++ b/Engine/Source/Programs/Shared/EpicGames.Horde/Compute/IComputeMessageChannel.cs
@@ -45,6 +45,23 @@ namespace EpicGames.Horde.Compute
 	/// </summary>
 	public static class ComputeMessageChannelExtensions
 	{
+		/// <summary>
+		/// Reads a message from the channel
+		/// </summary>
+		/// <param name="channel">Channel to receive on</param>
+		/// <param name="type">Expected type of the message</param>
+		/// <param name="cancellationToken">Cancellation token for the operation</param>
+		/// <returns>Data for a message that was read. Must be disposed.</returns>
+		public static async ValueTask<IComputeMessage> ReceiveAsync(this IComputeMessageChannel channel, ComputeMessageType type, CancellationToken cancellationToken = default)
+		{
+			IComputeMessage message = await channel.ReceiveAsync(cancellationToken);
+			if (message.Type != type)
+			{
+				throw new ComputeInvalidMessageException(message);
+			}
+			return message;
+		}
+
 		/// <summary>
 		/// Waits for the remote machine to send a 'ready' response
 		/// </summary>
