Horde: Support communicating with child processes via custom channel ids, making it possible to run multiple processes on forked message loops.

diff --git a/Engine/Source/Programs/Horde/Horde.Agent/Commands/Compute/CppComputeCommand.cs b/Engine/Source/Programs/Horde/Horde.Agent/Commands/Compute/CppComputeCommand.cs
--- a/Engine/Source/Programs/Horde/Horde.Agent/Commands/Compute/CppComputeCommand.cs
+++ b/Engine/Source/Programs/Horde/Horde.Agent/Commands/Compute/CppComputeCommand.cs
@@ -47,6 +47,9 @@ namespace Horde.Agent.Commands.Compute
 		/// <inheritdoc/>
 		protected override async Task<bool> HandleRequestAsync(IComputeLease lease, CancellationToken cancellationToken)
 		{
+			const int ControlChannelId = 0;
+			const int ProcessChannelId = 1;
+
 			// Read the task definition
 			byte[] data = await FileReference.ReadAllBytesAsync(TaskFile, cancellationToken);
 			JsonComputeTask jsonComputeTask = JsonSerializer.Deserialize<JsonComputeTask>(data, new JsonSerializerOptions { AllowTrailingCommas = true, PropertyNameCaseInsensitive = true, PropertyNamingPolicy = JsonNamingPolicy.CamelCase })!;
@@ -56,11 +59,11 @@ namespace Horde.Agent.Commands.Compute
 			NodeLocator sandbox = await CreateSandboxAsync(TaskFile, storage, cancellationToken);
 
 			// Open a socket and upload the sandbox
-			using (IComputeMessageChannel channel = lease.Socket.CreateMessageChannel(0, 4 * 1024 * 1024, _logger))
+			using (IComputeMessageChannel channel = lease.Socket.CreateMessageChannel(ControlChannelId, 4 * 1024 * 1024, _logger))
 			{
 				await channel.UploadFilesAsync("", sandbox, storage, cancellationToken);
 
-				await using (IComputeProcess process = await channel.ExecuteAsync(jsonComputeTask.Executable, jsonComputeTask.Arguments, jsonComputeTask.WorkingDir, jsonComputeTask.EnvVars, cancellationToken))
+				await using (IComputeProcess process = await channel.ExecuteAsync(ProcessChannelId, jsonComputeTask.Executable, jsonComputeTask.Arguments, jsonComputeTask.WorkingDir, jsonComputeTask.EnvVars, cancellationToken))
 				{
 					string? line;
 					while ((line = await process.ReadLineAsync(cancellationToken)) != null)
diff --git a/Engine/Source/Programs/Horde/Samples/RemoteClient/RemoteClientApp.cs b/Engine/Source/Programs/Horde/Samples/RemoteClient/RemoteClientApp.cs
--- a/Engine/Source/Programs/Horde/Samples/RemoteClient/RemoteClientApp.cs
+++ b/Engine/Source/Programs/Horde/Samples/RemoteClient/RemoteClientApp.cs
@@ -76,16 +76,18 @@ namespace RemoteClient
 			}
 		}
 
+		const int PrimaryChannelId = 0;
+		const int BackgroundChannelId = 1;
+		const int ChildProcessChannelId = 2;
+
 		static async Task RunRemoteAsync(IComputeLease lease, DirectoryReference uploadDir, string executable, List<string> arguments, ILogger logger)
 		{
 			// Create a message channel on channel id 0. The Horde Agent always listens on this channel for requests.
-			const int PrimaryChannelId = 0;
 			using (IComputeMessageChannel channel = lease.Socket.CreateMessageChannel(PrimaryChannelId, 4 * 1024 * 1024, logger))
 			{
 				await channel.WaitForAttachAsync();
 
 				// Fork another message loop on channel id 2. We'll use this to run an XOR task in the background.
-				const int BackgroundChannelId = 2;
 				using IComputeMessageChannel backgroundChannel = lease.Socket.CreateMessageChannel(BackgroundChannelId, 4 * 1024 * 1024, logger);
 				await using BackgroundTask otherChannelTask = BackgroundTask.StartNew(ctx => RunBackgroundXorAsync(backgroundChannel));
 				await channel.ForkAsync(BackgroundChannelId, 4 * 1024 * 1024, default);
@@ -101,7 +103,7 @@ namespace RemoteClient
 				}
 
 				// Run the task remotely in the background and echo the output to the console
-				await using (IComputeProcess process = await channel.ExecuteAsync(executable, arguments, null, null))
+				await using (IComputeProcess process = await channel.ExecuteAsync(ChildProcessChannelId, executable, arguments, null, null))
 				{
 					string? line = await process.ReadLineAsync();
 					logger.LogInformation("[REMOTE] {Line}", line);
@@ -119,8 +121,7 @@ namespace RemoteClient
 		
 		static async Task WriteNumbersAsync(IComputeSocket socket, ILogger logger, CancellationToken cancellationToken)
 		{
-			// Generate data into a buffer attached to channel 1. The remote server will echo them back to us as it receives them, then exit when the channel is complete/closed.
-			const int DataChannelId = 1;
+			// Write data to the child process channel. The remote server will echo them back to us as it receives them, then exit when the channel is complete/closed.
 
 			byte[] buffer = new byte[4];
 			for (int idx = 0; idx < 3; idx++)
@@ -128,11 +129,11 @@ namespace RemoteClient
 				cancellationToken.ThrowIfCancellationRequested();
 				logger.LogInformation("Writing value: {Value}", idx);
 				BinaryPrimitives.WriteInt32LittleEndian(buffer, idx);
-				await socket.SendAsync(DataChannelId, buffer, cancellationToken);
+				await socket.SendAsync(ChildProcessChannelId, buffer, cancellationToken);
 				await Task.Delay(1000, cancellationToken);
 			}
 
-			await socket.MarkCompleteAsync(DataChannelId, cancellationToken);
+			await socket.MarkCompleteAsync(ChildProcessChannelId, cancellationToken);
 		}
 
 		static async Task RunBackgroundXorAsync(IComputeMessageChannel channel)
diff --git a/Engine/Source/Programs/Shared/EpicGames.Horde/Compute/ComputeWorker.cs b/Engine/Source/Programs/Shared/EpicGames.Horde/Compute/ComputeWorker.cs
--- a/Engine/Source/Programs/Shared/EpicGames.Horde/Compute/ComputeWorker.cs
+++ b/Engine/Source/Programs/Shared/EpicGames.Horde/Compute/ComputeWorker.cs
@@ -85,7 +85,7 @@ namespace EpicGames.Horde.Compute
 						case ComputeMessageType.Execute:
 							{
 								ExecuteProcessMessage executeProcess = message.ParseExecuteProcessMessage();
-								await ExecuteProcessAsync(socket, channel, executeProcess.Executable, executeProcess.Arguments, executeProcess.WorkingDir, executeProcess.EnvVars, cancellationToken);
+								await ExecuteProcessAsync(socket, channel, executeProcess.ChannelId, executeProcess.Executable, executeProcess.Arguments, executeProcess.WorkingDir, executeProcess.EnvVars, cancellationToken);
 							}
 							break;
 						case ComputeMessageType.XorRequest:
@@ -148,13 +148,13 @@ namespace EpicGames.Horde.Compute
 			}
 		}
 
-		async Task ExecuteProcessAsync(IComputeSocket socket, IComputeMessageChannel channel, string executable, IReadOnlyList<string> arguments, string? workingDir, IReadOnlyDictionary<string, string?>? envVars, CancellationToken cancellationToken)
+		async Task ExecuteProcessAsync(IComputeSocket socket, IComputeMessageChannel channel, int newChannelId, string executable, IReadOnlyList<string> arguments, string? workingDir, IReadOnlyDictionary<string, string?>? envVars, CancellationToken cancellationToken)
 		{
 			try
 			{
 				if (RuntimeInformation.IsOSPlatform(OSPlatform.Windows))
 				{
-					await ExecuteProcessWindowsAsync(socket, channel, executable, arguments, workingDir, envVars, cancellationToken);
+					await ExecuteProcessWindowsAsync(socket, channel, newChannelId, executable, arguments, workingDir, envVars, cancellationToken);
 				}
 				else
 				{
@@ -167,16 +167,15 @@ namespace EpicGames.Horde.Compute
 			}
 		}
 
-		async Task ExecuteProcessWindowsAsync(IComputeSocket socket, IComputeMessageChannel channel, string executable, IReadOnlyList<string> arguments, string? workingDir, IReadOnlyDictionary<string, string?>? envVars, CancellationToken cancellationToken)
+		async Task ExecuteProcessWindowsAsync(IComputeSocket socket, IComputeMessageChannel channel, int newChannelId, string executable, IReadOnlyList<string> arguments, string? workingDir, IReadOnlyDictionary<string, string?>? envVars, CancellationToken cancellationToken)
 		{
-			const int ChannelId = 1;
 			string baseBufferName = $"Local\\Compute-{Guid.NewGuid()}";
 
 			using IComputeBuffer sendBuffer = SharedMemoryBuffer.CreateNew(ComputeChannel.GetSendBufferName(baseBufferName), 4096).ToSharedInstance();
-			socket.AttachSendBuffer(ChannelId, sendBuffer);
+			socket.AttachSendBuffer(newChannelId, sendBuffer);
 
 			using IComputeBuffer recvBuffer = SharedMemoryBuffer.CreateNew(ComputeChannel.GetRecvBufferName(baseBufferName), 4096).ToSharedInstance();
-			socket.AttachRecvBuffer(ChannelId, recvBuffer);
+			socket.AttachRecvBuffer(newChannelId, recvBuffer);
 
 			Dictionary<string, string?> newEnvVars = new Dictionary<string, string?>();
 			if (envVars != null)
diff --git a/Engine/Source/Programs/Shared/EpicGames.Horde/Compute/IComputeMessage.cs b/Engine/Source/Programs/Shared/EpicGames.Horde/Compute/IComputeMessage.cs
--- a/Engine/Source/Programs/Shared/EpicGames.Horde/Compute/IComputeMessage.cs
+++ b/Engine/Source/Programs/Shared/EpicGames.Horde/Compute/IComputeMessage.cs
@@ -144,11 +144,12 @@ namespace EpicGames.Horde.Compute
 	/// <summary>
 	/// Message to execute a new child process
 	/// </summary>
+	/// <param name="ChannelId">Channel id</param>
 	/// <param name="Executable">Executable path</param>
 	/// <param name="Arguments">Arguments for the executable</param>
 	/// <param name="WorkingDir">Working directory to execute in</param>
 	/// <param name="EnvVars">Environment variables for the child process. Null values unset variables.</param>
-	public record struct ExecuteProcessMessage(string Executable, IReadOnlyList<string> Arguments, string? WorkingDir, IReadOnlyDictionary<string, string?> EnvVars);
+	public record struct ExecuteProcessMessage(int ChannelId, string Executable, IReadOnlyList<string> Arguments, string? WorkingDir, IReadOnlyDictionary<string, string?> EnvVars);
 
 	/// <summary>
 	/// Response from executing a child process
@@ -322,15 +323,17 @@ namespace EpicGames.Horde.Compute
 		/// Executes a remote process
 		/// </summary>
 		/// <param name="channel">Current channel</param>
+		/// <param name="newChannelId">Channel to use for communicating with the child process</param>
 		/// <param name="executable">Executable to run, relative to the sandbox root</param>
 		/// <param name="arguments">Arguments for the child process</param>
 		/// <param name="workingDir">Working directory for the process</param>
 		/// <param name="envVars">Environment variables for the child process</param>
 		/// <param name="cancellationToken">Cancellation token for the operation</param>
-		public static async Task<IComputeProcess> ExecuteAsync(this IComputeMessageChannel channel, string executable, IReadOnlyList<string> arguments, string? workingDir, IReadOnlyDictionary<string, string?>? envVars, CancellationToken cancellationToken = default)
+		public static async Task<IComputeProcess> ExecuteAsync(this IComputeMessageChannel channel, int newChannelId, string executable, IReadOnlyList<string> arguments, string? workingDir, IReadOnlyDictionary<string, string?>? envVars, CancellationToken cancellationToken = default)
 		{
 			using (IComputeMessageBuilder request = await channel.CreateMessageAsync(ComputeMessageType.Execute, cancellationToken))
 			{
+				request.WriteInt32(newChannelId);
 				request.WriteString(executable);
 				request.WriteList(arguments, MemoryWriterExtensions.WriteString);
 				request.WriteOptionalString(workingDir);
@@ -345,11 +348,12 @@ namespace EpicGames.Horde.Compute
 		/// </summary>
 		public static ExecuteProcessMessage ParseExecuteProcessMessage(this IComputeMessage message)
 		{
+			int channelId = message.ReadInt32();
 			string executable = message.ReadString();
 			List<string> arguments = message.ReadList(MemoryReaderExtensions.ReadString);
 			string? workingDir = message.ReadOptionalString();
 			Dictionary<string, string?> envVars = message.ReadDictionary(MemoryReaderExtensions.ReadString, MemoryReaderExtensions.ReadOptionalString);
-			return new ExecuteProcessMessage(executable, arguments, workingDir, envVars);
+			return new ExecuteProcessMessage(channelId, executable, arguments, workingDir, envVars);
 		}
 
 		/// <summary>
