Horde: Various compute API improvements.

* All standard buffers now implement IComputeBuffer directly, and do not need to have ToSharedInstance() called to create a ref-counted version.
* Workers can now construct a socket directly, allowing multiple buffers to be attached.
* Control messages are sent whenever a receive buffer is attached, allowing the remote to wait for it to be available.

diff --git a/Engine/Source/Programs/Horde/Samples/RemoteClient/RemoteClientApp.cs b/Engine/Source/Programs/Horde/Samples/RemoteClient/RemoteClientApp.cs
--- a/Engine/Source/Programs/Horde/Samples/RemoteClient/RemoteClientApp.cs
+++ b/Engine/Source/Programs/Horde/Samples/RemoteClient/RemoteClientApp.cs
@@ -78,7 +78,7 @@ namespace RemoteClient
 
 		const int PrimaryChannelId = 0;
 		const int BackgroundChannelId = 1;
-		const int ChildProcessChannelId = 2;
+		const int ChildProcessChannelId = ComputeChannel.DefaultWorkerChannelId;
 
 		static async Task RunRemoteAsync(IComputeLease lease, DirectoryReference uploadDir, string executable, List<string> arguments, ILogger logger)
 		{
@@ -105,11 +105,9 @@ namespace RemoteClient
 				// Run the task remotely in the background and echo the output to the console
 				await using (IComputeProcess process = await channel.ExecuteAsync(ChildProcessChannelId, executable, arguments, null, null))
 				{
-					string? line = await process.ReadLineAsync();
-					logger.LogInformation("[REMOTE] {Line}", line);
-
 					await using BackgroundTask tickTask = BackgroundTask.StartNew(ctx => WriteNumbersAsync(lease.Socket, logger, ctx));
 
+					string? line;
 					while ((line = await process.ReadLineAsync()) != null)
 					{
 						logger.LogInformation("[REMOTE] {Line}", line);
@@ -122,6 +120,9 @@ namespace RemoteClient
 		static async Task WriteNumbersAsync(IComputeSocket socket, ILogger logger, CancellationToken cancellationToken)
 		{
 			// Write data to the child process channel. The remote server will echo them back to us as it receives them, then exit when the channel is complete/closed.
+			await socket.WaitForAttachAsync(ChildProcessChannelId, cancellationToken);
+
+			logger.LogInformation("Connected");
 
 			byte[] buffer = new byte[4];
 			for (int idx = 0; idx < 3; idx++)
diff --git a/Engine/Source/Programs/Horde/Samples/RemoteWorkerCpp/ComputeChannel.cpp b/Engine/Source/Programs/Horde/Samples/RemoteWorkerCpp/ComputeChannel.cpp
--- a/Engine/Source/Programs/Horde/Samples/RemoteWorkerCpp/ComputeChannel.cpp
+++ b/Engine/Source/Programs/Horde/Samples/RemoteWorkerCpp/ComputeChannel.cpp
@@ -5,8 +5,6 @@
 #include <assert.h>
 #include "ComputeChannel.h"
 
-const char* const FComputeChannel::EnvVarName = "UE_HORDE_COMPUTE_IPC";
-
 FComputeChannel::FComputeChannel()
 {
 }
@@ -16,40 +14,38 @@ FComputeChannel::~FComputeChannel()
 	Close();
 }
 
-bool FComputeChannel::Open()
+bool FComputeChannel::Open(FComputeSocket& Socket, int ChannelId)
 {
-	Close();
+	return Open(Socket, ChannelId, ChannelId, 1024 * 1024);
+}
 
-	char EnvVar[MAX_PATH];
-	int Length = GetEnvironmentVariableA(EnvVarName, EnvVar, sizeof(EnvVar) / sizeof(EnvVar[0]));
+bool FComputeChannel::Open(FComputeSocket& Socket, int ChannelId, int NumChunks, int ChunkSize)
+{
+	return Open(Socket, ChannelId, NumChunks, ChunkSize, NumChunks, ChunkSize);
+}
 
-	if (Length <= 0 || Length >= sizeof(EnvVar))
+bool FComputeChannel::Open(FComputeSocket& Socket, int ChannelId, int NumSendChunks, int SendChunkSize, int NumRecvChunks, int RecvChunkSize)
+{
+	if (!SendBuffer.CreateNew(nullptr, NumSendChunks, SendChunkSize))
 	{
 		Close();
 		return false;
 	}
-
-	char Buffer[MAX_PATH];
-
-	sprintf_s(Buffer, "%s_SEND", EnvVar);
-	if (!SendBuffer.OpenExisting(Buffer))
+	if (!RecvBuffer.CreateNew(nullptr, NumRecvChunks, RecvChunkSize))
 	{
 		Close();
 		return false;
 	}
 
-	sprintf_s(Buffer, "%s_RECV", EnvVar);
-	if (!RecvBuffer.OpenExisting(Buffer))
-	{
-		Close();
-		return false;
-	}
+	Socket.AttachSendBuffer(ChannelId, SendBuffer);
+	Socket.AttachRecvBuffer(ChannelId, RecvBuffer);
 
 	return true;
 }
 
 void FComputeChannel::Close()
 {
+	SendBuffer.MarkComplete();
 	SendBuffer.Close();
 	RecvBuffer.Close();
 }
diff --git a/Engine/Source/Programs/Horde/Samples/RemoteWorkerCpp/ComputeChannel.h b/Engine/Source/Programs/Horde/Samples/RemoteWorkerCpp/ComputeChannel.h
--- a/Engine/Source/Programs/Horde/Samples/RemoteWorkerCpp/ComputeChannel.h
+++ b/Engine/Source/Programs/Horde/Samples/RemoteWorkerCpp/ComputeChannel.h
@@ -3,16 +3,19 @@
 #pragma once
 
 #include "SharedMemoryBuffer.h"
+#include "ComputeSocket.h"
 
 class FComputeChannel
 {
-	static const char* const EnvVarName;
-
 public:
+	static const int DefaultWorkerChannelId = 100;
+
 	FComputeChannel();
 	~FComputeChannel();
 
-	bool Open();
+	bool Open(FComputeSocket& Socket, int ChannelId);
+	bool Open(FComputeSocket& Socket, int ChannelId, int NumChunks, int ChunkSize);
+	bool Open(FComputeSocket& Socket, int ChannelId, int NumSendChunks, int SendChunkSize, int NumRecvChunks, int RecvChunkSize);
 	void Close();
 
 	void Send(const void* Data, size_t Length);
diff --git a/Engine/Source/Programs/Horde/Samples/RemoteWorkerCpp/ComputeSocket.cpp b/Engine/Source/Programs/Horde/Samples/RemoteWorkerCpp/ComputeSocket.cpp
new file mode 100644
--- /dev/null
+++ b/Engine/Source/Programs/Horde/Samples/RemoteWorkerCpp/ComputeSocket.cpp
@@ -0,0 +1,93 @@
+// Copyright Epic Games, Inc. All Rights Reserved.
+
+#include <windows.h>
+#include <iostream>
+#include <assert.h>
+#include "ComputeSocket.h"
+
+const char* const FComputeSocket::EnvVarName = "UE_HORDE_COMPUTE_IPC";
+
+FComputeSocket::FComputeSocket()
+{
+}
+
+FComputeSocket::~FComputeSocket()
+{
+	Close();
+}
+
+bool FComputeSocket::Open()
+{
+	Close();
+
+	char EnvVar[MAX_PATH];
+	int Length = GetEnvironmentVariableA(EnvVarName, EnvVar, sizeof(EnvVar) / sizeof(EnvVar[0]));
+
+	if (Length <= 0 || Length >= sizeof(EnvVar))
+	{
+		return false;
+	}
+
+	return CommandBuffer.OpenExisting(EnvVar);
+}
+
+void FComputeSocket::Close()
+{
+	CommandBuffer.Close();
+}
+
+void FComputeSocket::AttachRecvBuffer(int ChannelId, FSharedMemoryBuffer& Buffer)
+{
+	const int AttachRecvBufferType = 0;
+	AttachBuffer(ChannelId, AttachRecvBufferType, Buffer);
+}
+
+void FComputeSocket::AttachSendBuffer(int ChannelId, FSharedMemoryBuffer& Buffer)
+{
+	const int AttachSendBufferType = 1;
+	AttachBuffer(ChannelId, AttachSendBufferType, Buffer);
+}
+
+void FComputeSocket::AttachBuffer(int ChannelId, int Type, FSharedMemoryBuffer& Buffer)
+{
+	size_t Size;
+	unsigned char* Data = CommandBuffer.GetWriteMemory(Size);
+
+	size_t Len = 0;
+	Len += WriteVarUInt(Data + Len, Type);
+	Len += WriteVarUInt(Data + Len, (unsigned int)ChannelId);
+	Len += WriteString(Data + Len, Buffer.GetName());
+
+	CommandBuffer.AdvanceWritePosition(Len);
+}
+
+size_t FComputeSocket::WriteVarUInt(unsigned char* Pos, unsigned int Value)
+{
+	// Use BSR to return the log2 of the integer
+	// return 0 if value is 0
+	unsigned long BitIndex;
+	unsigned int FloorLog2 = _BitScanReverse(&BitIndex, Value) ? BitIndex : 0;
+	unsigned int ByteCount = (unsigned int)(int(FloorLog2) / 7 + 1);
+
+	unsigned char* OutBytes = Pos + ByteCount - 1;
+	switch (ByteCount - 1)
+	{
+	case 4: *OutBytes-- = (unsigned char)(Value); Value >>= 8;
+	case 3: *OutBytes-- = (unsigned char)(Value); Value >>= 8;
+	case 2: *OutBytes-- = (unsigned char)(Value); Value >>= 8;
+	case 1: *OutBytes-- = (unsigned char)(Value); Value >>= 8;
+	default: break;
+	}
+	*OutBytes = (unsigned char)(0xff << (9 - ByteCount)) | (unsigned char)(Value);
+
+	return ByteCount;
+}
+
+size_t FComputeSocket::WriteString(unsigned char* Pos, const char* Text)
+{
+	unsigned int NameLen = (unsigned int)strlen(Text);
+	size_t Len = WriteVarUInt(Pos, NameLen);
+	memcpy(Pos + Len, Text, (unsigned int)NameLen);
+	return Len + NameLen;
+}
+
diff --git a/Engine/Source/Programs/Horde/Samples/RemoteWorkerCpp/ComputeSocket.h b/Engine/Source/Programs/Horde/Samples/RemoteWorkerCpp/ComputeSocket.h
new file mode 100644
--- /dev/null
+++ b/Engine/Source/Programs/Horde/Samples/RemoteWorkerCpp/ComputeSocket.h
@@ -0,0 +1,34 @@
+// Copyright Epic Games, Inc. All Rights Reserved.
+
+#pragma once
+
+#include "SharedMemoryBuffer.h"
+
+class FComputeSocket
+{
+	static const char* const EnvVarName;
+
+public:
+	FComputeSocket();
+	~FComputeSocket();
+
+	/** Opens a connection to the agent process */
+	bool Open();
+
+	/** Close the current connection */
+	void Close();
+
+	/** Attaches a new buffer for receiving data */
+	void AttachRecvBuffer(int ChannelId, FSharedMemoryBuffer& Buffer);
+
+	/** Attaches a new buffer for sending data */
+	void AttachSendBuffer(int ChannelId, FSharedMemoryBuffer& Buffer);
+
+private:
+	FSharedMemoryBuffer CommandBuffer;
+
+	void AttachBuffer(int ChannelId, int Type, FSharedMemoryBuffer& Buffer);
+
+	static size_t WriteVarUInt(unsigned char* Pos, unsigned int Value);
+	static size_t WriteString(unsigned char* Pos, const char* Text);
+};
diff --git a/Engine/Source/Programs/Horde/Samples/RemoteWorkerCpp/Main.cpp b/Engine/Source/Programs/Horde/Samples/RemoteWorkerCpp/Main.cpp
--- a/Engine/Source/Programs/Horde/Samples/RemoteWorkerCpp/Main.cpp
+++ b/Engine/Source/Programs/Horde/Samples/RemoteWorkerCpp/Main.cpp
@@ -6,14 +6,21 @@
 
 int main()
 {
-	FComputeChannel Channel;
-	if (!Channel.Open())
+	FComputeSocket Socket;
+	if (!Socket.Open())
 	{
 		std::cout << "Environment variable not set correctly" << std::endl;
 		return 1;
 	}
 
-	std::cout << "Connected to client" << std::endl; // // Client will wait for output before sending data on channel 1
+	FComputeChannel Channel;
+	if (!Channel.Open(Socket, FComputeChannel::DefaultWorkerChannelId))
+	{
+		std::cout << "Unable to create channel to host" << std::endl;
+		return 1;
+	}
+
+	std::cout << "Connected to client" << std::endl;
 
 	size_t Length = 0;
 	char Buffer[4];
@@ -35,6 +42,7 @@ int main()
 		}
 	}
 
+	Channel.Close();
 	return 0;
 }
 
diff --git a/Engine/Source/Programs/Horde/Samples/RemoteWorkerCpp/SharedMemoryBuffer.cpp b/Engine/Source/Programs/Horde/Samples/RemoteWorkerCpp/SharedMemoryBuffer.cpp
--- a/Engine/Source/Programs/Horde/Samples/RemoteWorkerCpp/SharedMemoryBuffer.cpp
+++ b/Engine/Source/Programs/Horde/Samples/RemoteWorkerCpp/SharedMemoryBuffer.cpp
@@ -3,8 +3,11 @@
 #include <windows.h>
 #include <assert.h>
 #include <iostream>
+#include <rpcdcep.h>
 #include "SharedMemoryBuffer.h"
 
+unsigned int FSharedMemoryBuffer::Counter = 0;
+
 enum class FSharedMemoryBuffer::EWriteState : int
 {
 	// Chunk is still being appended to
@@ -112,12 +115,28 @@ FSharedMemoryBuffer::~FSharedMemoryBuffer()
 	Close();
 }
 
-bool FSharedMemoryBuffer::OpenExisting(const char* Name)
+bool FSharedMemoryBuffer::CreateNew(const char* Name, int NumChunks, int ChunkLength)
 {
+	Close();
+
+	char BaseNameBuffer[MAX_PATH];
+	if (Name == nullptr)
+	{
+		DWORD Pid = GetCurrentProcessId();
+		ULONGLONG TickCount = GetTickCount64();
+		sprintf_s(BaseNameBuffer, "Local\\COMPUTE_%u_%llu_%u", Pid, TickCount, InterlockedIncrement(&Counter));
+		Name = BaseNameBuffer;
+	}
+
+	unsigned long long Capacity = sizeof(FHeader) + (NumChunks * sizeof(unsigned int)) + (NumChunks * ChunkLength);
+
 	char NameBuffer[MAX_PATH];
 	sprintf_s(NameBuffer, "%s_M", Name);
 
-	MemoryMappedFile = OpenFileMappingA(FILE_MAP_ALL_ACCESS, TRUE, NameBuffer);
+	LARGE_INTEGER LargeInteger;
+	LargeInteger.QuadPart = Capacity;
+
+	MemoryMappedFile = CreateFileMappingA(INVALID_HANDLE_VALUE, NULL, PAGE_READWRITE, LargeInteger.HighPart, LargeInteger.LowPart, NameBuffer);
 	if (MemoryMappedFile == nullptr)
 	{
 		Close();
@@ -131,8 +150,11 @@ bool FSharedMemoryBuffer::OpenExisting(const char* Name)
 		return false;
 	}
 
+	Header->NumChunks = NumChunks;
+	Header->ChunkLength = ChunkLength;
+
 	sprintf_s(NameBuffer, "%s_R", Name);
-	ReaderEvent = OpenEventA(SYNCHRONIZE | EVENT_MODIFY_STATE, TRUE, NameBuffer);
+	ReaderEvent = CreateEventA(NULL, TRUE, FALSE, NameBuffer);
 	if (ReaderEvent == nullptr)
 	{
 		Close();
@@ -140,6 +162,46 @@ bool FSharedMemoryBuffer::OpenExisting(const char* Name)
 	}
 
 	sprintf_s(NameBuffer, "%s_W", Name);
+	WriterEvent = CreateEventA(NULL, TRUE, FALSE, NameBuffer);
+	if (WriterEvent == nullptr)
+	{
+		Close();
+		return false;
+	}
+
+	return OpenInternal(Name);
+}
+
+bool FSharedMemoryBuffer::OpenExisting(const char* InName)
+{
+	Close();
+
+	char NameBuffer[MAX_PATH];
+	sprintf_s(NameBuffer, "%s_M", InName);
+
+	MemoryMappedFile = OpenFileMappingA(FILE_MAP_ALL_ACCESS, TRUE, NameBuffer);
+	if (MemoryMappedFile == nullptr)
+	{
+		Close();
+		return false;
+	}
+
+	Header = (FHeader*)MapViewOfFile(MemoryMappedFile, FILE_MAP_ALL_ACCESS, 0, 0, 0);
+	if (Header == nullptr)
+	{
+		Close();
+		return false;
+	}
+
+	sprintf_s(NameBuffer, "%s_R", InName);
+	ReaderEvent = OpenEventA(SYNCHRONIZE | EVENT_MODIFY_STATE, TRUE, NameBuffer);
+	if (ReaderEvent == nullptr)
+	{
+		Close();
+		return false;
+	}
+
+	sprintf_s(NameBuffer, "%s_W", InName);
 	WriterEvent = OpenEventA(SYNCHRONIZE | EVENT_MODIFY_STATE, TRUE, NameBuffer);
 	if (WriterEvent == nullptr)
 	{
@@ -147,6 +209,11 @@ bool FSharedMemoryBuffer::OpenExisting(const char* Name)
 		return false;
 	}
 
+	return OpenInternal(InName);
+}
+
+bool FSharedMemoryBuffer::OpenInternal(const char* InName)
+{
 	ReadChunkIdx = 0;
 	ReadChunkDataPtr = GetChunkDataPtr(ReadChunkIdx);
 	ReadChunkStatePtr = GetChunkStatePtr(ReadChunkIdx);
@@ -155,6 +222,7 @@ bool FSharedMemoryBuffer::OpenExisting(const char* Name)
 	WriteChunkDataPtr = GetChunkDataPtr(WriteChunkIdx);
 	WriteChunkStatePtr = GetChunkStatePtr(WriteChunkIdx);
 
+	strcpy_s(Name, InName);
 	return true;
 }
 
@@ -192,6 +260,8 @@ void FSharedMemoryBuffer::Close()
 	WriteChunkIdx = -1;
 	WriteChunkStatePtr = nullptr;
 	WriteChunkDataPtr = nullptr;
+
+	Name[0] = 0;
 }
 
 bool FSharedMemoryBuffer::IsComplete() const
@@ -273,21 +343,21 @@ void FSharedMemoryBuffer::WaitToRead(size_t CurrentLength)
 
 void FSharedMemoryBuffer::MarkComplete()
 {
-	FChunkState::MarkComplete(ReadChunkStatePtr);
+	FChunkState::MarkComplete(WriteChunkStatePtr);
 	SetEvent(ReaderEvent);
 }
 
 void FSharedMemoryBuffer::AdvanceWritePosition(size_t Size)
 {
-	FChunkState::Append(ReadChunkStatePtr, Size);
+	FChunkState::Append(WriteChunkStatePtr, Size);
 	SetEvent(ReaderEvent);
 }
 
 unsigned char* FSharedMemoryBuffer::GetWriteMemory(size_t& OutSize)
 {
-	FChunkState State = FChunkState::Read(ReadChunkStatePtr);
-	OutSize = State.GetLength() - ReadOffset;
-	return ReadChunkDataPtr + ReadOffset;
+	FChunkState State = FChunkState::Read(WriteChunkStatePtr);
+	OutSize = Header->ChunkLength - State.GetLength();
+	return WriteChunkDataPtr + State.GetLength();
 }
 
 void FSharedMemoryBuffer::WaitToWrite(size_t CurrentLength)
diff --git a/Engine/Source/Programs/Horde/Samples/RemoteWorkerCpp/SharedMemoryBuffer.h b/Engine/Source/Programs/Horde/Samples/RemoteWorkerCpp/SharedMemoryBuffer.h
--- a/Engine/Source/Programs/Horde/Samples/RemoteWorkerCpp/SharedMemoryBuffer.h
+++ b/Engine/Source/Programs/Horde/Samples/RemoteWorkerCpp/SharedMemoryBuffer.h
@@ -12,12 +12,18 @@ public:
 	FSharedMemoryBuffer();
 	~FSharedMemoryBuffer();
 
+	// Creates a new buffer
+	bool CreateNew(const char* Name, int NumChunks, int ChunkLength);
+
 	// Opens an existing shared memory buffer (typically from handles created in another process)
 	bool OpenExisting(const char* Name);
 
 	// Close the current buffer and release all allocated resources
 	void Close();
 
+	// Gets the name of this buffer
+	const char* GetName() const { return Name; }
+
 
 	/*** Reader Interface ***/
 
@@ -49,11 +55,16 @@ public:
 	void WaitToWrite(size_t currentLength);
 
 private:
+	bool OpenInternal(const char* InName);
 	unsigned char* GetChunkDataPtr(int chunkIdx) const;
 	volatile long long* GetChunkStatePtr(int chunkIdx) const;
 
+	static unsigned int Counter;
+
 	const int ReaderIdx = 0;
 
+	char Name[260];
+
 	void* MemoryMappedFile;
 	FHeader* Header;
 	void* ReaderEvent;
diff --git a/Engine/Source/Programs/Shared/EpicGames.Horde.Tests/BufferTests.cs b/Engine/Source/Programs/Shared/EpicGames.Horde.Tests/BufferTests.cs
--- a/Engine/Source/Programs/Shared/EpicGames.Horde.Tests/BufferTests.cs
+++ b/Engine/Source/Programs/Shared/EpicGames.Horde.Tests/BufferTests.cs
@@ -13,7 +13,6 @@ using EpicGames.Horde.Compute;
 using EpicGames.Horde.Compute.Buffers;
 using EpicGames.Horde.Compute.Transports;
 using Microsoft.Extensions.Logging.Abstractions;
-using Microsoft.VisualStudio.TestPlatform.CrossPlatEngine;
 using Microsoft.VisualStudio.TestTools.UnitTesting;
 
 namespace EpicGames.Horde.Tests
@@ -26,7 +25,7 @@ namespace EpicGames.Horde.Tests
 		[TestMethod]
 		public async Task TestPooledBuffer()
 		{
-			await TestProducerConsumerAsync(length => new PooledBuffer(length).ToSharedInstance(), CancellationToken.None);
+			await TestProducerConsumerAsync(length => new PooledBuffer(length), CancellationToken.None);
 		}
 
 		[TestMethod]
@@ -34,7 +33,7 @@ namespace EpicGames.Horde.Tests
 		{
 			if (RuntimeInformation.IsOSPlatform(OSPlatform.Windows))
 			{
-				await TestProducerConsumerAsync(length => SharedMemoryBuffer.CreateNew(null, length).ToSharedInstance(), CancellationToken.None);
+				await TestProducerConsumerAsync(length => SharedMemoryBuffer.CreateNew(null, length), CancellationToken.None);
 			}
 		}
 
diff --git a/Engine/Source/Programs/Shared/EpicGames.Horde/Compute/Buffers/PooledBuffer.cs b/Engine/Source/Programs/Shared/EpicGames.Horde/Compute/Buffers/PooledBuffer.cs
--- a/Engine/Source/Programs/Shared/EpicGames.Horde/Compute/Buffers/PooledBuffer.cs
+++ b/Engine/Source/Programs/Shared/EpicGames.Horde/Compute/Buffers/PooledBuffer.cs
@@ -8,10 +8,65 @@ using EpicGames.Core;
 
 namespace EpicGames.Horde.Compute.Buffers
 {
+	/// <summary>
+	/// Implementation of <see cref="IComputeBuffer"/> suitable for cross-process communication
+	/// </summary>
+	public sealed class PooledBuffer : IComputeBuffer
+	{
+		PooledBufferCore _core;
+
+		/// <summary>
+		/// Constructor
+		/// </summary>
+		/// <param name="capacity">Total capacity of the buffer</param>
+		public PooledBuffer(int capacity)
+			: this(2, capacity / 2)
+		{
+		}
+
+		/// <summary>
+		/// Constructor
+		/// </summary>
+		/// <param name="numChunks">Number of chunks in the buffer</param>
+		/// <param name="chunkLength">Length of each chunk</param>
+		public PooledBuffer(int numChunks, int chunkLength)
+			: this(new PooledBufferCore(numChunks, chunkLength))
+		{
+		}
+
+		private PooledBuffer(PooledBufferCore core) => _core = core;
+
+		/// <inheritdoc/>
+		public IComputeBufferReader Reader => _core.Reader;
+
+		/// <inheritdoc/>
+		public IComputeBufferWriter Writer => _core.Writer;
+
+		/// <inheritdoc cref="IComputeBuffer.AddRef"/>
+		public PooledBuffer AddRef()
+		{
+			_core.AddRef();
+			return new PooledBuffer(_core);
+		}
+
+		/// <inheritdoc/>
+		IComputeBuffer IComputeBuffer.AddRef() => AddRef();
+
+		/// <inheritdoc/>
+		public void Dispose()
+		{
+			if (_core != null)
+			{
+				_core.Release();
+				_core = null!;
+			}
+		}
+	}
+
 	/// <summary>
 	/// In-process buffer used to store compute messages
 	/// </summary>
-	public sealed class PooledBuffer : ComputeBufferBase
+	sealed class PooledBufferCore : ComputeBufferBase
 	{
 		class Chunk : ChunkBase, IDisposable
 		{
@@ -34,21 +89,12 @@ namespace EpicGames.Horde.Compute.Buffers
 		readonly AsyncEvent _writerEvent = new AsyncEvent();
 		readonly AsyncEvent _readerEvent = new AsyncEvent();
 
-		/// <summary>
-		/// Constructor
-		/// </summary>
-		/// <param name="capacity">Total capacity of the buffer</param>
-		public PooledBuffer(int capacity)
-			: this(2, capacity / 2)
-		{
-		}
-
 		/// <summary>
 		/// Constructor
 		/// </summary>
 		/// <param name="numChunks">Number of chunks in the buffer</param>
 		/// <param name="chunkLength">Length of each chunk</param>
-		public PooledBuffer(int numChunks, int chunkLength)
+		public PooledBufferCore(int numChunks, int chunkLength)
 			: this(CreateChunks(numChunks, chunkLength), 1)
 		{
 		}
@@ -56,7 +102,7 @@ namespace EpicGames.Horde.Compute.Buffers
 		/// <summary>
 		/// Constructor
 		/// </summary>
-		private PooledBuffer(Chunk[] chunks, int readerIdx)
+		private PooledBufferCore(Chunk[] chunks, int readerIdx)
 			: base(chunks, readerIdx)
 		{
 			_chunks = chunks;
diff --git a/Engine/Source/Programs/Shared/EpicGames.Horde/Compute/Buffers/SharedMemoryBuffer.cs b/Engine/Source/Programs/Shared/EpicGames.Horde/Compute/Buffers/SharedMemoryBuffer.cs
--- a/Engine/Source/Programs/Shared/EpicGames.Horde/Compute/Buffers/SharedMemoryBuffer.cs
+++ b/Engine/Source/Programs/Shared/EpicGames.Horde/Compute/Buffers/SharedMemoryBuffer.cs
@@ -14,47 +14,44 @@
 	/// <summary>
 	/// Implementation of <see cref="IComputeBuffer"/> suitable for cross-process communication
 	/// </summary>
-	public sealed class SharedMemoryBuffer : ComputeBufferBase
+	public sealed class SharedMemoryBuffer : IComputeBuffer
 	{
-		const int HeaderLength = sizeof(int) + sizeof(int); // num chunks, chunk length
+		SharedMemoryBufferCore _core;
 
-		sealed unsafe class Chunk : ChunkBase
+		/// <summary>
+		/// Name of this shared memory buffer
+		/// </summary>
+		public string Name => _core.Name;
+
+		private SharedMemoryBuffer(SharedMemoryBufferCore core)
 		{
-			readonly ulong* _statePtr;
+			_core = core;
+		}
 
-			protected override ref ulong StateValue => ref *_statePtr;
+		/// <inheritdoc/>
+		public IComputeBufferReader Reader => _core.Reader;
 
-			public Chunk(ulong* statePtr, Memory<byte> memory)
-				: base(memory)
+		/// <inheritdoc/>
+		public IComputeBufferWriter Writer => _core.Writer;
+
+		/// <inheritdoc cref="IComputeBuffer.AddRef"/>
+		public SharedMemoryBuffer AddRef()
+		{
+			_core.AddRef();
+			return new SharedMemoryBuffer(_core);
+		}
+
+		/// <inheritdoc/>
+		IComputeBuffer IComputeBuffer.AddRef() => AddRef();
+
+		/// <inheritdoc/>
+		public void Dispose()
+		{
+			if (_core != null)
 			{
-				_statePtr = statePtr;
+				_core.Release();
+				_core = null!;
 			}
 		}
 
-		readonly MemoryMappedFile _memoryMappedFile;
-		readonly MemoryMappedViewAccessor _memoryMappedViewAccessor;
-		readonly MemoryMappedView _memoryMappedView;
-
-		readonly Native.EventHandle _readerEvent;
-		readonly Native.EventHandle _writerEvent;
-
-		/// <inheritdoc/>
-		public string Name { get; }
-
 		/// <summary>
-		/// Constructor
-		/// </summary>
-		private unsafe SharedMemoryBuffer(string name, MemoryMappedFile memoryMappedFile, MemoryMappedViewAccessor memoryMappedViewAccessor, MemoryMappedView memoryMappedView, Chunk[] chunks, Native.EventHandle readerEvent, Native.EventHandle writerEvent)
-			: base(chunks, 1)
-		{
-			Name = name;
-
-			_memoryMappedFile = memoryMappedFile;
-			_memoryMappedViewAccessor = memoryMappedViewAccessor;
-			_memoryMappedView = memoryMappedView;
-
-			_readerEvent = readerEvent;
-			_writerEvent = writerEvent;
-		}
-
-		/// <summary>
@@ -76,7 +73,7 @@
 		/// <param name="chunkLength">Length of each chunk</param>
 		public static unsafe SharedMemoryBuffer CreateNew(string? name, int numChunks, int chunkLength)
 		{
-			long capacity = HeaderLength + (numChunks * sizeof(ulong)) + (numChunks * chunkLength);
+			long capacity = SharedMemoryBufferCore.HeaderLength + (numChunks * sizeof(ulong)) + (numChunks * chunkLength);
 
 			name ??= $"Local\\COMPUTE_{Guid.NewGuid()}";
 
@@ -91,22 +88,8 @@
 			Native.EventHandle writerEvent = Native.EventHandle.CreateNew($"{name}_W", EventResetMode.ManualReset, true, HandleInheritability.Inheritable);
 			Native.EventHandle readerEvent = Native.EventHandle.CreateNew($"{name}_R", EventResetMode.ManualReset, true, HandleInheritability.Inheritable);
 
-			Chunk[] chunks = CreateChunks(numChunks, chunkLength, memoryMappedView);
-			return new SharedMemoryBuffer(name, memoryMappedFile, memoryMappedViewAccessor, memoryMappedView, chunks, readerEvent, writerEvent);
+			SharedMemoryBufferCore core = new SharedMemoryBufferCore(name, memoryMappedFile, memoryMappedViewAccessor, memoryMappedView, numChunks, chunkLength, readerEvent, writerEvent);
+			return new SharedMemoryBuffer(core);
 		}
 
-		static unsafe Chunk[] CreateChunks(int numChunks, int chunkLength, MemoryMappedView memoryMappedView)
-		{
-			Chunk[] chunks = new Chunk[numChunks];
-			for (int chunkIdx = 0; chunkIdx < numChunks; chunkIdx++)
-			{
-				byte* statePtr = memoryMappedView.GetPointer() + HeaderLength + (chunkIdx * sizeof(ulong));
-				int chunkOffset = HeaderLength + (numChunks * sizeof(ulong)) + (chunkLength * chunkIdx);
-
-				Chunk chunk = new Chunk((ulong*)statePtr, memoryMappedView.GetMemory(chunkOffset, chunkLength));
-				chunks[chunkIdx] = chunk;
-			}
-			return chunks;
-		}
-
 		/// <summary>
@@ -124,15 +107,76 @@
 			MemoryMappedViewAccessor memoryMappedViewAccessor = memoryMappedFile.CreateViewAccessor();
 			MemoryMappedView memoryMappedView = new MemoryMappedView(memoryMappedViewAccessor);
 
-			ReadOnlySpan<byte> header = memoryMappedView.GetMemory(0, HeaderLength).Span;
+			ReadOnlySpan<byte> header = memoryMappedView.GetMemory(0, SharedMemoryBufferCore.HeaderLength).Span;
 			int numChunks = BinaryPrimitives.ReadInt32LittleEndian(header);
 			int chunkLength = BinaryPrimitives.ReadInt32LittleEndian(header.Slice(4));
 
 			Native.EventHandle readerEvent = Native.EventHandle.OpenExisting($"{name}_R");
 			Native.EventHandle writerEvent = Native.EventHandle.OpenExisting($"{name}_W");
 
-			Chunk[] chunks = CreateChunks(numChunks, chunkLength, memoryMappedView);
-			return new SharedMemoryBuffer(name, memoryMappedFile, memoryMappedViewAccessor, memoryMappedView, chunks, readerEvent, writerEvent);
+			SharedMemoryBufferCore core = new SharedMemoryBufferCore(name, memoryMappedFile, memoryMappedViewAccessor, memoryMappedView, numChunks, chunkLength, readerEvent, writerEvent);
+			return new SharedMemoryBuffer(core);
+		}
+	}
+
+	/// <summary>
+	/// Core implementation of <see cref="SharedMemoryBuffer"/>
+	/// </summary>
+	sealed class SharedMemoryBufferCore : ComputeBufferBase
+	{
+		public const int HeaderLength = sizeof(int) + sizeof(int); // num chunks, chunk length
+
+		sealed unsafe class Chunk : ChunkBase
+		{
+			readonly ulong* _statePtr;
+
+			protected override ref ulong StateValue => ref *_statePtr;
+
+			public Chunk(ulong* statePtr, Memory<byte> memory)
+				: base(memory)
+			{
+				_statePtr = statePtr;
+			}
+		}
+
+		readonly MemoryMappedFile _memoryMappedFile;
+		readonly MemoryMappedViewAccessor _memoryMappedViewAccessor;
+		readonly MemoryMappedView _memoryMappedView;
+
+		readonly Native.EventHandle _readerEvent;
+		readonly Native.EventHandle _writerEvent;
+
+		/// <inheritdoc/>
+		public string Name { get; }
+
+		/// <summary>
+		/// Constructor
+		/// </summary>
+		public unsafe SharedMemoryBufferCore(string name, MemoryMappedFile memoryMappedFile, MemoryMappedViewAccessor memoryMappedViewAccessor, MemoryMappedView memoryMappedView, int numChunks, int chunkLength, Native.EventHandle readerEvent, Native.EventHandle writerEvent)
+			: base(CreateChunks(numChunks, chunkLength, memoryMappedView), 1)
+		{
+			Name = name;
+
+			_memoryMappedFile = memoryMappedFile;
+			_memoryMappedViewAccessor = memoryMappedViewAccessor;
+			_memoryMappedView = memoryMappedView;
+
+			_readerEvent = readerEvent;
+			_writerEvent = writerEvent;
+		}
+
+		static unsafe Chunk[] CreateChunks(int numChunks, int chunkLength, MemoryMappedView memoryMappedView)
+		{
+			Chunk[] chunks = new Chunk[numChunks];
+			for (int chunkIdx = 0; chunkIdx < numChunks; chunkIdx++)
+			{
+				byte* statePtr = memoryMappedView.GetPointer() + HeaderLength + (chunkIdx * sizeof(ulong));
+				int chunkOffset = HeaderLength + (numChunks * sizeof(ulong)) + (chunkLength * chunkIdx);
+
+				Chunk chunk = new Chunk((ulong*)statePtr, memoryMappedView.GetMemory(chunkOffset, chunkLength));
+				chunks[chunkIdx] = chunk;
+			}
+			return chunks;
 		}
 
 		/// <inheritdoc/>
diff --git a/Engine/Source/Programs/Shared/EpicGames.Horde/Compute/ComputeBufferBase.cs b/Engine/Source/Programs/Shared/EpicGames.Horde/Compute/ComputeBufferBase.cs
--- a/Engine/Source/Programs/Shared/EpicGames.Horde/Compute/ComputeBufferBase.cs
+++ b/Engine/Source/Programs/Shared/EpicGames.Horde/Compute/ComputeBufferBase.cs
@@ -160,7 +160,7 @@ namespace EpicGames.Horde.Compute
 			}
 
 			/// <inheritdoc/>
-			public async ValueTask WaitToReadAsync(int currentLength, CancellationToken cancellationToken = default)
+			public async ValueTask<bool> WaitToReadAsync(int currentLength, CancellationToken cancellationToken = default)
 			{
 				for (; ; )
 				{
@@ -176,12 +176,17 @@ namespace EpicGames.Horde.Compute
 							await _buffer.WaitForReadEvent(_readerIdx, cancellationToken);
 						}
 					}
-					else if (_offset + currentLength < state.Length || state.WriteState == WriteState.Complete)
+					else if (_offset + currentLength < state.Length)
 					{
-						// Still have data to read from this chunk
-						break;
+						// Still have more data to read from the current buffer
+						return true;
+					}
+					else if (state.WriteState == WriteState.Complete)
+					{
+						// Buffer is complete
+						return false;
 					}
-					else if (_offset + currentLength >= state.Length && state.WriteState == WriteState.Writing)
+					else if (state.WriteState == WriteState.Writing)
 					{
 						// Wait until there is more data in the chunk
 						_buffer.ResetReadEvent(_readerIdx);
@@ -205,8 +210,7 @@ namespace EpicGames.Horde.Compute
 					}
 					else
 					{
-						// Still need to read data from the current buffer
-						break;
+						throw new NotImplementedException($"Invalid write state for buffer: {state.WriteState}");
 					}
 				}
 			}
@@ -298,46 +302,6 @@ namespace EpicGames.Horde.Compute
 			}
 		}
 
-		/// <summary>
-		/// Implementation of <see cref="IComputeBuffer"/> which handles reference counting 
-		/// </summary>
-		sealed class RefCountedImpl : IComputeBuffer
-		{
-			ComputeBufferBase _inner;
-
-			/// <summary>
-			/// Constructor
-			/// </summary>
-			/// <param name="buffer"></param>
-			public RefCountedImpl(ComputeBufferBase buffer) => _inner = buffer;
-
-			/// <inheritdoc/>
-			public IComputeBufferReader Reader => _inner.Reader;
-
-			/// <inheritdoc/>
-			public IComputeBufferWriter Writer => _inner.Writer;
-
-			/// <inheritdoc/>
-			public IComputeBuffer AddRef()
-			{
-				Interlocked.Increment(ref _inner._refCount);
-				return new RefCountedImpl(_inner);
-			}
-
-			/// <inheritdoc/>
-			public void Dispose()
-			{
-				if (_inner != null)
-				{
-					if (Interlocked.Decrement(ref _inner._refCount) == 0)
-					{
-						_inner.Dispose();
-					}
-					_inner = null!;
-				}
-			}
-		}
-
 		readonly ChunkBase[] _chunks;
 		readonly ReaderImpl _reader;
 		readonly WriterImpl _writer;
@@ -367,6 +331,21 @@ namespace EpicGames.Horde.Compute
 			_writer = new WriterImpl(this, _chunks);
 		}
 
+		/// <inheritdoc/>
+		public void AddRef()
+		{
+			Interlocked.Increment(ref _refCount);
+		}
+
+		/// <inheritdoc/>
+		public void Release()
+		{
+			if (Interlocked.Decrement(ref _refCount) == 0)
+			{
+				Dispose();
+			}
+		}
+
 		/// <inheritdoc/>
 		public void Dispose()
 		{
@@ -389,11 +368,6 @@ namespace EpicGames.Horde.Compute
 			}
 		}
 
-		/// <summary>
-		/// Creates a sharable <see cref="IComputeBuffer"/> implementation from this object.
-		/// </summary>
-		public IComputeBuffer ToSharedInstance() => new RefCountedImpl(this);
-
 		/// <summary>
 		/// Signals a read event
 		/// </summary>
diff --git a/Engine/Source/Programs/Shared/EpicGames.Horde/Compute/ComputeMessageChannel.cs b/Engine/Source/Programs/Shared/EpicGames.Horde/Compute/ComputeMessageChannel.cs
--- a/Engine/Source/Programs/Shared/EpicGames.Horde/Compute/ComputeMessageChannel.cs
+++ b/Engine/Source/Programs/Shared/EpicGames.Horde/Compute/ComputeMessageChannel.cs
@@ -120,6 +120,22 @@ namespace EpicGames.Horde.Compute
 		MessageBuilder? _currentBuilder;
 #pragma warning restore CA2213
 
+		/// <summary>
+		/// Constructor
+		/// </summary>
+		/// <param name="channelId"></param>
+		/// <param name="recvBuffer"></param>
+		/// <param name="sendBuffer"></param>
+		/// <param name="logger">Logger for diagnostic output</param>
+		public ComputeMessageChannel(int channelId, IComputeBuffer recvBuffer, IComputeBuffer sendBuffer, ILogger logger)
+		{
+			_socket = null!;
+			_channelId = channelId;
+			_recvBuffer = recvBuffer;
+			_sendBuffer = sendBuffer;
+			_logger = logger;
+		}
+
 		/// <summary>
 		/// Constructor
 		/// </summary>
@@ -158,6 +174,9 @@ namespace EpicGames.Horde.Compute
 			_sendBuffer.Writer.MarkComplete();
 		}
 
+		/// <inheritdoc/>
+		public ValueTask WaitForAttachAsync(CancellationToken cancellationToken = default) => _socket.WaitForAttachAsync(_channelId, cancellationToken);
+
 		/// <inheritdoc/>
 		public async ValueTask<IComputeMessage> ReceiveAsync(CancellationToken cancellationToken)
 		{
diff --git a/Engine/Source/Programs/Shared/EpicGames.Horde/Compute/ComputeWorker.cs b/Engine/Source/Programs/Shared/EpicGames.Horde/Compute/ComputeWorker.cs
--- a/Engine/Source/Programs/Shared/EpicGames.Horde/Compute/ComputeWorker.cs
+++ b/Engine/Source/Programs/Shared/EpicGames.Horde/Compute/ComputeWorker.cs
@@ -51,8 +51,6 @@ namespace EpicGames.Horde.Compute
 		{
 			using (IComputeMessageChannel channel = socket.CreateMessageChannel(channelId, bufferSize, _logger))
 			{
-				await channel.SendReadyAsync(cancellationToken);
-
 				List<Task> childTasks = new List<Task>();
 				for (; ; )
 				{
@@ -169,14 +167,6 @@ namespace EpicGames.Horde.Compute
 
 		async Task ExecuteProcessWindowsAsync(IComputeSocket socket, IComputeMessageChannel channel, int newChannelId, string executable, IReadOnlyList<string> arguments, string? workingDir, IReadOnlyDictionary<string, string?>? envVars, CancellationToken cancellationToken)
 		{
-			string baseBufferName = $"Local\\Compute-{Guid.NewGuid()}";
-
-			using IComputeBuffer sendBuffer = SharedMemoryBuffer.CreateNew(ComputeChannel.GetSendBufferName(baseBufferName), 4096).ToSharedInstance();
-			socket.AttachSendBuffer(newChannelId, sendBuffer);
-
-			using IComputeBuffer recvBuffer = SharedMemoryBuffer.CreateNew(ComputeChannel.GetRecvBufferName(baseBufferName), 4096).ToSharedInstance();
-			socket.AttachRecvBuffer(newChannelId, recvBuffer);
-
 			Dictionary<string, string?> newEnvVars = new Dictionary<string, string?>();
 			if (envVars != null)
 			{
@@ -186,7 +176,10 @@ namespace EpicGames.Horde.Compute
 				}
 			}
 
-			newEnvVars[ComputeChannel.IpcEnvVar] = baseBufferName;
+			using SharedMemoryBuffer ipcBuffer = SharedMemoryBuffer.CreateNew(null, 1, 64 * 1024);
+			newEnvVars[WorkerComputeSocket.IpcEnvVar] = ipcBuffer.Name;
+
+			using BackgroundTask backgroundTask = BackgroundTask.StartNew(ctx => ProcessIpcMessagesAsync(socket, ipcBuffer.Reader, cancellationToken));
 
 			_logger.LogInformation("Launching {Executable} {Arguments}", CommandLineArguments.Quote(executable), CommandLineArguments.Join(arguments));
 			await ExecuteProcessInternalAsync(channel, executable, arguments, workingDir, newEnvVars, cancellationToken);
@@ -194,8 +187,43 @@ namespace EpicGames.Horde.Compute
 
 			_logger.LogInformation("Child process has shut down");
 
-			recvBuffer.Writer.MarkComplete();
-			sendBuffer.Writer.MarkComplete();
+			ipcBuffer.Writer.MarkComplete();
+		}
+
+		static async Task ProcessIpcMessagesAsync(IComputeSocket socket, IComputeBufferReader ipcReader, CancellationToken cancellationToken)
+		{
+			while(await ipcReader.WaitToReadAsync(0, cancellationToken))
+			{
+				ReadOnlyMemory<byte> memory = ipcReader.GetMemory();
+				MemoryReader reader = new MemoryReader(memory);
+
+				IpcMessage message = (IpcMessage)reader.ReadUnsignedVarInt();
+				switch (message)
+				{
+					case IpcMessage.AttachSendBuffer:
+						{
+							int channelId = (int)reader.ReadUnsignedVarInt();
+							string name = reader.ReadString();
+#pragma warning disable CA2000 // Dispose objects before losing scope
+							socket.AttachSendBuffer(channelId, SharedMemoryBuffer.OpenExisting(name));
+#pragma warning restore CA2000 // Dispose objects before losing scope
+						}
+						break;
+					case IpcMessage.AttachRecvBuffer:
+						{
+							int channelId = (int)reader.ReadUnsignedVarInt();
+							string name = reader.ReadString();
+#pragma warning disable CA2000 // Dispose objects before losing scope
+							socket.AttachRecvBuffer(channelId, SharedMemoryBuffer.OpenExisting(name));
+#pragma warning restore CA2000 // Dispose objects before losing scope
+						}
+						break;
+					default:
+						throw new InvalidOperationException($"Invalid IPC message: {message}");
+				}
+
+				ipcReader.Advance(memory.Length - reader.RemainingMemory.Length);
+			}
 		}
 
 		async Task ExecuteProcessInternalAsync(IComputeMessageChannel channel, string executable, IReadOnlyList<string> arguments, string? workingDir, IReadOnlyDictionary<string, string?>? envVars, CancellationToken cancellationToken)
diff --git a/Engine/Source/Programs/Shared/EpicGames.Horde/Compute/DefaultComputeSocket.cs b/Engine/Source/Programs/Shared/EpicGames.Horde/Compute/DefaultComputeSocket.cs
--- a/Engine/Source/Programs/Shared/EpicGames.Horde/Compute/DefaultComputeSocket.cs
+++ b/Engine/Source/Programs/Shared/EpicGames.Horde/Compute/DefaultComputeSocket.cs
@@ -6,6 +6,7 @@ using System.Buffers.Binary;
 using System.Collections.Generic;
 using System.Linq;
 using System.Threading;
+using System.Threading.Channels;
 using System.Threading.Tasks;
 using EpicGames.Core;
 using Microsoft.Extensions.Logging;
@@ -17,6 +18,12 @@ namespace EpicGames.Horde.Compute
 	/// </summary>
 	class DefaultComputeSocket : IComputeSocket, IAsyncDisposable
 	{
+		enum ControlMessageType
+		{
+			Attach = -1,
+			Detach = -2,
+		}
+
 		readonly object _lockObject = new object();
 
 		bool _complete;
@@ -33,6 +40,12 @@ namespace EpicGames.Horde.Compute
 		readonly Dictionary<int, Task> _sendTasks = new Dictionary<int, Task>();
 		readonly Dictionary<IComputeBuffer, int> _sendBuffers = new Dictionary<IComputeBuffer, int>();
 
+		readonly HashSet<int> _remoteRecvBuffers = new HashSet<int>();
+		readonly AsyncEvent _remoteBuffersChanged = new AsyncEvent();
+
+		readonly BackgroundTask _attachTask;
+		readonly Channel<int> _attachChannel = Channel.CreateUnbounded<int>();
+
 		/// <summary>
 		/// Constructor
 		/// </summary>
@@ -43,6 +56,7 @@ namespace EpicGames.Horde.Compute
 			_transport = transport;
 			_logger = logger;
 			_recvTask = BackgroundTask.StartNew(ctx => RunReaderAsync(transport, ctx));
+			_attachTask = BackgroundTask.StartNew(ctx => SendAttachNotificationsAsync(ctx));
 		}
 
 		/// <summary>
@@ -73,6 +87,11 @@ namespace EpicGames.Horde.Compute
 			// a shutdown event to be sent back to our read task allowing it to shut down gracefully.
 			await _transport.MarkCompleteAsync(cancellationToken);
 
+			// Stop sending 
+			_attachChannel.Writer.Complete();
+			await _attachTask.Task.WaitAsync(TimeSpan.FromSeconds(5), cancellationToken);
+			await _attachTask.StopAsync();
+
 			// Wait for the reader to stop
 			await _recvTask.Task.WaitAsync(TimeSpan.FromSeconds(5), cancellationToken);
 			await _recvTask.StopAsync();
@@ -84,6 +103,7 @@ namespace EpicGames.Horde.Compute
 			_cancellationSource.Cancel();
 			await Task.WhenAll(_sendTasks.Values);
 			await _recvTask.DisposeAsync();
+			await _attachTask.DisposeAsync();
 			_sendSemaphore.Dispose();
 			_cancellationSource.Dispose();
 			GC.SuppressFinalize(this);
@@ -116,15 +136,24 @@ namespace EpicGames.Horde.Compute
 					int size = BinaryPrimitives.ReadInt32LittleEndian(header.AsSpan(4));
 
 					// Dispatch it to the correct place
-					if (size < 0)
+					if (size >= 0)
+					{
+						IComputeBufferWriter writer = GetReceiveBuffer(cachedWriters, id);
+						await ReadPacketAsync(transport, id, size, writer, cancellationToken);
+					}
+					else if (size == (int)ControlMessageType.Attach)
 					{
-						_logger.LogTrace("Detaching buffer {Id}", id);
+						_logger.LogTrace("Attaching recv buffer {Id}", id);
+						AttachRemoteRecvBuffer(id);
+					}
+					else if (size == (int)ControlMessageType.Detach)
+					{
+						_logger.LogTrace("Detaching recv buffer {Id}", id);
 						DetachRecvBuffer(cachedWriters, id);
 					}
 					else
 					{
-						IComputeBufferWriter writer = GetReceiveBuffer(cachedWriters, id);
-						await ReadPacketAsync(transport, id, size, writer, cancellationToken);
+						_logger.LogWarning("Unrecognized control message: {Message}", size);
 					}
 				}
 			}
@@ -186,6 +215,23 @@ namespace EpicGames.Horde.Compute
 			throw new ComputeInternalException($"No buffer is attached to channel {id}");
 		}
 
+		/// <inheritdoc/>
+		public async ValueTask WaitForAttachAsync(int channelId, CancellationToken cancellationToken = default)
+		{
+			for (; ; )
+			{
+				Task task = _remoteBuffersChanged.Task;
+				lock (_lockObject)
+				{
+					if (_remoteRecvBuffers.Contains(channelId))
+					{
+						return;
+					}
+				}
+				await task.WaitAsync(cancellationToken);
+			}
+		}
+
 		class SendSegment : ReadOnlySequenceSegment<byte>
 		{
 			public void Set(ReadOnlyMemory<byte> memory, ReadOnlySequenceSegment<byte>? next, long runningIndex)
@@ -210,7 +256,7 @@ namespace EpicGames.Horde.Compute
 		}
 
 		/// <inheritdoc/>
-		public ValueTask MarkCompleteAsync(int id, CancellationToken cancellationToken) => SendInternalAsync(id, -1, ReadOnlyMemory<byte>.Empty, cancellationToken);
+		public ValueTask MarkCompleteAsync(int id, CancellationToken cancellationToken) => SendInternalAsync(id, (int)ControlMessageType.Detach, ReadOnlyMemory<byte>.Empty, cancellationToken);
 
 		async ValueTask SendInternalAsync(int id, int size, ReadOnlyMemory<byte> memory, CancellationToken cancellationToken)
 		{
@@ -232,7 +278,7 @@ namespace EpicGames.Horde.Compute
 		}
 
 		/// <inheritdoc/>
-		public void AttachRecvBuffer(int channelId, IComputeBuffer recvBuffer)
+		public async void AttachRecvBuffer(int channelId, IComputeBuffer recvBuffer)
 		{
 			IComputeBuffer shared = recvBuffer.AddRef();
 
@@ -251,6 +297,32 @@ namespace EpicGames.Horde.Compute
 				shared.Writer.MarkComplete();
 				shared.Dispose();
 			}
+
+			if (!complete)
+			{
+				await _attachChannel.Writer.WriteAsync(channelId);
+			}
+		}
+
+		async Task SendAttachNotificationsAsync(CancellationToken cancellationToken)
+		{
+			while (await _attachChannel.Reader.WaitToReadAsync(cancellationToken))
+			{
+				int channelId;
+				if (_attachChannel.Reader.TryRead(out channelId))
+				{
+					await SendInternalAsync(channelId, (int)ControlMessageType.Attach, ReadOnlyMemory<byte>.Empty, cancellationToken);
+				}
+			}
+		}
+
+		void AttachRemoteRecvBuffer(int id)
+		{
+			lock (_lockObject)
+			{
+				_remoteRecvBuffers.Add(id);
+			}
+			_remoteBuffersChanged.Pulse();
 		}
 
 		void DetachRecvBuffer(Dictionary<int, IComputeBufferWriter> cachedWriters, int id)
diff --git a/Engine/Source/Programs/Shared/EpicGames.Horde/Compute/IComputeBuffer.cs b/Engine/Source/Programs/Shared/EpicGames.Horde/Compute/IComputeBuffer.cs
--- a/Engine/Source/Programs/Shared/EpicGames.Horde/Compute/IComputeBuffer.cs
+++ b/Engine/Source/Programs/Shared/EpicGames.Horde/Compute/IComputeBuffer.cs
@@ -54,7 +54,8 @@ namespace EpicGames.Horde.Compute
 		/// </summary>
 		/// <param name="currentLength">Current length of the buffer</param>
 		/// <param name="cancellationToken">Cancellation token for the operation</param>
-		ValueTask WaitToReadAsync(int currentLength, CancellationToken cancellationToken = default);
+		/// <returns>True if new data is available, false if the buffer is complete</returns>
+		ValueTask<bool> WaitToReadAsync(int currentLength, CancellationToken cancellationToken = default);
 	}
 
 	/// <summary>
diff --git a/Engine/Source/Programs/Shared/EpicGames.Horde/Compute/IComputeChannel.cs b/Engine/Source/Programs/Shared/EpicGames.Horde/Compute/IComputeChannel.cs
--- a/Engine/Source/Programs/Shared/EpicGames.Horde/Compute/IComputeChannel.cs
+++ b/Engine/Source/Programs/Shared/EpicGames.Horde/Compute/IComputeChannel.cs
@@ -5,7 +5,6 @@ using System.IO;
 using System.Threading;
 using System.Threading.Tasks;
 using EpicGames.Horde.Compute.Buffers;
-using Microsoft.Extensions.Logging;
 
 namespace EpicGames.Horde.Compute
 {
@@ -34,6 +33,11 @@ namespace EpicGames.Horde.Compute
 	/// </summary>
 	public static class ComputeChannel
 	{
+		/// <summary>
+		/// Default channel id to use for worker messages
+		/// </summary>
+		public const int DefaultWorkerChannelId = 100;
+
 		class BufferedReaderChannel : IComputeChannel
 		{
 			readonly IComputeSocket _socket;
@@ -59,22 +63,15 @@ namespace EpicGames.Horde.Compute
 			public ValueTask SendAsync(ReadOnlyMemory<byte> memory, CancellationToken cancellationToken = default) => _socket.SendAsync(_channelId, memory, cancellationToken);
 		}
 
-		class BufferedReaderWriterChannel : IComputeChannel
+		internal sealed class BufferedReaderWriterChannel : IComputeChannel
 		{
-			readonly IComputeSocket _socket;
-			readonly int _channelId;
 			readonly IComputeBuffer _recvBuffer;
 			readonly IComputeBuffer _sendBuffer;
 
-			public BufferedReaderWriterChannel(IComputeSocket socket, int channelId, IComputeBuffer recvBuffer, IComputeBuffer sendBuffer)
+			public BufferedReaderWriterChannel(IComputeBuffer recvBuffer, IComputeBuffer sendBuffer)
 			{
-				_socket = socket;
-				_channelId = channelId;
 				_recvBuffer = recvBuffer;
 				_sendBuffer = sendBuffer;
-
-				_socket.AttachRecvBuffer(_channelId, recvBuffer);
-				_socket.AttachSendBuffer(_channelId, sendBuffer);
 			}
 
 			public void Dispose()
@@ -88,47 +85,40 @@ namespace EpicGames.Horde.Compute
 			public ValueTask SendAsync(ReadOnlyMemory<byte> memory, CancellationToken cancellationToken = default) => _sendBuffer.Writer.WriteAsync(memory, cancellationToken);
 		}
 
-		class WorkerChannel : IComputeChannel
+		class SimpleWorkerChannel : IComputeChannel
 		{
-			readonly IComputeBuffer _recvBuffer;
-			readonly IComputeBuffer _sendBuffer;
+			readonly WorkerComputeSocket _socket;
+			readonly SharedMemoryBuffer _recvBuffer;
+			readonly SharedMemoryBuffer _sendBuffer;
 
-			public WorkerChannel(IComputeBuffer recvBuffer, IComputeBuffer sendBuffer)
+			public SimpleWorkerChannel(int channelId)
 			{
-				_recvBuffer = recvBuffer;
-				_sendBuffer = sendBuffer;
+				_socket = new WorkerComputeSocket();
+
+				_recvBuffer = SharedMemoryBuffer.CreateNew(null, 1024 * 1024);
+				_socket.AttachRecvBuffer(channelId, _recvBuffer);
+
+				_sendBuffer = SharedMemoryBuffer.CreateNew(null, 1024 * 1024);
+				_socket.AttachSendBuffer(channelId, _sendBuffer);
 			}
 
-			public void Dispose() { }
+			public void Dispose()
+			{
+				_sendBuffer.Writer.MarkComplete();
+				_sendBuffer.Dispose();
+				_recvBuffer.Dispose();
+				_socket.Dispose();
+			}
 
 			public ValueTask<int> ReceiveAsync(Memory<byte> buffer, CancellationToken cancellationToken = default) => _recvBuffer.Reader.ReadAsync(buffer, cancellationToken);
 
 			public ValueTask SendAsync(ReadOnlyMemory<byte> memory, CancellationToken cancellationToken = default) => _sendBuffer.Writer.WriteAsync(memory, cancellationToken);
 		}
 
-		/// <summary>
-		/// Name of the environment variable for passing the name of the compute channel
-		/// </summary>
-		public const string IpcEnvVar = "UE_HORDE_COMPUTE_IPC";
-
-		internal static string GetSendBufferName(string baseName) => $"{baseName}_SEND";
-		internal static string GetRecvBufferName(string baseName) => $"{baseName}_RECV";
-
 		/// <summary>
 		/// Creates a socket for a worker
 		/// </summary>
-		public static IComputeChannel ConnectAsWorker()
-		{
-			string? baseName = Environment.GetEnvironmentVariable(IpcEnvVar);
-			if (baseName == null)
-			{
-				throw new InvalidOperationException($"Environment variable {IpcEnvVar} is not defined; cannot connect as worker.");
-			}
-
-			IComputeBuffer recvBuffer = SharedMemoryBuffer.OpenExisting(GetRecvBufferName(baseName)).ToSharedInstance();
-			IComputeBuffer sendBuffer = SharedMemoryBuffer.OpenExisting(GetSendBufferName(baseName)).ToSharedInstance();
-			return new WorkerChannel(recvBuffer, sendBuffer);
-		}
+		public static IComputeChannel ConnectAsWorker(int channelId = DefaultWorkerChannelId) => new SimpleWorkerChannel(channelId);
 
 		/// <summary>
 		/// Creates a channel using a socket and receive buffer
@@ -145,7 +135,12 @@ namespace EpicGames.Horde.Compute
 		/// <param name="channelId">Channel id to send and receive data</param>
 		/// <param name="recvBuffer">Buffer for receiving data</param>
 		/// <param name="sendBuffer">Buffer for sending data</param>
-		public static IComputeChannel CreateChannel(this IComputeSocket socket, int channelId, IComputeBuffer recvBuffer, IComputeBuffer sendBuffer) => new BufferedReaderWriterChannel(socket, channelId, recvBuffer, sendBuffer);
+		public static IComputeChannel CreateChannel(this IComputeSocket socket, int channelId, IComputeBuffer recvBuffer, IComputeBuffer sendBuffer)
+		{
+			socket.AttachRecvBuffer(channelId, recvBuffer);
+			socket.AttachSendBuffer(channelId, sendBuffer);
+			return new BufferedReaderWriterChannel(recvBuffer, sendBuffer);
+		}
 	}
 
 	/// <summary>
diff --git a/Engine/Source/Programs/Shared/EpicGames.Horde/Compute/IComputeMessage.cs b/Engine/Source/Programs/Shared/EpicGames.Horde/Compute/IComputeMessage.cs
--- a/Engine/Source/Programs/Shared/EpicGames.Horde/Compute/IComputeMessage.cs
+++ b/Engine/Source/Programs/Shared/EpicGames.Horde/Compute/IComputeMessage.cs
@@ -22,11 +22,6 @@ namespace EpicGames.Horde.Compute
 		/// </summary>
 		None = 0x00,
 
-		/// <summary>
-		/// Indicates that the remote is ready to receive messages on the current channel
-		/// </summary>
-		Ready = 0x01,
-
 		/// <summary>
 		/// Sent in place of a regular response if an error occurs on the remote
 		/// </summary>
@@ -183,15 +178,6 @@ namespace EpicGames.Horde.Compute
 			message.Send();
 		}
 
-		/// <summary>
-		/// Sends a message to the remote indicating that we're ready to receive messages on the given channel
-		/// </summary>
-		public static async ValueTask SendReadyAsync(this IComputeMessageChannel channel, CancellationToken cancellationToken = default)
-		{
-			using IComputeMessageBuilder message = await channel.CreateMessageAsync(ComputeMessageType.Ready, cancellationToken);
-			message.Send();
-		}
-
 		/// <summary>
 		/// Sends an exception response to the remote
 		/// </summary>
diff --git a/Engine/Source/Programs/Shared/EpicGames.Horde/Compute/IComputeMessageChannel.cs b/Engine/Source/Programs/Shared/EpicGames.Horde/Compute/IComputeMessageChannel.cs
--- a/Engine/Source/Programs/Shared/EpicGames.Horde/Compute/IComputeMessageChannel.cs
+++ b/Engine/Source/Programs/Shared/EpicGames.Horde/Compute/IComputeMessageChannel.cs
@@ -12,6 +12,12 @@ namespace EpicGames.Horde.Compute
 	/// </summary>
 	public interface IComputeMessageChannel : IDisposable
 	{
+		/// <summary>
+		/// Wait for the remote end to attach a reader
+		/// </summary>
+		/// <param name="cancellationToken">Cancellation token for the operation</param>
+		ValueTask WaitForAttachAsync(CancellationToken cancellationToken = default);
+
 		/// <summary>
 		/// Reads a message from the channel
 		/// </summary>
@@ -62,20 +68,6 @@ namespace EpicGames.Horde.Compute
 			return message;
 		}
 
-		/// <summary>
-		/// Waits for the remote machine to send a 'ready' response
-		/// </summary>
-		/// <param name="channel">Channel to receive on</param>
-		/// <param name="cancellationToken">Cancellation token for the operation</param>
-		public static async ValueTask WaitForAttachAsync(this IComputeMessageChannel channel, CancellationToken cancellationToken = default)
-		{
-			IComputeMessage message = await channel.ReceiveAsync(cancellationToken);
-			if (message.Type != ComputeMessageType.Ready)
-			{
-				throw new ComputeInvalidMessageException(message);
-			}
-		}
-
 		/// <summary>
 		/// Creates a new builder for a message
 		/// </summary>
diff --git a/Engine/Source/Programs/Shared/EpicGames.Horde/Compute/IComputeSocket.cs b/Engine/Source/Programs/Shared/EpicGames.Horde/Compute/IComputeSocket.cs
--- a/Engine/Source/Programs/Shared/EpicGames.Horde/Compute/IComputeSocket.cs
+++ b/Engine/Source/Programs/Shared/EpicGames.Horde/Compute/IComputeSocket.cs
@@ -42,6 +42,13 @@ namespace EpicGames.Horde.Compute
 		/// <param name="buffer">Buffer to send data from</param>
 		void AttachSendBuffer(int channelId, IComputeBuffer buffer);
 
+		/// <summary>
+		/// Waits until the remote has attached a receive buffer to the given channel
+		/// </summary>
+		/// <param name="channelId">Channel id</param>
+		/// <param name="cancellationToken">Cancellation token for the operation</param>
+		ValueTask WaitForAttachAsync(int channelId, CancellationToken cancellationToken = default);
+
 		/// <summary>
 		/// Attempt to gracefully close the current connection and shutdown both ends of the transport
 		/// </summary>
@@ -96,8 +103,8 @@ namespace EpicGames.Horde.Compute
 		/// <param name="logger">Logger for the channel</param>
 		public static IComputeMessageChannel CreateMessageChannel(this IComputeSocket socket, int channelId, int sendBufferSize, int recvBufferSize, ILogger logger)
 		{
-			IComputeBuffer sendBuffer = new PooledBuffer(sendBufferSize).ToSharedInstance();
-			IComputeBuffer recvBuffer = new PooledBuffer(recvBufferSize).ToSharedInstance();
+			IComputeBuffer sendBuffer = new PooledBuffer(sendBufferSize);
+			IComputeBuffer recvBuffer = new PooledBuffer(recvBufferSize);
 			return new ComputeMessageChannel(socket, channelId, sendBuffer, recvBuffer, logger);
 		}
 	}
diff --git a/Engine/Source/Programs/Shared/EpicGames.Horde/Compute/WorkerComputeSocket.cs b/Engine/Source/Programs/Shared/EpicGames.Horde/Compute/WorkerComputeSocket.cs
new file mode 100644
--- /dev/null
+++ b/Engine/Source/Programs/Shared/EpicGames.Horde/Compute/WorkerComputeSocket.cs
@@ -0,0 +1,90 @@
+// Copyright Epic Games, Inc. All Rights Reserved.
+
+using System;
+using EpicGames.Core;
+using EpicGames.Horde.Compute.Buffers;
+
+namespace EpicGames.Horde.Compute
+{
+	internal enum IpcMessage
+	{
+		AttachRecvBuffer = 0,
+		AttachSendBuffer = 1,
+	}
+
+	/// <summary>
+	/// Provides functionality for attaching buffers for compute workers 
+	/// </summary>
+	public sealed class WorkerComputeSocket : IDisposable
+	{
+		/// <summary>
+		/// Name of the environment variable for passing the name of the compute channel
+		/// </summary>
+		public const string IpcEnvVar = "UE_HORDE_COMPUTE_IPC";
+
+		readonly SharedMemoryBuffer _commandBuffer;
+
+		/// <summary>
+		/// Creates a socket for a worker
+		/// </summary>
+		public WorkerComputeSocket()
+		{
+			string? baseName = Environment.GetEnvironmentVariable(IpcEnvVar);
+			if (baseName == null)
+			{
+				throw new InvalidOperationException($"Environment variable {IpcEnvVar} is not defined; cannot connect as worker.");
+			}
+
+			_commandBuffer = SharedMemoryBuffer.OpenExisting(baseName);
+		}
+
+		/// <inheritdoc/>
+		public void Dispose()
+		{
+			_commandBuffer.Dispose();
+		}
+
+		/// <summary>
+		/// Attaches a new buffer for receiving data
+		/// </summary>
+		/// <param name="channelId">Channel id for the buffer</param>
+		/// <param name="buffer">Buffer to attach</param>
+		public void AttachRecvBuffer(int channelId, SharedMemoryBuffer buffer) => AttachBuffer(IpcMessage.AttachRecvBuffer, channelId, buffer);
+
+		/// <summary>
+		/// Attaches a new buffer for sending data
+		/// </summary>
+		/// <param name="channelId">Channel id for the buffer</param>
+		/// <param name="buffer">Buffer to attach</param>
+		public void AttachSendBuffer(int channelId, SharedMemoryBuffer buffer) => AttachBuffer(IpcMessage.AttachSendBuffer, channelId, buffer);
+
+		void AttachBuffer(IpcMessage message, int channelId, SharedMemoryBuffer buffer)
+		{
+			MemoryWriter writer = new MemoryWriter(_commandBuffer.Writer.GetMemory());
+			writer.WriteUnsignedVarInt((int)message);
+			writer.WriteUnsignedVarInt(channelId);
+			writer.WriteString(buffer.Name);
+			_commandBuffer.Writer.Advance(writer.Length);
+		}
+	}
+
+	/// <summary>
+	/// Extension methods for worker sockets
+	/// </summary>
+	public static class WorkerComputeSocketExtensions
+	{
+		/// <summary>
+		/// Creates a channel using a socket and receive buffer
+		/// </summary>
+		/// <param name="socket">Socket to use for sending data</param>
+		/// <param name="channelId">Channel id to send and receive data</param>
+		/// <param name="recvBuffer">Buffer for receiving data</param>
+		/// <param name="sendBuffer">Buffer for sending data</param>
+		public static IComputeChannel CreateChannel(this WorkerComputeSocket socket, int channelId, SharedMemoryBuffer recvBuffer, SharedMemoryBuffer sendBuffer)
+		{
+			socket.AttachRecvBuffer(channelId, recvBuffer);
+			socket.AttachSendBuffer(channelId, sendBuffer);
+			return new ComputeChannel.BufferedReaderWriterChannel(recvBuffer, sendBuffer);
+		}
+	}
+}
