Horde: Update timestamps for files retrieved from the cache. Tools like UHT rely on timestamps to be able to detect modified files, and moving files does not update them by default.

diff --git a/Engine/Source/Programs/Shared/EpicGames.Perforce.Managed/ManagedWorkspace.cs b/Engine/Source/Programs/Shared/EpicGames.Perforce.Managed/ManagedWorkspace.cs
--- a/Engine/Source/Programs/Shared/EpicGames.Perforce.Managed/ManagedWorkspace.cs
+++ b/Engine/Source/Programs/Shared/EpicGames.Perforce.Managed/ManagedWorkspace.cs
@@ -2035,7 +2035,16 @@ namespace EpicGames.Perforce.Managed
 		{
 			try
 			{
-				FileReference.Move(fileToMove._trackedFile.GetLocation(), fileToMove._workspaceFile.GetLocation());
+				FileReference targetFile = fileToMove._workspaceFile.GetLocation();
+				FileReference.Move(fileToMove._trackedFile.GetLocation(), targetFile);
+				try
+				{ 
+					FileReference.SetLastWriteTimeUtc(targetFile, DateTime.UtcNow); 
+				} 
+				catch 
+				{
+					_logger.LogWarning("Unable to update timestamp on {TargetFile}", targetFile);
+				}
 			}
 			catch (Exception ex)
 			{
