EpicGames.Perforce: Add undo integration actions into the IntegrationAction enum, so UGS doesn't choke if the Perforce server returns them.

diff --git a/Engine/Source/Programs/Shared/EpicGames.Perforce/Records/IntegrationAction.cs b/Engine/Source/Programs/Shared/EpicGames.Perforce/Records/IntegrationAction.cs
--- a/Engine/Source/Programs/Shared/EpicGames.Perforce/Records/IntegrationAction.cs
+++ b/Engine/Source/Programs/Shared/EpicGames.Perforce/Records/IntegrationAction.cs
@@ -102,5 +102,17 @@ namespace EpicGames.Perforce
 		/// </summary>
 		[PerforceEnum("add into")]
 		AddInto,
+
+		/// <summary>
+		/// The resulting revision has "undid" the edited changes (content and history) from the previously submitted revision(s)
+		/// </summary>
+		[PerforceEnum("undid")]
+		Undid,
+
+		/// <summary>
+		/// The previous revision(s) have been "undone by" the resulting revision, known as the "undone" revision.
+		/// </summary>
+		[PerforceEnum("undone by")]
+		UndoneBy
 	}
 }
