Removed the code analysis on build option in C# UHT project.

diff --git a/Engine/Source/Programs/Shared/EpicGames.UHT/EpicGames.UHT.csproj b/Engine/Source/Programs/Shared/EpicGames.UHT/EpicGames.UHT.csproj
--- a/Engine/Source/Programs/Shared/EpicGames.UHT/EpicGames.UHT.csproj
+++ b/Engine/Source/Programs/Shared/EpicGames.UHT/EpicGames.UHT.csproj
@@ -13,7 +13,7 @@
     <AllowUnsafeBlocks>true</AllowUnsafeBlocks>
     <GenerateDocumentationFile>True</GenerateDocumentationFile>
     <DebugType>pdbonly</DebugType>
-    <EnforceCodeStyleInBuild>true</EnforceCodeStyleInBuild>
+    <EnforceCodeStyleInBuild>False</EnforceCodeStyleInBuild>
     <ProduceReferenceAssemblyInOutDir>true</ProduceReferenceAssemblyInOutDir>
     <RunAnalyzersDuringBuild>false</RunAnalyzersDuringBuild>
   </PropertyGroup>
