Merging Horde and OIDC changes from //UE5/Main.

diff --git a/Engine/Source/Programs/Shared/EpicGames.Horde/Compute/Clients/ServerComputeClient.cs b/Engine/Source/Programs/Shared/EpicGames.Horde/Compute/Clients/ServerComputeClient.cs
--- a/Engine/Source/Programs/Shared/EpicGames.Horde/Compute/Clients/ServerComputeClient.cs
+++ b/Engine/Source/Programs/Shared/EpicGames.Horde/Compute/Clients/ServerComputeClient.cs
@@ -28,7 +28,7 @@ namespace EpicGames.Horde.Compute.Clients
 		public const int NonceLength = 64;
 
 		class AssignComputeRequest
-		{
+	{
 			public Requirements? Requirements { get; set; }
 		}
 
@@ -45,7 +45,7 @@ namespace EpicGames.Horde.Compute.Clients
 		record class LeaseInfo(IReadOnlyList<string> Properties, IReadOnlyDictionary<string, int> AssignedResources, IComputeSocket Socket);
 
 		class LeaseImpl : IComputeLease
-		{
+			{
 			readonly IAsyncEnumerator<LeaseInfo> _source;
 
 			public IReadOnlyList<string> Properties => _source.Current.Properties;
@@ -62,11 +62,11 @@ namespace EpicGames.Horde.Compute.Clients
 			{
 				await _source.MoveNextAsync();
 				await _source.DisposeAsync();
-			}
+		}
 
 			/// <inheritdoc/>
 			public ValueTask CloseAsync(CancellationToken cancellationToken) => Socket.CloseAsync(cancellationToken);
-		}
+			}
 
 		readonly HttpClient? _defaultHttpClient;
 		readonly Func<HttpClient> _createHttpClient;
@@ -80,7 +80,7 @@ namespace EpicGames.Horde.Compute.Clients
 		/// <param name="authHeader">Authentication header</param>
 		/// <param name="logger">Logger for diagnostic messages</param>
 		public ServerComputeClient(Uri serverUri, AuthenticationHeaderValue? authHeader, ILogger logger)
-		{
+			{
 #pragma warning disable CA2000 // Dispose objects before losing scope
 			// This is disposed via HttpClient
 			SocketsHttpHandler handler = new SocketsHttpHandler();
@@ -93,7 +93,7 @@ namespace EpicGames.Horde.Compute.Clients
 
 			_createHttpClient = GetDefaultHttpClient;
 			_logger = logger;
-		}
+			}
 
 		/// <summary>
 		/// Constructor
@@ -131,7 +131,7 @@ namespace EpicGames.Horde.Compute.Clients
 		{
 			IAsyncEnumerator<LeaseInfo> source = ConnectAsync(clusterId, requirements, cancellationToken).GetAsyncEnumerator(cancellationToken);
 			if (!await source.MoveNextAsync())
-			{
+		{
 				await source.DisposeAsync();
 				return null;
 			}
@@ -144,28 +144,28 @@ namespace EpicGames.Horde.Compute.Clients
 			_logger.LogInformation("Requesting compute resource");
 
 			// Assign a compute worker
-			HttpClient client = _createHttpClient();
+				HttpClient client = _createHttpClient();
 
 			AssignComputeRequest request = new AssignComputeRequest();
-			request.Requirements = requirements;
+				request.Requirements = requirements;
 
 			AssignComputeResponse? responseMessage;
-			using (HttpResponseMessage response = await client.PostAsync($"api/v2/compute/{clusterId}", request, _cancellationSource.Token))
-			{
-				if (response.StatusCode == HttpStatusCode.ServiceUnavailable)
+				using (HttpResponseMessage response = await client.PostAsync($"api/v2/compute/{clusterId}", request, _cancellationSource.Token))
 				{
+				if (response.StatusCode == HttpStatusCode.ServiceUnavailable)
+			{
 					_logger.LogInformation("No compute resource is available.");
 					yield break;
-				}
+			}
 
 				response.EnsureSuccessStatusCode();
 
 				responseMessage = await response.Content.ReadFromJsonAsync<AssignComputeResponse>(cancellationToken: cancellationToken);
 				if (responseMessage == null)
-				{
+			{
 					throw new InvalidOperationException();
-				}
 			}
+		}
 
 			_logger.LogInformation("Connecting to {Ip} with nonce {Nonce}...", responseMessage.Ip, responseMessage.Nonce);
 
