Expose accessors for any sparse data super structs

This allows a class to change from using FBaseSparseData to FDerivedSparseData without having to fix-up all the existing accessors for GetBaseSparseData()

diff --git a/Engine/Source/Programs/Shared/EpicGames.UHT/Exporters/CodeGen/UhtHeaderCodeGenerator.cs b/Engine/Source/Programs/Shared/EpicGames.UHT/Exporters/CodeGen/UhtHeaderCodeGenerator.cs
--- a/Engine/Source/Programs/Shared/EpicGames.UHT/Exporters/CodeGen/UhtHeaderCodeGenerator.cs
+++ b/Engine/Source/Programs/Shared/EpicGames.UHT/Exporters/CodeGen/UhtHeaderCodeGenerator.cs
@@ -1,6 +1,7 @@
 // Copyright Epic Games, Inc. All Rights Reserved.
 
 using System;
+using System.Collections.Generic;
 using System.Text;
 using EpicGames.Core;
 using EpicGames.UHT.Types;
@@ -530,6 +531,49 @@ namespace EpicGames.UHT.Exporters.CodeGen
 			}
 		}
 
+		/// <summary>
+		/// A list of sparse data structs that should export accessors for this class.
+		/// This list excludes anything that has already been exported by the super class.
+		/// </summary>
+		/// <param name="classObj">Class in question</param>
+		/// <returns>List of structs</returns>
+		protected static List<UhtScriptStruct> GetSparseDataStructsToExport(UhtClass classObj)
+		{
+			List<UhtScriptStruct> sparseScriptStructs = new List<UhtScriptStruct>();
+
+			string[]? sparseDataTypes = classObj.MetaData.GetStringArray(UhtNames.SparseClassDataTypes);
+			if (sparseDataTypes != null)
+			{
+				List<UhtScriptStruct> baseSparseScriptStructs = new List<UhtScriptStruct>();
+				{
+					string[]? baseSparseDataTypes = classObj.SuperClass?.MetaData.GetStringArray(UhtNames.SparseClassDataTypes);
+					if (baseSparseDataTypes != null)
+					{
+						foreach (string baseSparseDataType in baseSparseDataTypes)
+						{
+							UhtScriptStruct? baseSparseScriptStruct = classObj.FindType(UhtFindOptions.EngineName | UhtFindOptions.ScriptStruct | UhtFindOptions.NoSelf, baseSparseDataType) as UhtScriptStruct;
+							if (baseSparseScriptStruct != null)
+							{
+								baseSparseScriptStructs.Add(baseSparseScriptStruct);
+							}
+						}
+					}
+				}
+
+				foreach (string sparseDataType in sparseDataTypes)
+				{
+					UhtScriptStruct? sparseScriptStruct = classObj.FindType(UhtFindOptions.EngineName | UhtFindOptions.ScriptStruct | UhtFindOptions.NoSelf, sparseDataType) as UhtScriptStruct;
+					while (sparseScriptStruct != null && !baseSparseScriptStructs.Contains(sparseScriptStruct))
+					{
+						sparseScriptStructs.Add(sparseScriptStruct);
+						sparseScriptStruct = sparseScriptStruct.SuperScriptStruct;
+					}
+				}
+			}
+
+			return sparseScriptStructs;
+		}
+
 		protected static string GetDelegateFunctionExportName(UhtFunction function)
 		{
 			const string DelegatePrefix = "delegate";
diff --git a/Engine/Source/Programs/Shared/EpicGames.UHT/Exporters/CodeGen/UhtHeaderCodeGeneratorCppFile.cs b/Engine/Source/Programs/Shared/EpicGames.UHT/Exporters/CodeGen/UhtHeaderCodeGeneratorCppFile.cs
--- a/Engine/Source/Programs/Shared/EpicGames.UHT/Exporters/CodeGen/UhtHeaderCodeGeneratorCppFile.cs
+++ b/Engine/Source/Programs/Shared/EpicGames.UHT/Exporters/CodeGen/UhtHeaderCodeGeneratorCppFile.cs
@@ -1298,20 +1298,12 @@ namespace EpicGames.UHT.Exporters.CodeGen
 
 		private static StringBuilder AppendSparseAccessors(StringBuilder builder, UhtClass classObj)
 		{
-			string[]? sparseDataTypes = classObj.MetaData.GetStringArray(UhtNames.SparseClassDataTypes);
-			if (sparseDataTypes == null)
-			{
-				return builder;
-			}
+			List<UhtScriptStruct> sparseScriptStructs = GetSparseDataStructsToExport(classObj);
 
-			string[]? baseSparseDataTypes = classObj.SuperClass?.MetaData.GetStringArray(UhtNames.SparseClassDataTypes);
-			if (baseSparseDataTypes != null && Enumerable.SequenceEqual(sparseDataTypes, baseSparseDataTypes))
+			foreach (UhtScriptStruct sparseScriptStruct in sparseScriptStructs)
 			{
-				return builder;
-			}
+				string sparseDataType = sparseScriptStruct.EngineName;
 
-			foreach (string sparseDataType in sparseDataTypes)
-			{
 				builder.Append('F').Append(sparseDataType).Append("* ").Append(classObj.SourceName).Append("::Get").Append(sparseDataType).Append("() \r\n");
 				builder.Append("{ \r\n");
 				builder.Append("\treturn static_cast<F").Append(sparseDataType).Append("*>(GetClass()->GetOrCreateSparseClassData()); \r\n");
diff --git a/Engine/Source/Programs/Shared/EpicGames.UHT/Exporters/CodeGen/UhtHeaderCodeGeneratorHFile.cs b/Engine/Source/Programs/Shared/EpicGames.UHT/Exporters/CodeGen/UhtHeaderCodeGeneratorHFile.cs
--- a/Engine/Source/Programs/Shared/EpicGames.UHT/Exporters/CodeGen/UhtHeaderCodeGeneratorHFile.cs
+++ b/Engine/Source/Programs/Shared/EpicGames.UHT/Exporters/CodeGen/UhtHeaderCodeGeneratorHFile.cs
@@ -632,68 +632,63 @@ namespace EpicGames.UHT.Exporters.CodeGen
 			// Format the sparse data
 			using (UhtMacroCreator macro = new(builder, this, classObj, SparseDataMacroSuffix))
 			{
-				string[]? sparseDataTypes = classObj.MetaData.GetStringArray(UhtNames.SparseClassDataTypes);
-				if (sparseDataTypes != null)
+				List<UhtScriptStruct> sparseScriptStructs = GetSparseDataStructsToExport(classObj);
+
+				if (sparseScriptStructs.Count > 0)
 				{
-					string[]? baseSparseDataTypes = classObj.SuperClass?.MetaData.GetStringArray(UhtNames.SparseClassDataTypes);
-					if (baseSparseDataTypes == null || !Enumerable.SequenceEqual(sparseDataTypes, baseSparseDataTypes))
+					string api = classObj.ClassFlags.HasAnyFlags(EClassFlags.MinimalAPI) ? PackageApi : "";
+
+					foreach (UhtScriptStruct sparseScriptStruct in sparseScriptStructs)
 					{
-						string api = classObj.ClassFlags.HasAnyFlags(EClassFlags.MinimalAPI) ? PackageApi : "";
+						string sparseDataType = sparseScriptStruct.EngineName;
 
-						foreach (string sparseDataType in sparseDataTypes)
-						{
-							builder.Append(api).Append('F').Append(sparseDataType).Append("* Get").Append(sparseDataType).Append("(); \\\r\n");
-							builder.Append(api).Append('F').Append(sparseDataType).Append("* Get").Append(sparseDataType).Append("() const; \\\r\n");
-							builder.Append(api).Append("const F").Append(sparseDataType).Append("* Get").Append(sparseDataType).Append("(EGetSparseClassDataMethod GetMethod) const; \\\r\n");
-							builder.Append(api).Append("static UScriptStruct* StaticGet").Append(sparseDataType).Append("ScriptStruct(); \\\r\n");
-						}
+						builder.Append(api).Append('F').Append(sparseDataType).Append("* Get").Append(sparseDataType).Append("(); \\\r\n");
+						builder.Append(api).Append('F').Append(sparseDataType).Append("* Get").Append(sparseDataType).Append("() const; \\\r\n");
+						builder.Append(api).Append("const F").Append(sparseDataType).Append("* Get").Append(sparseDataType).Append("(EGetSparseClassDataMethod GetMethod) const; \\\r\n");
+						builder.Append(api).Append("static UScriptStruct* StaticGet").Append(sparseDataType).Append("ScriptStruct(); \\\r\n");
 					}
+				}
+
+				foreach (UhtScriptStruct sparseScriptStruct in sparseScriptStructs)
+				{
+					string sparseDataType = sparseScriptStruct.EngineName;
 
-					foreach (string sparseDataType in sparseDataTypes)
+					foreach (UhtProperty sparseProperty in sparseScriptStruct.Properties)
 					{
-						UhtScriptStruct? sparseScriptStruct = classObj.FindType(UhtFindOptions.EngineName | UhtFindOptions.ScriptStruct | UhtFindOptions.NoSelf, sparseDataType) as UhtScriptStruct;
-						while (sparseScriptStruct != null)
+						if (!sparseProperty.MetaData.ContainsKey(UhtNames.NoGetter))
 						{
-							foreach (UhtProperty sparseProperty in sparseScriptStruct.Properties)
+							string propertyName = sparseProperty.SourceName;
+							string cleanPropertyName = propertyName;
+							if (sparseProperty is UhtBoolProperty && propertyName.StartsWith("b", StringComparison.Ordinal))
 							{
-								if (!sparseProperty.MetaData.ContainsKey(UhtNames.NoGetter))
-								{
-									string propertyName = sparseProperty.SourceName;
-									string cleanPropertyName = propertyName;
-									if (sparseProperty is UhtBoolProperty && propertyName.StartsWith("b", StringComparison.Ordinal))
-									{
-										cleanPropertyName = propertyName[1..];
-									}
-
-									bool getByRef = sparseProperty.MetaData.ContainsKey(UhtNames.GetByRef);
+								cleanPropertyName = propertyName[1..];
+							}
 
-									if (getByRef)
-									{
-										builder.Append("const ").AppendSparse(sparseProperty).Append("& Get").Append(cleanPropertyName).Append("() \\\r\n");
-									}
-									else
-									{
-										builder.AppendSparse(sparseProperty).Append(" Get").Append(cleanPropertyName).Append("() \\\r\n");
-									}
-									builder.Append("{ \\\r\n");
-									builder.Append("\treturn Get").Append(sparseDataType).Append("()->").Append(propertyName).Append("; \\\r\n");
-									builder.Append("} \\\r\n");
+							bool getByRef = sparseProperty.MetaData.ContainsKey(UhtNames.GetByRef);
 
-									if (getByRef)
-									{
-										builder.Append("const ").AppendSparse(sparseProperty).Append("& Get").Append(cleanPropertyName).Append("() const \\\r\n");
-									}
-									else
-									{
-										builder.AppendSparse(sparseProperty).Append(" Get").Append(cleanPropertyName).Append("() const \\\r\n");
-									}
-									builder.Append("{ \\\r\n");
-									builder.Append("\treturn Get").Append(sparseDataType).Append("()->").Append(propertyName).Append("; \\\r\n");
-									builder.Append("} \\\r\n");
-								}
+							if (getByRef)
+							{
+								builder.Append("const ").AppendSparse(sparseProperty).Append("& Get").Append(cleanPropertyName).Append("() \\\r\n");
+							}
+							else
+							{
+								builder.AppendSparse(sparseProperty).Append(" Get").Append(cleanPropertyName).Append("() \\\r\n");
 							}
+							builder.Append("{ \\\r\n");
+							builder.Append("\treturn Get").Append(sparseDataType).Append("()->").Append(propertyName).Append("; \\\r\n");
+							builder.Append("} \\\r\n");
 
-							sparseScriptStruct = sparseScriptStruct.SuperScriptStruct;
+							if (getByRef)
+							{
+								builder.Append("const ").AppendSparse(sparseProperty).Append("& Get").Append(cleanPropertyName).Append("() const \\\r\n");
+							}
+							else
+							{
+								builder.AppendSparse(sparseProperty).Append(" Get").Append(cleanPropertyName).Append("() const \\\r\n");
+							}
+							builder.Append("{ \\\r\n");
+							builder.Append("\treturn Get").Append(sparseDataType).Append("()->").Append(propertyName).Append("; \\\r\n");
+							builder.Append("} \\\r\n");
 						}
 					}
 				}
