Removed non-const overloads of the sparse data accessors

These always returned the same types and had the same implementation as the const overload

#rnx

diff --git a/Engine/Source/Programs/Shared/EpicGames.UHT/Exporters/CodeGen/UhtHeaderCodeGeneratorCppFile.cs b/Engine/Source/Programs/Shared/EpicGames.UHT/Exporters/CodeGen/UhtHeaderCodeGeneratorCppFile.cs
--- a/Engine/Source/Programs/Shared/EpicGames.UHT/Exporters/CodeGen/UhtHeaderCodeGeneratorCppFile.cs
+++ b/Engine/Source/Programs/Shared/EpicGames.UHT/Exporters/CodeGen/UhtHeaderCodeGeneratorCppFile.cs
@@ -1304,11 +1304,6 @@ namespace EpicGames.UHT.Exporters.CodeGen
 			{
 				string sparseDataType = sparseScriptStruct.EngineName;
 
-				builder.Append('F').Append(sparseDataType).Append("* ").Append(classObj.SourceName).Append("::Get").Append(sparseDataType).Append("() \r\n");
-				builder.Append("{ \r\n");
-				builder.Append("\treturn static_cast<F").Append(sparseDataType).Append("*>(GetClass()->GetOrCreateSparseClassData()); \r\n");
-				builder.Append("} \r\n");
-
 				builder.Append('F').Append(sparseDataType).Append("* ").Append(classObj.SourceName).Append("::Get").Append(sparseDataType).Append("() const \r\n");
 				builder.Append("{ \r\n");
 				builder.Append("\treturn static_cast<F").Append(sparseDataType).Append("*>(GetClass()->GetOrCreateSparseClassData()); \r\n");
diff --git a/Engine/Source/Programs/Shared/EpicGames.UHT/Exporters/CodeGen/UhtHeaderCodeGeneratorHFile.cs b/Engine/Source/Programs/Shared/EpicGames.UHT/Exporters/CodeGen/UhtHeaderCodeGeneratorHFile.cs
--- a/Engine/Source/Programs/Shared/EpicGames.UHT/Exporters/CodeGen/UhtHeaderCodeGeneratorHFile.cs
+++ b/Engine/Source/Programs/Shared/EpicGames.UHT/Exporters/CodeGen/UhtHeaderCodeGeneratorHFile.cs
@@ -642,7 +642,6 @@ namespace EpicGames.UHT.Exporters.CodeGen
 					{
 						string sparseDataType = sparseScriptStruct.EngineName;
 
-						builder.Append(api).Append('F').Append(sparseDataType).Append("* Get").Append(sparseDataType).Append("(); \\\r\n");
 						builder.Append(api).Append('F').Append(sparseDataType).Append("* Get").Append(sparseDataType).Append("() const; \\\r\n");
 						builder.Append(api).Append("const F").Append(sparseDataType).Append("* Get").Append(sparseDataType).Append("(EGetSparseClassDataMethod GetMethod) const; \\\r\n");
 						builder.Append(api).Append("static UScriptStruct* StaticGet").Append(sparseDataType).Append("ScriptStruct(); \\\r\n");
@@ -666,18 +665,6 @@ namespace EpicGames.UHT.Exporters.CodeGen
 
 							bool getByRef = sparseProperty.MetaData.ContainsKey(UhtNames.GetByRef);
 
-							if (getByRef)
-							{
-								builder.Append("const ").AppendSparse(sparseProperty).Append("& Get").Append(cleanPropertyName).Append("() \\\r\n");
-							}
-							else
-							{
-								builder.AppendSparse(sparseProperty).Append(" Get").Append(cleanPropertyName).Append("() \\\r\n");
-							}
-							builder.Append("{ \\\r\n");
-							builder.Append("\treturn Get").Append(sparseDataType).Append("()->").Append(propertyName).Append("; \\\r\n");
-							builder.Append("} \\\r\n");
-
 							if (getByRef)
 							{
 								builder.Append("const ").AppendSparse(sparseProperty).Append("& Get").Append(cleanPropertyName).Append("() const \\\r\n");
