EpicGames.Core: Add support for more expression types into the csproj parser to fix build error.

diff --git a/Engine/Source/Programs/Shared/EpicGames.Core/CsProjectInfo.cs b/Engine/Source/Programs/Shared/EpicGames.Core/CsProjectInfo.cs
--- a/Engine/Source/Programs/Shared/EpicGames.Core/CsProjectInfo.cs
+++ b/Engine/Source/Programs/Shared/EpicGames.Core/CsProjectInfo.cs
@@ -15,7 +15,7 @@ namespace EpicGames.Core
 	/// </summary>
 	public sealed class CsProjectParseException : Exception
 	{
-		internal CsProjectParseException(string? message, Exception? exception) : base(message, exception)
+		internal CsProjectParseException(string? message, Exception? exception = null) : base(message, exception)
 		{
 		}
 	}
@@ -707,6 +707,8 @@ namespace EpicGames.Core
 			}
 		}
 
+		public record class ConditionContext(DirectoryReference BaseDir);
+
 		/// <summary>
 		/// Evaluate whether the optional MSBuild condition on an XML element evaluates to true. Currently only supports 'ABC' == 'DEF' style expressions, but can be expanded as needed.
 		/// </summary>
@@ -735,44 +737,119 @@ namespace EpicGames.Core
 			// Tokenize the condition
 			string[] tokens = Tokenize(condition);
 
-			char[] tokenQuotes = new[] { '\'', '(', ')', '{', '}', '[', ']' };
-			
 			// Try to evaluate it. We only support a very limited class of condition expressions at the moment, but it's enough to parse standard projects
-			bool bResult;
+			int tokenIdx = 0;
 
-			// Handle Exists('Platform\Windows\Gauntlet.TargetDeviceWindows.cs')
-			if (tokens[0] == "Exists")
+			ConditionContext context = new ConditionContext(projectInfo.ProjectPath.Directory);
+			try
 			{
-				// remove all quotes, apostrophes etc that are either tokens or wrap tokens (The Tokenize() function is a bit suspect).
-				string[] arguments = tokens.Select(s => s.Trim(tokenQuotes)).Where(s => s.Length > 0).ToArray();
-
-				if (tokens.Length > 1)
+				bool bResult = CoerceToBool(EvaluateLogicalAnd(context, tokens, ref tokenIdx));
+				if (tokenIdx != tokens.Length)
 				{
-					FileSystemReference dependency = DirectoryReference.Combine(projectInfo.ProjectPath.Directory, arguments[1]);
+					throw new CsProjectParseException("Unexpected tokens at end of condition");
+				}
+				return bResult;
+			}
+			catch (CsProjectParseException Ex)
+			{
+				throw new CsProjectParseException($"{Ex.Message} while parsing {element} in project file {projectInfo.ProjectPath}", Ex);
+			}
+		}
 
-					if (File.Exists(dependency.FullName) || Directory.Exists(dependency.FullName))
-					{
-						return true;
-					}
+		static string EvaluateLogicalAnd(ConditionContext context, string[] tokens, ref int tokenIdx)
+		{
+			string result = EvaluateLogicalOr(context, tokens, ref tokenIdx);
+			while (tokenIdx < tokens.Length && tokens[tokenIdx].Equals("And", StringComparison.OrdinalIgnoreCase))
+			{
+				tokenIdx++;
+				string rhs = EvaluateEquality(context, tokens, ref tokenIdx);
+				result = (CoerceToBool(result) && CoerceToBool(rhs)).ToString();
+			}
+			return result;
+		}
+
+		static string EvaluateLogicalOr(ConditionContext context, string[] tokens, ref int tokenIdx)
+		{
+			string result = EvaluateEquality(context, tokens, ref tokenIdx);
+			while (tokenIdx < tokens.Length && tokens[tokenIdx].Equals("Or", StringComparison.OrdinalIgnoreCase))
+			{
+				tokenIdx++;
+				string rhs = EvaluateEquality(context, tokens, ref tokenIdx);
+				result = (CoerceToBool(result) || CoerceToBool(rhs)).ToString();
+			}
+			return result;
+		}
 
-					return false;
+		static string EvaluateEquality(ConditionContext context, string[] tokens, ref int tokenIdx)
+		{
+			// Otherwise try to parse an equality or inequality expression
+			string lhs = EvaluateValue(context, tokens, ref tokenIdx);
+			if (tokenIdx < tokens.Length)
+			{
+				if (tokens[tokenIdx] == "==")
+				{
+					tokenIdx++;
+					string rhs = EvaluateValue(context, tokens, ref tokenIdx);
+					return lhs.Equals(rhs, StringComparison.OrdinalIgnoreCase).ToString();
+				}
+				else if (tokens[tokenIdx] == "!=")
+				{
+					tokenIdx++;
+					string rhs = EvaluateValue(context, tokens, ref tokenIdx);
+					return lhs.Equals(rhs, StringComparison.OrdinalIgnoreCase).ToString();
 				}
 			}
+			return lhs;
+		}
 
-			if (tokens.Length == 3 && tokens[0].StartsWith("'", StringComparison.Ordinal) && tokens[1] == "==" && tokens[2].StartsWith("'", StringComparison.Ordinal))
+		static string EvaluateValue(ConditionContext context, string[] tokens, ref int tokenIdx)
+		{
+			// Handle Exists('Platform\Windows\Gauntlet.TargetDeviceWindows.cs')
+			if (tokens[tokenIdx].Equals("Exists", StringComparison.OrdinalIgnoreCase))
 			{
-				bResult = String.Equals(tokens[0], tokens[2], StringComparison.OrdinalIgnoreCase);
+				if (tokenIdx + 3 >= tokens.Length || !tokens[tokenIdx + 1].Equals("(", StringComparison.Ordinal) || !tokens[tokenIdx + 3].Equals(")", StringComparison.Ordinal))
+				{
+					throw new CsProjectParseException("Invalid 'Exists' expression", null);
+				}
+
+				// remove all quotes, apostrophes etc that are either tokens or wrap tokens (The Tokenize() function is a bit suspect).
+				string path = tokens[tokenIdx + 2].Trim('\'', '(', ')', '{', '}', '[', ']');
+				tokenIdx += 4;
+
+				FileSystemReference dependency = DirectoryReference.Combine(context.BaseDir, path);
+				bool exists = File.Exists(dependency.FullName) || Directory.Exists(dependency.FullName);
+				return exists.ToString();
 			}
-			else if (tokens.Length == 3 && tokens[0].StartsWith("'", StringComparison.Ordinal) && tokens[1] == "!=" && tokens[2].StartsWith("'", StringComparison.Ordinal))
+
+			// Handle negation
+			if (tokens[tokenIdx].Equals("!", StringComparison.Ordinal))
 			{
-				bResult = !String.Equals(tokens[0], tokens[2], StringComparison.OrdinalIgnoreCase);
+				tokenIdx++;
+				bool value = CoerceToBool(EvaluateValue(context, tokens, ref tokenIdx));
+				return (!value).ToString();
 			}
-			else
+
+			// Handle subexpressions
+			if (tokens[tokenIdx].Equals("(", StringComparison.Ordinal))
 			{
-				string msg = String.Format("Couldn't parse condition {0} in project file {1}", element.ToString(), projectInfo.ProjectPath);
-				throw new Exception(msg);
+				tokenIdx++;
+
+				string result = EvaluateLogicalAnd(context, tokens, ref tokenIdx);
+				if (!tokens[tokenIdx].Equals(")", StringComparison.Ordinal))
+				{
+					throw new CsProjectParseException("Missing ')'", null);
+				}
+				tokenIdx++;
+
+				return result;
 			}
-			return bResult;
+
+			return tokens[tokenIdx++];
+		}
+
+		static bool CoerceToBool(string value)
+		{
+			return !value.Equals("false", StringComparison.OrdinalIgnoreCase) && !value.Equals("0", StringComparison.Ordinal);
 		}
 
 		/// <summary>
@@ -1000,7 +1077,7 @@ namespace EpicGames.Core
 
 					methodEndIdx = idx;
 
-					if (condition[idx] == '(')
+					if (idx < condition.Length && condition[idx] == '(')
 					{
 						// a method invoke
 						for (; ; idx++)
@@ -1021,10 +1098,10 @@ namespace EpicGames.Core
 								break;
 							}
 						}
+						idx++;
 					}
 #pragma warning restore IDE0059 // Unnecessary assignment of a value
 
-					idx++;
 					tokens.Add(condition.Substring(startIdx, idx - startIdx));
 				}
 				else if(Char.IsLetterOrDigit(condition[idx]) || condition[idx] == '_')
