Fix warning about public type without documentation.

diff --git a/Engine/Source/Programs/Shared/EpicGames.Core/CsProjectInfo.cs b/Engine/Source/Programs/Shared/EpicGames.Core/CsProjectInfo.cs
--- a/Engine/Source/Programs/Shared/EpicGames.Core/CsProjectInfo.cs
+++ b/Engine/Source/Programs/Shared/EpicGames.Core/CsProjectInfo.cs
@@ -707,7 +707,7 @@ namespace EpicGames.Core
 			}
 		}
 
-		public record class ConditionContext(DirectoryReference BaseDir);
+		record class ConditionContext(DirectoryReference BaseDir);
 
 		/// <summary>
 		/// Evaluate whether the optional MSBuild condition on an XML element evaluates to true. Currently only supports 'ABC' == 'DEF' style expressions, but can be expanded as needed.
