Horde: Fix exception due to incorrect cast.

diff --git a/Engine/Source/Programs/Shared/EpicGames.Horde/Compute/DefaultComputeSocket.cs b/Engine/Source/Programs/Shared/EpicGames.Horde/Compute/DefaultComputeSocket.cs
--- a/Engine/Source/Programs/Shared/EpicGames.Horde/Compute/DefaultComputeSocket.cs
+++ b/Engine/Source/Programs/Shared/EpicGames.Horde/Compute/DefaultComputeSocket.cs
@@ -165,9 +165,9 @@ namespace EpicGames.Horde.Compute
 			lock (_lockObject)
 			{
 				_complete = true;
-				foreach (IComputeBufferWriter writer in _recvBuffers.Values)
+				foreach (IComputeBuffer buffer in _recvBuffers.Values)
 				{
-					writer.MarkComplete();
+					buffer.Writer.MarkComplete();
 				}
 			}
 
