Horde: Fix chunked tail data not correctly being delimited with newline characters when sent through the tailing endpoint.

diff --git a/Engine/Source/Programs/Shared/EpicGames.Horde/Logs/LogChunkNode.cs b/Engine/Source/Programs/Shared/EpicGames.Horde/Logs/LogChunkNode.cs
--- a/Engine/Source/Programs/Shared/EpicGames.Horde/Logs/LogChunkNode.cs
+++ b/Engine/Source/Programs/Shared/EpicGames.Horde/Logs/LogChunkNode.cs
@@ -134,6 +134,13 @@ namespace EpicGames.Horde.Logs
 		/// <returns>Line at the given index</returns>
 		public Utf8String GetLine(int idx) => new Utf8String(Data.Slice(LineOffsets[idx], LineOffsets[idx + 1] - LineOffsets[idx] - 1));
 
+		/// <summary>
+		/// Accessor for an individual line, including the trailing newline character
+		/// </summary>
+		/// <param name="idx">Index of the line to retrieve</param>
+		/// <returns>Line at the given index</returns>
+		public Utf8String GetLineWithNewline(int idx) => new Utf8String(Data.Slice(LineOffsets[idx], LineOffsets[idx + 1] - LineOffsets[idx]));
+
 		/// <summary>
 		/// Find the line index for a particular offset
 		/// </summary>
@@ -679,7 +686,7 @@ namespace EpicGames.Horde.Logs
 			{
 				for (; lineIdx < chunk.LineCount; lineIdx++)
 				{
-					yield return chunk.GetLine(lineIdx);
+					yield return chunk.GetLineWithNewline(lineIdx);
 				}
 				lineIdx -= chunk.LineCount;
 			}
