Updated UHT for sparse class data accessors to use EGetSparseClassDataMethod::ArchetypeIfNull.

diff --git a/Engine/Source/Programs/Shared/EpicGames.UHT/Exporters/CodeGen/UhtHeaderCodeGeneratorHFile.cs b/Engine/Source/Programs/Shared/EpicGames.UHT/Exporters/CodeGen/UhtHeaderCodeGeneratorHFile.cs
--- a/Engine/Source/Programs/Shared/EpicGames.UHT/Exporters/CodeGen/UhtHeaderCodeGeneratorHFile.cs
+++ b/Engine/Source/Programs/Shared/EpicGames.UHT/Exporters/CodeGen/UhtHeaderCodeGeneratorHFile.cs
@@ -674,7 +674,7 @@ namespace EpicGames.UHT.Exporters.CodeGen
 								builder.AppendSparse(sparseProperty).Append(" Get").Append(cleanPropertyName).Append("() const \\\r\n");
 							}
 							builder.Append("{ \\\r\n");
-							builder.Append("\treturn Get").Append(sparseDataType).Append("()->").Append(propertyName).Append("; \\\r\n");
+							builder.Append("\treturn Get").Append(sparseDataType).Append("(EGetSparseClassDataMethod::ArchetypeIfNull)->").Append(propertyName).Append("; \\\r\n");
 							builder.Append("} \\\r\n");
 						}
 					}
