Minor refactor of the sparse data structure collection method.

diff --git a/Engine/Source/Programs/Shared/EpicGames.UHT/Exporters/CodeGen/UhtHeaderCodeGenerator.cs b/Engine/Source/Programs/Shared/EpicGames.UHT/Exporters/CodeGen/UhtHeaderCodeGenerator.cs
--- a/Engine/Source/Programs/Shared/EpicGames.UHT/Exporters/CodeGen/UhtHeaderCodeGenerator.cs
+++ b/Engine/Source/Programs/Shared/EpicGames.UHT/Exporters/CodeGen/UhtHeaderCodeGenerator.cs
@@ -536,42 +536,23 @@ namespace EpicGames.UHT.Exporters.CodeGen
 		/// This list excludes anything that has already been exported by the super class.
 		/// </summary>
 		/// <param name="classObj">Class in question</param>
-		/// <returns>List of structs</returns>
-		protected static List<UhtScriptStruct> GetSparseDataStructsToExport(UhtClass classObj)
+		/// <returns>Enumeration of structs</returns>
+		protected static IEnumerable<UhtScriptStruct> GetSparseDataStructsToExport(UhtClass classObj)
 		{
-			List<UhtScriptStruct> sparseScriptStructs = new List<UhtScriptStruct>();
-
 			string[]? sparseDataTypes = classObj.MetaData.GetStringArray(UhtNames.SparseClassDataTypes);
 			if (sparseDataTypes != null)
 			{
-				List<UhtScriptStruct> baseSparseScriptStructs = new List<UhtScriptStruct>();
-				{
-					string[]? baseSparseDataTypes = classObj.SuperClass?.MetaData.GetStringArray(UhtNames.SparseClassDataTypes);
-					if (baseSparseDataTypes != null)
-					{
-						foreach (string baseSparseDataType in baseSparseDataTypes)
-						{
-							UhtScriptStruct? baseSparseScriptStruct = classObj.FindType(UhtFindOptions.EngineName | UhtFindOptions.ScriptStruct | UhtFindOptions.NoSelf, baseSparseDataType) as UhtScriptStruct;
-							if (baseSparseScriptStruct != null)
-							{
-								baseSparseScriptStructs.Add(baseSparseScriptStruct);
-							}
-						}
-					}
-				}
-
+				HashSet<string> baseSparseDataTypes = new(classObj.SuperClass?.MetaData.GetStringArray(UhtNames.SparseClassDataTypes) ?? Array.Empty<string>(), StringComparer.OrdinalIgnoreCase);
 				foreach (string sparseDataType in sparseDataTypes)
 				{
 					UhtScriptStruct? sparseScriptStruct = classObj.FindType(UhtFindOptions.EngineName | UhtFindOptions.ScriptStruct | UhtFindOptions.NoSelf, sparseDataType) as UhtScriptStruct;
-					while (sparseScriptStruct != null && !baseSparseScriptStructs.Contains(sparseScriptStruct))
+					while (sparseScriptStruct != null && !baseSparseDataTypes.Contains(sparseScriptStruct.EngineName))
 					{
-						sparseScriptStructs.Add(sparseScriptStruct);
+						yield return sparseScriptStruct;
 						sparseScriptStruct = sparseScriptStruct.SuperScriptStruct;
 					}
 				}
 			}
-
-			return sparseScriptStructs;
 		}
 
 		protected static string GetDelegateFunctionExportName(UhtFunction function)
diff --git a/Engine/Source/Programs/Shared/EpicGames.UHT/Exporters/CodeGen/UhtHeaderCodeGeneratorCppFile.cs b/Engine/Source/Programs/Shared/EpicGames.UHT/Exporters/CodeGen/UhtHeaderCodeGeneratorCppFile.cs
--- a/Engine/Source/Programs/Shared/EpicGames.UHT/Exporters/CodeGen/UhtHeaderCodeGeneratorCppFile.cs
+++ b/Engine/Source/Programs/Shared/EpicGames.UHT/Exporters/CodeGen/UhtHeaderCodeGeneratorCppFile.cs
@@ -1298,9 +1298,7 @@ namespace EpicGames.UHT.Exporters.CodeGen
 
 		private static StringBuilder AppendSparseAccessors(StringBuilder builder, UhtClass classObj)
 		{
-			List<UhtScriptStruct> sparseScriptStructs = GetSparseDataStructsToExport(classObj);
-
-			foreach (UhtScriptStruct sparseScriptStruct in sparseScriptStructs)
+			foreach (UhtScriptStruct sparseScriptStruct in GetSparseDataStructsToExport(classObj))
 			{
 				string sparseDataType = sparseScriptStruct.EngineName;
 
diff --git a/Engine/Source/Programs/Shared/EpicGames.UHT/Exporters/CodeGen/UhtHeaderCodeGeneratorHFile.cs b/Engine/Source/Programs/Shared/EpicGames.UHT/Exporters/CodeGen/UhtHeaderCodeGeneratorHFile.cs
--- a/Engine/Source/Programs/Shared/EpicGames.UHT/Exporters/CodeGen/UhtHeaderCodeGeneratorHFile.cs
+++ b/Engine/Source/Programs/Shared/EpicGames.UHT/Exporters/CodeGen/UhtHeaderCodeGeneratorHFile.cs
@@ -632,9 +632,9 @@ namespace EpicGames.UHT.Exporters.CodeGen
 			// Format the sparse data
 			using (UhtMacroCreator macro = new(builder, this, classObj, SparseDataMacroSuffix))
 			{
-				List<UhtScriptStruct> sparseScriptStructs = GetSparseDataStructsToExport(classObj);
+				IEnumerable<UhtScriptStruct> sparseScriptStructs = GetSparseDataStructsToExport(classObj);
 
-				if (sparseScriptStructs.Count > 0)
+				if (sparseScriptStructs.Any())
 				{
 					string api = classObj.ClassFlags.HasAnyFlags(EClassFlags.MinimalAPI) ? PackageApi : "";
 
