Horde: Reduce the number of small bundles created when archiving a mixture of small and large files.

diff --git a/Engine/Source/Programs/Shared/EpicGames.Horde/Storage/Nodes/DirectoryNode.cs b/Engine/Source/Programs/Shared/EpicGames.Horde/Storage/Nodes/DirectoryNode.cs
--- a/Engine/Source/Programs/Shared/EpicGames.Horde/Storage/Nodes/DirectoryNode.cs
+++ b/Engine/Source/Programs/Shared/EpicGames.Horde/Storage/Nodes/DirectoryNode.cs
@@ -683,7 +683,6 @@ namespace EpicGames.Horde.Storage.Nodes
 
 			// Compute the total size
 			long totalSize = files.Sum(x => x.Item2.Length);
-			long chunkSize = Math.Max(MinSizePerWriter, totalSize / MaxWriters);
 
 			// Create the progress reporting object
 			CopyStats? copyStats = null;
@@ -693,17 +692,17 @@ namespace EpicGames.Horde.Storage.Nodes
 			}
 
 			List<Task> tasks = new List<Task>();
-			long currentSize = 0;
-			long targetSize = chunkSize;
 			FileEntry[] fileEntries = new FileEntry[files.Count];
 
 			// Split it into separate writers
+			long remainingSize = totalSize;
 			for (int minIdx = 0; minIdx < files.Count; )
 			{
-				currentSize += files[minIdx].FileInfo.Length;
+				long chunkSize = Math.Max(MinSizePerWriter, remainingSize / Math.Max(1, MaxWriters - tasks.Count));
 
 				int maxIdx = minIdx + 1;
-				while (maxIdx < files.Count && currentSize + files[maxIdx].FileInfo.Length <= targetSize)
+				long currentSize = files[minIdx].FileInfo.Length;
+				while (maxIdx < files.Count && currentSize <= chunkSize)
 				{
 					currentSize += files[maxIdx].FileInfo.Length;
 					maxIdx++;
@@ -712,7 +711,7 @@ namespace EpicGames.Horde.Storage.Nodes
 				int minIdxCopy = minIdx;
 				tasks.Add(Task.Run(() => CopyFilesAsync(files, minIdxCopy, maxIdx, fileEntries, options, writer, copyStats, cancellationToken), cancellationToken));
 
-				targetSize += chunkSize;
+				remainingSize -= currentSize;
 				minIdx = maxIdx;
 			}
 
