Epicgames.Core: Move log backup to separate function

#rnx

diff --git a/Engine/Source/Programs/Shared/EpicGames.Core/Log.cs b/Engine/Source/Programs/Shared/EpicGames.Core/Log.cs
--- a/Engine/Source/Programs/Shared/EpicGames.Core/Log.cs
+++ b/Engine/Source/Programs/Shared/EpicGames.Core/Log.cs
@@ -205,60 +205,70 @@ namespace EpicGames.Core
 		}
 
 		/// <summary>
-		/// Adds a trace listener that writes to a log file.
-		/// If a StartupTraceListener was in use, this function will copy its captured data to the log file(s)
-		/// and remove the startup listener from the list of registered listeners.
+		/// Backup an existing log file if it already exists at the outputpath
 		/// </summary>
-		/// <param name="name">Identifier for the writer</param>
-		/// <param name="outputFile">The file to write to</param>
-		/// <returns>The created trace listener</returns>
-		public static void AddFileWriter(string name, FileReference outputFile)
+		/// <param name="outputFile">The file to back up</param>
+		public static void BackupLogFile(FileReference outputFile)
 		{
-			Logger.LogInformation("Log file: {OutputFile}", outputFile);
+			if (!Log.BackupLogFiles || !FileReference.Exists(outputFile))
+			{
+				return;
+			}
 
-			if (Log.BackupLogFiles && FileReference.Exists(outputFile))
+			// before creating a new backup, cap the number of existing files
+			string filenameWithoutExtension = outputFile.GetFileNameWithoutExtension();
+			string extension = outputFile.GetExtension();
+
+			Regex backupForm =
+				new Regex(filenameWithoutExtension + @"-backup-\d\d\d\d\.\d\d\.\d\d-\d\d\.\d\d\.\d\d" + extension);
+
+			foreach (FileReference oldBackup in DirectoryReference
+				.EnumerateFiles(outputFile.Directory)
+				// find files that match the way that we name backup files
+				.Where(x => backupForm.IsMatch(x.GetFileName()))
+				// sort them from newest to oldest
+				.OrderByDescending(x => x.GetFileName())
+				// skip the newest ones that are to be kept; -1 because we're about to create another backup.
+				.Skip(Log.LogFileBackupCount - 1))
 			{
-				// before creating a new backup, cap the number of existing files
-				string filenameWithoutExtension = outputFile.GetFileNameWithoutExtension();
-				string extension = outputFile.GetExtension();
-
-				Regex backupForm =
-					new Regex(filenameWithoutExtension + @"-backup-\d\d\d\d\.\d\d\.\d\d-\d\d\.\d\d\.\d\d" + extension);
-
-				foreach (FileReference oldBackup in DirectoryReference
-					.EnumerateFiles(outputFile.Directory)
-					// find files that match the way that we name backup files
-					.Where(x => backupForm.IsMatch(x.GetFileName()))
-					// sort them from newest to oldest
-					.OrderByDescending(x => x.GetFileName())
-					// skip the newest ones that are to be kept; -1 because we're about to create another backup.
-					.Skip(Log.LogFileBackupCount - 1))
-				{
-					Logger.LogDebug("Deleting old log file: {File}", oldBackup);
-					FileReference.Delete(oldBackup);
-				}
+				Logger.TraceLog("Deleting old log file: {File}", oldBackup);
+				FileReference.Delete(oldBackup);
+			}
 
-				// Ensure that the backup gets a unique name, in the extremely unlikely case that UBT was run twice during
-				// the same second.
-				DateTime fileTime = File.GetCreationTimeUtc(outputFile.FullName);
+			// Ensure that the backup gets a unique name, in the extremely unlikely case that UBT was run twice during
+			// the same second.
+			DateTime fileTime = File.GetCreationTimeUtc(outputFile.FullName);
 
-				FileReference backupFile;
-				for (; ; )
+			FileReference backupFile;
+			for (; ; )
+			{
+				string timestamp = $"{fileTime:yyyy.MM.dd-HH.mm.ss}";
+				backupFile = FileReference.Combine(outputFile.Directory,
+					$"{filenameWithoutExtension}-backup-{timestamp}{extension}");
+				if (!FileReference.Exists(backupFile))
 				{
-					string timestamp = $"{fileTime:yyyy.MM.dd-HH.mm.ss}";
-					backupFile = FileReference.Combine(outputFile.Directory,
-						$"{filenameWithoutExtension}-backup-{timestamp}{extension}");
-					if (!FileReference.Exists(backupFile))
-					{
-						break;
-					}
-
-					fileTime = fileTime.AddSeconds(1);
+					break;
 				}
 
-				FileReference.Move(outputFile, backupFile);
+				fileTime = fileTime.AddSeconds(1);
 			}
 
+			FileReference.Move(outputFile, backupFile);
+		}
+
+		/// <summary>
+		/// Adds a trace listener that writes to a log file.
+		/// If a StartupTraceListener was in use, this function will copy its captured data to the log file(s)
+		/// and remove the startup listener from the list of registered listeners.
+		/// </summary>
+		/// <param name="name">Identifier for the writer</param>
+		/// <param name="outputFile">The file to write to</param>
+		/// <returns>The created trace listener</returns>
+		public static void AddFileWriter(string name, FileReference outputFile)
+		{
+			Log.TraceInformation($"Log file: {outputFile}");
+
+			BackupLogFile(outputFile);
 			AddFileWriterWithoutBackup(name, outputFile);
 		}
 
