UnrealBuildTool: Fix bad merge in Log.cs

#rnx

diff --git a/Engine/Source/Programs/Shared/EpicGames.Core/Log.cs b/Engine/Source/Programs/Shared/EpicGames.Core/Log.cs
--- a/Engine/Source/Programs/Shared/EpicGames.Core/Log.cs
+++ b/Engine/Source/Programs/Shared/EpicGames.Core/Log.cs
@@ -231,7 +231,7 @@ namespace EpicGames.Core
 				// skip the newest ones that are to be kept; -1 because we're about to create another backup.
 				.Skip(Log.LogFileBackupCount - 1))
 			{
-				Logger.TraceLog("Deleting old log file: {File}", oldBackup);
+				Logger.LogDebug("Deleting old log file: {File}", oldBackup);
 				FileReference.Delete(oldBackup);
 			}
 
@@ -266,7 +266,7 @@ namespace EpicGames.Core
 		/// <returns>The created trace listener</returns>
 		public static void AddFileWriter(string name, FileReference outputFile)
 		{
-			Log.TraceInformation($"Log file: {outputFile}");
+			Logger.LogInformation("Log file: {OutputFile}", outputFile);
 
 			BackupLogFile(outputFile);
 			AddFileWriterWithoutBackup(name, outputFile);
