Horde: First phase of refactoring bundles towards load-in-place data structures.

Most changes are mainly superficial at this point, but move some post-load computed offsets and data structures into the bundle.

diff --git a/Engine/Source/Programs/Horde/Horde.Server.Tests/BlobStoreTests.cs b/Engine/Source/Programs/Horde/Horde.Server.Tests/BlobStoreTests.cs
--- a/Engine/Source/Programs/Horde/Horde.Server.Tests/BlobStoreTests.cs
+++ b/Engine/Source/Programs/Horde/Horde.Server.Tests/BlobStoreTests.cs
@@ -47,27 +47,21 @@ namespace Horde.Server.Tests
 			List<BundleType> types = new List<BundleType>();
 			types.Add(new BundleType(Guid.Parse("{AFDF76A7-5333-4DEE-B837-B5F5CA511245}"), 1));
 
-			List<BundleImport> imports = new List<BundleImport>();
+			List<BlobLocator> imports = new List<BlobLocator>(refs);
 
-			Dictionary<BlobLocator, int> locatorToIndex = new Dictionary<BlobLocator, int>();
-			foreach (BlobLocator locator in refs)
+			List<BundleNodeRef> exportRefs = new List<BundleNodeRef>();
+			for(int idx = 0; idx < refs.Count; idx++)
 			{
-				int index;
-				if (!locatorToIndex.TryGetValue(locator, out index))
-				{
-					BundleImport import = new BundleImport(locator, new int[] { 0 });
-					imports.Add(import);
-					locatorToIndex.Add(locator, imports.Count - 1);
-				}
+				exportRefs.Add(new BundleNodeRef(idx, 0));
 			}
 
 			List<BundleExport> exports = new List<BundleExport>();
-			exports.Add(new BundleExport(0, IoHash.Compute(data.Span), (int)data.Length, refs.Select(x => locatorToIndex[x]).ToArray()));
+			exports.Add(new BundleExport(0, IoHash.Compute(data.Span), 0, 0, (int)data.Length, exportRefs));
 
 			List<BundlePacket> packets = new List<BundlePacket>();
-			packets.Add(new BundlePacket((int)data.Length, (int)data.Length));
+			packets.Add(new BundlePacket(BundleCompressionFormat.None, 0, (int)data.Length, (int)data.Length));
 
-			BundleHeader header = new BundleHeader(BundleCompressionFormat.None, types, imports, exports, packets);
+			BundleHeader header = new BundleHeader(types, imports, exports, packets);
 			return new Bundle(header, new[] { data });
 		}
 
@@ -81,7 +75,7 @@ namespace Horde.Server.Tests
 		{
 			Blob blob = new Blob();
 			blob.Data = bundle.Packets[0];
-			blob.References = bundle.Header.Exports[0].References.Select(x => bundle.Header.Imports[x].Locator).ToList();
+			blob.References = bundle.Header.Imports.ToList();
 			return blob;
 		}
 
@@ -182,3 +176,4 @@ namespace Horde.Server.Tests
 		}
 	}
 }
+
diff --git a/Engine/Source/Programs/Horde/Horde.Server.Tests/GcServiceTests.cs b/Engine/Source/Programs/Horde/Horde.Server.Tests/GcServiceTests.cs
--- a/Engine/Source/Programs/Horde/Horde.Server.Tests/GcServiceTests.cs
+++ b/Engine/Source/Programs/Horde/Horde.Server.Tests/GcServiceTests.cs
@@ -70,7 +70,7 @@ namespace Horde.Server.Tests
 				if (nodes.Add(root))
 				{
 					Bundle bundle = await store.ReadBundleAsync(root);
-					await FindNodes(store, bundle.Header.Imports.Select(x => x.Locator), nodes);
+					await FindNodes(store, bundle.Header.Imports, nodes);
 				}
 			}
 		}
@@ -103,8 +103,8 @@ namespace Horde.Server.Tests
 			for (int idx = numNodes - 1; idx >= 0; idx--)
 			{
 				List<BundleType> types = new List<BundleType> { new BundleType(Guid.Parse("{AFDF76A7-5333-4DEE-B837-B5F5CA511245}"), 0) };
-				List<BundleImport> imports = children[idx].ConvertAll(x => new BundleImport(locators[x], Array.Empty<int>()));
-				BundleHeader header = new BundleHeader(BundleCompressionFormat.None, types, imports, Array.Empty<BundleExport>(), Array.Empty<BundlePacket>());
+				List<BlobLocator> imports = children[idx].ConvertAll(x => locators[x]);
+				BundleHeader header = new BundleHeader(types, imports, Array.Empty<BundleExport>(), Array.Empty<BundlePacket>());
 				Bundle bundle = new Bundle(header, Array.Empty<ReadOnlyMemory<byte>>());
 				locators[idx] = await store.WriteBundleAsync(bundle, prefix: "gctest");
 			}
diff --git a/Engine/Source/Programs/Horde/Horde.Server.Tests/StorageServiceTests.cs b/Engine/Source/Programs/Horde/Horde.Server.Tests/StorageServiceTests.cs
--- a/Engine/Source/Programs/Horde/Horde.Server.Tests/StorageServiceTests.cs
+++ b/Engine/Source/Programs/Horde/Horde.Server.Tests/StorageServiceTests.cs
@@ -36,11 +36,11 @@ namespace Horde.Server.Tests
 
 			List<BundleExport> exports = new List<BundleExport>();
 
-			exports.Add(new BundleExport(0, hash1, data1.Length, Array.Empty<int>()));
-			exports.Add(new BundleExport(0, hash1, data1.Length, Array.Empty<int>()));
-			exports.Add(new BundleExport(0, hash2, data2.Length, Array.Empty<int>()));
+			exports.Add(new BundleExport(0, hash1, 0, 0, data1.Length, Array.Empty<BundleNodeRef>()));
+			exports.Add(new BundleExport(0, hash1, 0, 0, data1.Length, Array.Empty<BundleNodeRef>()));
+			exports.Add(new BundleExport(0, hash2, 0, data1.Length, data2.Length, Array.Empty<BundleNodeRef>()));
 
-			BundleHeader header = new BundleHeader(BundleCompressionFormat.None, types, Array.Empty<BundleImport>(), exports, Array.Empty<BundlePacket>());
+			BundleHeader header = new BundleHeader(types, Array.Empty<BlobLocator>(), exports, Array.Empty<BundlePacket>());
 			Bundle bundle = new Bundle(header, Array.Empty<ReadOnlyMemory<byte>>());
 			BlobLocator locator = await client.WriteBundleAsync(bundle);
 
@@ -61,3 +61,4 @@ namespace Horde.Server.Tests
 		}
 	}
 }
+
diff --git a/Engine/Source/Programs/Horde/Horde.Server/Storage/StorageController.cs b/Engine/Source/Programs/Horde/Horde.Server/Storage/StorageController.cs
--- a/Engine/Source/Programs/Horde/Horde.Server/Storage/StorageController.cs
+++ b/Engine/Source/Programs/Horde/Horde.Server/Storage/StorageController.cs
@@ -500,11 +500,9 @@ namespace Horde.Server.Storage
 			if (includeImports)
 			{
 				responseImports = new List<object>();
-				foreach (BundleImport import in header.Imports)
+				foreach (BlobLocator import in header.Imports)
 				{
-					string link = $"{linkBase}/bundles/{import.Locator}";
-					List<string> exports = import.Exports.ConvertAll(x => $"{linkBase}/nodes/{import.Locator}?export={x}");
-					responseImports.Add(new { bundle = link, exports });
+					responseImports.Add($"{linkBase}/bundles/{import}");
 				}
 			}
 
@@ -546,7 +544,7 @@ namespace Horde.Server.Storage
 				}
 			}
 
-			return new { header.CompressionFormat, imports = responseImports, exports = responseExports, packets = responsePackets };
+			return new { imports = responseImports, exports = responseExports, packets = responsePackets };
 		}
 
 		/// <summary>
diff --git a/Engine/Source/Programs/Horde/Horde.Server/Storage/StorageService.cs b/Engine/Source/Programs/Horde/Horde.Server/Storage/StorageService.cs
--- a/Engine/Source/Programs/Horde/Horde.Server/Storage/StorageService.cs
+++ b/Engine/Source/Programs/Horde/Horde.Server/Storage/StorageService.cs
@@ -651,10 +651,10 @@ namespace Horde.Server.Storage
 							if (header != null)
 							{
 								List<ObjectId> importInfoIds = new List<ObjectId>();
-								foreach (BundleImport import in header.Imports)
+								foreach (BlobLocator import in header.Imports)
 								{
-									FilterDefinition<BlobInfo> filter = Builders<BlobInfo>.Filter.Expr(x => x.NamespaceId == blobInfo.NamespaceId && x.BlobId == import.Locator.BlobId);
-									UpdateDefinition<BlobInfo> update = Builders<BlobInfo>.Update.SetOnInsert(x => x.HostId, import.Locator.HostId);
+									FilterDefinition<BlobInfo> filter = Builders<BlobInfo>.Filter.Expr(x => x.NamespaceId == blobInfo.NamespaceId && x.BlobId == import.BlobId);
+									UpdateDefinition<BlobInfo> update = Builders<BlobInfo>.Update.SetOnInsert(x => x.HostId, import.HostId);
 									BlobInfo blobInfoDoc = await _blobCollection.FindOneAndUpdateAsync(filter, update, new FindOneAndUpdateOptions<BlobInfo> { IsUpsert = true, ReturnDocument = ReturnDocument.After }, cancellationToken);
 									importInfoIds.Add(blobInfoDoc.Id);
 								}
diff --git a/Engine/Source/Programs/Horde/Horde/Commands/Bundles/BundleDump.cs b/Engine/Source/Programs/Horde/Horde/Commands/Bundles/BundleDump.cs
--- a/Engine/Source/Programs/Horde/Horde/Commands/Bundles/BundleDump.cs
+++ b/Engine/Source/Programs/Horde/Horde/Commands/Bundles/BundleDump.cs
@@ -56,20 +56,9 @@ namespace Horde.Commands.Bundles
 			logger.LogInformation("Imports: {NumImports}", header.Imports.Count);
 
 			List<string> nodeNames = new List<string>();
-
-			int refIdx = 0;
-			foreach (BundleImport import in header.Imports)
+			for (int importIdx = 0; importIdx < header.Imports.Count; importIdx++)
 			{
-				logger.LogInformation("  From blob {BlobId} ({NumExports} nodes)", import.Locator, import.Exports.Count);
-				foreach (int exportIdx in import.Exports)
-				{
-					nodeNames.Add($"IMP {import.Locator}#{exportIdx}");
-					if (Verbose)
-					{
-						logger.LogInformation("    [{Index}] IMP {ExportIdx}", refIdx, exportIdx);
-					}
-					refIdx++;
-				}
+				logger.LogInformation("  IMP [{ImportIdx}] = {BlobId}", importIdx, header.Imports[importIdx]);
 			}
 
 			logger.LogInformation("");
@@ -80,17 +69,15 @@ namespace Horde.Commands.Bundles
 			for(int exportIdx = 0; exportIdx < header.Exports.Count; exportIdx++)
 			{
 				BundleExport export = header.Exports[exportIdx];
-				logger.LogInformation("  [{Index}] EXP {ExportHash} (type: {Type}, length: {NumBytes:n0}, packet: {PacketIdx})", refIdx, export.Hash, types[export.TypeIdx], export.Length, packetIdx);
-
-				nodeNames.Add($"EXP {exportIdx}");
-				refIdx++;
+				logger.LogInformation("  EXP [{ExportIdx}] EXP = hash: {ExportHash}, type: {Type}, length: {NumBytes:n0}, packet: {PacketIdx}", exportIdx, export.Hash, types[export.TypeIdx], export.Length, packetIdx);
 
 				if (Verbose)
 				{
-					for (int idx = 0; idx < export.References.Count; idx++)
+					for(int referenceIdx = 0; referenceIdx < export.References.Count; referenceIdx++)
 					{
-						int nodeIdx = export.References[idx];
-						logger.LogInformation("            REF {RefIdx,-3} -> {NodeIdx,-3} ({RefName})", idx, nodeIdx, nodeNames[nodeIdx]);
+						BundleNodeRef exportRef = export.References[referenceIdx];
+						NodeLocator nodeLocator = new NodeLocator(header.Imports[exportRef.ImportIdx], exportRef.NodeIdx);
+						logger.LogInformation("            REF {RefIdx,-3} -> {Node}", referenceIdx, nodeLocator);
 					}
 				}
 
diff --git a/Engine/Source/Programs/Shared/EpicGames.Horde.Tests/BundleStoreTests.cs b/Engine/Source/Programs/Shared/EpicGames.Horde.Tests/BundleStoreTests.cs
--- a/Engine/Source/Programs/Shared/EpicGames.Horde.Tests/BundleStoreTests.cs
+++ b/Engine/Source/Programs/Shared/EpicGames.Horde.Tests/BundleStoreTests.cs
@@ -75,12 +75,12 @@ namespace EpicGames.Horde.Tests
 			types.Add(new BundleType(Guid.Parse("F63606D4-5DBB-4061-A655-6F444F65229E"), 1));
 
 			List<BundleExport> exports = new List<BundleExport>();
-			exports.Add(new BundleExport(0, IoHash.Compute(payload), payload.Length, Array.Empty<int>()));
+			exports.Add(new BundleExport(0, IoHash.Compute(payload), 0, 0, payload.Length, Array.Empty<BundleNodeRef>()));
 
 			List<BundlePacket> packets = new List<BundlePacket>();
-			packets.Add(new BundlePacket(payload.Length, payload.Length));
+			packets.Add(new BundlePacket(BundleCompressionFormat.None, 0, payload.Length, payload.Length));
 
-			BundleHeader header = new BundleHeader(BundleCompressionFormat.None, types, Array.Empty<BundleImport>(), exports, packets);
+			BundleHeader header = new BundleHeader(types, Array.Empty<BlobLocator>(), exports, packets);
 			return new Bundle(header, new List<ReadOnlyMemory<byte>> { payload });
 		}
 
diff --git a/Engine/Source/Programs/Shared/EpicGames.Horde.Tests/BundleTests.cs b/Engine/Source/Programs/Shared/EpicGames.Horde.Tests/BundleTests.cs
--- a/Engine/Source/Programs/Shared/EpicGames.Horde.Tests/BundleTests.cs
+++ b/Engine/Source/Programs/Shared/EpicGames.Horde.Tests/BundleTests.cs
@@ -140,19 +140,19 @@ namespace EpicGames.Horde.Tests
 				List<BundleType> types = new List<BundleType>();
 				types.Add(new BundleType(Guid.NewGuid(), 0));
 
-				List<BundleImport> imports = new List<BundleImport>();
-				imports.Add(new BundleImport(new BlobLocator("import1"), new[] { 5 }));
-				imports.Add(new BundleImport(new BlobLocator("import2"), new[] { 6 }));
+				List<BlobLocator> imports = new List<BlobLocator>();
+				imports.Add(new BlobLocator("import1"));
+				imports.Add(new BlobLocator("import2"));
 
 				List<BundleExport> exports = new List<BundleExport>();
-				exports.Add(new BundleExport(0, IoHash.Compute(Encoding.UTF8.GetBytes("export1")), 2, new int[] { 1, 2 }));
-				exports.Add(new BundleExport(0, IoHash.Compute(Encoding.UTF8.GetBytes("export2")), 3, new int[] { 3 }));
+				exports.Add(new BundleExport(0, IoHash.Compute(Encoding.UTF8.GetBytes("export1")), 0, 0, 2, new BundleNodeRef[] { new BundleNodeRef(0, 5), new BundleNodeRef(0, 6) }));
+				exports.Add(new BundleExport(0, IoHash.Compute(Encoding.UTF8.GetBytes("export2")), 1, 0, 3, new BundleNodeRef[] { new BundleNodeRef(-1, 0) }));
 
 				List<BundlePacket> packets = new List<BundlePacket>();
-				packets.Add(new BundlePacket(20, 40));
-				packets.Add(new BundlePacket(10, 20));
+				packets.Add(new BundlePacket(BundleCompressionFormat.LZ4, 0, 20, 40));
+				packets.Add(new BundlePacket(BundleCompressionFormat.LZ4, 20, 10, 20));
 
-				oldHeader = new BundleHeader(BundleCompressionFormat.LZ4, types, imports, exports, packets);
+				oldHeader = new BundleHeader(types, imports, exports, packets);
 			}
 
 			ByteArrayBuilder writer = new ByteArrayBuilder();
@@ -164,19 +164,7 @@ namespace EpicGames.Horde.Tests
 			Assert.AreEqual(oldHeader.Imports.Count, newHeader.Imports.Count);
 			for (int idx = 0; idx < oldHeader.Imports.Count; idx++)
 			{
-				BundleImport oldImport = oldHeader.Imports[idx];
-				BundleImport newImport = newHeader.Imports[idx];
-
-				Assert.AreEqual(oldImport.Locator, newImport.Locator);
-				Assert.AreEqual(oldImport.Exports.Count, newImport.Exports.Count);
-
-				for (int importIdx = 0; importIdx < oldImport.Exports.Count; importIdx++)
-				{
-					int oldExportIdx = oldImport.Exports[importIdx];
-					int newExportIdx = newImport.Exports[importIdx];
-
-					Assert.AreEqual(oldExportIdx, newExportIdx);
-				}
+				Assert.AreEqual(oldHeader.Imports[idx], newHeader.Imports[idx]);
 			}
 
 			Assert.AreEqual(oldHeader.Exports.Count, newHeader.Exports.Count);
diff --git a/Engine/Source/Programs/Shared/EpicGames.Horde/Storage/Bundle.cs b/Engine/Source/Programs/Shared/EpicGames.Horde/Storage/Bundle.cs
--- a/Engine/Source/Programs/Shared/EpicGames.Horde/Storage/Bundle.cs
+++ b/Engine/Source/Programs/Shared/EpicGames.Horde/Storage/Bundle.cs
@@ -179,20 +179,15 @@ namespace EpicGames.Horde.Storage
 		/// </summary>
 		public const int PreludeLength = 8;
 
-		/// <summary>
-		/// Compression format for the bundle
-		/// </summary>
-		public BundleCompressionFormat CompressionFormat { get; }
-
 		/// <summary>
 		/// Types for exports within this bundle
 		/// </summary>
 		public IReadOnlyList<BundleType> Types { get; }
 
 		/// <summary>
-		/// References to nodes in other bundles
+		/// Bundles that we reference nodes in
 		/// </summary>
-		public IReadOnlyList<BundleImport> Imports { get; }
+		public IReadOnlyList<BlobLocator> Imports { get; }
 
 		/// <summary>
 		/// Nodes exported from this bundle
@@ -207,20 +202,20 @@ namespace EpicGames.Horde.Storage
 		/// <summary>
 		/// Constructs a new bundle header
 		/// </summary>
-		/// <param name="compressionFormat">Compression format for bundle packets</param>
 		/// <param name="types">Type array indexed by each export</param>
-		/// <param name="imports">Imports from other bundles</param>
+		/// <param name="imports">Imported bundles</param>
 		/// <param name="exports">Exports for nodes</param>
 		/// <param name="packets">Compression packets within the bundle</param>
-		public BundleHeader(BundleCompressionFormat compressionFormat, IReadOnlyList<BundleType> types, IReadOnlyList<BundleImport> imports, IReadOnlyList<BundleExport> exports, IReadOnlyList<BundlePacket> packets)
+		public BundleHeader(IReadOnlyList<BundleType> types, IReadOnlyList<BlobLocator> imports, IReadOnlyList<BundleExport> exports, IReadOnlyList<BundlePacket> packets)
 		{
-			CompressionFormat = compressionFormat;
 			Types = types;
 			Imports = imports;
 			Exports = exports;
 			Packets = packets;
 		}
 
+		record class ExportInfo(int TypeIdx, IoHash Hash, int Length, List<BundleNodeRef> References);
+
 		/// <summary>
 		/// Reads a bundle header from memory
 		/// </summary>
@@ -238,19 +233,124 @@ namespace EpicGames.Horde.Storage
 				throw new InvalidDataException($"Unknown bundle version {(int)version}. Max supported is {(int)BundleVersion.Latest}.");
 			}
 
-			CompressionFormat = (BundleCompressionFormat)reader.ReadUnsignedVarInt();
-			Types = reader.ReadVariableLengthArray(() => new BundleType(reader));
-			Imports = reader.ReadVariableLengthArray(() => new BundleImport(reader));
-			Exports = reader.ReadVariableLengthArray(() => new BundleExport(reader, version));
+			BundleCompressionFormat compressionFormat = (BundleCompressionFormat)reader.ReadUnsignedVarInt();
 
-			if (CompressionFormat == BundleCompressionFormat.None)
+			// Read the types
+			int numTypes = (int)reader.ReadUnsignedVarInt();
+			List<BundleType> types = new List<BundleType>(numTypes);
+
+			for (int typeIdx = 0; typeIdx < numTypes; typeIdx++)
 			{
-				Packets = Exports.ConvertAll(x => new BundlePacket(x.Length, x.Length));
+				Guid guid = reader.ReadGuid();
+				int serializerVersion = (int)reader.ReadUnsignedVarInt();
+
+				types.Add(new BundleType(guid, serializerVersion));
+			}
+
+			Types = types;
+
+			// Read the imports
+			int numImports = (int)reader.ReadUnsignedVarInt();
+			List<BlobLocator> imports = new List<BlobLocator>(numImports);
+			List<BundleNodeRef> allExportReferences = new List<BundleNodeRef>();
+
+			for (int importIdx = 0; importIdx < numImports; importIdx++)
+			{
+				BlobLocator locator = reader.ReadBlobLocator();
+				imports.Add(locator);
+
+				int[] exportIndexes = reader.ReadVariableLengthArray(() => (int)reader.ReadUnsignedVarInt());
+				for (int exportIdx = 0; exportIdx < exportIndexes.Length; exportIdx++)
+				{
+					BundleNodeRef exportReference = new BundleNodeRef(imports.Count - 1, exportIndexes[exportIdx]);
+					allExportReferences.Add(exportReference);
+				}
+			}
+
+			Imports = imports;
+
+			// Read the exports
+			int numExports = (int)reader.ReadUnsignedVarInt();
+			List<ExportInfo> exportInfos = new List<ExportInfo>(numExports);
+
+			for (int exportIdx = 0; exportIdx < numExports; exportIdx++)
+			{
+				allExportReferences.Add(new BundleNodeRef(-1, exportIdx));
+
+				int typeIdx = (int)reader.ReadUnsignedVarInt();
+				IoHash hash = reader.ReadIoHash();
+				int length = (int)reader.ReadUnsignedVarInt();
+
+				List<BundleNodeRef> exportReferences = new List<BundleNodeRef>();
+
+				int numReferences = (int)reader.ReadUnsignedVarInt();
+				if (numReferences > 0)
+				{
+					for (int idx = 0; idx < numReferences; idx++)
+					{
+						int referenceIdx = (int)reader.ReadUnsignedVarInt();
+						exportReferences.Add(allExportReferences[referenceIdx]);
+					}
+				}
+
+				if (version == BundleVersion.ExportAliases)
+				{
+					_ = reader.ReadUtf8String();
+				}
+
+				exportInfos.Add(new ExportInfo(typeIdx, hash, length, exportReferences));
+			}
+
+			// Read the compression packets
+			if (compressionFormat == BundleCompressionFormat.None)
+			{
+				List<BundlePacket> packets = new List<BundlePacket>(exportInfos.Count);
+
+				int encodedOffset = 0;
+				foreach (ExportInfo export in exportInfos)
+				{
+					packets.Add(new BundlePacket(compressionFormat, encodedOffset, export.Length, export.Length));
+					encodedOffset += export.Length;
+				}
+
+				Packets = packets;
 			}
 			else
 			{
-				Packets = reader.ReadVariableLengthArray(() => new BundlePacket(reader));
+				int numPackets = (int)reader.ReadUnsignedVarInt();
+				List<BundlePacket> packets = new List<BundlePacket>(numPackets);
+
+				int encodedOffset = 0;
+				for (int packetIdx = 0; packetIdx < numPackets; packetIdx++)
+				{
+					int encodedLength = (int)reader.ReadUnsignedVarInt();
+					int decodedLength = (int)reader.ReadUnsignedVarInt();
+					packets.Add(new BundlePacket(compressionFormat, encodedOffset, encodedLength, decodedLength));
+					encodedOffset += encodedLength;
+				}
+
+				Packets = packets;
 			}
+
+			// Create the final export list
+			List<BundleExport> exports = new List<BundleExport>();
+			{
+				int packetIdx = 0;
+				int packetOffset = 0;
+
+				foreach (ExportInfo exportInfo in exportInfos)
+				{
+					if (packetOffset + exportInfo.Length > Packets[packetIdx].DecodedLength)
+					{
+						packetIdx++;
+						packetOffset = 0;
+					}
+
+					exports.Add(new BundleExport(exportInfo.TypeIdx, exportInfo.Hash, packetIdx, packetOffset, exportInfo.Length, exportInfo.References));
+					packetOffset += exportInfo.Length;
+				}
+			}
+			Exports = exports;
 		}
 
 		/// <summary>
@@ -315,14 +415,71 @@ namespace EpicGames.Horde.Storage
 			// Write the header data
 			writer.WriteUnsignedVarInt((ulong)BundleVersion.Latest);
 
-			writer.WriteUnsignedVarInt((ulong)CompressionFormat);
+			BundleCompressionFormat compressionFormat = (Packets.Count > 0) ? Packets[0].CompressionFormat : BundleCompressionFormat.None;
+
+			writer.WriteUnsignedVarInt((ulong)compressionFormat);
 			writer.WriteVariableLengthArray(Types, x => x.Write(writer));
-			writer.WriteVariableLengthArray(Imports, x => x.Write(writer));
-			writer.WriteVariableLengthArray(Exports, x => x.Write(writer));
 
-			if (CompressionFormat != BundleCompressionFormat.None)
+			// Find all the referenced nodes in each import
+			SortedSet<int>[] nodes = new SortedSet<int>[Imports.Count];
+			for (int idx = 0; idx < Imports.Count; idx++)
+			{
+				nodes[idx] = new SortedSet<int>();
+			}
+			foreach (BundleNodeRef exportRef in Exports.SelectMany(x => x.References).Where(x => x.ImportIdx != -1))
 			{
-				writer.WriteVariableLengthArray(Packets, x => x.Write(writer));
+				nodes[exportRef.ImportIdx].Add(exportRef.NodeIdx);
+			}
+
+			// Map all the imports to an index
+			Dictionary<BundleNodeRef, int> exportRefToIndex = new Dictionary<BundleNodeRef, int>();
+			for (int importIdx = 0; importIdx < Imports.Count; importIdx++)
+			{
+				foreach (int nodeIdx in nodes[importIdx])
+				{
+					int index = exportRefToIndex.Count;
+					exportRefToIndex[new BundleNodeRef(importIdx, nodeIdx)] = index;
+				}
+			}
+			for (int exportIdx = 0; exportIdx < Exports.Count; exportIdx++)
+			{
+				int index = exportRefToIndex.Count;
+				exportRefToIndex[new BundleNodeRef(-1, exportIdx)] = index;
+			}
+
+			// Write all the imports
+			writer.WriteUnsignedVarInt(Imports.Count);
+			for (int importIdx = 0; importIdx < Imports.Count; importIdx++)
+			{
+				Debug.Assert(Imports[importIdx].IsValid());
+				writer.WriteBlobLocator(Imports[importIdx]);
+				writer.WriteVariableLengthArray(nodes[importIdx].ToArray(), x => writer.WriteUnsignedVarInt(x));
+			}
+
+			// Write all the exports
+			writer.WriteUnsignedVarInt(Exports.Count);
+			foreach (BundleExport export in Exports)
+			{
+				writer.WriteUnsignedVarInt(export.TypeIdx);
+				writer.WriteIoHash(export.Hash);
+				writer.WriteUnsignedVarInt(export.Length);
+
+				writer.WriteUnsignedVarInt(export.References.Count);
+				for (int idx = 0; idx < export.References.Count; idx++)
+				{
+					writer.WriteUnsignedVarInt(exportRefToIndex[export.References[idx]]);
+				}
+			}
+
+			// Write the packets
+			if (compressionFormat != BundleCompressionFormat.None)
+			{
+				writer.WriteUnsignedVarInt(Packets.Count);
+				foreach(BundlePacket packet in Packets)
+				{
+					writer.WriteUnsignedVarInt(packet.EncodedLength);
+					writer.WriteUnsignedVarInt(packet.DecodedLength);
+				}
 			}
 
 			// Go back and write the length of the header
@@ -357,16 +514,6 @@ namespace EpicGames.Horde.Storage
 			Version = version;
 		}
 
-		/// <summary>
-		/// Deserialize a bundle type
-		/// </summary>
-		/// <param name="reader">Reader to deserialize from</param>
-		public BundleType(IMemoryReader reader)
-		{
-			Guid = reader.ReadGuid();
-			Version = (int)reader.ReadUnsignedVarInt();
-		}
-
 		/// <summary>
 		/// Serializes a type to storage
 		/// </summary>
@@ -410,55 +557,20 @@ namespace EpicGames.Horde.Storage
 	}
 
 	/// <summary>
-	/// Reference to nodes in another bundle
+	/// Descriptor for a compression packet
 	/// </summary>
-	public class BundleImport
+	public class BundlePacket
 	{
 		/// <summary>
-		/// Blob containing the bundle data.
+		/// Compression format for this packet
 		/// </summary>
-		public BlobLocator Locator { get; }
-
-		/// <summary>
-		/// Entries for each node
-		/// </summary>
-		public IReadOnlyList<int> Exports { get; }
-
-		/// <summary>
-		/// Constructor
-		/// </summary>
-		public BundleImport(BlobLocator locator, IReadOnlyList<int> exports)
-		{
-			Locator = locator;
-			Exports = exports;
-		}
-
-		/// <summary>
-		/// Deserialize a bundle import
-		/// </summary>
-		/// <param name="reader">Reader to deserialize from</param>
-		internal BundleImport(IMemoryReader reader)
-		{
-			Locator = reader.ReadBlobLocator();
-			Exports = reader.ReadVariableLengthArray(() => (int)reader.ReadUnsignedVarInt());
-		}
+		public BundleCompressionFormat CompressionFormat { get; }
 
 		/// <summary>
-		/// Serializes a bundle import
+		/// Offset of the packet within the payload stream
 		/// </summary>
-		/// <param name="writer">Writer to serialize to</param>
-		public void Write(IMemoryWriter writer)
-		{
-			writer.WriteBlobLocator(Locator);
-			writer.WriteVariableLengthArray(Exports, x => writer.WriteUnsignedVarInt(x));
-		}
-	}
+		public int EncodedOffset { get; }
 
-	/// <summary>
-	/// Descriptor for a compression packet
-	/// </summary>
-	public class BundlePacket
-	{
 		/// <summary>
 		/// Encoded length of the packet
 		/// </summary>
@@ -472,33 +584,17 @@ namespace EpicGames.Horde.Storage
 		/// <summary>
 		/// Constructor
 		/// </summary>
+		/// <param name="compressionFormat">Compression format for the packet</param>
+		/// <param name="encodedOffset">Offset of the data within the payload stream</param>
 		/// <param name="encodedLength">Size of the encoded data</param>
 		/// <param name="decodedLength">Size of the decoded data</param>
-		public BundlePacket(int encodedLength, int decodedLength)
+		public BundlePacket(BundleCompressionFormat compressionFormat, int encodedOffset, int encodedLength, int decodedLength)
 		{
+			CompressionFormat = compressionFormat;
+			EncodedOffset = encodedOffset;
 			EncodedLength = encodedLength;
 			DecodedLength = decodedLength;
 		}
-
-		/// <summary>
-		/// Deserialize a packet header from a memory reader
-		/// </summary>
-		/// <param name="reader">Reader to deserialize from</param>
-		public BundlePacket(IMemoryReader reader)
-		{
-			EncodedLength = (int)reader.ReadUnsignedVarInt();
-			DecodedLength = (int)reader.ReadUnsignedVarInt();
-		}
-
-		/// <summary>
-		/// Serializes this packet header 
-		/// </summary>
-		/// <param name="writer"></param>
-		public void Write(IMemoryWriter writer)
-		{
-			writer.WriteUnsignedVarInt(EncodedLength);
-			writer.WriteUnsignedVarInt(DecodedLength);
-		}
 	}
 
 	/// <summary>
@@ -516,79 +612,47 @@ namespace EpicGames.Horde.Storage
 		/// </summary>
 		public IoHash Hash { get; }
 
+		/// <summary>
+		/// Index of the packet containing this bundle
+		/// </summary>
+		public int Packet { get; }
+
+		/// <summary>
+		/// Offset of this export's data within the (decompressed) packet
+		/// </summary>
+		public int Offset { get; }
+
 		/// <summary>
 		/// Uncompressed length of this node
 		/// </summary>
 		public int Length { get; }
 
 		/// <summary>
-		/// Nodes referenced by this export. Indices in this array correspond to a lookup table consisting
-        /// of the imported nodes in the order they are declared in the header, followed by nodes listed in the
-        /// export table itself.
-        /// </summary>
-		public IReadOnlyList<int> References { get; }
+		/// References to other nodes
+		/// </summary>
+		public IReadOnlyList<BundleNodeRef> References { get; }
 
 		/// <summary>
 		/// Constructor
 		/// </summary>
-		public BundleExport(int typeIdx, IoHash hash, int length, IReadOnlyList<int> references)
+		public BundleExport(int typeIdx, IoHash hash, int packet, int offset, int length, IReadOnlyList<BundleNodeRef> references)
 		{
 			TypeIdx = typeIdx;
 			Hash = hash;
+			Packet = packet;
+			Offset = offset;
 			Length = length;
 			References = references;
 		}
-
-		/// <summary>
-		/// Deserialize an export
-		/// </summary>
-		/// <param name="reader">Reader to deserialize from</param>
-		/// <param name="version">Version number of the bundle</param>
-		public BundleExport(IMemoryReader reader, BundleVersion version)
-		{
-			TypeIdx = (int)reader.ReadUnsignedVarInt();
-			Hash = reader.ReadIoHash();
-			Length = (int)reader.ReadUnsignedVarInt();
-
-			int numReferences = (int)reader.ReadUnsignedVarInt();
-			if (numReferences == 0)
-			{
-				References = Array.Empty<int>();
-			}
-			else
-			{
-				int[] references = new int[numReferences];
-				for (int idx = 0; idx < numReferences; idx++)
-				{
-					references[idx] = (int)reader.ReadUnsignedVarInt();
-				}
-				References = references;
-			}
-
-			if (version == BundleVersion.ExportAliases)
-			{
-				_ = reader.ReadUtf8String();
-			}
-		}
-
-		/// <summary>
-		/// Serializes an export
-		/// </summary>
-		/// <param name="writer">Writer to serialize to</param>
-		public void Write(IMemoryWriter writer)
-		{
-			writer.WriteUnsignedVarInt(TypeIdx);
-			writer.WriteIoHash(Hash);
-			writer.WriteUnsignedVarInt(Length);
-
-			writer.WriteUnsignedVarInt(References.Count);
-			for(int idx = 0; idx < References.Count; idx++)
-			{
-				writer.WriteUnsignedVarInt(References[idx]);
-			}
-		}
 	}
 
+	/// <summary>
+	/// Reference to a node in another bundle
+	/// </summary>
+	/// <param name="ImportIdx">Index into the import table of the blob containing the referenced node. Can be -1 for references within the same bundle.</param>
+	/// <param name="NodeIdx">Node imported from the bundle</param>
+	public record class BundleNodeRef(int ImportIdx, int NodeIdx);
+
 	/// <summary>
 	/// Utility methods for bundles
 	/// </summary>
diff --git a/Engine/Source/Programs/Shared/EpicGames.Horde/Storage/TreeReader.cs b/Engine/Source/Programs/Shared/EpicGames.Horde/Storage/TreeReader.cs
--- a/Engine/Source/Programs/Shared/EpicGames.Horde/Storage/TreeReader.cs
+++ b/Engine/Source/Programs/Shared/EpicGames.Horde/Storage/TreeReader.cs
@@ -14,6 +14,7 @@ using Microsoft.Extensions.Logging;
 using System.Collections.Concurrent;
 using Microsoft.CodeAnalysis;
 using Blake3;
+using System.Diagnostics;
 
 namespace EpicGames.Horde.Storage
 {
@@ -89,13 +90,6 @@ namespace EpicGames.Horde.Storage
 	/// </summary>
 	public class TreeReader
 	{
-		/// <summary>
-		/// Describes a bundle export
-		/// </summary>
-		/// <param name="PacketIdx">Index of the packet within a bundle</param>
-		/// <param name="Offset">Offset within the packet</param>
-		record struct ExportInfo(int PacketIdx, int Offset);
-
 		/// <summary>
 		/// Information about a known type that can be deserialized from a bundle
 		/// </summary>
@@ -151,52 +145,14 @@ namespace EpicGames.Horde.Storage
 		{
 			public readonly BlobLocator Locator;
 			public readonly BundleHeader Header;
-			public readonly int[] PacketOffsets;
-			public readonly ExportInfo[] Exports;
-			public readonly NodeLocator[] References;
+			public readonly int HeaderLength;
 			public readonly TypeInfo?[] Types;
 
 			public BundleInfo(BlobLocator locator, BundleHeader header, int headerLength, TypeInfo?[] types)
 			{
 				Locator = locator;
 				Header = header;
-
-				PacketOffsets = new int[header.Packets.Count];
-				Exports = new ExportInfo[header.Exports.Count];
-
-				int exportIdx = 0;
-				int packetOffset = headerLength;
-				for (int packetIdx = 0; packetIdx < header.Packets.Count; packetIdx++)
-				{
-					BundlePacket packet = header.Packets[packetIdx];
-					PacketOffsets[packetIdx] = packetOffset;
-
-					int nodeOffset = 0;
-					for (; exportIdx < header.Exports.Count && nodeOffset + header.Exports[exportIdx].Length <= packet.DecodedLength; exportIdx++)
-					{
-						Exports[exportIdx] = new ExportInfo(packetIdx, nodeOffset);
-						nodeOffset += header.Exports[exportIdx].Length;
-					}
-
-					packetOffset += packet.EncodedLength;
-				}
-
-				References = new NodeLocator[header.Imports.Sum(x => x.Exports.Count) + header.Exports.Count];
-
-				int referenceIdx = 0;
-				foreach (BundleImport import in header.Imports)
-				{
-					foreach (int importExportIdx in import.Exports)
-					{
-						NodeLocator importLocator = new NodeLocator(import.Locator, importExportIdx);
-						References[referenceIdx++] = importLocator;
-					}
-				}
-				for (int idx = 0; idx < header.Exports.Count; idx++)
-				{
-					References[referenceIdx++] = new NodeLocator(locator, idx);
-				}
-
+				HeaderLength = headerLength;
 				Types = types;
 			}
 		}
@@ -555,7 +511,7 @@ namespace EpicGames.Horde.Storage
 			using (IMemoryOwner<byte> owner = MemoryPool<byte>.Shared.Rent(readLength))
 			{
 				Memory<byte> buffer = owner.Memory.Slice(0, readLength);
-				buffer = await _store.ReadBlobRangeAsync(bundleInfo.Locator, bundleInfo.PacketOffsets[minPacketIdx], buffer, cancellationToken);
+				buffer = await _store.ReadBlobRangeAsync(bundleInfo.Locator, bundleInfo.HeaderLength + bundleInfo.Header.Packets[minPacketIdx].EncodedOffset, buffer, cancellationToken);
 
 				// Copy all the packets that have been read into separate buffers, so we can cache them individually.
 				ReadOnlyMemory<byte>?[] packets = new ReadOnlyMemory<byte>?[maxPacketIdx - minPacketIdx];
@@ -564,7 +520,7 @@ namespace EpicGames.Horde.Storage
 					for (int idx = minPacketIdx; idx < maxPacketIdx; idx++)
 					{
 						string cacheKey = GetEncodedPacketCacheKey(bundleInfo.Locator.BlobId, idx);
-						ReadOnlyMemory<byte> data = buffer.Slice(bundleInfo.PacketOffsets[idx] - bundleInfo.PacketOffsets[minPacketIdx], bundleInfo.Header.Packets[idx].EncodedLength).ToArray();
+						ReadOnlyMemory<byte> data = buffer.Slice(bundleInfo.Header.Packets[idx].EncodedOffset - bundleInfo.Header.Packets[minPacketIdx].EncodedOffset, bundleInfo.Header.Packets[idx].EncodedLength).ToArray();
 						packets[idx - minPacketIdx] = data;
 						AddToCache(cacheKey, data, data.Length);
 					}
@@ -583,7 +539,7 @@ namespace EpicGames.Horde.Storage
 					ReadOnlyMemory<byte>? data = packets[updatePacket.PacketIdx - minPacketIdx];
 					if (data == null)
 					{
-						data = buffer.Slice(bundleInfo.PacketOffsets[updatePacket.PacketIdx] - bundleInfo.PacketOffsets[minPacketIdx], bundleInfo.Header.Packets[updatePacket.PacketIdx].EncodedLength).ToArray();
+						data = buffer.Slice(bundleInfo.Header.Packets[updatePacket.PacketIdx].EncodedOffset - bundleInfo.Header.Packets[minPacketIdx].EncodedOffset, bundleInfo.Header.Packets[updatePacket.PacketIdx].EncodedLength).ToArray();
 					}
 					updatePacket.CompletionSource.SetResult((ReadOnlyMemory<byte>)data);
 				}
@@ -616,6 +572,8 @@ namespace EpicGames.Horde.Storage
 		/// <returns>Information about the bundle</returns>
 		async ValueTask<BundleInfo> GetBundleInfoAsync(BlobLocator locator, CancellationToken cancellationToken = default)
 		{
+			Debug.Assert(locator.IsValid());
+
 			string cacheKey = GetBundleInfoCacheKey(locator.BlobId);
 			if (_cache != null && _cache.TryGetValue(cacheKey, out BundleInfo bundleInfo))
 			{
@@ -695,7 +653,7 @@ namespace EpicGames.Horde.Storage
 			BundlePacket packet = bundleInfo.Header.Packets[packetIdx];
 			byte[] decodedPacket = new byte[packet.DecodedLength];
 
-			BundleData.Decompress(bundleInfo.Header.CompressionFormat, encodedPacket, decodedPacket);
+			BundleData.Decompress(packet.CompressionFormat, encodedPacket, decodedPacket);
 
 			string decodedCacheKey = GetDecodedPacketCacheKey(bundleInfo.Locator.BlobId, packetIdx);
 			AddToCache(decodedCacheKey, (ReadOnlyMemory<byte>)decodedPacket, decodedPacket.Length);
@@ -779,20 +737,27 @@ namespace EpicGames.Horde.Storage
 			BundleInfo bundleInfo = await GetBundleInfoAsync(locator.Blob, cancellationToken);
 			BundleExport export = bundleInfo.Header.Exports[locator.ExportIdx];
 
-			ExportInfo exportInfo = bundleInfo.Exports[locator.ExportIdx];
-
 			List<NodeLocator> refs = new List<NodeLocator>(export.References.Count);
-			for (int idx = 0; idx < export.References.Count; idx++)
+			foreach (BundleNodeRef reference in export.References)
 			{
-				NodeLocator reference = bundleInfo.References[export.References[idx]];
-				refs.Add(reference);
+				BlobLocator importBlob;
+				if (reference.ImportIdx == -1)
+				{
+					importBlob = locator.Blob;
+				}
+				else
+				{
+					importBlob = bundleInfo.Header.Imports[reference.ImportIdx];
+				}
+				Debug.Assert(importBlob.IsValid());
+				refs.Add(new NodeLocator(importBlob, reference.NodeIdx));
 			}
 
 			ReadOnlyMemory<byte> nodeData = ReadOnlyMemory<byte>.Empty;
 			if (export.Length > 0)
 			{
-				ReadOnlyMemory<byte> packetData = await ReadBundlePacketAsync(bundleInfo, exportInfo.PacketIdx, cancellationToken);
-				nodeData = packetData.Slice(exportInfo.Offset, export.Length);
+				ReadOnlyMemory<byte> packetData = await ReadBundlePacketAsync(bundleInfo, export.Packet, cancellationToken);
+				nodeData = packetData.Slice(export.Offset, export.Length);
 			}
 
 			TypeInfo? typeInfo = bundleInfo.Types[export.TypeIdx];
diff --git a/Engine/Source/Programs/Shared/EpicGames.Horde/Storage/TreeWriter.cs b/Engine/Source/Programs/Shared/EpicGames.Horde/Storage/TreeWriter.cs
--- a/Engine/Source/Programs/Shared/EpicGames.Horde/Storage/TreeWriter.cs
+++ b/Engine/Source/Programs/Shared/EpicGames.Horde/Storage/TreeWriter.cs
@@ -247,13 +247,17 @@ namespace EpicGames.Horde.Storage
 		internal class PendingNode : NodeHandle
 		{
 			public readonly NodeKey Key;
+			public readonly int Packet;
+			public readonly int Offset;
 			public readonly int Length;
 			public readonly NodeHandle[] Refs;
 
-			public PendingNode(NodeKey key, int length, IReadOnlyList<NodeHandle> refs, PendingBundle pendingBundle)
+			public PendingNode(NodeKey key, int packet, int offset, int length, IReadOnlyList<NodeHandle> refs, PendingBundle pendingBundle)
 				: base(key.Hash, default)
 			{
 				Key = key;
+				Packet = packet;
+				Offset = offset;
 				Length = length;
 				Refs = refs.ToArray();
 				PendingBundle = pendingBundle;
@@ -286,6 +290,9 @@ namespace EpicGames.Horde.Storage
 			// The current packet being built
 			IMemoryOwner<byte>? _currentPacket;
 
+			// Index of the current packet being written to
+			int _currentPacketIdx;
+
 			// Length of the packet that has been written
 			int _currentPacketLength;
 
@@ -414,10 +421,12 @@ namespace EpicGames.Horde.Storage
 			// Finish a node write
 			public PendingNode WriteNode(NodeKey nodeKey, int size, IReadOnlyList<NodeHandle> refs)
 			{
+				int offset = _currentPacketLength;
+
 				_currentPacketLength += size;
 				UncompressedLength += size;
 
-				PendingNode pendingNode = new PendingNode(nodeKey, (int)size, refs, this);
+				PendingNode pendingNode = new PendingNode(nodeKey, _currentPacketIdx, offset, (int)size, refs, this);
 				_queue.Add(pendingNode);
 				_nodeKeyToInfo.Add(nodeKey, pendingNode);
 
@@ -442,6 +451,7 @@ namespace EpicGames.Horde.Storage
 					}
 
 					_currentPacket = null;
+					_currentPacketIdx++;
 					_currentPacketLength = 0;
 				}
 			}
@@ -450,7 +460,8 @@ namespace EpicGames.Horde.Storage
 			{
 				int encodedLength = BundleData.Compress(_compressionFormat, packetData.Memory.Slice(0, length), _encodedPacketWriter);
 
-				BundlePacket packet = new BundlePacket(encodedLength, length);
+				int encodedOffset = (_packets.Count == 0) ? 0 : _packets[^1].EncodedOffset + _packets[^1].EncodedLength;
+				BundlePacket packet = new BundlePacket(_compressionFormat, encodedOffset, encodedLength, length);
 				_packets.Add(packet);
 
 				packetData.Dispose();
@@ -538,31 +549,20 @@ namespace EpicGames.Horde.Storage
 					}
 				}
 
-				// Map from node hash to index, with imported nodes first, ordered by blob, and exported nodes second.
-				int nodeIdx = 0;
-				Dictionary<NodeHandle, int> nodeToIndex = new Dictionary<NodeHandle, int>();
-
-				// Add all the imports and assign them identifiers
-				List<BundleImport> imports = new List<BundleImport>();
-				foreach ((BlobLocator blobLocator, List<(int, NodeHandle)> importEntries) in bundleToImports)
-				{
-					importEntries.SortBy(x => x.Item1);
-
-					int[] entries = new int[importEntries.Count];
-					for (int idx = 0; idx < importEntries.Count; idx++)
-					{
-						NodeHandle key = importEntries[idx].Item2;
-						nodeToIndex.Add(key, nodeIdx++);
-						entries[idx] = importEntries[idx].Item1;
-					}
-
-					imports.Add(new BundleImport(blobLocator, entries));
-				}
+				List<BlobLocator> imports = new List<BlobLocator>();
+				Dictionary<BlobLocator, int> importToIndex = new Dictionary<BlobLocator, int>();
 
 				// List of types in the bundle
 				List<BundleType> types = new List<BundleType>();
 				Dictionary<BundleType, int> typeToIndex = new Dictionary<BundleType, int>();
 
+				// Map of node handle to reference
+				Dictionary<NodeHandle, BundleNodeRef> nodeHandleToExportRef = new Dictionary<NodeHandle, BundleNodeRef>();
+				for (int exportIdx = 0; exportIdx < _queue.Count; exportIdx++)
+				{
+					nodeHandleToExportRef[_queue[exportIdx]] = new BundleNodeRef(-1, exportIdx);
+				}
+
 				// Create the export list
 				List<BundleExport> exports = new List<BundleExport>(_queue.Count);
 				foreach (PendingNode nodeInfo in _queue)
@@ -575,11 +575,26 @@ namespace EpicGames.Horde.Storage
 						types.Add(nodeInfo.Key.Type);
 					}
 
-					int[] references = nodeInfo.Refs.Select(x => nodeToIndex[x]).ToArray();
-					BundleExport export = new BundleExport(typeIdx, nodeInfo.Key.Hash, nodeInfo.Length, references);
-					nodeToIndex[nodeInfo] = nodeIdx;
+					List<BundleNodeRef> exportRefs = new List<BundleNodeRef>();
+					foreach (NodeHandle nodeRef in nodeInfo.Refs)
+					{
+						BundleNodeRef? exportRef;
+						if (!nodeHandleToExportRef.TryGetValue(nodeRef, out exportRef))
+						{
+							int importIdx;
+							if (!importToIndex.TryGetValue(nodeRef.Locator.Blob, out importIdx))
+							{
+								importIdx = imports.Count;
+								imports.Add(nodeRef.Locator.Blob);
+								importToIndex.Add(nodeRef.Locator.Blob, importIdx);
+							}
+							exportRef = new BundleNodeRef(importIdx, nodeRef.Locator.ExportIdx);
+						}
+						exportRefs.Add(exportRef);
+					}
+
+					BundleExport export = new BundleExport(typeIdx, nodeInfo.Key.Hash, nodeInfo.Packet, nodeInfo.Offset, nodeInfo.Length, exportRefs);
 					exports.Add(export);
-					nodeIdx++;
 				}
 
 				// Get the memory for each packet
@@ -602,7 +617,7 @@ namespace EpicGames.Horde.Storage
 				}
 
 				// Create the bundle
-				BundleHeader header = new BundleHeader(_compressionFormat, types, imports, exports, _packets.ToArray());
+				BundleHeader header = new BundleHeader(types, imports, exports, _packets.ToArray());
 				return new Bundle(header, packetData);
 			}
 		}
