Horde: Build fixes.

diff --git a/Engine/Source/Programs/Horde/Horde.Server.Tests/BlobStoreTests.cs b/Engine/Source/Programs/Horde/Horde.Server.Tests/BlobStoreTests.cs
--- a/Engine/Source/Programs/Horde/Horde.Server.Tests/BlobStoreTests.cs
+++ b/Engine/Source/Programs/Horde/Horde.Server.Tests/BlobStoreTests.cs
@@ -49,10 +49,10 @@ namespace Horde.Server.Tests
 
 			List<BlobLocator> imports = new List<BlobLocator>(refs);
 
-			List<BundleNodeRef> exportRefs = new List<BundleNodeRef>();
+			List<BundleExportRef> exportRefs = new List<BundleExportRef>();
 			for(int idx = 0; idx < refs.Count; idx++)
 			{
-				exportRefs.Add(new BundleNodeRef(idx, 0));
+				exportRefs.Add(new BundleExportRef(idx, 0));
 			}
 
 			List<BundleExport> exports = new List<BundleExport>();
diff --git a/Engine/Source/Programs/Horde/Horde.Server.Tests/StorageServiceTests.cs b/Engine/Source/Programs/Horde/Horde.Server.Tests/StorageServiceTests.cs
--- a/Engine/Source/Programs/Horde/Horde.Server.Tests/StorageServiceTests.cs
+++ b/Engine/Source/Programs/Horde/Horde.Server.Tests/StorageServiceTests.cs
@@ -36,9 +36,9 @@ namespace Horde.Server.Tests
 
 			List<BundleExport> exports = new List<BundleExport>();
 
-			exports.Add(new BundleExport(0, hash1, 0, 0, data1.Length, Array.Empty<BundleNodeRef>()));
-			exports.Add(new BundleExport(0, hash1, 0, 0, data1.Length, Array.Empty<BundleNodeRef>()));
-			exports.Add(new BundleExport(0, hash2, 0, data1.Length, data2.Length, Array.Empty<BundleNodeRef>()));
+			exports.Add(new BundleExport(0, hash1, 0, 0, data1.Length, Array.Empty<BundleExportRef>()));
+			exports.Add(new BundleExport(0, hash1, 0, 0, data1.Length, Array.Empty<BundleExportRef>()));
+			exports.Add(new BundleExport(0, hash2, 0, data1.Length, data2.Length, Array.Empty<BundleExportRef>()));
 
 			BundleHeader header = new BundleHeader(types, Array.Empty<BlobLocator>(), exports, Array.Empty<BundlePacket>());
 			Bundle bundle = new Bundle(header, Array.Empty<ReadOnlyMemory<byte>>());
diff --git a/Engine/Source/Programs/Horde/Horde/Commands/Bundles/BundleDump.cs b/Engine/Source/Programs/Horde/Horde/Commands/Bundles/BundleDump.cs
--- a/Engine/Source/Programs/Horde/Horde/Commands/Bundles/BundleDump.cs
+++ b/Engine/Source/Programs/Horde/Horde/Commands/Bundles/BundleDump.cs
@@ -55,7 +55,6 @@ namespace Horde.Commands.Bundles
 			logger.LogInformation("");
 			logger.LogInformation("Imports: {NumImports}", header.Imports.Count);
 
-			List<string> nodeNames = new List<string>();
 			for (int importIdx = 0; importIdx < header.Imports.Count; importIdx++)
 			{
 				logger.LogInformation("  IMP [{ImportIdx}] = {BlobId}", importIdx, header.Imports[importIdx]);
@@ -75,7 +74,7 @@ namespace Horde.Commands.Bundles
 				{
 					for(int referenceIdx = 0; referenceIdx < export.References.Count; referenceIdx++)
 					{
-						BundleNodeRef exportRef = export.References[referenceIdx];
+						BundleExportRef exportRef = export.References[referenceIdx];
 						NodeLocator nodeLocator = new NodeLocator(header.Imports[exportRef.ImportIdx], exportRef.NodeIdx);
 						logger.LogInformation("            REF {RefIdx,-3} -> {Node}", referenceIdx, nodeLocator);
 					}
diff --git a/Engine/Source/Programs/Shared/EpicGames.Horde.Tests/BundleStoreTests.cs b/Engine/Source/Programs/Shared/EpicGames.Horde.Tests/BundleStoreTests.cs
--- a/Engine/Source/Programs/Shared/EpicGames.Horde.Tests/BundleStoreTests.cs
+++ b/Engine/Source/Programs/Shared/EpicGames.Horde.Tests/BundleStoreTests.cs
@@ -75,7 +75,7 @@ namespace EpicGames.Horde.Tests
 			types.Add(new BundleType(Guid.Parse("F63606D4-5DBB-4061-A655-6F444F65229E"), 1));
 
 			List<BundleExport> exports = new List<BundleExport>();
-			exports.Add(new BundleExport(0, IoHash.Compute(payload), 0, 0, payload.Length, Array.Empty<BundleNodeRef>()));
+			exports.Add(new BundleExport(0, IoHash.Compute(payload), 0, 0, payload.Length, Array.Empty<BundleExportRef>()));
 
 			List<BundlePacket> packets = new List<BundlePacket>();
 			packets.Add(new BundlePacket(BundleCompressionFormat.None, 0, payload.Length, payload.Length));
diff --git a/Engine/Source/Programs/Shared/EpicGames.Horde.Tests/BundleTests.cs b/Engine/Source/Programs/Shared/EpicGames.Horde.Tests/BundleTests.cs
--- a/Engine/Source/Programs/Shared/EpicGames.Horde.Tests/BundleTests.cs
+++ b/Engine/Source/Programs/Shared/EpicGames.Horde.Tests/BundleTests.cs
@@ -145,8 +145,8 @@ namespace EpicGames.Horde.Tests
 				imports.Add(new BlobLocator("import2"));
 
 				List<BundleExport> exports = new List<BundleExport>();
-				exports.Add(new BundleExport(0, IoHash.Compute(Encoding.UTF8.GetBytes("export1")), 0, 0, 2, new BundleNodeRef[] { new BundleNodeRef(0, 5), new BundleNodeRef(0, 6) }));
-				exports.Add(new BundleExport(0, IoHash.Compute(Encoding.UTF8.GetBytes("export2")), 1, 0, 3, new BundleNodeRef[] { new BundleNodeRef(-1, 0) }));
+				exports.Add(new BundleExport(0, IoHash.Compute(Encoding.UTF8.GetBytes("export1")), 0, 0, 2, new BundleExportRef[] { new BundleExportRef(0, 5), new BundleExportRef(0, 6) }));
+				exports.Add(new BundleExport(0, IoHash.Compute(Encoding.UTF8.GetBytes("export2")), 1, 0, 3, new BundleExportRef[] { new BundleExportRef(-1, 0) }));
 
 				List<BundlePacket> packets = new List<BundlePacket>();
 				packets.Add(new BundlePacket(BundleCompressionFormat.LZ4, 0, 20, 40));
