Change Epicgames.ODIC pdbs to be non portable

Fixes Microsoft symbol stripping failing on this project

#rnx

diff --git a/Engine/Source/Programs/Shared/EpicGames.OIDC/EpicGames.OIDC.csproj b/Engine/Source/Programs/Shared/EpicGames.OIDC/EpicGames.OIDC.csproj
--- a/Engine/Source/Programs/Shared/EpicGames.OIDC/EpicGames.OIDC.csproj
+++ b/Engine/Source/Programs/Shared/EpicGames.OIDC/EpicGames.OIDC.csproj
@@ -4,6 +4,7 @@
     <OutputType>Library</OutputType>
     <RootNamespace>EpicGames.OIDC</RootNamespace>
     <TargetFramework>netcoreapp3.1</TargetFramework>
+    <DebugType>pdbonly</DebugType>
     <Nullable>enable</Nullable>
     <GenerateAssemblyInfo>false</GenerateAssemblyInfo>
     <GenerateTargetFrameworkAttribute>false</GenerateTargetFrameworkAttribute>
