Horde: Use flat byte arrays to store collections in bundles.

diff --git a/Engine/Source/Programs/Shared/EpicGames.Horde/Storage/Bundle.cs b/Engine/Source/Programs/Shared/EpicGames.Horde/Storage/Bundle.cs
--- a/Engine/Source/Programs/Shared/EpicGames.Horde/Storage/Bundle.cs
+++ b/Engine/Source/Programs/Shared/EpicGames.Horde/Storage/Bundle.cs
@@ -187,17 +187,17 @@ namespace EpicGames.Horde.Storage
 		/// <summary>
 		/// Bundles that we reference nodes in
 		/// </summary>
-		public IReadOnlyList<BlobLocator> Imports { get; }
+		public BundleImportCollection Imports { get; }
 
 		/// <summary>
 		/// Nodes exported from this bundle
 		/// </summary>
-		public IReadOnlyList<BundleExport> Exports { get; }
+		public BundleExportCollection Exports { get; }
 
 		/// <summary>
 		/// List of data packets within this bundle
 		/// </summary>
-		public IReadOnlyList<BundlePacket> Packets { get; }
+		public BundlePacketCollection Packets { get; }
 
 		/// <summary>
 		/// Constructs a new bundle header
@@ -207,6 +207,18 @@ namespace EpicGames.Horde.Storage
 		/// <param name="exports">Exports for nodes</param>
 		/// <param name="packets">Compression packets within the bundle</param>
 		public BundleHeader(IReadOnlyList<BundleType> types, IReadOnlyList<BlobLocator> imports, IReadOnlyList<BundleExport> exports, IReadOnlyList<BundlePacket> packets)
+			: this(types, new BundleImportCollection(imports), new BundleExportCollection(exports), new BundlePacketCollection(packets))
+		{
+		}
+
+		/// <summary>
+		/// Constructs a new bundle header
+		/// </summary>
+		/// <param name="types">Type array indexed by each export</param>
+		/// <param name="imports">Imported bundles</param>
+		/// <param name="exports">Exports for nodes</param>
+		/// <param name="packets">Compression packets within the bundle</param>
+		public BundleHeader(IReadOnlyList<BundleType> types, BundleImportCollection imports, BundleExportCollection exports, BundlePacketCollection packets)
 		{
 			Types = types;
 			Imports = imports;
@@ -267,7 +279,7 @@ namespace EpicGames.Horde.Storage
 				}
 			}
 
-			Imports = imports;
+			Imports = new BundleImportCollection(imports);
 
 			// Read the exports
 			int numExports = (int)reader.ReadUnsignedVarInt();
@@ -281,7 +293,7 @@ namespace EpicGames.Horde.Storage
 				IoHash hash = reader.ReadIoHash();
 				int length = (int)reader.ReadUnsignedVarInt();
 
-				List<BundleExportRef> exportReferences = new List<BundleExportRef>();
+				List<BundleExportRef> exportRefs = new List<BundleExportRef>();
 
 				int numReferences = (int)reader.ReadUnsignedVarInt();
 				if (numReferences > 0)
@@ -289,7 +301,7 @@ namespace EpicGames.Horde.Storage
 					for (int idx = 0; idx < numReferences; idx++)
 					{
 						int referenceIdx = (int)reader.ReadUnsignedVarInt();
-						exportReferences.Add(allExportReferences[referenceIdx]);
+						exportRefs.Add(allExportReferences[referenceIdx]);
 					}
 				}
 
@@ -298,7 +310,7 @@ namespace EpicGames.Horde.Storage
 					_ = reader.ReadUtf8String();
 				}
 
-				exportInfos.Add(new ExportInfo(typeIdx, hash, length, exportReferences));
+				exportInfos.Add(new ExportInfo(typeIdx, hash, length, exportRefs));
 			}
 
 			// Read the compression packets
@@ -313,7 +325,7 @@ namespace EpicGames.Horde.Storage
 					encodedOffset += export.Length;
 				}
 
-				Packets = packets;
+				Packets = new BundlePacketCollection(packets);
 			}
 			else
 			{
@@ -329,7 +341,7 @@ namespace EpicGames.Horde.Storage
 					encodedOffset += encodedLength;
 				}
 
-				Packets = packets;
+				Packets = new BundlePacketCollection(packets);
 			}
 
 			// Create the final export list
@@ -350,7 +362,7 @@ namespace EpicGames.Horde.Storage
 					packetOffset += exportInfo.Length;
 				}
 			}
-			Exports = exports;
+			Exports = new BundleExportCollection(exports);
 		}
 
 		/// <summary>
@@ -493,6 +505,11 @@ namespace EpicGames.Horde.Storage
 	[DebuggerDisplay("{Guid}#{Version}")]
 	public class BundleType : IEquatable<BundleType>
 	{
+		/// <summary>
+		/// Number of bytes in a serialized object
+		/// </summary>
+		public const int NumBytes = 20;
+
 		/// <summary>
 		/// Nominal identifier for the type
 		/// </summary>
@@ -556,11 +573,112 @@ namespace EpicGames.Horde.Storage
 		public static bool operator !=(BundleType lhs, BundleType rhs) => !lhs.Equals(rhs);
 	}
 
+	/// <summary>
+	/// Collection of imported node references
+	/// </summary>
+	public struct BundleImportCollection : IReadOnlyList<BlobLocator>
+	{
+		readonly ReadOnlyMemory<byte> _data;
+
+		/// <summary>
+		/// Deserializing constructor
+		/// </summary>
+		public BundleImportCollection(ReadOnlyMemory<byte> data) => _data = data;
+
+		/// <summary>
+		/// Constructor
+		/// </summary>
+		/// <param name="locators">Locators to write to the </param>
+		public BundleImportCollection(IReadOnlyCollection<BlobLocator> locators)
+		{
+			byte[] data = new byte[Measure(locators)];
+			Write(data, locators);
+			_data = data;
+		}
+
+		/// <summary>
+		/// Retrieve a single import from the collection
+		/// </summary>
+		public BlobLocator this[int index]
+		{
+			get
+			{
+				ReadOnlySpan<byte> span = _data.Span;
+				int offset = BinaryPrimitives.ReadInt32LittleEndian(span.Slice(index * sizeof(int)));
+				int length = span.Slice(offset).IndexOf((byte)0);
+				return new BlobLocator(new Utf8String(_data.Slice(offset, length)));
+			}
+		}
+
+		/// <inheritdoc/>
+		public int Count => (_data.Length == 0)? 0 : BinaryPrimitives.ReadInt32LittleEndian(_data.Span) / sizeof(int);
+
+		/// <inheritdoc/>
+		public IEnumerator<BlobLocator> GetEnumerator()
+		{
+			int count = Count;
+			for(int idx = 0; idx < count; idx++)
+			{
+				yield return this[idx];
+			}
+		}
+
+		/// <inheritdoc/>
+		IEnumerator IEnumerable.GetEnumerator() => GetEnumerator();
+
+		/// <summary>
+		/// Measure the size of memory required to store a collection of import locators
+		/// </summary>
+		/// <param name="locators">Locators to write</param>
+		/// <returns>Size in bytes of the output buffer</returns>
+		public static int Measure(IReadOnlyCollection<BlobLocator> locators)
+		{
+			int length = 0;
+			foreach (BlobLocator locator in locators)
+			{
+				length += sizeof(int) + locator.Inner.Length + 1;
+			}
+			return length;
+		}
+
+		/// <summary>
+		/// Serialize a collection of locators to memory
+		/// </summary>
+		/// <param name="data">Output buffer for the serialized data</param>
+		/// <param name="locators">Locators to write</param>
+		public static void Write(Span<byte> data, IReadOnlyCollection<BlobLocator> locators)
+		{
+			Span<byte> next = data;
+
+			int offset = locators.Count * sizeof(int);
+			foreach(BlobLocator locator in locators)
+			{
+				BinaryPrimitives.WriteInt32LittleEndian(next, offset);
+				offset += locator.Inner.Length + 1;
+				next = next.Slice(sizeof(int));
+			}
+
+			foreach (BlobLocator locator in locators)
+			{
+				locator.Inner.Span.CopyTo(next);
+				next = next.Slice(locator.Inner.Length);
+
+				next[0] = 0;
+				next = next.Slice(1);
+			}
+		}
+	}
+
 	/// <summary>
 	/// Descriptor for a compression packet
 	/// </summary>
-	public class BundlePacket
+	public struct BundlePacket
 	{
+		/// <summary>
+		/// Size of this structure when serialized
+		/// </summary>
+		public const int NumBytes = 16;
+
 		/// <summary>
 		/// Compression format for this packet
 		/// </summary>
@@ -595,55 +713,274 @@ namespace EpicGames.Horde.Storage
 			EncodedLength = encodedLength;
 			DecodedLength = decodedLength;
 		}
+
+		/// <summary>
+		/// Constructor
+		/// </summary>
+		public BundlePacket(ReadOnlySpan<byte> span)
+		{
+			CompressionFormat = (BundleCompressionFormat)BinaryPrimitives.ReadInt32LittleEndian(span);
+			EncodedOffset = BinaryPrimitives.ReadInt32LittleEndian(span.Slice(4));
+			EncodedLength = BinaryPrimitives.ReadInt32LittleEndian(span.Slice(8));
+			DecodedLength = BinaryPrimitives.ReadInt32LittleEndian(span.Slice(12));
+		}
+
+		/// <summary>
+		/// Serialize the struct to memory
+		/// </summary>
+		public void CopyTo(Span<byte> span)
+		{
+			BinaryPrimitives.WriteInt32LittleEndian(span, (int)CompressionFormat);
+			BinaryPrimitives.WriteInt32LittleEndian(span.Slice(4), EncodedOffset);
+			BinaryPrimitives.WriteInt32LittleEndian(span.Slice(8), EncodedLength);
+			BinaryPrimitives.WriteInt32LittleEndian(span.Slice(12), DecodedLength);
+		}
+	}
+
+	/// <summary>
+	/// Collection of information about packets in a bundle
+	/// </summary>
+	public struct BundlePacketCollection : IReadOnlyCollection<BundlePacket>
+	{
+		readonly ReadOnlyMemory<byte> _data;
+
+		/// <summary>
+		/// Constructor
+		/// </summary>
+		/// <param name="data"></param>
+		public BundlePacketCollection(ReadOnlyMemory<byte> data) => _data = data;
+
+		/// <summary>
+		/// Constructor
+		/// </summary>
+		/// <param name="packets">Packets to include in this collection</param>
+		public BundlePacketCollection(IReadOnlyCollection<BundlePacket> packets)
+		{
+			byte[] data = new byte[Measure(packets)];
+			Write(data, packets);
+			_data = data;
+		}
+
+		/// <inheritdoc/>
+		public int Count => _data.Length / BundlePacket.NumBytes;
+
+		/// <inheritdoc/>
+		public BundlePacket this[int index] => new BundlePacket(_data.Slice(index * BundlePacket.NumBytes, BundlePacket.NumBytes).Span);
+
+		/// <inheritdoc/>
+		public IEnumerator<BundlePacket> GetEnumerator()
+		{
+			ReadOnlyMemory<byte> remaining = _data;
+			while (remaining.Length > 0)
+			{
+				yield return new BundlePacket(remaining.Slice(0, BundlePacket.NumBytes).Span);
+				remaining = remaining.Slice(BundlePacket.NumBytes);
+			}
+		}
+
+		/// <inheritdoc/>
+		IEnumerator IEnumerable.GetEnumerator() => GetEnumerator();
+
+		/// <summary>
+		/// Measure the size of memory required to store a collection of import locators
+		/// </summary>
+		/// <param name="packets">Locators to write</param>
+		/// <returns>Size in bytes of the output buffer</returns>
+		public static int Measure(IReadOnlyCollection<BundlePacket> packets) => BundlePacket.NumBytes * packets.Count;
+
+		/// <summary>
+		/// Serialize a collection of packets to memory
+		/// </summary>
+		/// <param name="data">Output buffer for the serialized data</param>
+		/// <param name="packets">Packets to write</param>
+		public static void Write(Span<byte> data, IReadOnlyCollection<BundlePacket> packets)
+		{
+			Span<byte> next = data;
+			foreach (BundlePacket packet in packets)
+			{
+				packet.CopyTo(next);
+				next = next.Slice(BundlePacket.NumBytes);
+			}
+		}
 	}
 
 	/// <summary>
 	/// Entry for a node exported from an object
 	/// </summary>
-	public class BundleExport
+	public struct BundleExport
 	{
 		/// <summary>
-		/// Index of the type within the type list
+		/// Number of bytes in a serialized export object
 		/// </summary>
-		public int TypeIdx { get; }
+		public const int NumBytes = 32;
 
 		/// <summary>
-		/// Hash of the node
+		/// Raw data for this export
 		/// </summary>
-		public IoHash Hash { get; }
+		public ReadOnlyMemory<byte> Data { get; }
 
 		/// <summary>
-		/// Index of the packet containing this bundle
+		/// Hash of the node data
 		/// </summary>
-		public int Packet { get; }
+		public IoHash Hash => new IoHash(Data.Slice(0, 20).Span);
 
 		/// <summary>
-		/// Offset of this export's data within the (decompressed) packet
+		/// Type id of the node. Can be used to look up the type information from the bundle header.
 		/// </summary>
-		public int Offset { get; }
+		public int TypeIdx => BinaryPrimitives.ReadUInt16LittleEndian(Data.Slice(20).Span);
 
 		/// <summary>
-		/// Uncompressed length of this node
+		/// Packet containing this export's data
 		/// </summary>
-		public int Length { get; }
+		public int Packet => BinaryPrimitives.ReadUInt16LittleEndian(Data.Slice(22).Span);
+
+		/// <summary>
+		/// Offset within the packet of the node data
+		/// </summary>
+		public int Offset => BinaryPrimitives.ReadInt32LittleEndian(Data.Slice(24).Span);
+
+		/// <summary>
+		/// Length of the node
+		/// </summary>
+		public int Length => BinaryPrimitives.ReadInt32LittleEndian(Data.Slice(28).Span);
 
 		/// <summary>
 		/// References to other nodes
 		/// </summary>
-		public IReadOnlyList<BundleExportRef> References { get; }
+		public BundleExportRefCollection References { get; }
 
 		/// <summary>
 		/// Constructor
 		/// </summary>
-		public BundleExport(int typeIdx, IoHash hash, int packet, int offset, int length, IReadOnlyList<BundleExportRef> references)
+		public BundleExport(ReadOnlyMemory<byte> data, BundleExportRefCollection references)
 		{
-			TypeIdx = typeIdx;
-			Hash = hash;
-			Packet = packet;
-			Offset = offset;
-			Length = length;
+			Data = data;
 			References = references;
 		}
+
+		/// <summary>
+		/// Constructor
+		/// </summary>
+		public BundleExport(int typeIdx, IoHash hash, int packetIdx, int offset, int length, IReadOnlyList<BundleExportRef> references)
+		{
+			byte[] data = new byte[NumBytes];
+			Write(data, typeIdx, hash, packetIdx, offset, length);
+			Data = data;
+			References = new BundleExportRefCollection(references);
+		}
+
+		/// <summary>
+		/// Writes a new export to a block of memory
+		/// </summary>
+		public static void Write(Span<byte> data, int typeIdx, IoHash hash, int packet, int offset, int length)
+		{
+			hash.CopyTo(data);
+			BinaryPrimitives.WriteUInt16LittleEndian(data.Slice(20), (ushort)typeIdx);
+			BinaryPrimitives.WriteUInt16LittleEndian(data.Slice(22), (ushort)packet);
+			BinaryPrimitives.WriteInt32LittleEndian(data.Slice(24), offset);
+			BinaryPrimitives.WriteInt32LittleEndian(data.Slice(28), length);
+		}
+	}
+
+	/// <summary>
+	/// Entry for a node exported from an object
+	/// </summary>
+	public class BundleExportCollection : IReadOnlyList<BundleExport>
+	{
+		readonly int _count;
+		readonly ReadOnlyMemory<byte> _data;
+		readonly ReadOnlyMemory<byte> _refIndexData;
+		readonly BundleExportRefCollection _refs;
+
+		/// <summary>
+		/// Constructor
+		/// </summary>
+		public BundleExportCollection(ReadOnlyMemory<byte> data, ReadOnlyMemory<byte> refData)
+		{
+			_count = data.Length / BundleExport.NumBytes;
+			_data = data;
+
+			if (refData.Length > 0)
+			{
+				_refIndexData = refData.Slice(0, _count * sizeof(int));
+				_refs = new BundleExportRefCollection(refData.Slice(_refIndexData.Length));
+			}
+		}
+
+		/// <summary>
+		/// Constructor
+		/// </summary>
+		public BundleExportCollection(IReadOnlyList<BundleExport> exports)
+		{
+			int exportRefLength = 0;
+			foreach (BundleExport export in exports)
+			{
+				exportRefLength += export.References.Data.Length;
+			}
+
+			if (exportRefLength > 0)
+			{
+				byte[] refData = new byte[exportRefLength];
+				byte[] refIndexData = new byte[exports.Count * sizeof(int)];
+
+				int refDataOffset = 0;
+				int refIndexOffset = 0;
+
+				foreach (BundleExport export in exports)
+				{
+					BinaryPrimitives.WriteInt32LittleEndian(refIndexData.AsSpan(refIndexOffset), refDataOffset);
+					refIndexOffset += sizeof(int);
+
+					export.References.Data.Span.CopyTo(refData.AsSpan(refDataOffset));
+					refDataOffset += export.References.Data.Length;
+				}
+
+				_refIndexData = refIndexData;
+				_refs = new BundleExportRefCollection(refData);
+			}
+
+			byte[] data = new byte[exports.Count * BundleExport.NumBytes];
+			for (int idx = 0; idx < exports.Count; idx++)
+			{
+				exports[idx].Data.CopyTo(data.AsMemory(idx * BundleExport.NumBytes));
+			}
+			_data = data;
+		}
+
+		/// <inheritdoc/>
+		public int Count => _data.Length / BundleExport.NumBytes;
+
+		/// <inheritdoc/>
+		public BundleExport this[int index]
+		{
+			get
+			{
+				ReadOnlyMemory<byte> exportData = _data.Slice(index * BundleExport.NumBytes);
+
+				BundleExportRefCollection exportRefs = new BundleExportRefCollection();
+				if (_refIndexData.Length > 0)
+				{
+					int minOffset = BinaryPrimitives.ReadInt32LittleEndian(_refIndexData.Span.Slice(index * sizeof(int)));
+					int maxOffset = (index == Count - 1)? _refs.Data.Length : BinaryPrimitives.ReadInt32LittleEndian(_refIndexData.Span.Slice((index + 1) * sizeof(int)));
+					exportRefs = new BundleExportRefCollection(_refs.Data.Slice(minOffset, maxOffset - minOffset));
+				}
+
+				return new BundleExport(exportData, exportRefs);
+			}
+		}
+
+		/// <inheritdoc/>
+		public IEnumerator<BundleExport> GetEnumerator()
+		{
+			int count = Count;
+			for (int idx = 0; idx < count; idx++)
+			{
+				yield return this[idx];
+			}
+		}
+
+		/// <inheritdoc/>
+		IEnumerator IEnumerable.GetEnumerator() => GetEnumerator();
 	}
 
 	/// <summary>
@@ -651,7 +988,79 @@ namespace EpicGames.Horde.Storage
 	/// </summary>
 	/// <param name="ImportIdx">Index into the import table of the blob containing the referenced node. Can be -1 for references within the same bundle.</param>
 	/// <param name="NodeIdx">Node imported from the bundle</param>
-	public record class BundleExportRef(int ImportIdx, int NodeIdx);
+	public record struct BundleExportRef(int ImportIdx, int NodeIdx)
+	{
+		/// <summary>
+		/// Number of bytes in the serialized object
+		/// </summary>
+		public const int NumBytes = 4;
+
+		/// <summary>
+		/// Deserialize this object from memory
+		/// </summary>
+		public static BundleExportRef Read(ReadOnlySpan<byte> data)
+		{
+			return new BundleExportRef(BinaryPrimitives.ReadInt16LittleEndian(data), BinaryPrimitives.ReadInt16LittleEndian(data.Slice(2)));
+		}
+
+		/// <summary>
+		/// Serialize this object to memory
+		/// </summary>
+		public void CopyTo(Span<byte> data)
+		{
+			BinaryPrimitives.WriteInt16LittleEndian(data, (short)ImportIdx);
+			BinaryPrimitives.WriteUInt16LittleEndian(data.Slice(2), (ushort)NodeIdx);
+		}
+	}
+
+	/// <summary>
+	/// Collection of information about exported nodes
+	/// </summary>
+	public struct BundleExportRefCollection : IReadOnlyList<BundleExportRef>
+	{
+		/// <summary>
+		/// Data used to store this collection
+		/// </summary>
+		public ReadOnlyMemory<byte> Data { get; }
+
+		/// <summary>
+		/// Constructor
+		/// </summary>
+		public BundleExportRefCollection(ReadOnlyMemory<byte> data) => Data = data;
+
+		/// <summary>
+		/// Constructor
+		/// </summary>
+		public BundleExportRefCollection(IReadOnlyList<BundleExportRef> exportRefs)
+		{
+			byte[] data = new byte[exportRefs.Count * BundleExportRef.NumBytes];
+			for (int idx = 0; idx < exportRefs.Count; idx++)
+			{
+				exportRefs[idx].CopyTo(data.AsSpan(idx * BundleExportRef.NumBytes));
+			}
+			Data = data;
+		}
+
+		/// <inheritdoc/>
+		public int Count => Data.Length / BundleExportRef.NumBytes;
+
+		/// <inheritdoc/>
+		public BundleExportRef this[int index] => BundleExportRef.Read(Data.Span.Slice(index * BundleExportRef.NumBytes));
+
+		/// <inheritdoc/>
+		public IEnumerator<BundleExportRef> GetEnumerator()
+		{
+			ReadOnlyMemory<byte> data = Data;
+			while (data.Length > 0)
+			{
+				yield return BundleExportRef.Read(data.Span);
+				data = data.Slice(BundleExportRef.NumBytes);
+			}
+		}
+
+		/// <inheritdoc/>
+		IEnumerator IEnumerable.GetEnumerator() => GetEnumerator();
+	}
 
 	/// <summary>
 	/// Utility methods for bundles
diff --git a/Engine/Source/Programs/Shared/EpicGames.Horde/Storage/TreeWriter.cs b/Engine/Source/Programs/Shared/EpicGames.Horde/Storage/TreeWriter.cs
--- a/Engine/Source/Programs/Shared/EpicGames.Horde/Storage/TreeWriter.cs
+++ b/Engine/Source/Programs/Shared/EpicGames.Horde/Storage/TreeWriter.cs
@@ -578,7 +578,7 @@ namespace EpicGames.Horde.Storage
 					List<BundleExportRef> exportRefs = new List<BundleExportRef>();
 					foreach (NodeHandle nodeRef in nodeInfo.Refs)
 					{
-						BundleExportRef? exportRef;
+						BundleExportRef exportRef;
 						if (!nodeHandleToExportRef.TryGetValue(nodeRef, out exportRef))
 						{
 							int importIdx;
@@ -617,7 +617,7 @@ namespace EpicGames.Horde.Storage
 				}
 
 				// Create the bundle
-				BundleHeader header = new BundleHeader(types, imports, exports, _packets.ToArray());
+				BundleHeader header = new BundleHeader(types, new BundleImportCollection(imports), new BundleExportCollection(exports), new BundlePacketCollection(_packets));
 				return new Bundle(header, packetData);
 			}
 		}
