Horde: Add a collection for storing a flat list of bundle node types.

diff --git a/Engine/Source/Programs/Shared/EpicGames.Horde/Storage/Bundle.cs b/Engine/Source/Programs/Shared/EpicGames.Horde/Storage/Bundle.cs
--- a/Engine/Source/Programs/Shared/EpicGames.Horde/Storage/Bundle.cs
+++ b/Engine/Source/Programs/Shared/EpicGames.Horde/Storage/Bundle.cs
@@ -182,7 +182,7 @@ namespace EpicGames.Horde.Storage
 		/// <summary>
 		/// Types for exports within this bundle
 		/// </summary>
-		public IReadOnlyList<BundleType> Types { get; }
+		public BundleTypeCollection Types { get; }
 
 		/// <summary>
 		/// Bundles that we reference nodes in
@@ -207,7 +207,7 @@ namespace EpicGames.Horde.Storage
 		/// <param name="exports">Exports for nodes</param>
 		/// <param name="packets">Compression packets within the bundle</param>
 		public BundleHeader(IReadOnlyList<BundleType> types, IReadOnlyList<BlobLocator> imports, IReadOnlyList<BundleExport> exports, IReadOnlyList<BundlePacket> packets)
-			: this(types, new BundleImportCollection(imports), new BundleExportCollection(exports), new BundlePacketCollection(packets))
+			: this(new BundleTypeCollection(types), new BundleImportCollection(imports), new BundleExportCollection(exports), new BundlePacketCollection(packets))
 		{
 		}
 
@@ -218,7 +218,7 @@ namespace EpicGames.Horde.Storage
 		/// <param name="imports">Imported bundles</param>
 		/// <param name="exports">Exports for nodes</param>
 		/// <param name="packets">Compression packets within the bundle</param>
-		public BundleHeader(IReadOnlyList<BundleType> types, BundleImportCollection imports, BundleExportCollection exports, BundlePacketCollection packets)
+		public BundleHeader(BundleTypeCollection types, BundleImportCollection imports, BundleExportCollection exports, BundlePacketCollection packets)
 		{
 			Types = types;
 			Imports = imports;
@@ -259,7 +259,7 @@ namespace EpicGames.Horde.Storage
 				types.Add(new BundleType(guid, serializerVersion));
 			}
 
-			Types = types;
+			Types = new BundleTypeCollection(types);
 
 			// Read the imports
 			int numImports = (int)reader.ReadUnsignedVarInt();
@@ -430,7 +430,6 @@ namespace EpicGames.Horde.Storage
 			BundleCompressionFormat compressionFormat = (Packets.Count > 0) ? Packets[0].CompressionFormat : BundleCompressionFormat.None;
 
 			writer.WriteUnsignedVarInt((ulong)compressionFormat);
-			writer.WriteVariableLengthArray(Types, x => x.Write(writer));
 
 			// Find all the referenced nodes in each import
 			SortedSet<int>[] nodes = new SortedSet<int>[Imports.Count];
@@ -459,6 +458,14 @@ namespace EpicGames.Horde.Storage
 				exportRefToIndex[new BundleExportRef(-1, exportIdx)] = index;
 			}
 
+			// Write all the types
+			writer.WriteUnsignedVarInt(Types.Count);
+			for (int typeIdx = 0; typeIdx < Types.Count; typeIdx++)
+			{
+				writer.WriteGuid(Types[typeIdx].Guid);
+				writer.WriteUnsignedVarInt(Types[typeIdx].Version);
+			}
+
 			// Write all the imports
 			writer.WriteUnsignedVarInt(Imports.Count);
 			for (int importIdx = 0; importIdx < Imports.Count; importIdx++)
@@ -531,14 +538,23 @@ namespace EpicGames.Horde.Storage
 			Version = version;
 		}
 
+		/// <summary>
+		/// Constructor
+		/// </summary>
+		/// <param name="data">Data to read from</param>
+		public BundleType(ReadOnlySpan<byte> data)
+			: this(new Guid(data.Slice(0, 16)), BinaryPrimitives.ReadInt32LittleEndian(data.Slice(16)))
+		{
+		}
+
 		/// <summary>
 		/// Serializes a type to storage
 		/// </summary>
-		/// <param name="writer">Writer for serialization</param>
-		public void Write(IMemoryWriter writer)
+		/// <param name="data">Data to write to</param>
+		public void CopyTo(Span<byte> data)
 		{
-			writer.WriteGuid(Guid);
-			writer.WriteUnsignedVarInt(Version);
+			Guid.TryWriteBytes(data.Slice(0, 16));
+			BinaryPrimitives.WriteInt32LittleEndian(data.Slice(16), Version);
 		}
 
 		/// <summary>
@@ -573,6 +589,58 @@ namespace EpicGames.Horde.Storage
 		public static bool operator !=(BundleType lhs, BundleType rhs) => !lhs.Equals(rhs);
 	}
 
+	/// <summary>
+	/// Collection of node types in a bundle
+	/// </summary>
+	public class BundleTypeCollection : IReadOnlyList<BundleType>
+	{
+		/// <summary>
+		/// Data for this collection
+		/// </summary>
+		public ReadOnlyMemory<byte> Data { get; }
+
+		/// <summary>
+		/// Constructor
+		/// </summary>
+		public BundleTypeCollection(ReadOnlyMemory<byte> data) => Data = data;
+
+		/// <summary>
+		/// Constructor
+		/// </summary>
+		public BundleTypeCollection(IReadOnlyList<BundleType> types)
+		{
+			byte[] data = new byte[types.Count * BundleType.NumBytes];
+
+			int offset = 0;
+			foreach (BundleType type in types)
+			{
+				type.CopyTo(data.AsSpan(offset));
+				offset += BundleType.NumBytes;
+			}
+
+			Data = data;
+		}
+
+		/// <inheritdoc/>
+		public int Count => Data.Length / BundleType.NumBytes;
+
+		/// <inheritdoc/>
+		public BundleType this[int index] => new BundleType(Data.Slice(index * BundleType.NumBytes).Span);
+
+		/// <inheritdoc/>
+		public IEnumerator<BundleType> GetEnumerator()
+		{
+			int count = Count;
+			for (int idx = 0; idx < count; idx++)
+			{
+				yield return this[idx];
+			}
+		}
+
+		/// <inheritdoc/>
+		IEnumerator IEnumerable.GetEnumerator() => GetEnumerator();
+	}
+
 	/// <summary>
 	/// Collection of imported node references
 	/// </summary>
diff --git a/Engine/Source/Programs/Shared/EpicGames.Horde/Storage/TreeWriter.cs b/Engine/Source/Programs/Shared/EpicGames.Horde/Storage/TreeWriter.cs
--- a/Engine/Source/Programs/Shared/EpicGames.Horde/Storage/TreeWriter.cs
+++ b/Engine/Source/Programs/Shared/EpicGames.Horde/Storage/TreeWriter.cs
@@ -617,7 +617,7 @@ namespace EpicGames.Horde.Storage
 				}
 
 				// Create the bundle
-				BundleHeader header = new BundleHeader(types, new BundleImportCollection(imports), new BundleExportCollection(exports), new BundlePacketCollection(_packets));
+				BundleHeader header = new BundleHeader(new BundleTypeCollection(types), new BundleImportCollection(imports), new BundleExportCollection(exports), new BundlePacketCollection(_packets));
 				return new Bundle(header, packetData);
 			}
 		}
