EpicGames.Perforce fixing parameter string for --parallel

#rnx

diff --git a/Engine/Source/Programs/Shared/EpicGames.Perforce/PerforceConnection.cs b/Engine/Source/Programs/Shared/EpicGames.Perforce/PerforceConnection.cs
--- a/Engine/Source/Programs/Shared/EpicGames.Perforce/PerforceConnection.cs
+++ b/Engine/Source/Programs/Shared/EpicGames.Perforce/PerforceConnection.cs
@@ -3808,7 +3808,7 @@ namespace EpicGames.Perforce
 			// Using multiple threads is not supported through p4.exe due to threaded output not being parsable
 			if(numThreads != -1 && (connection is NativePerforceConnection))
 			{
-				StringBuilder argument = new StringBuilder($"--parallel-threads={numThreads}");
+				StringBuilder argument = new StringBuilder($"--parallel=threads={numThreads}");
 				if(batch != -1)
 				{
 					argument.Append($",batch={batch}");
