Horde: Store all bundle data in a flat memory buffer. Still uses legacy serialization format for now.

diff --git a/Engine/Source/Programs/Horde/Horde.Server.Tests/BlobStoreTests.cs b/Engine/Source/Programs/Horde/Horde.Server.Tests/BlobStoreTests.cs
--- a/Engine/Source/Programs/Horde/Horde.Server.Tests/BlobStoreTests.cs
+++ b/Engine/Source/Programs/Horde/Horde.Server.Tests/BlobStoreTests.cs
@@ -61,7 +61,7 @@ namespace Horde.Server.Tests
 			List<BundlePacket> packets = new List<BundlePacket>();
 			packets.Add(new BundlePacket(BundleCompressionFormat.None, 0, (int)data.Length, (int)data.Length));
 
-			BundleHeader header = new BundleHeader(types, imports, exports, packets);
+			BundleHeader header = BundleHeader.Create(types, imports, exports, packets);
 			return new Bundle(header, new[] { data });
 		}
 
diff --git a/Engine/Source/Programs/Horde/Horde.Server.Tests/GcServiceTests.cs b/Engine/Source/Programs/Horde/Horde.Server.Tests/GcServiceTests.cs
--- a/Engine/Source/Programs/Horde/Horde.Server.Tests/GcServiceTests.cs
+++ b/Engine/Source/Programs/Horde/Horde.Server.Tests/GcServiceTests.cs
@@ -104,7 +104,7 @@ namespace Horde.Server.Tests
 			{
 				List<BundleType> types = new List<BundleType> { new BundleType(Guid.Parse("{AFDF76A7-5333-4DEE-B837-B5F5CA511245}"), 0) };
 				List<BlobLocator> imports = children[idx].ConvertAll(x => locators[x]);
-				BundleHeader header = new BundleHeader(types, imports, Array.Empty<BundleExport>(), Array.Empty<BundlePacket>());
+				BundleHeader header = BundleHeader.Create(types, imports, Array.Empty<BundleExport>(), Array.Empty<BundlePacket>());
 				Bundle bundle = new Bundle(header, Array.Empty<ReadOnlyMemory<byte>>());
 				locators[idx] = await store.WriteBundleAsync(bundle, prefix: "gctest");
 			}
diff --git a/Engine/Source/Programs/Horde/Horde.Server.Tests/StorageServiceTests.cs b/Engine/Source/Programs/Horde/Horde.Server.Tests/StorageServiceTests.cs
--- a/Engine/Source/Programs/Horde/Horde.Server.Tests/StorageServiceTests.cs
+++ b/Engine/Source/Programs/Horde/Horde.Server.Tests/StorageServiceTests.cs
@@ -40,7 +40,7 @@ namespace Horde.Server.Tests
 			exports.Add(new BundleExport(0, hash1, 0, 0, data1.Length, Array.Empty<BundleExportRef>()));
 			exports.Add(new BundleExport(0, hash2, 0, data1.Length, data2.Length, Array.Empty<BundleExportRef>()));
 
-			BundleHeader header = new BundleHeader(types, Array.Empty<BlobLocator>(), exports, Array.Empty<BundlePacket>());
+			BundleHeader header = BundleHeader.Create(types, Array.Empty<BlobLocator>(), exports, Array.Empty<BundlePacket>());
 			Bundle bundle = new Bundle(header, Array.Empty<ReadOnlyMemory<byte>>());
 			BlobLocator locator = await client.WriteBundleAsync(bundle);
 
diff --git a/Engine/Source/Programs/Horde/Horde.Server/Storage/StorageService.cs b/Engine/Source/Programs/Horde/Horde.Server/Storage/StorageService.cs
--- a/Engine/Source/Programs/Horde/Horde.Server/Storage/StorageService.cs
+++ b/Engine/Source/Programs/Horde/Horde.Server/Storage/StorageService.cs
@@ -1047,7 +1047,7 @@ namespace Horde.Server.Storage
 						int headerSize = BundleHeader.ReadPrelude(memory);
 						if (headerSize <= fetchSize)
 						{
-							return new BundleHeader(new MemoryReader(memory));
+							return BundleHeader.Read(new MemoryReader(memory));
 						}
 
 						// Increase the fetch size and retry
diff --git a/Engine/Source/Programs/Shared/EpicGames.Horde.Tests/BundleStoreTests.cs b/Engine/Source/Programs/Shared/EpicGames.Horde.Tests/BundleStoreTests.cs
--- a/Engine/Source/Programs/Shared/EpicGames.Horde.Tests/BundleStoreTests.cs
+++ b/Engine/Source/Programs/Shared/EpicGames.Horde.Tests/BundleStoreTests.cs
@@ -80,7 +80,7 @@ namespace EpicGames.Horde.Tests
 			List<BundlePacket> packets = new List<BundlePacket>();
 			packets.Add(new BundlePacket(BundleCompressionFormat.None, 0, payload.Length, payload.Length));
 
-			BundleHeader header = new BundleHeader(types, Array.Empty<BlobLocator>(), exports, packets);
+			BundleHeader header = BundleHeader.Create(types, Array.Empty<BlobLocator>(), exports, packets);
 			return new Bundle(header, new List<ReadOnlyMemory<byte>> { payload });
 		}
 
diff --git a/Engine/Source/Programs/Shared/EpicGames.Horde.Tests/BundleTests.cs b/Engine/Source/Programs/Shared/EpicGames.Horde.Tests/BundleTests.cs
--- a/Engine/Source/Programs/Shared/EpicGames.Horde.Tests/BundleTests.cs
+++ b/Engine/Source/Programs/Shared/EpicGames.Horde.Tests/BundleTests.cs
@@ -152,14 +152,14 @@ namespace EpicGames.Horde.Tests
 				packets.Add(new BundlePacket(BundleCompressionFormat.LZ4, 0, 20, 40));
 				packets.Add(new BundlePacket(BundleCompressionFormat.LZ4, 20, 10, 20));
 
-				oldHeader = new BundleHeader(types, imports, exports, packets);
+				oldHeader = BundleHeader.Create(types, imports, exports, packets);
 			}
 
 			ByteArrayBuilder writer = new ByteArrayBuilder();
 			oldHeader.Write(writer);
 			byte[] serializedData = writer.ToByteArray();
 
-			BundleHeader newHeader = new BundleHeader(new MemoryReader(serializedData));
+			BundleHeader newHeader = BundleHeader.Read(new MemoryReader(serializedData));
 				
 			Assert.AreEqual(oldHeader.Imports.Count, newHeader.Imports.Count);
 			for (int idx = 0; idx < oldHeader.Imports.Count; idx++)
diff --git a/Engine/Source/Programs/Shared/EpicGames.Horde/Storage/Bundle.cs b/Engine/Source/Programs/Shared/EpicGames.Horde/Storage/Bundle.cs
--- a/Engine/Source/Programs/Shared/EpicGames.Horde/Storage/Bundle.cs
+++ b/Engine/Source/Programs/Shared/EpicGames.Horde/Storage/Bundle.cs
@@ -81,7 +81,7 @@ namespace EpicGames.Horde.Storage
 		/// </summary>
 		public Bundle(IMemoryReader reader)
 		{
-			Header = new BundleHeader(reader);
+			Header = BundleHeader.Read(reader);
 
 			ReadOnlyMemory<byte>[] packets = new ReadOnlyMemory<byte>[Header.Packets.Count];
 			for (int idx = 0; idx < Header.Packets.Count; idx++)
@@ -164,6 +164,37 @@ namespace EpicGames.Horde.Storage
 		Oodle = 3,
 	}
 
+	/// <summary>
+	/// Identifier for the type of a section in the bundle header
+	/// </summary>
+	public enum BundleSectionType
+	{
+		/// <summary>
+		/// List of custom types
+		/// </summary>
+		Types = 0,
+
+		/// <summary>
+		/// Imports of other bundles
+		/// </summary>
+		Imports = 1,
+
+		/// <summary>
+		/// List of exports
+		/// </summary>
+		Exports = 2,
+
+		/// <summary>
+		/// References to exports in other bundles
+		/// </summary>
+		ExportRefs = 3,
+
+		/// <summary>
+		/// Packet headers
+		/// </summary>
+		Packets = 4,
+	}
+
 	/// <summary>
 	/// Header for the contents of a bundle.
 	/// </summary>
@@ -179,6 +210,11 @@ namespace EpicGames.Horde.Storage
 		/// </summary>
 		public const int PreludeLength = 8;
 
+		/// <summary>
+		/// Underlying data for the header
+		/// </summary>
+		public ReadOnlyMemory<byte> Data { get; }
+
 		/// <summary>
 		/// Types for exports within this bundle
 		/// </summary>
@@ -200,15 +236,48 @@ namespace EpicGames.Horde.Storage
 		public BundlePacketCollection Packets { get; }
 
 		/// <summary>
-		/// Constructs a new bundle header
+		/// Construct a header from the given data encoded in the latest format
 		/// </summary>
-		/// <param name="types">Type array indexed by each export</param>
-		/// <param name="imports">Imported bundles</param>
-		/// <param name="exports">Exports for nodes</param>
-		/// <param name="packets">Compression packets within the bundle</param>
-		public BundleHeader(IReadOnlyList<BundleType> types, IReadOnlyList<BlobLocator> imports, IReadOnlyList<BundleExport> exports, IReadOnlyList<BundlePacket> packets)
-			: this(new BundleTypeCollection(types), new BundleImportCollection(imports), new BundleExportCollection(exports), new BundlePacketCollection(packets))
+		/// <param name="data">Data for the header, including the prelude</param>
+		private BundleHeader(ReadOnlyMemory<byte> data)
 		{
+			Data = data;
+
+			ReadOnlyMemory<byte> exportData = ReadOnlyMemory<byte>.Empty;
+			ReadOnlyMemory<byte> exportRefData = ReadOnlyMemory<byte>.Empty;
+
+			for (int offset = PreludeLength; offset < data.Length;)
+			{
+				uint header = BinaryPrimitives.ReadUInt32LittleEndian(data.Slice(offset).Span);
+				offset += sizeof(int);
+
+				int length = (int)(header >> 8);
+				ReadOnlyMemory<byte> sectionData = data.Slice(offset, length);
+
+				BundleSectionType type = (BundleSectionType)(header & 255);
+				switch (type)
+				{
+					case BundleSectionType.Types:
+						Types = new BundleTypeCollection(sectionData);
+						break;
+					case BundleSectionType.Imports:
+						Imports = new BundleImportCollection(sectionData);
+						break;
+					case BundleSectionType.Exports:
+						exportData = sectionData;
+						break;
+					case BundleSectionType.ExportRefs:
+						exportRefData = sectionData;
+						break;
+					case BundleSectionType.Packets:
+						Packets = new BundlePacketCollection(sectionData);
+						break;
+				}
+
+				offset += length;
+			}
+
+			Exports = new BundleExportCollection(exportData, exportRefData);
 		}
 
 		/// <summary>
@@ -218,12 +287,112 @@ namespace EpicGames.Horde.Storage
 		/// <param name="imports">Imported bundles</param>
 		/// <param name="exports">Exports for nodes</param>
 		/// <param name="packets">Compression packets within the bundle</param>
-		public BundleHeader(BundleTypeCollection types, BundleImportCollection imports, BundleExportCollection exports, BundlePacketCollection packets)
+		public static BundleHeader Create(IReadOnlyList<BundleType> types, IReadOnlyList<BlobLocator> imports, IReadOnlyList<BundleExport> exports, IReadOnlyList<BundlePacket> packets)
 		{
-			Types = types;
-			Imports = imports;
-			Exports = exports;
-			Packets = packets;
+			// Find the size of all the sections
+			int typesLength = BundleTypeCollection.Measure(types);
+			int importsLength = BundleImportCollection.Measure(imports);
+			int exportsLength = BundleExportCollection.Measure(exports);
+
+			int exportRefsLength = exports.Sum(x => x.References.Data.Length);
+			if (exportRefsLength > 0)
+			{
+				exportRefsLength += exports.Count * sizeof(int);
+			}
+
+			int packetsLength = BundlePacketCollection.Measure(packets);
+
+			// Find the total size of the header
+			int length = PreludeLength;
+			if (typesLength > 0)
+			{
+				length += SectionHeaderLength + typesLength;
+			}
+			if (importsLength > 0)
+			{
+				length += SectionHeaderLength + importsLength;
+			}
+			if (exportsLength > 0)
+			{
+				length += SectionHeaderLength + exportsLength;
+			}
+			if (exportRefsLength > 0)
+			{
+				length += SectionHeaderLength + exportRefsLength;
+			}
+			if (packetsLength > 0)
+			{
+				length += SectionHeaderLength + packetsLength;
+			}
+
+			// Allocate the header data
+			byte[] data = new byte[length];
+
+			// Write all the sections
+			Span<byte> next = data.AsSpan(PreludeLength);
+			if (typesLength > 0)
+			{
+				WriteSectionHeader(next, BundleSectionType.Types, typesLength);
+				next = next.Slice(SectionHeaderLength);
+
+				BundleTypeCollection.Write(next.Slice(0, typesLength), types);
+				next = next.Slice(typesLength);
+			}
+			if (importsLength > 0)
+			{
+				WriteSectionHeader(next, BundleSectionType.Imports, importsLength);
+				next = next.Slice(SectionHeaderLength);
+
+				BundleImportCollection.Write(next.Slice(0, importsLength), imports);
+				next = next.Slice(importsLength);
+			}
+			if (exportsLength > 0)
+			{
+				WriteSectionHeader(next, BundleSectionType.Exports, exportsLength);
+				next = next.Slice(SectionHeaderLength);
+
+				BundleExportCollection.Write(next.Slice(0, exportsLength), exports);
+				next = next.Slice(exportsLength);
+			}
+			if (exportRefsLength > 0)
+			{
+				WriteSectionHeader(next, BundleSectionType.ExportRefs, exportRefsLength);
+				next = next.Slice(SectionHeaderLength);
+
+				int indexOffset = 0;
+				int dataOffset = sizeof(int) * exports.Count;
+
+				foreach (BundleExport export in exports)
+				{
+					BinaryPrimitives.WriteInt32LittleEndian(next.Slice(indexOffset), dataOffset);
+					indexOffset += sizeof(int);
+
+					export.References.Data.Span.CopyTo(next.Slice(dataOffset));
+					dataOffset += export.References.Data.Span.Length;
+				}
+
+				Debug.Assert(dataOffset == exportRefsLength);
+				next = next.Slice(exportRefsLength);
+			}
+			if (packetsLength > 0)
+			{
+				WriteSectionHeader(next, BundleSectionType.Packets, packetsLength);
+				next = next.Slice(SectionHeaderLength);
+
+				BundlePacketCollection.Write(next.Slice(0, packetsLength), packets);
+				next = next.Slice(packetsLength);
+			}
+			Debug.Assert(next.Length == 0);
+
+			// Create the header instance
+			return new BundleHeader(data);
+		}
+
+		const int SectionHeaderLength = 4;
+
+		static void WriteSectionHeader(Span<byte> span, BundleSectionType type, int length)
+		{
+			BinaryPrimitives.WriteInt32LittleEndian(span, (int)type | (length << 8));
 		}
 
 		record class ExportInfo(int TypeIdx, IoHash Hash, int Length, List<BundleExportRef> References);
@@ -233,7 +402,7 @@ namespace EpicGames.Horde.Storage
 		/// </summary>
 		/// <param name="reader">Reader to deserialize from</param>
 		/// <returns>New header object</returns>
-		public BundleHeader(IMemoryReader reader)
+		public static BundleHeader Read(IMemoryReader reader)
 		{
 			ReadOnlyMemory<byte> prelude = reader.GetMemory(8);
 			CheckPrelude(prelude);
@@ -259,8 +428,6 @@ namespace EpicGames.Horde.Storage
 				types.Add(new BundleType(guid, serializerVersion));
 			}
 
-			Types = new BundleTypeCollection(types);
-
 			// Read the imports
 			int numImports = (int)reader.ReadUnsignedVarInt();
 			List<BlobLocator> imports = new List<BlobLocator>(numImports);
@@ -279,8 +446,6 @@ namespace EpicGames.Horde.Storage
 				}
 			}
 
-			Imports = new BundleImportCollection(imports);
-
 			// Read the exports
 			int numExports = (int)reader.ReadUnsignedVarInt();
 			List<ExportInfo> exportInfos = new List<ExportInfo>(numExports);
@@ -314,9 +479,10 @@ namespace EpicGames.Horde.Storage
 			}
 
 			// Read the compression packets
+			List<BundlePacket> packets;
 			if (compressionFormat == BundleCompressionFormat.None)
 			{
-				List<BundlePacket> packets = new List<BundlePacket>(exportInfos.Count);
+				packets = new List<BundlePacket>(exportInfos.Count);
 
 				int encodedOffset = 0;
 				foreach (ExportInfo export in exportInfos)
@@ -324,13 +490,11 @@ namespace EpicGames.Horde.Storage
 					packets.Add(new BundlePacket(compressionFormat, encodedOffset, export.Length, export.Length));
 					encodedOffset += export.Length;
 				}
-
-				Packets = new BundlePacketCollection(packets);
 			}
 			else
 			{
 				int numPackets = (int)reader.ReadUnsignedVarInt();
-				List<BundlePacket> packets = new List<BundlePacket>(numPackets);
+				packets = new List<BundlePacket>(numPackets);
 
 				int encodedOffset = 0;
 				for (int packetIdx = 0; packetIdx < numPackets; packetIdx++)
@@ -340,8 +504,6 @@ namespace EpicGames.Horde.Storage
 					packets.Add(new BundlePacket(compressionFormat, encodedOffset, encodedLength, decodedLength));
 					encodedOffset += encodedLength;
 				}
-
-				Packets = new BundlePacketCollection(packets);
 			}
 
 			// Create the final export list
@@ -352,7 +514,7 @@ namespace EpicGames.Horde.Storage
 
 				foreach (ExportInfo exportInfo in exportInfos)
 				{
-					if (packetOffset + exportInfo.Length > Packets[packetIdx].DecodedLength)
+					if (packetOffset + exportInfo.Length > packets[packetIdx].DecodedLength)
 					{
 						packetIdx++;
 						packetOffset = 0;
@@ -362,7 +524,8 @@ namespace EpicGames.Horde.Storage
 					packetOffset += exportInfo.Length;
 				}
 			}
-			Exports = new BundleExportCollection(exports);
+
+			return BundleHeader.Create(types, imports, exports, packets);
 		}
 
 		/// <summary>
@@ -395,7 +558,7 @@ namespace EpicGames.Horde.Storage
 				prelude.CopyTo(header);
 
 				await stream.ReadFixedLengthBytesAsync(header.Slice(prelude.Length), cancellationToken);
-				return new BundleHeader(new MemoryReader(header));
+				return BundleHeader.Read(new MemoryReader(header));
 			}
 		}
 
@@ -592,7 +755,7 @@ namespace EpicGames.Horde.Storage
 	/// <summary>
 	/// Collection of node types in a bundle
 	/// </summary>
-	public class BundleTypeCollection : IReadOnlyList<BundleType>
+	public struct BundleTypeCollection : IReadOnlyList<BundleType>
 	{
 		/// <summary>
 		/// Data for this collection
@@ -609,15 +772,8 @@ namespace EpicGames.Horde.Storage
 		/// </summary>
 		public BundleTypeCollection(IReadOnlyList<BundleType> types)
 		{
-			byte[] data = new byte[types.Count * BundleType.NumBytes];
-
-			int offset = 0;
-			foreach (BundleType type in types)
-			{
-				type.CopyTo(data.AsSpan(offset));
-				offset += BundleType.NumBytes;
-			}
-
+			byte[] data = new byte[Measure(types)];
+			Write(data, types);
 			Data = data;
 		}
 
@@ -639,6 +795,29 @@ namespace EpicGames.Horde.Storage
 
 		/// <inheritdoc/>
 		IEnumerator IEnumerable.GetEnumerator() => GetEnumerator();
+
+		/// <summary>
+		/// Gets the size of memory required to serialize a collection of types
+		/// </summary>
+		/// <param name="types">Type collection</param>
+		/// <returns>Number of bytes required to serialize the types</returns>
+		public static int Measure(IReadOnlyCollection<BundleType> types) => types.Count * BundleType.NumBytes;
+
+		/// <summary>
+		/// Serializes a set of types to a fixed block of memory
+		/// </summary>
+		/// <param name="span">Span to write the types to</param>
+		/// <param name="types">Collection of types to be written</param>
+		public static void Write(Span<byte> span, IReadOnlyCollection<BundleType> types)
+		{
+			Span<byte> next = span;
+			foreach (BundleType type in types)
+			{
+				type.CopyTo(next);
+				next = next.Slice(BundleType.NumBytes);
+			}
+			Debug.Assert(next.Length == 0);
+		}
 	}
 
 	/// <summary>
@@ -734,6 +913,8 @@ namespace EpicGames.Horde.Storage
 				next[0] = 0;
 				next = next.Slice(1);
 			}
+
+			Debug.Assert(next.Length == 0);
 		}
 	}
 
@@ -869,6 +1050,7 @@ namespace EpicGames.Horde.Storage
 				packet.CopyTo(next);
 				next = next.Slice(BundlePacket.NumBytes);
 			}
+			Debug.Assert(next.Length == 0);
 		}
 	}
 
@@ -957,62 +1139,16 @@ namespace EpicGames.Horde.Storage
 	{
 		readonly int _count;
 		readonly ReadOnlyMemory<byte> _data;
-		readonly ReadOnlyMemory<byte> _refIndexData;
-		readonly BundleExportRefCollection _refs;
+		readonly ReadOnlyMemory<byte> _refs;
 
 		/// <summary>
 		/// Constructor
 		/// </summary>
-		public BundleExportCollection(ReadOnlyMemory<byte> data, ReadOnlyMemory<byte> refData)
+		public BundleExportCollection(ReadOnlyMemory<byte> data, ReadOnlyMemory<byte> refs)
 		{
 			_count = data.Length / BundleExport.NumBytes;
 			_data = data;
-
-			if (refData.Length > 0)
-			{
-				_refIndexData = refData.Slice(0, _count * sizeof(int));
-				_refs = new BundleExportRefCollection(refData.Slice(_refIndexData.Length));
-			}
-		}
-
-		/// <summary>
-		/// Constructor
-		/// </summary>
-		public BundleExportCollection(IReadOnlyList<BundleExport> exports)
-		{
-			int exportRefLength = 0;
-			foreach (BundleExport export in exports)
-			{
-				exportRefLength += export.References.Data.Length;
-			}
-
-			if (exportRefLength > 0)
-			{
-				byte[] refData = new byte[exportRefLength];
-				byte[] refIndexData = new byte[exports.Count * sizeof(int)];
-
-				int refDataOffset = 0;
-				int refIndexOffset = 0;
-
-				foreach (BundleExport export in exports)
-				{
-					BinaryPrimitives.WriteInt32LittleEndian(refIndexData.AsSpan(refIndexOffset), refDataOffset);
-					refIndexOffset += sizeof(int);
-
-					export.References.Data.Span.CopyTo(refData.AsSpan(refDataOffset));
-					refDataOffset += export.References.Data.Length;
-				}
-
-				_refIndexData = refIndexData;
-				_refs = new BundleExportRefCollection(refData);
-			}
-
-			byte[] data = new byte[exports.Count * BundleExport.NumBytes];
-			for (int idx = 0; idx < exports.Count; idx++)
-			{
-				exports[idx].Data.CopyTo(data.AsMemory(idx * BundleExport.NumBytes));
-			}
-			_data = data;
+			_refs = refs;
 		}
 
 		/// <inheritdoc/>
@@ -1026,11 +1162,11 @@ namespace EpicGames.Horde.Storage
 				ReadOnlyMemory<byte> exportData = _data.Slice(index * BundleExport.NumBytes);
 
 				BundleExportRefCollection exportRefs = new BundleExportRefCollection();
-				if (_refIndexData.Length > 0)
+				if (_refs.Length > 0)
 				{
-					int minOffset = BinaryPrimitives.ReadInt32LittleEndian(_refIndexData.Span.Slice(index * sizeof(int)));
-					int maxOffset = (index == Count - 1)? _refs.Data.Length : BinaryPrimitives.ReadInt32LittleEndian(_refIndexData.Span.Slice((index + 1) * sizeof(int)));
-					exportRefs = new BundleExportRefCollection(_refs.Data.Slice(minOffset, maxOffset - minOffset));
+					int minOffset = BinaryPrimitives.ReadInt32LittleEndian(_refs.Span.Slice(index * sizeof(int)));
+					int maxOffset = (index == Count - 1)? _refs.Length : BinaryPrimitives.ReadInt32LittleEndian(_refs.Span.Slice((index + 1) * sizeof(int)));
+					exportRefs = new BundleExportRefCollection(_refs.Slice(minOffset, maxOffset - minOffset));
 				}
 
 				return new BundleExport(exportData, exportRefs);
@@ -1049,6 +1185,28 @@ namespace EpicGames.Horde.Storage
 
 		/// <inheritdoc/>
 		IEnumerator IEnumerable.GetEnumerator() => GetEnumerator();
+
+		/// <summary>
+		/// Measure the size of memory required to store a collection of import locators
+		/// </summary>
+		/// <param name="exports">Locators to write</param>
+		/// <returns>Size in bytes of the output buffer</returns>
+		public static int Measure(IReadOnlyCollection<BundleExport> exports) => BundleExport.NumBytes * exports.Count;
+
+		/// <summary>
+		/// Serialize a collection of packets to memory
+		/// </summary>
+		/// <param name="data">Output buffer for the serialized data</param>
+		/// <param name="exports">Packets to write</param>
+		public static void Write(Span<byte> data, IReadOnlyCollection<BundleExport> exports)
+		{
+			Span<byte> next = data;
+			foreach (BundleExport export in exports)
+			{
+				export.Data.Span.CopyTo(next);
+				next = next.Slice(export.Data.Length);
+			}
+		}
 	}
 
 	/// <summary>
diff --git a/Engine/Source/Programs/Shared/EpicGames.Horde/Storage/TreeReader.cs b/Engine/Source/Programs/Shared/EpicGames.Horde/Storage/TreeReader.cs
--- a/Engine/Source/Programs/Shared/EpicGames.Horde/Storage/TreeReader.cs
+++ b/Engine/Source/Programs/Shared/EpicGames.Horde/Storage/TreeReader.cs
@@ -417,7 +417,7 @@ namespace EpicGames.Horde.Storage
 					}
 
 					// Parse the header and construct the bundle info from it
-					BundleHeader header = new BundleHeader(new MemoryReader(memory));
+					BundleHeader header = BundleHeader.Read(new MemoryReader(memory));
 
 					// Get the types within this bundle
 					TypeInfo?[] types = new TypeInfo?[header.Types.Count];
diff --git a/Engine/Source/Programs/Shared/EpicGames.Horde/Storage/TreeWriter.cs b/Engine/Source/Programs/Shared/EpicGames.Horde/Storage/TreeWriter.cs
--- a/Engine/Source/Programs/Shared/EpicGames.Horde/Storage/TreeWriter.cs
+++ b/Engine/Source/Programs/Shared/EpicGames.Horde/Storage/TreeWriter.cs
@@ -617,7 +617,7 @@ namespace EpicGames.Horde.Storage
 				}
 
 				// Create the bundle
-				BundleHeader header = new BundleHeader(new BundleTypeCollection(types), new BundleImportCollection(imports), new BundleExportCollection(exports), new BundlePacketCollection(_packets));
+				BundleHeader header = BundleHeader.Create(types, imports, exports, _packets);
 				return new Bundle(header, packetData);
 			}
 		}
