Horde: Change bundles to be stored in a read-in-place format.

diff --git a/Engine/Source/Programs/Horde/Horde.Server.Tests/StorageServiceTests.cs b/Engine/Source/Programs/Horde/Horde.Server.Tests/StorageServiceTests.cs
--- a/Engine/Source/Programs/Horde/Horde.Server.Tests/StorageServiceTests.cs
+++ b/Engine/Source/Programs/Horde/Horde.Server.Tests/StorageServiceTests.cs
@@ -40,7 +40,7 @@ namespace Horde.Server.Tests
 			exports.Add(new BundleExport(0, hash1, 0, 0, data1.Length, Array.Empty<BundleExportRef>()));
 			exports.Add(new BundleExport(0, hash2, 0, data1.Length, data2.Length, Array.Empty<BundleExportRef>()));
 
-			BundleHeader header = BundleHeader.Create(types, Array.Empty<BlobLocator>(), exports, Array.Empty<BundlePacket>());
+			BundleHeader header = BundleHeader.Create(types, Array.Empty<BlobLocator>(), exports, new BundlePacket[1]);
 			Bundle bundle = new Bundle(header, Array.Empty<ReadOnlyMemory<byte>>());
 			BlobLocator locator = await client.WriteBundleAsync(bundle);
 
diff --git a/Engine/Source/Programs/Horde/Horde.Server/Storage/StorageService.cs b/Engine/Source/Programs/Horde/Horde.Server/Storage/StorageService.cs
--- a/Engine/Source/Programs/Horde/Horde.Server/Storage/StorageService.cs
+++ b/Engine/Source/Programs/Horde/Horde.Server/Storage/StorageService.cs
@@ -1044,10 +1044,11 @@ namespace Horde.Server.Storage
 						memory = memory.Slice(0, length);
 
 						// Make sure it's large enough to hold the header
-						int headerSize = BundleHeader.ReadPrelude(memory);
+						int headerSize = BundleHeader.ReadPrelude(memory.Span);
 						if (headerSize <= fetchSize)
 						{
-							return BundleHeader.Read(new MemoryReader(memory));
+							byte[] data = memory.ToArray(); // Need to copy data since rented memory will be disposed
+							return BundleHeader.Read(data);
 						}
 
 						// Increase the fetch size and retry
diff --git a/Engine/Source/Programs/Horde/Horde/Commands/Bundles/BundleDump.cs b/Engine/Source/Programs/Horde/Horde/Commands/Bundles/BundleDump.cs
--- a/Engine/Source/Programs/Horde/Horde/Commands/Bundles/BundleDump.cs
+++ b/Engine/Source/Programs/Horde/Horde/Commands/Bundles/BundleDump.cs
@@ -22,7 +22,7 @@ namespace Horde.Commands.Bundles
 		public override async Task<int> ExecuteAsync(ILogger logger)
 		{
 			byte[] bytes = await FileReference.ReadAllBytesAsync(Input);
-			Bundle bundle = new Bundle(new MemoryReader(bytes));
+			Bundle bundle = new Bundle(bytes);
 			logger.LogInformation("Summary for blob {Location}", Input);
 
 			BundleHeader header = bundle.Header;
diff --git a/Engine/Source/Programs/Shared/EpicGames.Horde.Tests/BundleTests.cs b/Engine/Source/Programs/Shared/EpicGames.Horde.Tests/BundleTests.cs
--- a/Engine/Source/Programs/Shared/EpicGames.Horde.Tests/BundleTests.cs
+++ b/Engine/Source/Programs/Shared/EpicGames.Horde.Tests/BundleTests.cs
@@ -155,11 +155,9 @@ namespace EpicGames.Horde.Tests
 				oldHeader = BundleHeader.Create(types, imports, exports, packets);
 			}
 
-			ByteArrayBuilder writer = new ByteArrayBuilder();
-			oldHeader.Write(writer);
-			byte[] serializedData = writer.ToByteArray();
+			byte[] serializedData = oldHeader.Data.ToArray();
 
-			BundleHeader newHeader = BundleHeader.Read(new MemoryReader(serializedData));
+			BundleHeader newHeader = BundleHeader.Read(serializedData);
 				
 			Assert.AreEqual(oldHeader.Imports.Count, newHeader.Imports.Count);
 			for (int idx = 0; idx < oldHeader.Imports.Count; idx++)
diff --git a/Engine/Source/Programs/Shared/EpicGames.Horde/Storage/Bundle.cs b/Engine/Source/Programs/Shared/EpicGames.Horde/Storage/Bundle.cs
--- a/Engine/Source/Programs/Shared/EpicGames.Horde/Storage/Bundle.cs
+++ b/Engine/Source/Programs/Shared/EpicGames.Horde/Storage/Bundle.cs
@@ -2,6 +2,7 @@
 
 using EpicGames.Core;
 using K4os.Compression.LZ4;
+using Microsoft.CodeAnalysis.CSharp.Syntax;
 using System;
 using System.Buffers;
 using System.Buffers.Binary;
@@ -12,6 +13,7 @@ using System.Globalization;
 using System.IO;
 using System.IO.Compression;
 using System.Linq;
+using System.Runtime.CompilerServices;
 using System.Text;
 using System.Threading;
 using System.Threading.Tasks;
@@ -39,6 +41,11 @@ namespace EpicGames.Horde.Storage
 		/// </summary>
 		RemoveAliases = 2,
 
+		/// <summary>
+		/// Use data structures which support in-place reading and writing.
+		/// </summary>
+		InPlace = 3,
+
 		/// <summary>
 		/// Last item in the enum. Used for <see cref="Latest"/>
 		/// </summary>
@@ -79,15 +86,17 @@ namespace EpicGames.Horde.Storage
 		/// <summary>
 		/// Reads a bundle from a block of memory
 		/// </summary>
-		public Bundle(IMemoryReader reader)
+		public Bundle(ReadOnlyMemory<byte> memory)
 		{
-			Header = BundleHeader.Read(reader);
+			int offset = BundleHeader.ReadPrelude(memory.Span);
+			Header = BundleHeader.Read(memory.Slice(0, offset));
 
 			ReadOnlyMemory<byte>[] packets = new ReadOnlyMemory<byte>[Header.Packets.Count];
 			for (int idx = 0; idx < Header.Packets.Count; idx++)
 			{
 				BundlePacket packet = Header.Packets[idx];
-				packets[idx] = reader.ReadFixedLengthBytes(packet.EncodedLength);
+				packets[idx] = memory.Slice(offset, packet.EncodedLength);
+				offset += packet.EncodedLength;
 			}
 
 			Packets = packets;
@@ -123,11 +132,8 @@ namespace EpicGames.Horde.Storage
 		/// <returns>Sequence for the bundle</returns>
 		public ReadOnlySequence<byte> AsSequence()
 		{
-			ByteArrayBuilder builder = new ByteArrayBuilder();
-			Header.Write(builder);
-
 			ReadOnlySequenceBuilder<byte> sequence = new ReadOnlySequenceBuilder<byte>();
-			sequence.Append(builder.AsSequence());
+			sequence.Append(Header.Data);
 
 			foreach (ReadOnlyMemory<byte> packet in Packets)
 			{
@@ -243,13 +249,17 @@ namespace EpicGames.Horde.Storage
 		{
 			Data = data;
 
+			ReadOnlySpan<byte> span = data.Span;
+			// BundleVersion version = (BundleVersion)span[3];
+			int headerLength = BinaryPrimitives.ReadInt32LittleEndian(span.Slice(4));
+
 			ReadOnlyMemory<byte> exportData = ReadOnlyMemory<byte>.Empty;
 			ReadOnlyMemory<byte> exportRefData = ReadOnlyMemory<byte>.Empty;
 
-			for (int offset = PreludeLength; offset < data.Length;)
+			for (int offset = PreludeLength; offset < headerLength;)
 			{
 				uint header = BinaryPrimitives.ReadUInt32LittleEndian(data.Slice(offset).Span);
-				offset += sizeof(int);
+				offset += SectionHeaderLength;
 
 				int length = (int)(header >> 8);
 				ReadOnlyMemory<byte> sectionData = data.Slice(offset, length);
@@ -327,6 +337,11 @@ namespace EpicGames.Horde.Storage
 
 			// Allocate the header data
 			byte[] data = new byte[length];
+			data[0] = (byte)'U';
+			data[1] = (byte)'B';
+			data[2] = (byte)'N';
+			data[3] = (byte)BundleVersion.Latest;
+			BinaryPrimitives.WriteInt32LittleEndian(data.AsSpan(4), length);
 
 			// Write all the sections
 			Span<byte> next = data.AsSpan(PreludeLength);
@@ -400,18 +415,36 @@ namespace EpicGames.Horde.Storage
 		/// <summary>
 		/// Reads a bundle header from memory
 		/// </summary>
-		/// <param name="reader">Reader to deserialize from</param>
+		/// <param name="memory">Memory to deserialize from</param>
 		/// <returns>New header object</returns>
-		public static BundleHeader Read(IMemoryReader reader)
+		public static BundleHeader Read(ReadOnlyMemory<byte> memory)
+		{
+			ReadOnlySpan<byte> span = memory.Span;
+			if (span[0] == (byte)'U' && span[1] == (byte)'B' && span[2] == (byte)'N')
+			{
+				return new BundleHeader(memory);
+			}
+			else if (span[0] == (byte)'U' && span[1] == (byte)'E' && span[2] == (byte)'B' && span[3] == (byte)'N')
+			{
+				return ReadLegacy(memory);
+			}
+			else
+			{
+				throw new NotSupportedException();
+			}
+		}
+
+		static BundleHeader ReadLegacy(ReadOnlyMemory<byte> memory)
 		{
-			ReadOnlyMemory<byte> prelude = reader.GetMemory(8);
-			CheckPrelude(prelude);
-			reader.Advance(8);
+			int headerLength = BinaryPrimitives.ReadInt32BigEndian(memory.Span.Slice(4));
+
+			MemoryReader reader = new MemoryReader(memory.Slice(0, headerLength));
+			reader.Advance(PreludeLength);
 
 			BundleVersion version = (BundleVersion)reader.ReadUnsignedVarInt();
-			if (version > BundleVersion.Latest)
+			if (version > BundleVersion.RemoveAliases)
 			{
-				throw new InvalidDataException($"Unknown bundle version {(int)version}. Max supported is {(int)BundleVersion.Latest}.");
+				throw new InvalidDataException($"Invalid bundle version {(int)version}. Max supported version in legacy format is {(int)BundleVersion.RemoveAliases}.");
 			}
 
 			BundleCompressionFormat compressionFormat = (BundleCompressionFormat)reader.ReadUnsignedVarInt();
@@ -529,14 +562,24 @@ namespace EpicGames.Horde.Storage
 		}
 
 		/// <summary>
-		/// Reads the prelude bytes from a bundle, and returns the header size.
+		/// Validates that the prelude bytes for a bundle header are correct
 		/// </summary>
-		/// <param name="prelude">The prelude data</param>
-		/// <returns>Size of the header, including the prelude data</returns>
-		public static int ReadPrelude(ReadOnlyMemory<byte> prelude)
+		/// <param name="prelude">The prelude bytes</param>
+		/// <returns>Length of the header data, including the prelude</returns>
+		public static int ReadPrelude(ReadOnlySpan<byte> prelude)
 		{
-			CheckPrelude(prelude);
-			return BinaryPrimitives.ReadInt32BigEndian(prelude.Span.Slice(4));
+			if (prelude[0] == 'U' && prelude[1] == 'E' && prelude[2] == 'B' && prelude[3] == 'N')
+			{
+				return BinaryPrimitives.ReadInt32BigEndian(prelude.Slice(4));
+			}
+			else if (prelude[0] == 'U' && prelude[1] == 'B' && prelude[2] == 'N')
+			{
+				return BinaryPrimitives.ReadInt32LittleEndian(prelude.Slice(4));
+			}
+			else
+			{
+				throw new InvalidDataException("Invalid signature bytes for bundle. Corrupt data?");
+			}
 		}
 
 		/// <summary>
@@ -547,125 +590,16 @@ namespace EpicGames.Horde.Storage
 		/// <returns>New header</returns>
 		public static async Task<BundleHeader> FromStreamAsync(Stream stream, CancellationToken cancellationToken)
 		{
-			byte[] prelude = new byte[8];
+			byte[] prelude = new byte[PreludeLength];
 			await stream.ReadFixedLengthBytesAsync(prelude, cancellationToken);
-			CheckPrelude(prelude);
 
-			int headerLength = BinaryPrimitives.ReadInt32BigEndian(prelude.AsSpan(4));
-			using (IMemoryOwner<byte> owner = MemoryPool<byte>.Shared.Rent(headerLength))
-			{
-				Memory<byte> header = owner.Memory.Slice(0, headerLength);
-				prelude.CopyTo(header);
+			int headerLength = ReadPrelude(prelude);
 
-				await stream.ReadFixedLengthBytesAsync(header.Slice(prelude.Length), cancellationToken);
-				return BundleHeader.Read(new MemoryReader(header));
-			}
-		}
+			byte[] header = new byte[headerLength];
+			prelude.CopyTo(header.AsSpan());
 
-		/// <summary>
-		/// Validates that the prelude bytes for a bundle header are correct
-		/// </summary>
-		/// <param name="prelude">The prelude bytes</param>
-		static void CheckPrelude(ReadOnlyMemory<byte> prelude)
-		{
-			if (!prelude.Slice(0, 4).Span.SequenceEqual(Signature.Span))
-			{
-				throw new InvalidDataException("Invalid signature bytes for bundle. Corrupt data?");
-			}
-		}
-
-		/// <summary>
-		/// Serializes a header to memory
-		/// </summary>
-		/// <param name="writer">Writer to serialize to</param>
-		public void Write(IMemoryWriter writer)
-		{
-			int initialLength = writer.Length;
-
-			// Write the header prelude; a fixed-length block containing a signature and header size.
-			Memory<byte> prelude = writer.GetMemory(8);
-			Signature.CopyTo(prelude);
-			writer.Advance(8);
-
-			// Write the header data
-			writer.WriteUnsignedVarInt((ulong)BundleVersion.Latest);
-
-			BundleCompressionFormat compressionFormat = (Packets.Count > 0) ? Packets[0].CompressionFormat : BundleCompressionFormat.None;
-
-			writer.WriteUnsignedVarInt((ulong)compressionFormat);
-
-			// Find all the referenced nodes in each import
-			SortedSet<int>[] nodes = new SortedSet<int>[Imports.Count];
-			for (int idx = 0; idx < Imports.Count; idx++)
-			{
-				nodes[idx] = new SortedSet<int>();
-			}
-			foreach (BundleExportRef exportRef in Exports.SelectMany(x => x.References).Where(x => x.ImportIdx != -1))
-			{
-				nodes[exportRef.ImportIdx].Add(exportRef.NodeIdx);
-			}
-
-			// Map all the imports to an index
-			Dictionary<BundleExportRef, int> exportRefToIndex = new Dictionary<BundleExportRef, int>();
-			for (int importIdx = 0; importIdx < Imports.Count; importIdx++)
-			{
-				foreach (int nodeIdx in nodes[importIdx])
-				{
-					int index = exportRefToIndex.Count;
-					exportRefToIndex[new BundleExportRef(importIdx, nodeIdx)] = index;
-				}
-			}
-			for (int exportIdx = 0; exportIdx < Exports.Count; exportIdx++)
-			{
-				int index = exportRefToIndex.Count;
-				exportRefToIndex[new BundleExportRef(-1, exportIdx)] = index;
-			}
-
-			// Write all the types
-			writer.WriteUnsignedVarInt(Types.Count);
-			for (int typeIdx = 0; typeIdx < Types.Count; typeIdx++)
-			{
-				writer.WriteGuid(Types[typeIdx].Guid);
-				writer.WriteUnsignedVarInt(Types[typeIdx].Version);
-			}
-
-			// Write all the imports
-			writer.WriteUnsignedVarInt(Imports.Count);
-			for (int importIdx = 0; importIdx < Imports.Count; importIdx++)
-			{
-				Debug.Assert(Imports[importIdx].IsValid());
-				writer.WriteBlobLocator(Imports[importIdx]);
-				writer.WriteVariableLengthArray(nodes[importIdx].ToArray(), x => writer.WriteUnsignedVarInt(x));
-			}
-
-			// Write all the exports
-			writer.WriteUnsignedVarInt(Exports.Count);
-			foreach (BundleExport export in Exports)
-			{
-				writer.WriteUnsignedVarInt(export.TypeIdx);
-				writer.WriteIoHash(export.Hash);
-				writer.WriteUnsignedVarInt(export.Length);
-
-				writer.WriteUnsignedVarInt(export.References.Count);
-				for (int idx = 0; idx < export.References.Count; idx++)
-				{
-					writer.WriteUnsignedVarInt(exportRefToIndex[export.References[idx]]);
-				}
-			}
-
-			// Write the packets
-			if (compressionFormat != BundleCompressionFormat.None)
-			{
-				writer.WriteUnsignedVarInt(Packets.Count);
-				foreach(BundlePacket packet in Packets)
-				{
-					writer.WriteUnsignedVarInt(packet.EncodedLength);
-					writer.WriteUnsignedVarInt(packet.DecodedLength);
-				}
-			}
-
-			// Go back and write the length of the header
-			BinaryPrimitives.WriteInt32BigEndian(prelude.Slice(4).Span, writer.Length - initialLength);
+			await stream.ReadFixedLengthBytesAsync(header.AsMemory(prelude.Length), cancellationToken);
+			return BundleHeader.Read(header);
 		}
 	}
 
@@ -850,6 +784,7 @@ namespace EpicGames.Horde.Storage
 		{
 			get
 			{
+				Debug.Assert(index < Count);
 				ReadOnlySpan<byte> span = _data.Span;
 				int offset = BinaryPrimitives.ReadInt32LittleEndian(span.Slice(index * sizeof(int)));
 				int length = span.Slice(offset).IndexOf((byte)0);
@@ -1014,7 +949,14 @@ namespace EpicGames.Horde.Storage
 		public int Count => _data.Length / BundlePacket.NumBytes;
 
 		/// <inheritdoc/>
-		public BundlePacket this[int index] => new BundlePacket(_data.Slice(index * BundlePacket.NumBytes, BundlePacket.NumBytes).Span);
+		public BundlePacket this[int index]
+		{
+			get
+			{
+				Debug.Assert(index < Count);
+				return new BundlePacket(_data.Slice(index * BundlePacket.NumBytes, BundlePacket.NumBytes).Span);
+			}
+		}
 
 		/// <inheritdoc/>
 		public IEnumerator<BundlePacket> GetEnumerator()
@@ -1159,6 +1101,8 @@ namespace EpicGames.Horde.Storage
 		{
 			get
 			{
+				Debug.Assert(index < _count);
+
 				ReadOnlyMemory<byte> exportData = _data.Slice(index * BundleExport.NumBytes);
 
 				BundleExportRefCollection exportRefs = new BundleExportRefCollection();
@@ -1271,7 +1215,14 @@ namespace EpicGames.Horde.Storage
 		public int Count => Data.Length / BundleExportRef.NumBytes;
 
 		/// <inheritdoc/>
-		public BundleExportRef this[int index] => BundleExportRef.Read(Data.Span.Slice(index * BundleExportRef.NumBytes));
+		public BundleExportRef this[int index]
+		{
+			get
+			{
+				Debug.Assert(index < Count);
+				return BundleExportRef.Read(Data.Span.Slice(index * BundleExportRef.NumBytes));
+			}
+		}
 
 		/// <inheritdoc/>
 		public IEnumerator<BundleExportRef> GetEnumerator()
diff --git a/Engine/Source/Programs/Shared/EpicGames.Horde/Storage/TreeReader.cs b/Engine/Source/Programs/Shared/EpicGames.Horde/Storage/TreeReader.cs
--- a/Engine/Source/Programs/Shared/EpicGames.Horde/Storage/TreeReader.cs
+++ b/Engine/Source/Programs/Shared/EpicGames.Horde/Storage/TreeReader.cs
@@ -409,7 +409,7 @@ namespace EpicGames.Horde.Storage
 					memory = await _store.ReadBlobRangeAsync(queuedHeader.Blob, 0, memory, cancellationToken);
 
 					// Make sure it's large enough to hold the header
-					int headerSize = BundleHeader.ReadPrelude(memory);
+					int headerSize = BundleHeader.ReadPrelude(memory.Span);
 					if (headerSize > prefetchSize)
 					{
 						prefetchSize = headerSize;
@@ -417,7 +417,7 @@ namespace EpicGames.Horde.Storage
 					}
 
 					// Parse the header and construct the bundle info from it
-					BundleHeader header = BundleHeader.Read(new MemoryReader(memory));
+					BundleHeader header = BundleHeader.Read(memory.ToArray());
 
 					// Get the types within this bundle
 					TypeInfo?[] types = new TypeInfo?[header.Types.Count];
@@ -610,6 +610,9 @@ namespace EpicGames.Horde.Storage
 		/// <returns>The decoded data</returns>
 		async ValueTask<ReadOnlyMemory<byte>> ReadBundlePacketAsync(BundleInfo bundleInfo, int packetIdx, CancellationToken cancellationToken)
 		{
+			char c = bundleInfo.Locator.BlobId.ToString()[0];
+			Debug.Assert(Char.IsLetterOrDigit(c));
+
 			if (packetIdx < 0 || packetIdx >= bundleInfo.Header.Packets.Count)
 			{
 				throw new ArgumentException("Packet index is out of range", nameof(packetIdx));
@@ -674,6 +677,9 @@ namespace EpicGames.Horde.Storage
 		/// <returns>The encoded packet data</returns>
 		async ValueTask<ReadOnlyMemory<byte>> ReadEncodedPacketAsync(BundleInfo bundleInfo, int packetIdx)
 		{
+			char c = bundleInfo.Locator.BlobId.ToString()[0];
+			Debug.Assert(Char.IsLetterOrDigit(c));
+
 			if (packetIdx < 0 || packetIdx >= bundleInfo.Header.Packets.Count)
 			{
 				throw new ArgumentException("Packet index is out of range", nameof(packetIdx));
