EpicGames.Core: Add overloads for HTTP REST utility methods which take Uri arguments.

diff --git a/Engine/Source/Programs/Shared/EpicGames.Core/HttpClientExtensions.cs b/Engine/Source/Programs/Shared/EpicGames.Core/HttpClientExtensions.cs
--- a/Engine/Source/Programs/Shared/EpicGames.Core/HttpClientExtensions.cs
+++ b/Engine/Source/Programs/Shared/EpicGames.Core/HttpClientExtensions.cs
@@ -1,5 +1,6 @@
 // Copyright Epic Games, Inc. All Rights Reserved.
 
+using System;
 using System.Net.Http;
 using System.Text;
 using System.Text.Json;
@@ -16,6 +17,9 @@ namespace EpicGames.Core
 	/// </summary>
 	public static class HttpClientExtensions
 	{
+		/// <inheritdoc cref="GetAsync{TResponse}(HttpClient, Uri, CancellationToken)"/>
+		public static Task<TResponse> GetAsync<TResponse>(this HttpClient client, string url, CancellationToken cancellationToken) => GetAsync<TResponse>(client, new Uri(url), cancellationToken);
+
 		/// <summary>
 		/// Gets a resource from an HTTP endpoint and parses it as a JSON object
 		/// </summary>
@@ -24,7 +28,7 @@ namespace EpicGames.Core
 		/// <param name="url">The url to retrieve</param>
 		/// <param name="cancellationToken">Cancels the request</param>
 		/// <returns>New instance of the object</returns>
-		public static async Task<TResponse> GetAsync<TResponse>(this HttpClient client, string url, CancellationToken cancellationToken)
+		public static async Task<TResponse> GetAsync<TResponse>(this HttpClient client, Uri url, CancellationToken cancellationToken)
 		{
 			using (HttpResponseMessage response = await client.GetAsync(url, cancellationToken))
 			{
@@ -33,6 +37,9 @@ namespace EpicGames.Core
 			}
 		}
 
+		/// <inheritdoc cref="PostAsync{TRequest}(HttpClient, Uri, TRequest, CancellationToken)"/>
+		public static Task<HttpResponseMessage> PostAsync<TRequest>(this HttpClient client, string url, TRequest request, CancellationToken cancellationToken) => PostAsync<TRequest>(client, new Uri(url), request, cancellationToken);
+
 		/// <summary>
 		/// Posts an object to an HTTP endpoint as a JSON object
 		/// </summary>
@@ -42,12 +49,15 @@ namespace EpicGames.Core
 		/// <param name="request">The object to post</param>
 		/// <param name="cancellationToken">Cancels the request</param>
 		/// <returns>Response message</returns>
-		public static async Task<HttpResponseMessage> PostAsync<TRequest>(this HttpClient client, string url, TRequest request, CancellationToken cancellationToken)
+		public static async Task<HttpResponseMessage> PostAsync<TRequest>(this HttpClient client, Uri url, TRequest request, CancellationToken cancellationToken)
 		{
 			using HttpContent content = ToJsonContent(request);
 			return await client.PostAsync(url, content, cancellationToken);
 		}
 
+		/// <inheritdoc cref="PostAsync{TRequest}(HttpClient, Uri, TRequest, CancellationToken)"/>
+		public static Task<TResponse> PostAsync<TResponse, TRequest>(this HttpClient client, string url, TRequest request, CancellationToken cancellationToken) => PostAsync<TResponse, TRequest>(client, new Uri(url), request, cancellationToken);
+
 		/// <summary>
 		/// Posts an object to an HTTP endpoint as a JSON object, and parses the response object
 		/// </summary>
@@ -58,7 +68,7 @@ namespace EpicGames.Core
 		/// <param name="request">The object to post</param>
 		/// <param name="cancellationToken">Cancels the request</param>
 		/// <returns>The response parsed into the requested type</returns>
-		public static async Task<TResponse> PostAsync<TResponse, TRequest>(this HttpClient client, string url, TRequest request, CancellationToken cancellationToken)
+		public static async Task<TResponse> PostAsync<TResponse, TRequest>(this HttpClient client, Uri url, TRequest request, CancellationToken cancellationToken)
 		{
 			using (HttpResponseMessage response = await PostAsync(client, url, request, cancellationToken))
 			{
@@ -67,6 +77,9 @@ namespace EpicGames.Core
 			}
 		}
 
+		/// <inheritdoc cref="PutAsync{TRequest}(HttpClient, Uri, TRequest, CancellationToken)"/>
+		public static Task<HttpResponseMessage> PutAsync<TRequest>(this HttpClient client, string url, TRequest request, CancellationToken cancellationToken) => PutAsync<TRequest>(client, new Uri(url), request, cancellationToken);
+
 		/// <summary>
 		/// Puts an object to an HTTP endpoint as a JSON object
 		/// </summary>
@@ -76,12 +89,15 @@ namespace EpicGames.Core
 		/// <param name="obj">The object to post</param>
 		/// <param name="cancellationToken">Cancels the request</param>
 		/// <returns>Response message</returns>
-		public static async Task<HttpResponseMessage> PutAsync<TRequest>(this HttpClient client, string url, TRequest obj, CancellationToken cancellationToken)
+		public static async Task<HttpResponseMessage> PutAsync<TRequest>(this HttpClient client, Uri url, TRequest obj, CancellationToken cancellationToken)
 		{
 			using HttpContent content = ToJsonContent(obj);
 			return await client.PutAsync(url, content, cancellationToken);
 		}
 
+		/// <inheritdoc cref="PutAsync{TRequest}(HttpClient, Uri, TRequest, CancellationToken)"/>
+		public static Task<TResponse> PutAsync<TResponse, TRequest>(this HttpClient client, string url, TRequest request, CancellationToken cancellationToken) => PutAsync<TResponse, TRequest>(client, new Uri(url), request, cancellationToken);
+
 		/// <summary>
 		/// Puts an object to an HTTP endpoint as a JSON object, and parses the response object
 		/// </summary>
@@ -92,7 +108,7 @@ namespace EpicGames.Core
 		/// <param name="request">The object to post</param>
 		/// <param name="cancellationToken">Cancels the request</param>
 		/// <returns>The response parsed into the requested type</returns>
-		public static async Task<TResponse> PutAsync<TResponse, TRequest>(this HttpClient client, string url, TRequest request, CancellationToken cancellationToken)
+		public static async Task<TResponse> PutAsync<TResponse, TRequest>(this HttpClient client, Uri url, TRequest request, CancellationToken cancellationToken)
 		{
 			using (HttpResponseMessage response = await PutAsync(client, url, request, cancellationToken))
 			{
