Horde: Change compute buffer reader/writer to pass a minimum desired size rather than the current size.

diff --git a/Engine/Source/Programs/Shared/EpicGames.Horde/Compute/Buffers/PooledBuffer.cs b/Engine/Source/Programs/Shared/EpicGames.Horde/Compute/Buffers/PooledBuffer.cs
--- a/Engine/Source/Programs/Shared/EpicGames.Horde/Compute/Buffers/PooledBuffer.cs
+++ b/Engine/Source/Programs/Shared/EpicGames.Horde/Compute/Buffers/PooledBuffer.cs
@@ -103,7 +103,7 @@ namespace EpicGames.Horde.Compute.Buffers
 		/// Constructor
 		/// </summary>
 		private PooledBufferCore(Chunk[] chunks, int readerIdx)
-			: base(chunks, readerIdx)
+			: base(chunks, readerIdx, true)
 		{
 			_chunks = chunks;
 		}
diff --git a/Engine/Source/Programs/Shared/EpicGames.Horde/Compute/Buffers/SharedMemoryBuffer.cs b/Engine/Source/Programs/Shared/EpicGames.Horde/Compute/Buffers/SharedMemoryBuffer.cs
--- a/Engine/Source/Programs/Shared/EpicGames.Horde/Compute/Buffers/SharedMemoryBuffer.cs
+++ b/Engine/Source/Programs/Shared/EpicGames.Horde/Compute/Buffers/SharedMemoryBuffer.cs
@@ -88,7 +88,7 @@ namespace EpicGames.Horde.Compute.Buffers
 			Native.EventHandle writerEvent = Native.EventHandle.CreateNew($"{name}_W", EventResetMode.ManualReset, true, HandleInheritability.Inheritable);
 			Native.EventHandle readerEvent = Native.EventHandle.CreateNew($"{name}_R", EventResetMode.ManualReset, true, HandleInheritability.Inheritable);
 
-			SharedMemoryBufferCore core = new SharedMemoryBufferCore(name, memoryMappedFile, memoryMappedViewAccessor, memoryMappedView, numChunks, chunkLength, readerEvent, writerEvent);
+			SharedMemoryBufferCore core = new SharedMemoryBufferCore(name, memoryMappedFile, memoryMappedViewAccessor, memoryMappedView, numChunks, chunkLength, readerEvent, writerEvent, true);
 			return new SharedMemoryBuffer(core);
 		}
 
@@ -114,7 +114,7 @@ namespace EpicGames.Horde.Compute.Buffers
 			Native.EventHandle readerEvent = Native.EventHandle.OpenExisting($"{name}_R");
 			Native.EventHandle writerEvent = Native.EventHandle.OpenExisting($"{name}_W");
 
-			SharedMemoryBufferCore core = new SharedMemoryBufferCore(name, memoryMappedFile, memoryMappedViewAccessor, memoryMappedView, numChunks, chunkLength, readerEvent, writerEvent);
+			SharedMemoryBufferCore core = new SharedMemoryBufferCore(name, memoryMappedFile, memoryMappedViewAccessor, memoryMappedView, numChunks, chunkLength, readerEvent, writerEvent, false);
 			return new SharedMemoryBuffer(core);
 		}
 	}
@@ -152,8 +152,8 @@ namespace EpicGames.Horde.Compute.Buffers
 		/// <summary>
 		/// Constructor
 		/// </summary>
-		public unsafe SharedMemoryBufferCore(string name, MemoryMappedFile memoryMappedFile, MemoryMappedViewAccessor memoryMappedViewAccessor, MemoryMappedView memoryMappedView, int numChunks, int chunkLength, Native.EventHandle readerEvent, Native.EventHandle writerEvent)
-			: base(CreateChunks(numChunks, chunkLength, memoryMappedView), 1)
+		public unsafe SharedMemoryBufferCore(string name, MemoryMappedFile memoryMappedFile, MemoryMappedViewAccessor memoryMappedViewAccessor, MemoryMappedView memoryMappedView, int numChunks, int chunkLength, Native.EventHandle readerEvent, Native.EventHandle writerEvent, bool isNew)
+			: base(CreateChunks(numChunks, chunkLength, memoryMappedView), 1, isNew)
 		{
 			Name = name;
 
diff --git a/Engine/Source/Programs/Shared/EpicGames.Horde/Compute/ComputeBufferBase.cs b/Engine/Source/Programs/Shared/EpicGames.Horde/Compute/ComputeBufferBase.cs
--- a/Engine/Source/Programs/Shared/EpicGames.Horde/Compute/ComputeBufferBase.cs
+++ b/Engine/Source/Programs/Shared/EpicGames.Horde/Compute/ComputeBufferBase.cs
@@ -149,10 +149,10 @@ namespace EpicGames.Horde.Compute
 			}
 
 			/// <inheritdoc/>
-			public void Advance(int length) => _offset += length;
+			public void AdvanceReadPosition(int length) => _offset += length;
 
 			/// <inheritdoc/>
-			public ReadOnlyMemory<byte> GetMemory()
+			public ReadOnlyMemory<byte> GetReadBuffer()
 			{
 				ChunkBase chunk = _chunks[_chunkIdx];
 				ChunkState state = chunk.State;
@@ -160,7 +160,7 @@ namespace EpicGames.Horde.Compute
 			}
 
 			/// <inheritdoc/>
-			public async ValueTask<bool> WaitToReadAsync(int currentLength, CancellationToken cancellationToken = default)
+			public async ValueTask<bool> WaitToReadAsync(int minLength, CancellationToken cancellationToken = default)
 			{
 				for (; ; )
 				{
@@ -176,16 +176,11 @@ namespace EpicGames.Horde.Compute
 							await _buffer.WaitForReadEvent(_readerIdx, cancellationToken);
 						}
 					}
-					else if (_offset + currentLength < state.Length)
+					else if (_offset + minLength <= state.Length)
 					{
-						// Still have more data to read from the current buffer
+						// We have enough data in the chunk to be able to read a message
 						return true;
 					}
-					else if (state.WriteState == WriteState.Complete)
-					{
-						// Buffer is complete
-						return false;
-					}
 					else if (state.WriteState == WriteState.Writing)
 					{
 						// Wait until there is more data in the chunk
@@ -195,6 +190,11 @@ namespace EpicGames.Horde.Compute
 							await _buffer.WaitForReadEvent(_readerIdx, cancellationToken);
 						}
 					}
+					else if (_offset < state.Length || state.WriteState == WriteState.Complete)
+					{
+						// Cannot read the requested amount of data from this chunk.
+						return false;
+					}
 					else if (state.WriteState == WriteState.MovedToNext)
 					{
 						// Move to the next chunk
@@ -229,14 +229,14 @@ namespace EpicGames.Horde.Compute
 				_chunks = chunks;
 			}
 
-			public void Advance(int size)
+			public void AdvanceWritePosition(int size)
 			{
 				ChunkBase chunk = _chunks[_chunkIdx];
 				chunk.Append(size);
 				_buffer.SetAllReadEvents();
 			}
 
-			public Memory<byte> GetMemory()
+			public Memory<byte> GetWriteBuffer()
 			{
 				ChunkBase chunk = _chunks[_chunkIdx];
 
@@ -261,44 +261,46 @@ namespace EpicGames.Horde.Compute
 				return false;
 			}
 
-			public async ValueTask WaitToWriteAsync(int currentLength, CancellationToken cancellationToken)
+			public async ValueTask WaitToWriteAsync(int minSize, CancellationToken cancellationToken)
 			{
-				if (currentLength >= _chunks[_chunkIdx].Memory.Length)
+				if (minSize > _chunks[_chunkIdx].Memory.Length)
 				{
-					throw new ArgumentException("Current length is already complete buffer size; cannot allocate more.", nameof(currentLength));
+					throw new ArgumentException("Requested read size is larger than chunk size.", nameof(minSize));
 				}
 
 				for (; ; )
 				{
-					Memory<byte> memory = GetMemory();
-					if (memory.Length != currentLength)
+					ChunkBase chunk = _chunks[_chunkIdx];
+
+					ChunkState chunkState = chunk.State;
+					if (chunkState.WriteState == WriteState.Writing)
 					{
-						break;
+						int length = chunkState.Length;
+						if (length + minSize <= chunk.Memory.Length)
+						{
+							return;
+						}
+
+						chunk.MoveToNext();
+						_buffer.SetAllReadEvents();
 					}
-					await MoveNextWriteChunk(cancellationToken);
-				}
-			}
 
-			async ValueTask MoveNextWriteChunk(CancellationToken cancellationToken)
-			{
-				ChunkBase chunk = _chunks[_chunkIdx];
-				chunk.MoveToNext();
+					int nextWriteChunkIdx = _chunkIdx + 1;
+					if (nextWriteChunkIdx == _chunks.Length)
+					{
+						nextWriteChunkIdx = 0;
+					}
 
-				int nextWriteChunkIdx = _chunkIdx + 1;
-				if (nextWriteChunkIdx == _chunks.Length)
-				{
-					nextWriteChunkIdx = 0;
-				}
+					ChunkBase nextChunk = _chunks[nextWriteChunkIdx];
+					while (nextChunk.State.ReaderFlags != 0)
+					{
+						await _buffer.WaitForWriteEvent(cancellationToken);
+						_buffer.ResetWriteEvent();
+					}
 
-				ChunkBase nextChunk = _chunks[nextWriteChunkIdx];
-				while (nextChunk.State.ReaderFlags != 0)
-				{
-					await _buffer.WaitForWriteEvent(cancellationToken);
-					_buffer.ResetWriteEvent();
+					_chunkIdx = nextWriteChunkIdx;
+					nextChunk.StartWriting(_buffer._numReaders);
 				}
-
-				nextChunk.StartWriting(_buffer._numReaders);
-				_chunkIdx = nextWriteChunkIdx;
 			}
 		}
 
@@ -320,10 +322,14 @@ namespace EpicGames.Horde.Compute
 		/// </summary>
 		/// <param name="chunks">Data for the buffer</param>
 		/// <param name="numReaders">Number of readers for this buffer</param>
-		protected ComputeBufferBase(ChunkBase[] chunks, int numReaders)
+		/// <param name="isNew"></param>
+		protected ComputeBufferBase(ChunkBase[] chunks, int numReaders, bool isNew)
 		{
 			_chunks = chunks;
-			_chunks[0].StartWriting(numReaders);
+			if (isNew)
+			{
+				_chunks[0].StartWriting(numReaders);
+			}
 
 			_numReaders = numReaders;
 
diff --git a/Engine/Source/Programs/Shared/EpicGames.Horde/Compute/ComputeMessageChannel.cs b/Engine/Source/Programs/Shared/EpicGames.Horde/Compute/ComputeMessageChannel.cs
--- a/Engine/Source/Programs/Shared/EpicGames.Horde/Compute/ComputeMessageChannel.cs
+++ b/Engine/Source/Programs/Shared/EpicGames.Horde/Compute/ComputeMessageChannel.cs
@@ -85,16 +85,16 @@ namespace EpicGames.Horde.Compute
 
 			public void Send()
 			{
-				Span<byte> header = _sendBufferWriter.GetMemory().Span;
+				Span<byte> header = _sendBufferWriter.GetWriteBuffer().Span;
 				header[0] = (byte)_type;
 				BinaryPrimitives.WriteInt32LittleEndian(header.Slice(1, 4), _length);
 
 				if (_channel._logger.IsEnabled(LogLevel.Trace))
 				{
-					_channel.LogMessageInfo("SEND", _type, _sendBufferWriter.GetMemory().Slice(HeaderLength, _length).Span);
+					_channel.LogMessageInfo("SEND", _type, _sendBufferWriter.GetWriteBuffer().Slice(HeaderLength, _length).Span);
 				}
 
-				_sendBufferWriter.Advance(HeaderLength + _length);
+				_sendBufferWriter.AdvanceWritePosition(HeaderLength + _length);
 				_length = 0;
 			}
 
@@ -102,7 +102,7 @@ namespace EpicGames.Horde.Compute
 			public void Advance(int count) => _length += count;
 
 			/// <inheritdoc/>
-			public Memory<byte> GetMemory(int sizeHint = 0) => _sendBufferWriter.GetMemory().Slice(HeaderLength + _length);
+			public Memory<byte> GetMemory(int sizeHint = 0) => _sendBufferWriter.GetWriteBuffer().Slice(HeaderLength + _length);
 
 			/// <inheritdoc/>
 			public Span<byte> GetSpan(int sizeHint = 0) => GetMemory(sizeHint).Span;
@@ -183,23 +183,28 @@ namespace EpicGames.Horde.Compute
 			IComputeBufferReader recvBufferReader = _recvBuffer.Reader;
 			while (!recvBufferReader.IsComplete)
 			{
-				ReadOnlyMemory<byte> memory = recvBufferReader.GetMemory();
-				if (memory.Length >= HeaderLength)
+				ReadOnlyMemory<byte> memory = recvBufferReader.GetReadBuffer();
+				if (memory.Length < HeaderLength)
 				{
-					int length = BinaryPrimitives.ReadInt32LittleEndian(memory.Span.Slice(1, 4));
-					if (memory.Length >= HeaderLength + length)
-					{
-						ComputeMessageType type = (ComputeMessageType)memory.Span[0];
-						Message message = new Message(type, memory.Slice(HeaderLength, length));
-						if (_logger.IsEnabled(LogLevel.Trace))
-						{
-							LogMessageInfo("RECV", message.Type, message.Data.Span);
-						}
-						recvBufferReader.Advance(HeaderLength + length);
-						return message;
-					}
+					await recvBufferReader.WaitToReadAsync(HeaderLength, cancellationToken);
+					continue;
 				}
-				await recvBufferReader.WaitToReadAsync(memory.Length, cancellationToken);
+
+				int messageLength = BinaryPrimitives.ReadInt32LittleEndian(memory.Span.Slice(1, 4));
+				if (memory.Length < HeaderLength + messageLength)
+				{
+					await recvBufferReader.WaitToReadAsync(HeaderLength + messageLength, cancellationToken);
+					continue;
+				}
+
+				ComputeMessageType type = (ComputeMessageType)memory.Span[0];
+				Message message = new Message(type, memory.Slice(HeaderLength, messageLength));
+				if (_logger.IsEnabled(LogLevel.Trace))
+				{
+					LogMessageInfo("RECV", message.Type, message.Data.Span);
+				}
+				recvBufferReader.AdvanceReadPosition(HeaderLength + messageLength);
+				return message;
 			}
 			return new Message(ComputeMessageType.None, ReadOnlyMemory<byte>.Empty);
 		}
@@ -227,16 +232,10 @@ namespace EpicGames.Horde.Compute
 			}
 
 			IComputeBufferWriter sendBufferWriter = _sendBuffer.Writer;
-			for(; ;)
-			{
-				Memory<byte> memory = sendBufferWriter.GetMemory();
-				if (memory.Length >= maxSize)
-				{
-					_currentBuilder = new MessageBuilder(this, sendBufferWriter, type);
-					return _currentBuilder;
-				}
-				await sendBufferWriter.WaitToWriteAsync(memory.Length, cancellationToken);
-			}
+			await sendBufferWriter.WaitToWriteAsync(maxSize, cancellationToken);
+
+			_currentBuilder = new MessageBuilder(this, sendBufferWriter, type);
+			return _currentBuilder;
 		}
 	}
 }
diff --git a/Engine/Source/Programs/Shared/EpicGames.Horde/Compute/ComputeWorker.cs b/Engine/Source/Programs/Shared/EpicGames.Horde/Compute/ComputeWorker.cs
--- a/Engine/Source/Programs/Shared/EpicGames.Horde/Compute/ComputeWorker.cs
+++ b/Engine/Source/Programs/Shared/EpicGames.Horde/Compute/ComputeWorker.cs
@@ -192,9 +192,9 @@ namespace EpicGames.Horde.Compute
 
 		static async Task ProcessIpcMessagesAsync(IComputeSocket socket, IComputeBufferReader ipcReader, CancellationToken cancellationToken)
 		{
-			while(await ipcReader.WaitToReadAsync(0, cancellationToken))
+			while(await ipcReader.WaitToReadAsync(1, cancellationToken))
 			{
-				ReadOnlyMemory<byte> memory = ipcReader.GetMemory();
+				ReadOnlyMemory<byte> memory = ipcReader.GetReadBuffer();
 				MemoryReader reader = new MemoryReader(memory);
 
 				IpcMessage message = (IpcMessage)reader.ReadUnsignedVarInt();
@@ -222,7 +222,7 @@ namespace EpicGames.Horde.Compute
 						throw new InvalidOperationException($"Invalid IPC message: {message}");
 				}
 
-				ipcReader.Advance(memory.Length - reader.RemainingMemory.Length);
+				ipcReader.AdvanceReadPosition(memory.Length - reader.RemainingMemory.Length);
 			}
 		}
 
diff --git a/Engine/Source/Programs/Shared/EpicGames.Horde/Compute/DefaultComputeSocket.cs b/Engine/Source/Programs/Shared/EpicGames.Horde/Compute/DefaultComputeSocket.cs
--- a/Engine/Source/Programs/Shared/EpicGames.Horde/Compute/DefaultComputeSocket.cs
+++ b/Engine/Source/Programs/Shared/EpicGames.Horde/Compute/DefaultComputeSocket.cs
@@ -176,12 +176,12 @@ namespace EpicGames.Horde.Compute
 
 		async Task ReadPacketAsync(IComputeTransport transport, int id, int size, IComputeBufferWriter writer, CancellationToken cancellationToken)
 		{
-			Memory<byte> memory = writer.GetMemory();
+			Memory<byte> memory = writer.GetWriteBuffer();
 			while (memory.Length < size)
 			{
 				_logger.LogTrace("No space in buffer {Id}, flushing", id);
-				await writer.WaitToWriteAsync(memory.Length, cancellationToken);
-				memory = writer.GetMemory();
+				await writer.WaitToWriteAsync(size, cancellationToken);
+				memory = writer.GetWriteBuffer();
 			}
 
 			for (int offset = 0; offset < size;)
@@ -190,7 +190,7 @@ namespace EpicGames.Horde.Compute
 				offset += read;
 			}
 
-			writer.Advance(size);
+			writer.AdvanceWritePosition(size);
 		}
 
 		IComputeBufferWriter GetReceiveBuffer(Dictionary<int, IComputeBufferWriter> cachedWriters, int id)
@@ -361,18 +361,18 @@ namespace EpicGames.Horde.Compute
 			IComputeBufferReader reader = buffer.Reader;
 			while (!cancellationToken.IsCancellationRequested)
 			{
-				ReadOnlyMemory<byte> memory = reader.GetMemory();
+				ReadOnlyMemory<byte> memory = reader.GetReadBuffer();
 				if (memory.Length > 0)
 				{
 					await SendAsync(channelId, memory, cancellationToken);
-					reader.Advance(memory.Length);
+					reader.AdvanceReadPosition(memory.Length);
 				}
 				if (reader.IsComplete)
 				{
 					await MarkCompleteAsync(channelId, cancellationToken);
 					break;
 				}
-				await reader.WaitToReadAsync(0, cancellationToken);
+				await reader.WaitToReadAsync(1, cancellationToken);
 			}
 
 			lock (_lockObject)
diff --git a/Engine/Source/Programs/Shared/EpicGames.Horde/Compute/IComputeBuffer.cs b/Engine/Source/Programs/Shared/EpicGames.Horde/Compute/IComputeBuffer.cs
--- a/Engine/Source/Programs/Shared/EpicGames.Horde/Compute/IComputeBuffer.cs
+++ b/Engine/Source/Programs/Shared/EpicGames.Horde/Compute/IComputeBuffer.cs
@@ -41,21 +41,21 @@ namespace EpicGames.Horde.Compute
 		/// Updates the read position
 		/// </summary>
 		/// <param name="size">Size of data that was read</param>
-		void Advance(int size);
+		void AdvanceReadPosition(int size);
 
 		/// <summary>
 		/// Gets the next data to read
 		/// </summary>
 		/// <returns>Memory to read from</returns>
-		ReadOnlyMemory<byte> GetMemory();
+		ReadOnlyMemory<byte> GetReadBuffer();
 
 		/// <summary>
 		/// Wait for data to be available, or for the buffer to be marked as complete
 		/// </summary>
-		/// <param name="currentLength">Current length of the buffer</param>
+		/// <param name="minLength">Minimum amount of data to read</param>
 		/// <param name="cancellationToken">Cancellation token for the operation</param>
 		/// <returns>True if new data is available, false if the buffer is complete</returns>
-		ValueTask<bool> WaitToReadAsync(int currentLength, CancellationToken cancellationToken = default);
+		ValueTask<bool> WaitToReadAsync(int minLength, CancellationToken cancellationToken = default);
 	}
 
 	/// <summary>
@@ -72,19 +72,21 @@ namespace EpicGames.Horde.Compute
 		/// <summary>
 		/// Updates the current write position within the buffer
 		/// </summary>
-		void Advance(int size);
+		void AdvanceWritePosition(int size);
 
 		/// <summary>
 		/// Gets memory to write to
 		/// </summary>
 		/// <returns>Memory to be written to</returns>
-		Memory<byte> GetMemory();
+		Memory<byte> GetWriteBuffer();
 
 		/// <summary>
 		/// Gets memory to write to
 		/// </summary>
+		/// <param name="minLength">Minimum size of the desired write buffer</param>
+		/// <param name="cancellationToken">Cancellation token for the operation</param>
 		/// <returns>Memory to be written to</returns>
-		ValueTask WaitToWriteAsync(int currentLength, CancellationToken cancellationToken = default);
+		ValueTask WaitToWriteAsync(int minLength, CancellationToken cancellationToken = default);
 	}
 
 	/// <summary>
@@ -103,15 +105,15 @@ namespace EpicGames.Horde.Compute
 		{
 			for (; ; )
 			{
-				ReadOnlyMemory<byte> readMemory = reader.GetMemory();
+				ReadOnlyMemory<byte> readMemory = reader.GetReadBuffer();
 				if (reader.IsComplete || readMemory.Length > 0)
 				{
 					int length = Math.Min(readMemory.Length, buffer.Length);
 					readMemory.Slice(0, length).CopyTo(buffer);
-					reader.Advance(length);
+					reader.AdvanceReadPosition(length);
 					return length;
 				}
-				await reader.WaitToReadAsync(0, cancellationToken);
+				await reader.WaitToReadAsync(1, cancellationToken);
 			}
 		}
 
@@ -125,11 +127,11 @@ namespace EpicGames.Horde.Compute
 		{
 			while (buffer.Length > 0)
 			{
-				Memory<byte> writeMemory = writer.GetMemory();
+				Memory<byte> writeMemory = writer.GetWriteBuffer();
 				if (writeMemory.Length >= buffer.Length)
 				{
 					buffer.CopyTo(writeMemory);
-					writer.Advance(buffer.Length);
+					writer.AdvanceWritePosition(buffer.Length);
 					break;
 				}
 				await writer.WaitToWriteAsync(writeMemory.Length, cancellationToken);
diff --git a/Engine/Source/Programs/Shared/EpicGames.Horde/Compute/WorkerComputeSocket.cs b/Engine/Source/Programs/Shared/EpicGames.Horde/Compute/WorkerComputeSocket.cs
--- a/Engine/Source/Programs/Shared/EpicGames.Horde/Compute/WorkerComputeSocket.cs
+++ b/Engine/Source/Programs/Shared/EpicGames.Horde/Compute/WorkerComputeSocket.cs
@@ -60,11 +60,11 @@ namespace EpicGames.Horde.Compute
 
 		void AttachBuffer(IpcMessage message, int channelId, SharedMemoryBuffer buffer)
 		{
-			MemoryWriter writer = new MemoryWriter(_commandBuffer.Writer.GetMemory());
+			MemoryWriter writer = new MemoryWriter(_commandBuffer.Writer.GetWriteBuffer());
 			writer.WriteUnsignedVarInt((int)message);
 			writer.WriteUnsignedVarInt(channelId);
 			writer.WriteString(buffer.Name);
-			_commandBuffer.Writer.Advance(writer.Length);
+			_commandBuffer.Writer.AdvanceWritePosition(writer.Length);
 		}
 	}
 
