Horde: Missing file.

diff --git a/Engine/Source/Programs/Shared/EpicGames.Horde.Tests/BufferTests.cs b/Engine/Source/Programs/Shared/EpicGames.Horde.Tests/BufferTests.cs
--- a/Engine/Source/Programs/Shared/EpicGames.Horde.Tests/BufferTests.cs
+++ b/Engine/Source/Programs/Shared/EpicGames.Horde.Tests/BufferTests.cs
@@ -78,16 +78,16 @@ namespace EpicGames.Horde.Tests
 			int offset = 0;
 			while (!reader.IsComplete)
 			{
-				ReadOnlyMemory<byte> memory = reader.GetMemory();
+				ReadOnlyMemory<byte> memory = reader.GetReadBuffer();
 				if (memory.Length == 0)
 				{
-					await reader.WaitToReadAsync(0, CancellationToken.None);
+					await reader.WaitToReadAsync(1, CancellationToken.None);
 					continue;
 				}
 
 				int length = Math.Min(memory.Length, 7);
 				memory.Slice(0, length).CopyTo(output.Slice(offset));
-				reader.Advance(length);
+				reader.AdvanceReadPosition(length);
 				offset += length;
 			}
 		}
