EpicGames.Perforce: Limit concurrency for file operations in ManagedWorkspace

diff --git a/Engine/Source/Programs/Shared/EpicGames.Perforce.Managed/ManagedWorkspace.cs b/Engine/Source/Programs/Shared/EpicGames.Perforce.Managed/ManagedWorkspace.cs
--- a/Engine/Source/Programs/Shared/EpicGames.Perforce.Managed/ManagedWorkspace.cs
+++ b/Engine/Source/Programs/Shared/EpicGames.Perforce.Managed/ManagedWorkspace.cs
@@ -114,6 +114,11 @@ namespace EpicGames.Perforce.Managed
 		/// Maximum number of threads to sync in parallel
 		/// </summary>
 		const int NumParallelSyncThreads = 4;
+		
+		/// <summary>
+		/// Maximum number of concurrent file system operations (copying, moving, deleting etc)
+		/// </summary>
+		const int MaxFileConcurrency = 4;
 
 		/// <summary>
 		/// Minimum amount of space that must be on a drive after a branch is synced
@@ -475,17 +480,20 @@ namespace EpicGames.Perforce.Managed
 				using (ILoggerProgress scope = _logger.BeginProgressScope("Cleaning files..."))
 				{
 					Stopwatch timer = Stopwatch.StartNew();
-					using (ThreadPoolWorkQueue queue = new ThreadPoolWorkQueue())
+					
+					ParallelOptions options = new() { MaxDegreeOfParallelism = MaxFileConcurrency, CancellationToken = cancellationToken };
+					await Parallel.ForEachAsync(filesToDelete, options, (fileToDelete, ct) =>
 					{
-						foreach (FileInfo fileToDelete in filesToDelete)
-						{
-							queue.Enqueue(() => FileUtils.ForceDeleteFile(fileToDelete));
-						}
-						foreach (DirectoryInfo directoryToDelete in directoriesToDelete)
-						{
-							queue.Enqueue(() => FileUtils.ForceDeleteDirectory(directoryToDelete));
-						}
-					}
+						FileUtils.ForceDeleteFile(fileToDelete);
+						return ValueTask.CompletedTask;
+					});
+					
+					await Parallel.ForEachAsync(directoriesToDelete, options, (directoryToDelete, ct) =>
+					{
+						FileUtils.ForceDeleteDirectory(directoryToDelete);
+						return ValueTask.CompletedTask;
+					});
+					
 					scope.Progress = $"({timer.Elapsed.TotalSeconds:0.0}s)";
 				}
 
@@ -1740,7 +1748,12 @@ namespace EpicGames.Perforce.Managed
 					await SaveAsync(TransactionState.Dirty, cancellationToken);
 
 					// Execute all the moves and deletes
-					await ParallelTask.ForEachAsync(sourceAndTargetFiles, sourceAndTargetFile => FileUtils.ForceMoveFile(sourceAndTargetFile.Key, sourceAndTargetFile.Value));
+					ParallelOptions options = new() { MaxDegreeOfParallelism = MaxFileConcurrency, CancellationToken = cancellationToken };
+					await Parallel.ForEachAsync(sourceAndTargetFiles, options, (sourceAndTargetFile, ct) =>
+					{
+						FileUtils.ForceMoveFile(sourceAndTargetFile.Key, sourceAndTargetFile.Value);
+						return ValueTask.CompletedTask;
+					});
 
 					scope.Progress = $"({timer.Elapsed.TotalSeconds:0.0}s)";
 				}
@@ -1755,7 +1768,12 @@ namespace EpicGames.Perforce.Managed
 				{
 					Stopwatch timer = Stopwatch.StartNew();
 
-					await ParallelTask.ForEachAsync(filesToDelete, fileToDelete => RemoveFile(fileToDelete));
+					ParallelOptions options = new() { MaxDegreeOfParallelism = MaxFileConcurrency, CancellationToken = cancellationToken };
+					await Parallel.ForEachAsync(filesToDelete, options, (fileToDelete, ct) =>
+					{
+						RemoveFile(fileToDelete);
+						return ValueTask.CompletedTask;
+					});
 
 					scope.Progress = $"({timer.Elapsed.TotalSeconds:0.0}s)";
 				}
@@ -1852,7 +1870,13 @@ namespace EpicGames.Perforce.Managed
 				using (ILoggerProgress status = _logger.BeginProgressScope(String.Format("Moving {0} {1} from cache...", filesToMove.Length, (filesToMove.Length == 1) ? "file" : "files")))
 				{
 					Stopwatch timer = Stopwatch.StartNew();
-					await ParallelTask.ForEachAsync(filesToMove, fileToMove => MoveFileFromCache(fileToMove, transaction._filesToSync));
+					ParallelOptions options = new() { MaxDegreeOfParallelism = MaxFileConcurrency, CancellationToken = cancellationToken };
+					await Parallel.ForEachAsync(filesToMove, options, (fileToMove, ct) =>
+					{
+						MoveFileFromCache(fileToMove, transaction._filesToSync);
+						return ValueTask.CompletedTask;
+					});
+					
 					_contentIdToTrackedFile = _contentIdToTrackedFile.Where(x => !transaction._filesToMove.ContainsKey(x.Value)).ToDictionary(x => x.Key, x => x.Value);
 					status.Progress = $"({timer.Elapsed.TotalSeconds:0.0}s)";
 				}
@@ -1866,7 +1890,12 @@ namespace EpicGames.Perforce.Managed
 				using (ILoggerProgress status = _logger.BeginProgressScope(String.Format("Copying {0} {1} within workspace...", filesToCopy.Length, (filesToCopy.Length == 1) ? "file" : "files")))
 				{
 					Stopwatch timer = Stopwatch.StartNew();
-					await ParallelTask.ForEachAsync(filesToCopy, fileToCopy => CopyFileWithinWorkspace(fileToCopy, transaction._filesToSync));
+					ParallelOptions options = new() { MaxDegreeOfParallelism = MaxFileConcurrency, CancellationToken = cancellationToken };
+					await Parallel.ForEachAsync(filesToCopy, options, (fileToCopy, ct) =>
+					{
+						CopyFileWithinWorkspace(fileToCopy, transaction._filesToSync);
+						return ValueTask.CompletedTask;
+					});
 					status.Progress = $"({timer.Elapsed.TotalSeconds:0.0}s)";
 				}
 			}
@@ -1946,7 +1975,14 @@ namespace EpicGames.Perforce.Managed
 
 							// Update metadata for the complete batch
 							(int beginIdx, int endIdx) = batches[batchIdx];
-							await ParallelTask.ForAsync(beginIdx, endIdx, idx => filesToSync[idx]._workspaceFile.UpdateMetadata());
+
+							int[] indexesToUpdate = Enumerable.Range(beginIdx, endIdx).ToArray();
+							ParallelOptions options = new() { MaxDegreeOfParallelism = MaxFileConcurrency, CancellationToken = cancellationToken };
+							await Parallel.ForEachAsync(indexesToUpdate, options, (idx, ct) =>
+							{
+								filesToSync[idx]._workspaceFile.UpdateMetadata();
+								return ValueTask.CompletedTask;
+							});
 
 							// Save the current state every minute
 							TimeSpan elapsed = timer.Elapsed;
