Set exception on a task only if it has not already finished running

diff --git a/Engine/Source/Programs/Shared/EpicGames.Horde/Storage/TreeReader.cs b/Engine/Source/Programs/Shared/EpicGames.Horde/Storage/TreeReader.cs
--- a/Engine/Source/Programs/Shared/EpicGames.Horde/Storage/TreeReader.cs
+++ b/Engine/Source/Programs/Shared/EpicGames.Horde/Storage/TreeReader.cs
@@ -388,7 +388,7 @@ namespace EpicGames.Horde.Storage
 			}
 			catch (Exception ex)
 			{
-				queuedHeader.CompletionSource.SetException(ex);
+				queuedHeader.CompletionSource.TrySetException(ex);
 			}
 		}
 
