EpicGames.Perforce: Fix index error in ManagedWorkspace

diff --git a/Engine/Source/Programs/Shared/EpicGames.Perforce.Managed/ManagedWorkspace.cs b/Engine/Source/Programs/Shared/EpicGames.Perforce.Managed/ManagedWorkspace.cs
--- a/Engine/Source/Programs/Shared/EpicGames.Perforce.Managed/ManagedWorkspace.cs
+++ b/Engine/Source/Programs/Shared/EpicGames.Perforce.Managed/ManagedWorkspace.cs
@@ -1976,7 +1976,7 @@ namespace EpicGames.Perforce.Managed
 							// Update metadata for the complete batch
 							(int beginIdx, int endIdx) = batches[batchIdx];
 
-							int[] indexesToUpdate = Enumerable.Range(beginIdx, endIdx).ToArray();
+							int[] indexesToUpdate = Enumerable.Range(beginIdx, endIdx - beginIdx).ToArray();
 							ParallelOptions options = new() { MaxDegreeOfParallelism = MaxFileConcurrency, CancellationToken = cancellationToken };
 							await Parallel.ForEachAsync(indexesToUpdate, options, (idx, ct) =>
 							{
