EpicGames.Core: Use separate overloads for HTTP requests that specify a string parameter for the path. This is treated differently by HttpClient, since it's combined with the HTTP base address.

diff --git a/Engine/Source/Programs/Shared/EpicGames.Core/HttpClientExtensions.cs b/Engine/Source/Programs/Shared/EpicGames.Core/HttpClientExtensions.cs
--- a/Engine/Source/Programs/Shared/EpicGames.Core/HttpClientExtensions.cs
+++ b/Engine/Source/Programs/Shared/EpicGames.Core/HttpClientExtensions.cs
@@ -17,9 +17,6 @@ namespace EpicGames.Core
 	/// </summary>
 	public static class HttpClientExtensions
 	{
-		/// <inheritdoc cref="GetAsync{TResponse}(HttpClient, Uri, CancellationToken)"/>
-		public static Task<TResponse> GetAsync<TResponse>(this HttpClient client, string url, CancellationToken cancellationToken) => GetAsync<TResponse>(client, new Uri(url), cancellationToken);
-
 		/// <summary>
 		/// Gets a resource from an HTTP endpoint and parses it as a JSON object
 		/// </summary>
@@ -28,7 +25,7 @@ namespace EpicGames.Core
 		/// <param name="url">The url to retrieve</param>
 		/// <param name="cancellationToken">Cancels the request</param>
 		/// <returns>New instance of the object</returns>
-		public static async Task<TResponse> GetAsync<TResponse>(this HttpClient client, Uri url, CancellationToken cancellationToken)
+		public static async Task<TResponse> GetAsync<TResponse>(this HttpClient client, string url, CancellationToken cancellationToken)
 		{
 			using (HttpResponseMessage response = await client.GetAsync(url, cancellationToken))
 			{
@@ -37,8 +34,15 @@ namespace EpicGames.Core
 			}
 		}
 
-		/// <inheritdoc cref="PostAsync{TRequest}(HttpClient, Uri, TRequest, CancellationToken)"/>
-		public static Task<HttpResponseMessage> PostAsync<TRequest>(this HttpClient client, string url, TRequest request, CancellationToken cancellationToken) => PostAsync<TRequest>(client, new Uri(url), request, cancellationToken);
+		/// <inheritdoc cref="GetAsync{TResponse}(HttpClient, String, CancellationToken)"/>
+		public static async Task<TResponse> GetAsync<TResponse>(this HttpClient client, Uri url, CancellationToken cancellationToken)
+		{
+			using (HttpResponseMessage response = await client.GetAsync(url, cancellationToken))
+			{
+				response.EnsureSuccessStatusCode();
+				return await ParseJsonContent<TResponse>(response);
+			}
+		}
 
 		/// <summary>
 		/// Posts an object to an HTTP endpoint as a JSON object
@@ -49,14 +53,18 @@ namespace EpicGames.Core
 		/// <param name="request">The object to post</param>
 		/// <param name="cancellationToken">Cancels the request</param>
 		/// <returns>Response message</returns>
-		public static async Task<HttpResponseMessage> PostAsync<TRequest>(this HttpClient client, Uri url, TRequest request, CancellationToken cancellationToken)
+		public static async Task<HttpResponseMessage> PostAsync<TRequest>(this HttpClient client, string url, TRequest request, CancellationToken cancellationToken)
 		{
 			using HttpContent content = ToJsonContent(request);
 			return await client.PostAsync(url, content, cancellationToken);
 		}
 
-		/// <inheritdoc cref="PostAsync{TRequest}(HttpClient, Uri, TRequest, CancellationToken)"/>
-		public static Task<TResponse> PostAsync<TResponse, TRequest>(this HttpClient client, string url, TRequest request, CancellationToken cancellationToken) => PostAsync<TResponse, TRequest>(client, new Uri(url), request, cancellationToken);
+		/// <inheritdoc cref="PostAsync{TRequest}(HttpClient, String, TRequest, CancellationToken)"/>
+		public static async Task<HttpResponseMessage> PostAsync<TRequest>(this HttpClient client, Uri url, TRequest request, CancellationToken cancellationToken)
+		{
+			using HttpContent content = ToJsonContent(request);
+			return await client.PostAsync(url, content, cancellationToken);
+		}
 
 		/// <summary>
 		/// Posts an object to an HTTP endpoint as a JSON object, and parses the response object
@@ -68,7 +76,7 @@ namespace EpicGames.Core
 		/// <param name="request">The object to post</param>
 		/// <param name="cancellationToken">Cancels the request</param>
 		/// <returns>The response parsed into the requested type</returns>
-		public static async Task<TResponse> PostAsync<TResponse, TRequest>(this HttpClient client, Uri url, TRequest request, CancellationToken cancellationToken)
+		public static async Task<TResponse> PostAsync<TResponse, TRequest>(this HttpClient client, string url, TRequest request, CancellationToken cancellationToken)
 		{
 			using (HttpResponseMessage response = await PostAsync(client, url, request, cancellationToken))
 			{
@@ -77,8 +85,15 @@ namespace EpicGames.Core
 			}
 		}
 
-		/// <inheritdoc cref="PutAsync{TRequest}(HttpClient, Uri, TRequest, CancellationToken)"/>
-		public static Task<HttpResponseMessage> PutAsync<TRequest>(this HttpClient client, string url, TRequest request, CancellationToken cancellationToken) => PutAsync<TRequest>(client, new Uri(url), request, cancellationToken);
+		/// <inheritdoc cref="PostAsync{TRequest}(HttpClient, String, TRequest, CancellationToken)"/>
+		public static async Task<TResponse> PostAsync<TResponse, TRequest>(this HttpClient client, Uri url, TRequest request, CancellationToken cancellationToken)
+		{
+			using (HttpResponseMessage response = await PostAsync(client, url, request, cancellationToken))
+			{
+				response.EnsureSuccessStatusCode();
+				return await ParseJsonContent<TResponse>(response);
+			}
+		}
 
 		/// <summary>
 		/// Puts an object to an HTTP endpoint as a JSON object
@@ -86,17 +101,21 @@ namespace EpicGames.Core
 		/// <typeparam name="TRequest">The object type to post</typeparam>
 		/// <param name="client">The http client instance</param>
 		/// <param name="url">The url to post to</param>
-		/// <param name="obj">The object to post</param>
+		/// <param name="request">The object to post</param>
 		/// <param name="cancellationToken">Cancels the request</param>
 		/// <returns>Response message</returns>
-		public static async Task<HttpResponseMessage> PutAsync<TRequest>(this HttpClient client, Uri url, TRequest obj, CancellationToken cancellationToken)
+		public static async Task<HttpResponseMessage> PutAsync<TRequest>(this HttpClient client, string url, TRequest request, CancellationToken cancellationToken)
 		{
-			using HttpContent content = ToJsonContent(obj);
+			using HttpContent content = ToJsonContent(request);
 			return await client.PutAsync(url, content, cancellationToken);
 		}
 
-		/// <inheritdoc cref="PutAsync{TRequest}(HttpClient, Uri, TRequest, CancellationToken)"/>
-		public static Task<TResponse> PutAsync<TResponse, TRequest>(this HttpClient client, string url, TRequest request, CancellationToken cancellationToken) => PutAsync<TResponse, TRequest>(client, new Uri(url), request, cancellationToken);
+		/// <inheritdoc cref="PutAsync{TRequest}(HttpClient, String, TRequest, CancellationToken)"/>
+		public static async Task<HttpResponseMessage> PutAsync<TRequest>(this HttpClient client, Uri url, TRequest request, CancellationToken cancellationToken)
+		{
+			using HttpContent content = ToJsonContent(request);
+			return await client.PutAsync(url, content, cancellationToken);
+		}
 
 		/// <summary>
 		/// Puts an object to an HTTP endpoint as a JSON object, and parses the response object
@@ -108,6 +127,16 @@ namespace EpicGames.Core
 		/// <param name="request">The object to post</param>
 		/// <param name="cancellationToken">Cancels the request</param>
 		/// <returns>The response parsed into the requested type</returns>
+		public static async Task<TResponse> PutAsync<TResponse, TRequest>(this HttpClient client, string url, TRequest request, CancellationToken cancellationToken)
+		{
+			using (HttpResponseMessage response = await PutAsync(client, url, request, cancellationToken))
+			{
+				response.EnsureSuccessStatusCode();
+				return await ParseJsonContent<TResponse>(response);
+			}
+		}
+
+		/// <inheritdoc cref="PutAsync{TRequest}(HttpClient, String, TRequest, CancellationToken)"/>
 		public static async Task<TResponse> PutAsync<TResponse, TRequest>(this HttpClient client, Uri url, TRequest request, CancellationToken cancellationToken)
 		{
 			using (HttpResponseMessage response = await PutAsync(client, url, request, cancellationToken))
