Horde: Move more code into the common ComputeBuffer implementation, and finish support for multiple readers.

diff --git a/Engine/Source/Programs/Shared/EpicGames.Horde.Tests/BufferTests.cs b/Engine/Source/Programs/Shared/EpicGames.Horde.Tests/BufferTests.cs
--- a/Engine/Source/Programs/Shared/EpicGames.Horde.Tests/BufferTests.cs
+++ b/Engine/Source/Programs/Shared/EpicGames.Horde.Tests/BufferTests.cs
@@ -1,9 +1,7 @@
 // Copyright Epic Games, Inc. All Rights Reserved.
 
 using System;
-using System.Diagnostics;
 using System.IO.Pipelines;
-using System.IO.Pipes;
 using System.Linq;
 using System.Runtime.InteropServices;
 using System.Security.Cryptography;
@@ -22,6 +20,61 @@ namespace EpicGames.Horde.Tests
 	{
 		const int ChannelId = 0;
 
+		[TestMethod]
+		public async Task TestSimpleBuffer()
+		{
+			using PooledBuffer buffer = new PooledBuffer(2, 1024);
+			buffer.Writer.AdvanceWritePosition(10);
+			await buffer.Reader.WaitToReadAsync(9);
+			buffer.Reader.AdvanceReadPosition(9);
+			await buffer.Reader.WaitToReadAsync(1);
+		}
+
+		[TestMethod]
+		public void TestOverflow()
+		{
+			using PooledBuffer buffer = new PooledBuffer(2, 20);
+
+			// Fill up the first chunk
+			Assert.AreEqual(20, buffer.Writer.GetWriteBuffer().Length);
+			buffer.Writer.AdvanceWritePosition(10);
+			Assert.AreEqual(10, buffer.Writer.GetWriteBuffer().Length);
+			buffer.Writer.AdvanceWritePosition(10);
+			Assert.AreEqual(0, buffer.Writer.GetWriteBuffer().Length);
+
+			Task waitToWriteTask = buffer.Writer.WaitToWriteAsync(1).AsTask();
+			Assert.IsTrue(waitToWriteTask.IsCompleted);
+
+			// Fill up the second chunk
+			Assert.AreEqual(20, buffer.Writer.GetWriteBuffer().Length);
+			buffer.Writer.AdvanceWritePosition(10);
+			Assert.AreEqual(10, buffer.Writer.GetWriteBuffer().Length);
+			buffer.Writer.AdvanceWritePosition(10);
+			Assert.AreEqual(0, buffer.Writer.GetWriteBuffer().Length);
+
+			waitToWriteTask = buffer.Writer.WaitToWriteAsync(1).AsTask();
+			Assert.IsFalse(waitToWriteTask.IsCompleted);
+
+			// Wait for data to be read
+			Assert.AreEqual(20, buffer.Reader.GetReadBuffer().Length);
+			buffer.Reader.AdvanceReadPosition(10);
+			Assert.IsFalse(waitToWriteTask.IsCompleted);
+
+			Assert.AreEqual(10, buffer.Reader.GetReadBuffer().Length);
+			buffer.Reader.AdvanceReadPosition(10);
+			Assert.AreEqual(0, buffer.Reader.GetReadBuffer().Length);
+
+			Task waitToReadTask = buffer.Reader.WaitToReadAsync(1).AsTask();
+			Assert.IsTrue(waitToReadTask.IsCompleted);
+
+			Assert.AreEqual(20, buffer.Reader.GetReadBuffer().Length);
+			Assert.IsTrue(waitToWriteTask.IsCompleted);
+
+			// Make sure both reader and writer have something to work with
+			Assert.AreEqual(20, buffer.Reader.GetReadBuffer().Length);
+			Assert.AreEqual(20, buffer.Writer.GetWriteBuffer().Length);
+		}
+
 		[TestMethod]
 		public async Task TestPooledBuffer()
 		{
diff --git a/Engine/Source/Programs/Shared/EpicGames.Horde/Compute/Buffers/PooledBuffer.cs b/Engine/Source/Programs/Shared/EpicGames.Horde/Compute/Buffers/PooledBuffer.cs
--- a/Engine/Source/Programs/Shared/EpicGames.Horde/Compute/Buffers/PooledBuffer.cs
+++ b/Engine/Source/Programs/Shared/EpicGames.Horde/Compute/Buffers/PooledBuffer.cs
@@ -2,9 +2,12 @@
 
 using System;
 using System.Buffers;
+using System.Linq;
+using System.Runtime.InteropServices;
 using System.Threading;
 using System.Threading.Tasks;
 using EpicGames.Core;
+using Microsoft.CodeAnalysis.CSharp.Syntax;
 
 namespace EpicGames.Horde.Compute.Buffers
 {
@@ -30,7 +33,7 @@ namespace EpicGames.Horde.Compute.Buffers
 		/// <param name="numChunks">Number of chunks in the buffer</param>
 		/// <param name="chunkLength">Length of each chunk</param>
 		public PooledBuffer(int numChunks, int chunkLength)
-			: this(new PooledBufferCore(numChunks, chunkLength))
+			: this(PooledBufferCore.Create(numChunks, chunkLength, 1))
 		{
 		}
 
@@ -68,55 +71,51 @@ namespace EpicGames.Horde.Compute.Buffers
 	/// </summary>
 	sealed class PooledBufferCore : ComputeBufferBase
 	{
-		class Chunk : ChunkBase, IDisposable
+		unsafe class PinnedBuffer : IDisposable
 		{
-			readonly IMemoryOwner<byte> _memoryOwner;
-			ulong _stateValue;
+			readonly byte[] _data;
+			readonly GCHandle _handle;
 
-			/// <inheritdoc/>
-			protected override ref ulong StateValue => ref _stateValue;
+			public Memory<byte> Data => _data;
+			public void* Ptr => _handle.AddrOfPinnedObject().ToPointer();
 
-			public Chunk(IMemoryOwner<byte> memoryOwner)
-				: base(memoryOwner.Memory)
+			public PinnedBuffer(int length)
 			{
-				_memoryOwner = memoryOwner;
+				_data = new byte[length];
+				_handle = GCHandle.Alloc(_data, GCHandleType.Pinned);
 			}
 
-			public void Dispose() => _memoryOwner.Dispose();
+			public void Dispose()
+			{
+				_handle.Free();
+			}
 		}
 
-		readonly Chunk[] _chunks;
+		readonly PinnedBuffer[] _buffers;
 		readonly AsyncEvent _writerEvent = new AsyncEvent();
 		readonly AsyncEvent _readerEvent = new AsyncEvent();
 
-		/// <summary>
-		/// Constructor
-		/// </summary>
-		/// <param name="numChunks">Number of chunks in the buffer</param>
-		/// <param name="chunkLength">Length of each chunk</param>
-		public PooledBufferCore(int numChunks, int chunkLength)
-			: this(CreateChunks(numChunks, chunkLength), 1)
+		private PooledBufferCore(HeaderPtr headerPtr, Memory<byte>[] chunks, PinnedBuffer[] buffers)
+			: base(headerPtr, chunks)
 		{
+			_buffers = buffers;
 		}
 
-		/// <summary>
-		/// Constructor
-		/// </summary>
-		private PooledBufferCore(Chunk[] chunks, int readerIdx)
-			: base(chunks, readerIdx, true)
+		public static unsafe PooledBufferCore Create(int numChunks, int chunkLength, int numReaders)
 		{
-			_chunks = chunks;
-		}
+			Memory<byte>[] chunks = new Memory<byte>[numChunks];
 
-		static Chunk[] CreateChunks(int numChunks, int chunkLength)
-		{
-			Chunk[] chunks = new Chunk[numChunks];
-			for (int chunkIdx = 0; chunkIdx < numChunks; chunkIdx++)
+			PinnedBuffer[] buffers = new PinnedBuffer[numChunks + 1];
+			for (int idx = 0; idx < numChunks; idx++)
 			{
-				Chunk chunk = new Chunk(MemoryPool<byte>.Shared.Rent(chunkLength));
-				chunks[chunkIdx] = chunk;
+				buffers[idx] = new PinnedBuffer(chunkLength);
+				chunks[idx] = buffers[idx].Data;
 			}
-			return chunks;
+
+			buffers[numChunks] = new PinnedBuffer(HeaderSize);
+
+			HeaderPtr headerPtr = new HeaderPtr((ulong*)buffers[numChunks].Ptr, numReaders, numChunks, chunkLength);
+			return new PooledBufferCore(headerPtr, chunks, buffers);
 		}
 
 		/// <inheritdoc/>
@@ -126,9 +125,9 @@ namespace EpicGames.Horde.Compute.Buffers
 
 			if (disposing)
 			{
-				for (int idx = 0; idx < _chunks.Length; idx++)
+				for (int idx = 0; idx < _buffers.Length; idx++)
 				{
-					_chunks[idx].Dispose();
+					_buffers[idx].Dispose();
 				}
 			}
 		}
diff --git a/Engine/Source/Programs/Shared/EpicGames.Horde/Compute/Buffers/SharedMemoryBuffer.cs b/Engine/Source/Programs/Shared/EpicGames.Horde/Compute/Buffers/SharedMemoryBuffer.cs
--- a/Engine/Source/Programs/Shared/EpicGames.Horde/Compute/Buffers/SharedMemoryBuffer.cs
+++ b/Engine/Source/Programs/Shared/EpicGames.Horde/Compute/Buffers/SharedMemoryBuffer.cs
@@ -6,7 +6,6 @@ using System.IO.MemoryMappedFiles;
 using System.IO;
 using System.Threading.Tasks;
 using System.Threading;
-using System.Buffers.Binary;
 using System.Runtime.InteropServices;
 
 namespace EpicGames.Horde.Compute.Buffers
@@ -73,22 +72,7 @@ namespace EpicGames.Horde.Compute.Buffers
 		/// <param name="chunkLength">Length of each chunk</param>
 		public static unsafe SharedMemoryBuffer CreateNew(string? name, int numChunks, int chunkLength)
 		{
-			long capacity = SharedMemoryBufferCore.HeaderLength + (numChunks * sizeof(ulong)) + (numChunks * chunkLength);
-
-			name ??= $"Local\\COMPUTE_{Guid.NewGuid()}";
-
-			MemoryMappedFile memoryMappedFile = MemoryMappedFile.CreateNew($"{name}_M", capacity, MemoryMappedFileAccess.ReadWrite, MemoryMappedFileOptions.None, HandleInheritability.Inheritable);
-			MemoryMappedViewAccessor memoryMappedViewAccessor = memoryMappedFile.CreateViewAccessor();
-			MemoryMappedView memoryMappedView = new MemoryMappedView(memoryMappedViewAccessor);
-
-			Span<byte> header = memoryMappedView.GetMemory(0, 8).Span;
-			BinaryPrimitives.WriteInt32LittleEndian(header, numChunks);
-			BinaryPrimitives.WriteInt32LittleEndian(header.Slice(4), chunkLength);
-
-			Native.EventHandle writerEvent = Native.EventHandle.CreateNew($"{name}_W", EventResetMode.ManualReset, true, HandleInheritability.Inheritable);
-			Native.EventHandle readerEvent = Native.EventHandle.CreateNew($"{name}_R", EventResetMode.ManualReset, true, HandleInheritability.Inheritable);
-
-			SharedMemoryBufferCore core = new SharedMemoryBufferCore(name, memoryMappedFile, memoryMappedViewAccessor, memoryMappedView, numChunks, chunkLength, readerEvent, writerEvent, true);
+			SharedMemoryBufferCore core = SharedMemoryBufferCore.CreateNew(name, numChunks, chunkLength);
 			return new SharedMemoryBuffer(core);
 		}
 
@@ -98,23 +82,7 @@ namespace EpicGames.Horde.Compute.Buffers
 		/// <param name="name">Name of the buffer to open</param>
 		public static SharedMemoryBuffer OpenExisting(string name)
 		{
-			if (!RuntimeInformation.IsOSPlatform(OSPlatform.Windows))
-			{
-				throw new NotSupportedException();
-			}
-
-			MemoryMappedFile memoryMappedFile = MemoryMappedFile.OpenExisting($"{name}_M");
-			MemoryMappedViewAccessor memoryMappedViewAccessor = memoryMappedFile.CreateViewAccessor();
-			MemoryMappedView memoryMappedView = new MemoryMappedView(memoryMappedViewAccessor);
-
-			ReadOnlySpan<byte> header = memoryMappedView.GetMemory(0, SharedMemoryBufferCore.HeaderLength).Span;
-			int numChunks = BinaryPrimitives.ReadInt32LittleEndian(header);
-			int chunkLength = BinaryPrimitives.ReadInt32LittleEndian(header.Slice(4));
-
-			Native.EventHandle readerEvent = Native.EventHandle.OpenExisting($"{name}_R");
-			Native.EventHandle writerEvent = Native.EventHandle.OpenExisting($"{name}_W");
-
-			SharedMemoryBufferCore core = new SharedMemoryBufferCore(name, memoryMappedFile, memoryMappedViewAccessor, memoryMappedView, numChunks, chunkLength, readerEvent, writerEvent, false);
+			SharedMemoryBufferCore core = SharedMemoryBufferCore.OpenExisting(name);
 			return new SharedMemoryBuffer(core);
 		}
 	}
@@ -124,21 +92,6 @@ namespace EpicGames.Horde.Compute.Buffers
 	/// </summary>
 	sealed class SharedMemoryBufferCore : ComputeBufferBase
 	{
-		public const int HeaderLength = sizeof(int) + sizeof(int); // num chunks, chunk length
-
-		sealed unsafe class Chunk : ChunkBase
-		{
-			readonly ulong* _statePtr;
-
-			protected override ref ulong StateValue => ref *_statePtr;
-
-			public Chunk(ulong* statePtr, Memory<byte> memory)
-				: base(memory)
-			{
-				_statePtr = statePtr;
-			}
-		}
-
 		readonly MemoryMappedFile _memoryMappedFile;
 		readonly MemoryMappedViewAccessor _memoryMappedViewAccessor;
 		readonly MemoryMappedView _memoryMappedView;
@@ -152,8 +105,8 @@ namespace EpicGames.Horde.Compute.Buffers
 		/// <summary>
 		/// Constructor
 		/// </summary>
-		public unsafe SharedMemoryBufferCore(string name, MemoryMappedFile memoryMappedFile, MemoryMappedViewAccessor memoryMappedViewAccessor, MemoryMappedView memoryMappedView, int numChunks, int chunkLength, Native.EventHandle readerEvent, Native.EventHandle writerEvent, bool isNew)
-			: base(CreateChunks(numChunks, chunkLength, memoryMappedView), 1, isNew)
+		private unsafe SharedMemoryBufferCore(HeaderPtr headerPtr, string name, MemoryMappedFile memoryMappedFile, MemoryMappedViewAccessor memoryMappedViewAccessor, MemoryMappedView memoryMappedView, Native.EventHandle readerEvent, Native.EventHandle writerEvent)
+			: base(headerPtr, CreateChunks(headerPtr.NumChunks, headerPtr.ChunkLength, memoryMappedView))
 		{
 			Name = name;
 
@@ -165,16 +118,52 @@ namespace EpicGames.Horde.Compute.Buffers
 			_writerEvent = writerEvent;
 		}
 
-		static unsafe Chunk[] CreateChunks(int numChunks, int chunkLength, MemoryMappedView memoryMappedView)
+		public static unsafe SharedMemoryBufferCore CreateNew(string? name, int numChunks, int chunkLength)
 		{
-			Chunk[] chunks = new Chunk[numChunks];
-			for (int chunkIdx = 0; chunkIdx < numChunks; chunkIdx++)
+			long capacity = HeaderSize + (numChunks * sizeof(ulong)) + (numChunks * chunkLength);
+
+			name ??= $"Local\\COMPUTE_{Guid.NewGuid()}";
+
+			MemoryMappedFile memoryMappedFile = MemoryMappedFile.CreateNew($"{name}_M", capacity, MemoryMappedFileAccess.ReadWrite, MemoryMappedFileOptions.None, HandleInheritability.Inheritable);
+			MemoryMappedViewAccessor memoryMappedViewAccessor = memoryMappedFile.CreateViewAccessor();
+			MemoryMappedView memoryMappedView = new MemoryMappedView(memoryMappedViewAccessor);
+
+			Native.EventHandle writerEvent = Native.EventHandle.CreateNew($"{name}_W", EventResetMode.ManualReset, true, HandleInheritability.Inheritable);
+			Native.EventHandle readerEvent = Native.EventHandle.CreateNew($"{name}_R", EventResetMode.ManualReset, true, HandleInheritability.Inheritable);
+
+			HeaderPtr headerPtr = new HeaderPtr((ulong*)memoryMappedView.GetPointer(), 1, numChunks, chunkLength);
+			return new SharedMemoryBufferCore(headerPtr, name, memoryMappedFile, memoryMappedViewAccessor, memoryMappedView, readerEvent, writerEvent);
+		}
+
+		/// <summary>
+		/// Open an existing buffer by name
+		/// </summary>
+		/// <param name="name">Name of the buffer to open</param>
+		public static unsafe SharedMemoryBufferCore OpenExisting(string name)
+		{
+			if (!RuntimeInformation.IsOSPlatform(OSPlatform.Windows))
 			{
-				byte* statePtr = memoryMappedView.GetPointer() + HeaderLength + (chunkIdx * sizeof(ulong));
-				int chunkOffset = HeaderLength + (numChunks * sizeof(ulong)) + (chunkLength * chunkIdx);
+				throw new NotSupportedException();
+			}
+
+			MemoryMappedFile memoryMappedFile = MemoryMappedFile.OpenExisting($"{name}_M");
+			MemoryMappedViewAccessor memoryMappedViewAccessor = memoryMappedFile.CreateViewAccessor();
+			MemoryMappedView memoryMappedView = new MemoryMappedView(memoryMappedViewAccessor);
 
-				Chunk chunk = new Chunk((ulong*)statePtr, memoryMappedView.GetMemory(chunkOffset, chunkLength));
-				chunks[chunkIdx] = chunk;
+			Native.EventHandle readerEvent = Native.EventHandle.OpenExisting($"{name}_R");
+			Native.EventHandle writerEvent = Native.EventHandle.OpenExisting($"{name}_W");
+
+			HeaderPtr headerPtr = new HeaderPtr((ulong*)memoryMappedView.GetPointer());
+			return new SharedMemoryBufferCore(headerPtr, name, memoryMappedFile, memoryMappedViewAccessor, memoryMappedView, readerEvent, writerEvent);
+		}
+
+		static unsafe Memory<byte>[] CreateChunks(int numChunks, int chunkLength, MemoryMappedView memoryMappedView)
+		{
+			Memory<byte>[] chunks = new Memory<byte>[numChunks];
+			for (int chunkIdx = 0; chunkIdx < numChunks; chunkIdx++)
+			{
+				int chunkOffset = HeaderSize + (chunkLength * chunkIdx);
+				chunks[chunkIdx] = memoryMappedView.GetMemory(chunkOffset, chunkLength);
 			}
 			return chunks;
 		}
diff --git a/Engine/Source/Programs/Shared/EpicGames.Horde/Compute/ComputeBufferBase.cs b/Engine/Source/Programs/Shared/EpicGames.Horde/Compute/ComputeBufferBase.cs
--- a/Engine/Source/Programs/Shared/EpicGames.Horde/Compute/ComputeBufferBase.cs
+++ b/Engine/Source/Programs/Shared/EpicGames.Horde/Compute/ComputeBufferBase.cs
@@ -1,6 +1,8 @@
 // Copyright Epic Games, Inc. All Rights Reserved.
 
 using System;
+using System.Diagnostics;
+using System.Linq;
 using System.Threading;
 using System.Threading.Tasks;
 
@@ -11,20 +13,58 @@ namespace EpicGames.Horde.Compute
 	/// </summary>
 	public abstract class ComputeBufferBase : IDisposable
 	{
+		const int MaxChunks = 16;
+		const int MaxReaders = 16;
+
+		internal const int HeaderSize = (2 + MaxChunks + MaxReaders) * sizeof(ulong);
+
+		/// <summary>
+		/// Tracked state of the buffer
+		/// </summary>
+		internal readonly unsafe struct HeaderPtr
+		{
+			readonly ulong* _data;
+
+			public HeaderPtr(ulong* data) => _data = data;
+
+			public HeaderPtr(ulong* data, int numReaders, int numChunks, int chunkLength)
+			{
+				_data = data;
+
+				data[0] = ((ulong)(uint)numChunks << 32) | (uint)numReaders;
+				data[1] = (uint)chunkLength;
+
+				GetChunkStatePtr(0).StartWriting(numReaders);
+			}
+
+			public int NumReaders => (int)_data[0];
+			public int NumChunks => (int)(_data[0] >> 32);
+			public int ChunkLength => (int)_data[1];
+			public int WriteChunkIdx
+			{
+				get => (int)(_data[1] >> 32);
+				set => _data[1] = ((ulong)value << 32) | (uint)ChunkLength;
+			}
+
+			public ChunkStatePtr GetChunkStatePtr(int chunkIdx) => new ChunkStatePtr(_data + 2 + chunkIdx);
+
+			public ReaderStatePtr GetReaderStatePtr(int readerIdx) => new ReaderStatePtr(_data + 2 + MaxChunks + readerIdx);
+		}
+
 		/// <summary>
 		/// Write state for a chunk
 		/// </summary>
-		protected enum WriteState
+		internal enum WriteState
 		{
 			/// <summary>
-			/// Chunk is still being appended to
+			/// Writer has moved to the next chunk
 			/// </summary>
-			Writing = 0,
+			MovedToNext = 0,
 
 			/// <summary>
-			/// Writer has moved to the next chunk
+			/// Chunk is still being appended to
 			/// </summary>
-			MovedToNext = 2,
+			Writing = 2,
 
 			/// <summary>
 			/// This chunk marks the end of the stream
@@ -35,31 +75,21 @@ namespace EpicGames.Horde.Compute
 		/// <summary>
 		/// Stores the state of a chunk in a 64-bit value, which can be updated atomically
 		/// </summary>
-		protected record struct ChunkState(ulong Value)
+		internal record struct ChunkState(ulong Value)
 		{
-			/// <summary>
-			/// Written length of this chunk
-			/// </summary>
+			// Written length of this chunk
 			public int Length => (int)(Value & 0x7fffffff);
 
-			/// <summary>
-			/// Set of flags which are set for each reader that still has to read from a chunk
-			/// </summary>
+			// Set of flags which are set for each reader that still has to read from a chunk
 			public int ReaderFlags => (int)((Value >> 31) & 0x7fffffff);
 
-			/// <summary>
-			/// State of the writer
-			/// </summary>
+			// State of the writer
 			public WriteState WriteState => (WriteState)(Value >> 62);
 
-			/// <summary>
-			/// Constructor
-			/// </summary>
+			// Constructor
 			public ChunkState(WriteState writerState, int readerFlags, int length) : this(((ulong)writerState << 62) | ((ulong)readerFlags << 31) | (uint)length) { }
 
-			/// <summary>
-			/// Test whether a particular reader is still referencing the chunk
-			/// </summary>
+			// Test whether a particular reader is still referencing the chunk
 			public bool HasReaderFlag(int readerIdx) => (Value & (1UL << (readerIdx + 31))) != 0;
 
 			/// <inheritdoc/>
@@ -67,96 +97,131 @@ namespace EpicGames.Horde.Compute
 		}
 
 		/// <summary>
-		/// Base class for chunks of data in the buffer
+		/// Wraps a pointer to the state of a chunk
 		/// </summary>
-		protected abstract class ChunkBase
+		internal readonly unsafe struct ChunkStatePtr
 		{
-			/// <summary>
-			/// Accessor for the current state
-			/// </summary>
-			public ChunkState State => new ChunkState(Interlocked.CompareExchange(ref StateValue, 0, 0));
+			readonly ulong* _data;
 
-			/// <summary>
-			/// Data underlying the chunk
-			/// </summary>
-			public Memory<byte> Memory { get; }
+			public ChunkStatePtr(ulong* data) => _data = data;
 
-			/// <summary>
-			/// Current state of this chunk
-			/// </summary>
-			protected abstract ref ulong StateValue { get; }
+			// Current value of the chunk state
+			public ChunkState Value => new ChunkState(Interlocked.CompareExchange(ref *_data, 0, 0));
 
-			/// <summary>
-			/// Constructor
-			/// </summary>
-			protected ChunkBase(Memory<byte> memory)
-			{
-				Memory = memory;
-			}
+			// Written length of this chunk
+			public int Length => Value.Length;
 
-			/// <summary>
-			/// Append data to the chunk
-			/// </summary>
-			public void Append(int length) => Interlocked.Add(ref StateValue, (uint)length);
+			// Set of flags which are set for each reader that still has to read from a chunk
+			public int ReaderFlags => Value.ReaderFlags;
 
-			/// <summary>
-			/// Mark the chunk as being written to
-			/// </summary>
-			public void StartWriting(int numReaders) => Interlocked.Exchange(ref StateValue, new ChunkState(WriteState.Writing, (1 << numReaders) - 1, 0).Value);
+			// State of the writer
+			public WriteState WriteState => Value.WriteState;
 
-			/// <summary>
-			/// Move to the next chunk
-			/// </summary>
-			public void MoveToNext() => Interlocked.Or(ref StateValue, (ulong)WriteState.MovedToNext << 62);
+			// Append data to the chunk
+			public void Append(int length) => Interlocked.Add(ref *_data, (uint)length);
 
-			/// <summary>
-			/// Move to the next chunk
-			/// </summary>
-			public void MarkComplete() => Interlocked.Or(ref StateValue, (ulong)WriteState.Complete << 62);
+			// Mark the chunk as being written to
+			public void StartWriting(int numReaders) => Interlocked.Exchange(ref *_data, new ChunkState(WriteState.Writing, (1 << numReaders) - 1, 0).Value);
 
-			/// <summary>
-			/// Clear the reader flag
-			/// </summary>
-			public void FinishReading(int readerIdx) => Interlocked.And(ref StateValue, ~(1UL << (readerIdx + 31)));
+			// Move to the next chunk
+			public void MarkComplete() => Interlocked.Or(ref *_data, (ulong)WriteState.Complete << 62);
+
+			// Clear the reader flag
+			public void FinishReading(int readerIdx) => Interlocked.And(ref *_data, ~(1UL << (readerIdx + 31)));
+
+			// Move to the next chunk
+			public void FinishWriting() => Interlocked.And(ref *_data, ~((ulong)WriteState.Writing << 62));
 
 			/// <inheritdoc/>
-			public override string ToString() => State.ToString();
+			public override string ToString() => Value.ToString();
 		}
 
+		/// <summary>
+		/// Encodes the state of a reader in a 64-bit value
+		/// </summary>
+		internal record struct ReaderState(ulong Value)
+		{
+			public ReaderState(int chunkIdx, int offset) : this(((ulong)chunkIdx << 32) | (uint)offset) { }
+
+			public int ChunkIdx => (int)(Value >> 32);
+			public int Offset => (int)Value;
+
+			public override string ToString() => $"Chunk {ChunkIdx}, Offset {Offset}";
+		}
+
+		/// <summary>
+		/// Pointer to a reader state value
+		/// </summary>
+		internal readonly unsafe struct ReaderStatePtr
+		{
+			readonly ulong* _data;
+
+			public ReaderStatePtr(ulong* data) => _data = data;
+
+			public ReaderState Value
+			{
+				get => new ReaderState(Interlocked.CompareExchange(ref *_data, 0, 0));
+				set => Interlocked.Exchange(ref *_data, value.Value);
+			}
+
+			public int ChunkIdx => Value.ChunkIdx;
+			public int Offset => Value.Offset;
+
+			public void Advance(int length) => Interlocked.Add(ref *_data, (ulong)length);
+
+			public override string ToString() => Value.ToString();
+		}
+
+		/// <inheritdoc cref="IComputeBufferReader"/>
 		sealed class ReaderImpl : IComputeBufferReader
 		{
 			readonly ComputeBufferBase _buffer;
-			readonly ChunkBase[] _chunks;
+			readonly HeaderPtr _header;
+			readonly Memory<byte>[] _chunks;
 			readonly int _readerIdx;
 
-			int _chunkIdx;
-			int _offset;
-
 			public bool IsComplete
 			{
 				get
 				{
-					ChunkState state = _chunks[_chunkIdx].State;
-					return state.WriteState == WriteState.Complete && _offset == state.Length;
+					ReaderState readerState = _header.GetReaderStatePtr(_readerIdx).Value;
+					ChunkState chunkState = _header.GetChunkStatePtr(readerState.ChunkIdx).Value;
+					return chunkState.WriteState == WriteState.Complete && readerState.Offset == chunkState.Length;
 				}
 			}
 
-			public ReaderImpl(ComputeBufferBase buffer, ChunkBase[] chunks, int readerIdx)
+			public ReaderImpl(ComputeBufferBase buffer, HeaderPtr header, Memory<byte>[] chunks, int readerIdx)
 			{
 				_buffer = buffer;
+				_header = header;
 				_chunks = chunks;
 				_readerIdx = readerIdx;
 			}
 
 			/// <inheritdoc/>
-			public void AdvanceReadPosition(int length) => _offset += length;
+			public void AdvanceReadPosition(int length)
+			{
+				ReaderStatePtr readerStatePtr = _header.GetReaderStatePtr(_readerIdx);
+				readerStatePtr.Advance(length);
+			}
 
 			/// <inheritdoc/>
 			public ReadOnlyMemory<byte> GetReadBuffer()
 			{
-				ChunkBase chunk = _chunks[_chunkIdx];
-				ChunkState state = chunk.State;
-				return state.HasReaderFlag(_readerIdx) ? chunk.Memory.Slice(_offset, state.Length - _offset) : ReadOnlyMemory<byte>.Empty;
+				ReaderStatePtr readerStatePtr = _header.GetReaderStatePtr(_readerIdx);
+				ReaderState readerState = readerStatePtr.Value;
+
+				ChunkStatePtr chunkStatePtr = _header.GetChunkStatePtr(readerState.ChunkIdx);
+				ChunkState chunkState = chunkStatePtr.Value;
+
+				if (chunkState.HasReaderFlag(_readerIdx))
+				{
+					return _chunks[readerState.ChunkIdx].Slice(readerState.Offset, chunkState.Length - readerState.Offset);
+				}
+				else
+				{
+					return ReadOnlyMemory<byte>.Empty;
+				}
 			}
 
 			/// <inheritdoc/>
@@ -164,97 +229,109 @@ namespace EpicGames.Horde.Compute
 			{
 				for (; ; )
 				{
-					ChunkBase chunk = _chunks[_chunkIdx];
-					ChunkState state = chunk.State;
+					ReaderStatePtr readerStatePtr = _header.GetReaderStatePtr(_readerIdx);
+					ReaderState readerState = readerStatePtr.Value;
+
+					ChunkStatePtr chunkStatePtr = _header.GetChunkStatePtr(readerState.ChunkIdx);
+					ChunkState chunkState = chunkStatePtr.Value;
 
-					if (!state.HasReaderFlag(_readerIdx))
+					if (!chunkState.HasReaderFlag(_readerIdx))
 					{
 						// Wait until the current chunk is readable
 						_buffer.ResetReadEvent(_readerIdx);
-						if (!chunk.State.HasReaderFlag(_readerIdx))
+						if (!chunkState.HasReaderFlag(_readerIdx))
 						{
 							await _buffer.WaitForReadEvent(_readerIdx, cancellationToken);
 						}
 					}
-					else if (_offset + minLength <= state.Length)
+					else if (readerState.Offset + minLength <= chunkState.Length)
 					{
 						// We have enough data in the chunk to be able to read a message
 						return true;
 					}
-					else if (state.WriteState == WriteState.Writing)
+					else if (chunkState.WriteState == WriteState.Writing)
 					{
 						// Wait until there is more data in the chunk
 						_buffer.ResetReadEvent(_readerIdx);
-						if (chunk.State == state)
+						if (_header.GetChunkStatePtr(readerState.ChunkIdx).Value == chunkState)
 						{
 							await _buffer.WaitForReadEvent(_readerIdx, cancellationToken);
 						}
 					}
-					else if (_offset < state.Length || state.WriteState == WriteState.Complete)
+					else if (readerState.Offset < chunkState.Length || chunkState.WriteState == WriteState.Complete)
 					{
 						// Cannot read the requested amount of data from this chunk.
 						return false;
 					}
-					else if (state.WriteState == WriteState.MovedToNext)
+					else if (chunkState.WriteState == WriteState.MovedToNext)
 					{
 						// Move to the next chunk
-						chunk.FinishReading(_readerIdx);
+						chunkStatePtr.FinishReading(_readerIdx);
 						_buffer.SetWriteEvent();
 
-						_chunkIdx++;
-						if (_chunkIdx == _chunks.Length)
+						int chunkIdx = readerStatePtr.ChunkIdx + 1;
+						if (chunkIdx == _chunks.Length)
 						{
-							_chunkIdx = 0;
+							chunkIdx = 0;
 						}
-						_offset = 0;
+
+						readerStatePtr.Value = new ReaderState(chunkIdx, 0);
 					}
 					else
 					{
-						throw new NotImplementedException($"Invalid write state for buffer: {state.WriteState}");
+						throw new NotImplementedException($"Invalid write state for buffer: {chunkState.WriteState}");
 					}
 				}
 			}
 		}
 
+		/// <inheritdoc cref="IComputeBufferWriter"/>
 		sealed class WriterImpl : IComputeBufferWriter
 		{
 			readonly ComputeBufferBase _buffer;
-			readonly ChunkBase[] _chunks;
-
-			int _chunkIdx;
+			readonly HeaderPtr _header;
+			readonly Memory<byte>[] _chunks;
 
-			public WriterImpl(ComputeBufferBase buffer, ChunkBase[] chunks)
+			public WriterImpl(ComputeBufferBase buffer, HeaderPtr header, Memory<byte>[] chunks)
 			{
 				_buffer = buffer;
+				_header = header;
 				_chunks = chunks;
 			}
 
 			public void AdvanceWritePosition(int size)
 			{
-				ChunkBase chunk = _chunks[_chunkIdx];
-				chunk.Append(size);
-				_buffer.SetAllReadEvents();
+				if (size > 0)
+				{
+					ChunkStatePtr chunkStatePtr = _header.GetChunkStatePtr(_header.WriteChunkIdx);
+					ChunkState chunkState = chunkStatePtr.Value;
+
+					Debug.Assert(chunkState.WriteState == WriteState.Writing);
+					chunkStatePtr.Append(size);
+
+					_buffer.SetAllReadEvents();
+				}
 			}
 
 			public Memory<byte> GetWriteBuffer()
 			{
-				ChunkBase chunk = _chunks[_chunkIdx];
-
-				ChunkState state = chunk.State;
-				if (state.WriteState != WriteState.Writing)
+				ChunkState state = _header.GetChunkStatePtr(_header.WriteChunkIdx).Value;
+				if (state.WriteState == WriteState.Writing)
 				{
-					throw new InvalidOperationException();
+					return _chunks[_header.WriteChunkIdx].Slice(state.Length);
+				}
+				else
+				{
+					return Memory<byte>.Empty;
 				}
-
-				return chunk.Memory.Slice(state.Length);
 			}
 
 			public bool MarkComplete()
 			{
-				ChunkBase chunk = _chunks[_chunkIdx];
-				if (chunk.State.WriteState != WriteState.Complete)
+				ChunkStatePtr chunkStatePtr = _header.GetChunkStatePtr(_header.WriteChunkIdx);
+				if (chunkStatePtr.WriteState != WriteState.Complete)
 				{
-					chunk.MarkComplete();
+					chunkStatePtr.MarkComplete();
 					_buffer.SetAllReadEvents();
 					return true;
 				}
@@ -263,51 +340,51 @@ namespace EpicGames.Horde.Compute
 
 			public async ValueTask WaitToWriteAsync(int minSize, CancellationToken cancellationToken)
 			{
-				if (minSize > _chunks[_chunkIdx].Memory.Length)
+				if (minSize > _header.ChunkLength)
 				{
 					throw new ArgumentException("Requested read size is larger than chunk size.", nameof(minSize));
 				}
 
 				for (; ; )
 				{
-					ChunkBase chunk = _chunks[_chunkIdx];
+					ChunkStatePtr chunkStatePtr = _header.GetChunkStatePtr(_header.WriteChunkIdx);
 
-					ChunkState chunkState = chunk.State;
+					ChunkState chunkState = chunkStatePtr.Value;
 					if (chunkState.WriteState == WriteState.Writing)
 					{
 						int length = chunkState.Length;
-						if (length + minSize <= chunk.Memory.Length)
+						if (length + minSize <= _header.ChunkLength)
 						{
 							return;
 						}
-
-						chunk.MoveToNext();
-						_buffer.SetAllReadEvents();
+						chunkStatePtr.FinishWriting();
 					}
 
-					int nextWriteChunkIdx = _chunkIdx + 1;
-					if (nextWriteChunkIdx == _chunks.Length)
+					int nextChunkIdx = _header.WriteChunkIdx + 1;
+					if (nextChunkIdx == _chunks.Length)
 					{
-						nextWriteChunkIdx = 0;
+						nextChunkIdx = 0;
 					}
 
-					ChunkBase nextChunk = _chunks[nextWriteChunkIdx];
-					while (nextChunk.State.ReaderFlags != 0)
+					ChunkStatePtr nextChunkStatePtr = _header.GetChunkStatePtr(nextChunkIdx);
+					while (nextChunkStatePtr.ReaderFlags != 0)
 					{
 						await _buffer.WaitForWriteEvent(cancellationToken);
 						_buffer.ResetWriteEvent();
 					}
 
-					_chunkIdx = nextWriteChunkIdx;
-					nextChunk.StartWriting(_buffer._numReaders);
+					_header.WriteChunkIdx = nextChunkIdx;
+					nextChunkStatePtr.StartWriting(_header.NumReaders);
+
+					_buffer.SetAllReadEvents();
 				}
 			}
 		}
 
-		readonly ChunkBase[] _chunks;
+		readonly HeaderPtr _headerPtr;
+		readonly Memory<byte>[] _chunks;
 		readonly ReaderImpl _reader;
 		readonly WriterImpl _writer;
-		readonly int _numReaders;
 
 		int _refCount = 1;
 
@@ -317,24 +394,25 @@ namespace EpicGames.Horde.Compute
 		/// <inheritdoc/>
 		public IComputeBufferWriter Writer => _writer;
 
+#pragma warning disable IDE0051 // Remove unused private members
+		// For debugging purposes only
+		ChunkState[] ChunkStates => Enumerable.Range(0, _headerPtr.NumChunks).Select(x => _headerPtr.GetChunkStatePtr(x).Value).ToArray();
+		ReaderState[] ReaderStates => Enumerable.Range(0, _headerPtr.NumReaders).Select(x => _headerPtr.GetReaderStatePtr(x).Value).ToArray();
+		int WriteChunkIdx => _headerPtr.WriteChunkIdx;
+#pragma warning restore IDE0051 // Remove unused private members
+
 		/// <summary>
 		/// Constructor
 		/// </summary>
+		/// <param name="headerPtr">Header for the buffer</param>
 		/// <param name="chunks">Data for the buffer</param>
-		/// <param name="numReaders">Number of readers for this buffer</param>
-		/// <param name="isNew"></param>
-		protected ComputeBufferBase(ChunkBase[] chunks, int numReaders, bool isNew)
+		internal ComputeBufferBase(HeaderPtr headerPtr, Memory<byte>[] chunks)
 		{
+			_headerPtr = headerPtr;
 			_chunks = chunks;
-			if (isNew)
-			{
-				_chunks[0].StartWriting(numReaders);
-			}
-
-			_numReaders = numReaders;
 
-			_reader = new ReaderImpl(this, _chunks, 0);
-			_writer = new WriterImpl(this, _chunks);
+			_reader = new ReaderImpl(this, _headerPtr, _chunks, 0);
+			_writer = new WriterImpl(this, _headerPtr, _chunks);
 		}
 
 		/// <inheritdoc/>
@@ -368,7 +446,7 @@ namespace EpicGames.Horde.Compute
 
 		private void SetAllReadEvents()
 		{
-			for (int readerIdx = 0; readerIdx < _numReaders; readerIdx++)
+			for (int readerIdx = 0; readerIdx < _headerPtr.NumReaders; readerIdx++)
 			{
 				SetReadEvent(readerIdx);
 			}
