Make response logging information so it appears in p4vutils output
#rnx

diff --git a/Engine/Source/Programs/Shared/EpicGames.Perforce/PerforceConnection.cs b/Engine/Source/Programs/Shared/EpicGames.Perforce/PerforceConnection.cs
--- a/Engine/Source/Programs/Shared/EpicGames.Perforce/PerforceConnection.cs
+++ b/Engine/Source/Programs/Shared/EpicGames.Perforce/PerforceConnection.cs
@@ -581,7 +581,7 @@ namespace EpicGames.Perforce
 			{
 				for (int idx = 0; idx < responses.Count; idx++)
 				{
-					connection.Logger.LogDebug("Unexpected response {Idx}: {Text}", idx, responses[idx].ToString());
+					connection.Logger.LogInformation("Unexpected response {Idx}: {Text}", idx, responses[idx].ToString());
 				}
 				throw new PerforceException("Expected one result from 'p4 {0}', got {1}", command, responses.Count);
 			}
