Horde: Push implementation of reference-counted data structures into the compute buffer base class.

diff --git a/Engine/Source/Programs/Shared/EpicGames.Horde/Compute/Buffers/PooledBuffer.cs b/Engine/Source/Programs/Shared/EpicGames.Horde/Compute/Buffers/PooledBuffer.cs
--- a/Engine/Source/Programs/Shared/EpicGames.Horde/Compute/Buffers/PooledBuffer.cs
+++ b/Engine/Source/Programs/Shared/EpicGames.Horde/Compute/Buffers/PooledBuffer.cs
@@ -1,153 +1,107 @@
 // Copyright Epic Games, Inc. All Rights Reserved.
 
 using System;
-using System.Buffers;
-using System.Linq;
 using System.Runtime.InteropServices;
 using System.Threading;
 using System.Threading.Tasks;
 using EpicGames.Core;
-using Microsoft.CodeAnalysis.CSharp.Syntax;
 
 namespace EpicGames.Horde.Compute.Buffers
 {
 	/// <summary>
-	/// Implementation of <see cref="IComputeBuffer"/> suitable for cross-process communication
+	/// In-process buffer used to store compute messages
 	/// </summary>
-	public sealed class PooledBuffer : IComputeBuffer
+	public sealed class PooledBuffer : ComputeBufferBase
 	{
-		PooledBufferCore _core;
-
-		/// <summary>
-		/// Constructor
-		/// </summary>
-		/// <param name="capacity">Total capacity of the buffer</param>
-		public PooledBuffer(int capacity)
-			: this(2, capacity / 2)
-		{
-		}
-
-		/// <summary>
-		/// Constructor
-		/// </summary>
-		/// <param name="numChunks">Number of chunks in the buffer</param>
-		/// <param name="chunkLength">Length of each chunk</param>
-		public PooledBuffer(int numChunks, int chunkLength)
-			: this(PooledBufferCore.Create(numChunks, chunkLength, 1))
+		class Resources : ResourcesBase
 		{
-		}
+			public HeaderPtr _headerPtr;
+			readonly GCHandle _headerHandle;
 
-		private PooledBuffer(PooledBufferCore core) => _core = core;
+			public Memory<byte>[] _chunks;
+			readonly GCHandle[] _chunkHandles;
 
-		/// <inheritdoc/>
-		public IComputeBufferReader Reader => _core.Reader;
+			readonly AsyncEvent _writerEvent = new AsyncEvent();
+			readonly AsyncEvent _readerEvent = new AsyncEvent();
 
-		/// <inheritdoc/>
-		public IComputeBufferWriter Writer => _core.Writer;
+			public unsafe Resources(int numChunks, int chunkLength, int numReaders)
+			{
+				_chunks = new Memory<byte>[numChunks];
+				_chunkHandles = new GCHandle[numChunks];
 
-		/// <inheritdoc cref="IComputeBuffer.AddRef"/>
-		public PooledBuffer AddRef()
-		{
-			_core.AddRef();
-			return new PooledBuffer(_core);
-		}
+				for (int idx = 0; idx < numChunks; idx++)
+				{
+					byte[] data = new byte[chunkLength];
+					_chunks[idx] = data;
+					_chunkHandles[idx] = GCHandle.Alloc(data, GCHandleType.Pinned);
+				}
 
-		/// <inheritdoc/>
-		IComputeBuffer IComputeBuffer.AddRef() => AddRef();
+				byte[] header = new byte[HeaderSize];
+				_headerHandle = GCHandle.Alloc(header, GCHandleType.Pinned);
+				_headerPtr = new HeaderPtr((ulong*)_headerHandle.AddrOfPinnedObject().ToPointer(), numReaders, numChunks, chunkLength);
+			}
 
-		/// <inheritdoc/>
-		public void Dispose()
-		{
-			if (_core != null)
+			public override void Dispose()
 			{
-				_core.Release();
-				_core = null!;
+				for (int idx = 0; idx < _chunkHandles.Length; idx++)
+				{
+					_chunkHandles[idx].Free();
+				}
+				_headerHandle.Free();
 			}
-		}
-	}
 
-	/// <summary>
-	/// In-process buffer used to store compute messages
-	/// </summary>
-	sealed class PooledBufferCore : ComputeBufferBase
-	{
-		unsafe class PinnedBuffer : IDisposable
-		{
-			readonly byte[] _data;
-			readonly GCHandle _handle;
+			/// <inheritdoc/>
+			public override void SetReadEvent(int readerIdx) => _readerEvent.Set();
 
-			public Memory<byte> Data => _data;
-			public void* Ptr => _handle.AddrOfPinnedObject().ToPointer();
+			/// <inheritdoc/>
+			public override void ResetReadEvent(int readerIdx) => _readerEvent.Reset();
 
-			public PinnedBuffer(int length)
-			{
-				_data = new byte[length];
-				_handle = GCHandle.Alloc(_data, GCHandleType.Pinned);
-			}
+			/// <inheritdoc/>
+			public override Task WaitForReadEvent(int readerIdx, CancellationToken cancellationToken) => _readerEvent.Task.WaitAsync(cancellationToken);
 
-			public void Dispose()
-			{
-				_handle.Free();
-			}
+			/// <inheritdoc/>
+			public override void SetWriteEvent() => _writerEvent.Set();
+
+			/// <inheritdoc/>
+			public override void ResetWriteEvent() => _writerEvent.Reset();
+
+			/// <inheritdoc/>
+			public override Task WaitForWriteEvent(CancellationToken cancellationToken) => _writerEvent.Task.WaitAsync(cancellationToken);
 		}
 
-		readonly PinnedBuffer[] _buffers;
-		readonly AsyncEvent _writerEvent = new AsyncEvent();
-		readonly AsyncEvent _readerEvent = new AsyncEvent();
+		readonly Resources _resources;
 
-		private PooledBufferCore(HeaderPtr headerPtr, Memory<byte>[] chunks, PinnedBuffer[] buffers)
-			: base(headerPtr, chunks)
+		/// <summary>
+		/// Constructor
+		/// </summary>
+		/// <param name="capacity">Total capacity of the buffer</param>
+		public PooledBuffer(int capacity)
+			: this(2, capacity / 2)
 		{
-			_buffers = buffers;
 		}
 
-		public static unsafe PooledBufferCore Create(int numChunks, int chunkLength, int numReaders)
+		/// <summary>
+		/// Constructor
+		/// </summary>
+		/// <param name="numChunks">Number of chunks in the buffer</param>
+		/// <param name="chunkLength">Length of each chunk</param>
+		/// <param name="numReaders">Number of readers for this buffer</param>
+		public PooledBuffer(int numChunks, int chunkLength, int numReaders = 1)
+			: this(new Resources(numChunks, chunkLength, numReaders))
 		{
-			Memory<byte>[] chunks = new Memory<byte>[numChunks];
-
-			PinnedBuffer[] buffers = new PinnedBuffer[numChunks + 1];
-			for (int idx = 0; idx < numChunks; idx++)
-			{
-				buffers[idx] = new PinnedBuffer(chunkLength);
-				chunks[idx] = buffers[idx].Data;
-			}
-
-			buffers[numChunks] = new PinnedBuffer(HeaderSize);
-
-			HeaderPtr headerPtr = new HeaderPtr((ulong*)buffers[numChunks].Ptr, numReaders, numChunks, chunkLength);
-			return new PooledBufferCore(headerPtr, chunks, buffers);
 		}
 
-		/// <inheritdoc/>
-		protected override void Dispose(bool disposing)
+		private PooledBuffer(Resources resources)
+			: base(resources._headerPtr, resources._chunks, resources)
 		{
-			base.Dispose(disposing);
-
-			if (disposing)
-			{
-				for (int idx = 0; idx < _buffers.Length; idx++)
-				{
-					_buffers[idx].Dispose();
-				}
-			}
+			_resources = resources;
 		}
 
 		/// <inheritdoc/>
-		protected override void SetReadEvent(int readerIdx) => _readerEvent.Set();
-
-		/// <inheritdoc/>
-		protected override void ResetReadEvent(int readerIdx) => _readerEvent.Reset();
-
-		/// <inheritdoc/>
-		protected override Task WaitForReadEvent(int readerIdx, CancellationToken cancellationToken) => _readerEvent.Task.WaitAsync(cancellationToken);
-
-		/// <inheritdoc/>
-		protected override void SetWriteEvent() => _writerEvent.Set();
-
-		/// <inheritdoc/>
-		protected override void ResetWriteEvent() => _writerEvent.Reset();
-
-		/// <inheritdoc/>
-		protected override Task WaitForWriteEvent(CancellationToken cancellationToken) => _writerEvent.Task.WaitAsync(cancellationToken);
+		public override IComputeBuffer AddRef()
+		{
+			_resources.AddRef();
+			return new PooledBuffer(_resources);
+		}
 	}
 }
diff --git a/Engine/Source/Programs/Shared/EpicGames.Horde/Compute/Buffers/SharedMemoryBuffer.cs b/Engine/Source/Programs/Shared/EpicGames.Horde/Compute/Buffers/SharedMemoryBuffer.cs
--- a/Engine/Source/Programs/Shared/EpicGames.Horde/Compute/Buffers/SharedMemoryBuffer.cs
+++ b/Engine/Source/Programs/Shared/EpicGames.Horde/Compute/Buffers/SharedMemoryBuffer.cs
@@ -11,46 +11,91 @@ using System.Runtime.InteropServices;
 namespace EpicGames.Horde.Compute.Buffers
 {
 	/// <summary>
-	/// Implementation of <see cref="IComputeBuffer"/> suitable for cross-process communication
+	/// Core implementation of <see cref="SharedMemoryBuffer"/>
 	/// </summary>
-	public sealed class SharedMemoryBuffer : IComputeBuffer
+	public sealed class SharedMemoryBuffer : ComputeBufferBase
 	{
-		SharedMemoryBufferCore _core;
+		class Resources : ResourcesBase
+		{
+			readonly MemoryMappedFile _memoryMappedFile;
+			readonly MemoryMappedViewAccessor _memoryMappedViewAccessor;
+			readonly MemoryMappedView _memoryMappedView;
 
-		/// <summary>
-		/// Name of this shared memory buffer
-		/// </summary>
-		public string Name => _core.Name;
+			readonly Native.EventHandle _readerEvent;
+			readonly Native.EventHandle _writerEvent;
 
-		private SharedMemoryBuffer(SharedMemoryBufferCore core)
-		{
-			_core = core;
+			public string Name { get; }
+			public HeaderPtr HeaderPtr { get; }
+			public Memory<byte>[] Chunks { get; }
+
+			public Resources(string name, HeaderPtr headerPtr, MemoryMappedFile memoryMappedFile, MemoryMappedViewAccessor memoryMappedViewAccessor, MemoryMappedView memoryMappedView, Native.EventHandle readerEvent, Native.EventHandle writerEvent)
+			{
+				Name = name;
+				HeaderPtr = headerPtr;
+
+				_memoryMappedFile = memoryMappedFile;
+				_memoryMappedViewAccessor = memoryMappedViewAccessor;
+				_memoryMappedView = memoryMappedView;
+
+				_readerEvent = readerEvent;
+				_writerEvent = writerEvent;
+
+				Chunks = new Memory<byte>[headerPtr.NumChunks];
+				for (int chunkIdx = 0; chunkIdx < headerPtr.NumChunks; chunkIdx++)
+				{
+					int chunkOffset = HeaderSize + (headerPtr.ChunkLength * chunkIdx);
+					Chunks[chunkIdx] = memoryMappedView.GetMemory(chunkOffset, headerPtr.ChunkLength);
+				}
+			}
+
+			public override void Dispose()
+			{
+				_readerEvent.Dispose();
+				_writerEvent.Dispose();
+
+				_memoryMappedView.Dispose();
+				_memoryMappedViewAccessor.Dispose();
+				_memoryMappedFile.Dispose();
+			}
+
+			/// <inheritdoc/>
+			public override void SetReadEvent(int readerIdx) => _readerEvent.Set();
+
+			/// <inheritdoc/>
+			public override void ResetReadEvent(int readerIdx) => _readerEvent.Reset();
+
+			/// <inheritdoc/>
+			public override Task WaitForReadEvent(int readerIdx, CancellationToken cancellationToken) => _readerEvent.WaitOneAsync(cancellationToken);
+
+			/// <inheritdoc/>
+			public override void SetWriteEvent() => _writerEvent.Set();
+
+			/// <inheritdoc/>
+			public override void ResetWriteEvent() => _writerEvent.Reset();
+
+			/// <inheritdoc/>
+			public override Task WaitForWriteEvent(CancellationToken cancellationToken) => _writerEvent.WaitOneAsync(cancellationToken);
 		}
 
 		/// <inheritdoc/>
-		public IComputeBufferReader Reader => _core.Reader;
+		public string Name => _resources.Name;
 
-		/// <inheritdoc/>
-		public IComputeBufferWriter Writer => _core.Writer;
+		readonly Resources _resources;
 
-		/// <inheritdoc cref="IComputeBuffer.AddRef"/>
-		public SharedMemoryBuffer AddRef()
+		/// <summary>
+		/// Constructor
+		/// </summary>
+		private unsafe SharedMemoryBuffer(Resources resources)
+			: base(resources.HeaderPtr, resources.Chunks, resources)
 		{
-			_core.AddRef();
-			return new SharedMemoryBuffer(_core);
+			_resources = resources;
 		}
 
 		/// <inheritdoc/>
-		IComputeBuffer IComputeBuffer.AddRef() => AddRef();
-
-		/// <inheritdoc/>
-		public void Dispose()
+		public override IComputeBuffer AddRef()
 		{
-			if (_core != null)
-			{
-				_core.Release();
-				_core = null!;
-			}
+			_resources.AddRef();
+			return new SharedMemoryBuffer(_resources);
 		}
 
 		/// <summary>
@@ -71,54 +116,6 @@ namespace EpicGames.Horde.Compute.Buffers
 		/// <param name="numChunks">Number of chunks in the buffer</param>
 		/// <param name="chunkLength">Length of each chunk</param>
 		public static unsafe SharedMemoryBuffer CreateNew(string? name, int numChunks, int chunkLength)
-		{
-			SharedMemoryBufferCore core = SharedMemoryBufferCore.CreateNew(name, numChunks, chunkLength);
-			return new SharedMemoryBuffer(core);
-		}
-
-		/// <summary>
-		/// Open an existing buffer by name
-		/// </summary>
-		/// <param name="name">Name of the buffer to open</param>
-		public static SharedMemoryBuffer OpenExisting(string name)
-		{
-			SharedMemoryBufferCore core = SharedMemoryBufferCore.OpenExisting(name);
-			return new SharedMemoryBuffer(core);
-		}
-	}
-
-	/// <summary>
-	/// Core implementation of <see cref="SharedMemoryBuffer"/>
-	/// </summary>
-	sealed class SharedMemoryBufferCore : ComputeBufferBase
-	{
-		readonly MemoryMappedFile _memoryMappedFile;
-		readonly MemoryMappedViewAccessor _memoryMappedViewAccessor;
-		readonly MemoryMappedView _memoryMappedView;
-
-		readonly Native.EventHandle _readerEvent;
-		readonly Native.EventHandle _writerEvent;
-
-		/// <inheritdoc/>
-		public string Name { get; }
-
-		/// <summary>
-		/// Constructor
-		/// </summary>
-		private unsafe SharedMemoryBufferCore(HeaderPtr headerPtr, string name, MemoryMappedFile memoryMappedFile, MemoryMappedViewAccessor memoryMappedViewAccessor, MemoryMappedView memoryMappedView, Native.EventHandle readerEvent, Native.EventHandle writerEvent)
-			: base(headerPtr, CreateChunks(headerPtr.NumChunks, headerPtr.ChunkLength, memoryMappedView))
-		{
-			Name = name;
-
-			_memoryMappedFile = memoryMappedFile;
-			_memoryMappedViewAccessor = memoryMappedViewAccessor;
-			_memoryMappedView = memoryMappedView;
-
-			_readerEvent = readerEvent;
-			_writerEvent = writerEvent;
-		}
-
-		public static unsafe SharedMemoryBufferCore CreateNew(string? name, int numChunks, int chunkLength)
 		{
 			long capacity = HeaderSize + (numChunks * sizeof(ulong)) + (numChunks * chunkLength);
 
@@ -132,14 +129,14 @@ namespace EpicGames.Horde.Compute.Buffers
 			Native.EventHandle readerEvent = Native.EventHandle.CreateNew($"{name}_R", EventResetMode.ManualReset, true, HandleInheritability.Inheritable);
 
 			HeaderPtr headerPtr = new HeaderPtr((ulong*)memoryMappedView.GetPointer(), 1, numChunks, chunkLength);
-			return new SharedMemoryBufferCore(headerPtr, name, memoryMappedFile, memoryMappedViewAccessor, memoryMappedView, readerEvent, writerEvent);
+			return new SharedMemoryBuffer(new Resources(name, headerPtr, memoryMappedFile, memoryMappedViewAccessor, memoryMappedView, readerEvent, writerEvent));
 		}
 
 		/// <summary>
 		/// Open an existing buffer by name
 		/// </summary>
 		/// <param name="name">Name of the buffer to open</param>
-		public static unsafe SharedMemoryBufferCore OpenExisting(string name)
+		public static unsafe SharedMemoryBuffer OpenExisting(string name)
 		{
 			if (!RuntimeInformation.IsOSPlatform(OSPlatform.Windows))
 			{
@@ -154,52 +151,7 @@ namespace EpicGames.Horde.Compute.Buffers
 			Native.EventHandle writerEvent = Native.EventHandle.OpenExisting($"{name}_W");
 
 			HeaderPtr headerPtr = new HeaderPtr((ulong*)memoryMappedView.GetPointer());
-			return new SharedMemoryBufferCore(headerPtr, name, memoryMappedFile, memoryMappedViewAccessor, memoryMappedView, readerEvent, writerEvent);
-		}
-
-		static unsafe Memory<byte>[] CreateChunks(int numChunks, int chunkLength, MemoryMappedView memoryMappedView)
-		{
-			Memory<byte>[] chunks = new Memory<byte>[numChunks];
-			for (int chunkIdx = 0; chunkIdx < numChunks; chunkIdx++)
-			{
-				int chunkOffset = HeaderSize + (chunkLength * chunkIdx);
-				chunks[chunkIdx] = memoryMappedView.GetMemory(chunkOffset, chunkLength);
-			}
-			return chunks;
-		}
-
-		/// <inheritdoc/>
-		protected override void Dispose(bool disposing)
-		{
-			base.Dispose(disposing);
-
-			if (disposing)
-			{
-				_readerEvent.Dispose();
-				_writerEvent.Dispose();
-
-				_memoryMappedView.Dispose();
-				_memoryMappedViewAccessor.Dispose();
-				_memoryMappedFile.Dispose();
-			}
+			return new SharedMemoryBuffer(new Resources(name, headerPtr, memoryMappedFile, memoryMappedViewAccessor, memoryMappedView, readerEvent, writerEvent));
 		}
-
-		/// <inheritdoc/>
-		protected override void SetReadEvent(int readerIdx) => _readerEvent.Set();
-
-		/// <inheritdoc/>
-		protected override void ResetReadEvent(int readerIdx) => _readerEvent.Reset();
-
-		/// <inheritdoc/>
-		protected override Task WaitForReadEvent(int readerIdx, CancellationToken cancellationToken) => _readerEvent.WaitOneAsync(cancellationToken);
-
-		/// <inheritdoc/>
-		protected override void SetWriteEvent() => _writerEvent.Set();
-
-		/// <inheritdoc/>
-		protected override void ResetWriteEvent() => _writerEvent.Reset();
-
-		/// <inheritdoc/>
-		protected override Task WaitForWriteEvent(CancellationToken cancellationToken) => _writerEvent.WaitOneAsync(cancellationToken);
 	}
 }
diff --git a/Engine/Source/Programs/Shared/EpicGames.Horde/Compute/ComputeBufferBase.cs b/Engine/Source/Programs/Shared/EpicGames.Horde/Compute/ComputeBufferBase.cs
--- a/Engine/Source/Programs/Shared/EpicGames.Horde/Compute/ComputeBufferBase.cs
+++ b/Engine/Source/Programs/Shared/EpicGames.Horde/Compute/ComputeBufferBase.cs
@@ -11,7 +11,7 @@ namespace EpicGames.Horde.Compute
 	/// <summary>
 	/// In-process buffer used to store compute messages
 	/// </summary>
-	public abstract class ComputeBufferBase : IDisposable
+	public abstract class ComputeBufferBase : IComputeBuffer
 	{
 		const int MaxChunks = 16;
 		const int MaxReaders = 16;
@@ -172,227 +172,101 @@ namespace EpicGames.Horde.Compute
 			public override string ToString() => Value.ToString();
 		}
 
-		/// <inheritdoc cref="IComputeBufferReader"/>
-		sealed class ReaderImpl : IComputeBufferReader
+		/// <summary>
+		/// State shared between buffer instances
+		/// </summary>
+		internal abstract class ResourcesBase : IDisposable
 		{
-			readonly ComputeBufferBase _buffer;
-			readonly HeaderPtr _header;
-			readonly Memory<byte>[] _chunks;
-			readonly int _readerIdx;
+			int _refCount = 1;
 
-			public bool IsComplete
+			public void AddRef()
 			{
-				get
-				{
-					ReaderState readerState = _header.GetReaderStatePtr(_readerIdx).Value;
-					ChunkState chunkState = _header.GetChunkStatePtr(readerState.ChunkIdx).Value;
-					return chunkState.WriteState == WriteState.Complete && readerState.Offset == chunkState.Length;
-				}
+				Interlocked.Increment(ref _refCount);
 			}
 
-			public ReaderImpl(ComputeBufferBase buffer, HeaderPtr header, Memory<byte>[] chunks, int readerIdx)
-			{
-				_buffer = buffer;
-				_header = header;
-				_chunks = chunks;
-				_readerIdx = readerIdx;
-			}
-
-			/// <inheritdoc/>
-			public void AdvanceReadPosition(int length)
-			{
-				ReaderStatePtr readerStatePtr = _header.GetReaderStatePtr(_readerIdx);
-				readerStatePtr.Advance(length);
-			}
-
-			/// <inheritdoc/>
-			public ReadOnlyMemory<byte> GetReadBuffer()
+			public void Release()
 			{
-				ReaderStatePtr readerStatePtr = _header.GetReaderStatePtr(_readerIdx);
-				ReaderState readerState = readerStatePtr.Value;
-
-				ChunkStatePtr chunkStatePtr = _header.GetChunkStatePtr(readerState.ChunkIdx);
-				ChunkState chunkState = chunkStatePtr.Value;
-
-				if (chunkState.HasReaderFlag(_readerIdx))
-				{
-					return _chunks[readerState.ChunkIdx].Slice(readerState.Offset, chunkState.Length - readerState.Offset);
-				}
-				else
+				if (Interlocked.Decrement(ref _refCount) == 0)
 				{
-					return ReadOnlyMemory<byte>.Empty;
+					Dispose();
 				}
 			}
 
-			/// <inheritdoc/>
-			public async ValueTask<bool> WaitToReadAsync(int minLength, CancellationToken cancellationToken = default)
-			{
-				for (; ; )
-				{
-					ReaderStatePtr readerStatePtr = _header.GetReaderStatePtr(_readerIdx);
-					ReaderState readerState = readerStatePtr.Value;
+			/// <summary>
+			/// Dispose of the 
+			/// </summary>
+			public abstract void Dispose();
+
+			/// <summary>
+			/// Signals a read event
+			/// </summary>
+			public abstract void SetReadEvent(int readerIdx);
 
-					ChunkStatePtr chunkStatePtr = _header.GetChunkStatePtr(readerState.ChunkIdx);
-					ChunkState chunkState = chunkStatePtr.Value;
+			/// <summary>
+			/// Resets a read event
+			/// </summary>
+			public abstract void ResetReadEvent(int readerIdx);
 
-					if (!chunkState.HasReaderFlag(_readerIdx))
-					{
-						// Wait until the current chunk is readable
-						_buffer.ResetReadEvent(_readerIdx);
-						if (!chunkState.HasReaderFlag(_readerIdx))
-						{
-							await _buffer.WaitForReadEvent(_readerIdx, cancellationToken);
-						}
-					}
-					else if (readerState.Offset + minLength <= chunkState.Length)
-					{
-						// We have enough data in the chunk to be able to read a message
-						return true;
-					}
-					else if (chunkState.WriteState == WriteState.Writing)
-					{
-						// Wait until there is more data in the chunk
-						_buffer.ResetReadEvent(_readerIdx);
-						if (_header.GetChunkStatePtr(readerState.ChunkIdx).Value == chunkState)
-						{
-							await _buffer.WaitForReadEvent(_readerIdx, cancellationToken);
-						}
-					}
-					else if (readerState.Offset < chunkState.Length || chunkState.WriteState == WriteState.Complete)
-					{
-						// Cannot read the requested amount of data from this chunk.
-						return false;
-					}
-					else if (chunkState.WriteState == WriteState.MovedToNext)
-					{
-						// Move to the next chunk
-						chunkStatePtr.FinishReading(_readerIdx);
-						_buffer.SetWriteEvent();
+			/// <summary>
+			/// Waits for a read event to be signalled
+			/// </summary>
+			public abstract Task WaitForReadEvent(int readerIdx, CancellationToken cancellationToken);
 
-						int chunkIdx = readerStatePtr.ChunkIdx + 1;
-						if (chunkIdx == _chunks.Length)
-						{
-							chunkIdx = 0;
-						}
+			/// <summary>
+			/// Signals the write event
+			/// </summary>
+			public abstract void SetWriteEvent();
 
-						readerStatePtr.Value = new ReaderState(chunkIdx, 0);
-					}
-					else
-					{
-						throw new NotImplementedException($"Invalid write state for buffer: {chunkState.WriteState}");
-					}
-				}
-			}
+			/// <summary>
+			/// Resets the write event
+			/// </summary>
+			public abstract void ResetWriteEvent();
+
+			/// <summary>
+			/// Waits for the write event to be signalled
+			/// </summary>
+			public abstract Task WaitForWriteEvent(CancellationToken cancellationToken);
 		}
 
-		/// <inheritdoc cref="IComputeBufferWriter"/>
-		sealed class WriterImpl : IComputeBufferWriter
+		class ReaderImpl : IComputeBufferReader
 		{
 			readonly ComputeBufferBase _buffer;
-			readonly HeaderPtr _header;
-			readonly Memory<byte>[] _chunks;
 
-			public WriterImpl(ComputeBufferBase buffer, HeaderPtr header, Memory<byte>[] chunks)
-			{
-				_buffer = buffer;
-				_header = header;
-				_chunks = chunks;
-			}
+			public ReaderImpl(ComputeBufferBase buffer) => _buffer = buffer;
 
-			public void AdvanceWritePosition(int size)
-			{
-				if (size > 0)
-				{
-					ChunkStatePtr chunkStatePtr = _header.GetChunkStatePtr(_header.WriteChunkIdx);
-					ChunkState chunkState = chunkStatePtr.Value;
-
-					Debug.Assert(chunkState.WriteState == WriteState.Writing);
-					chunkStatePtr.Append(size);
+			public bool IsComplete => _buffer.IsComplete(0);
 
-					_buffer.SetAllReadEvents();
-				}
-			}
-
-			public Memory<byte> GetWriteBuffer()
-			{
-				ChunkState state = _header.GetChunkStatePtr(_header.WriteChunkIdx).Value;
-				if (state.WriteState == WriteState.Writing)
-				{
-					return _chunks[_header.WriteChunkIdx].Slice(state.Length);
-				}
-				else
-				{
-					return Memory<byte>.Empty;
-				}
-			}
+			public void AdvanceReadPosition(int size) => _buffer.AdvanceReadPosition(0, size);
 
-			public bool MarkComplete()
-			{
-				ChunkStatePtr chunkStatePtr = _header.GetChunkStatePtr(_header.WriteChunkIdx);
-				if (chunkStatePtr.WriteState != WriteState.Complete)
-				{
-					chunkStatePtr.MarkComplete();
-					_buffer.SetAllReadEvents();
-					return true;
-				}
-				return false;
-			}
+			public ReadOnlyMemory<byte> GetReadBuffer() => _buffer.GetReadBuffer(0);
 
-			public async ValueTask WaitToWriteAsync(int minSize, CancellationToken cancellationToken)
-			{
-				if (minSize > _header.ChunkLength)
-				{
-					throw new ArgumentException("Requested read size is larger than chunk size.", nameof(minSize));
-				}
+			public ValueTask<bool> WaitToReadAsync(int minLength, CancellationToken cancellationToken = default) => _buffer.WaitToReadAsync(0, minLength, cancellationToken);
+		}
 
-				for (; ; )
-				{
-					ChunkStatePtr chunkStatePtr = _header.GetChunkStatePtr(_header.WriteChunkIdx);
+		class WriterImpl : IComputeBufferWriter
+		{
+			readonly ComputeBufferBase _buffer;
 
-					ChunkState chunkState = chunkStatePtr.Value;
-					if (chunkState.WriteState == WriteState.Writing)
-					{
-						int length = chunkState.Length;
-						if (length + minSize <= _header.ChunkLength)
-						{
-							return;
-						}
-						chunkStatePtr.FinishWriting();
-					}
+			public WriterImpl(ComputeBufferBase buffer) => _buffer = buffer;
 
-					int nextChunkIdx = _header.WriteChunkIdx + 1;
-					if (nextChunkIdx == _chunks.Length)
-					{
-						nextChunkIdx = 0;
-					}
+			public void AdvanceWritePosition(int size) => _buffer.AdvanceWritePosition(size);
 
-					ChunkStatePtr nextChunkStatePtr = _header.GetChunkStatePtr(nextChunkIdx);
-					while (nextChunkStatePtr.ReaderFlags != 0)
-					{
-						await _buffer.WaitForWriteEvent(cancellationToken);
-						_buffer.ResetWriteEvent();
-					}
+			public Memory<byte> GetWriteBuffer() => _buffer.GetWriteBuffer();
 
-					_header.WriteChunkIdx = nextChunkIdx;
-					nextChunkStatePtr.StartWriting(_header.NumReaders);
+			public bool MarkComplete() => _buffer.MarkComplete();
 
-					_buffer.SetAllReadEvents();
-				}
-			}
+			public ValueTask WaitToWriteAsync(int minLength, CancellationToken cancellationToken = default) => _buffer.WaitToWriteAsync(minLength, cancellationToken);
 		}
 
 		readonly HeaderPtr _headerPtr;
 		readonly Memory<byte>[] _chunks;
-		readonly ReaderImpl _reader;
-		readonly WriterImpl _writer;
-
-		int _refCount = 1;
+		ResourcesBase _resources;
 
 		/// <inheritdoc/>
-		public IComputeBufferReader Reader => _reader;
+		public IComputeBufferReader Reader { get; }
 
 		/// <inheritdoc/>
-		public IComputeBufferWriter Writer => _writer;
+		public IComputeBufferWriter Writer { get; }
 
 #pragma warning disable IDE0051 // Remove unused private members
 		// For debugging purposes only
@@ -406,80 +280,232 @@ namespace EpicGames.Horde.Compute
 		/// </summary>
 		/// <param name="headerPtr">Header for the buffer</param>
 		/// <param name="chunks">Data for the buffer</param>
-		internal ComputeBufferBase(HeaderPtr headerPtr, Memory<byte>[] chunks)
+		/// <param name="resources">Resources shared between instances of the buffer</param>
+		internal ComputeBufferBase(HeaderPtr headerPtr, Memory<byte>[] chunks, ResourcesBase resources)
 		{
 			_headerPtr = headerPtr;
 			_chunks = chunks;
+			_resources = resources;
 
-			_reader = new ReaderImpl(this, _headerPtr, _chunks, 0);
-			_writer = new WriterImpl(this, _headerPtr, _chunks);
+			Reader = new ReaderImpl(this);
+			Writer = new WriterImpl(this);
 		}
 
 		/// <inheritdoc/>
-		public void AddRef()
+		public void Dispose()
+		{
+			Dispose(true);
+			GC.SuppressFinalize(this);
+		}
+
+		/// <summary>
+		/// Overriable dispose method
+		/// </summary>
+		protected virtual void Dispose(bool disposing)
 		{
-			Interlocked.Increment(ref _refCount);
+			if (_resources != null)
+			{
+				_resources.Release();
+				_resources = null!;
+			}
+		}
+
+		private void SetAllReadEvents()
+		{
+			for (int readerIdx = 0; readerIdx < _headerPtr.NumReaders; readerIdx++)
+			{
+				_resources.SetReadEvent(readerIdx);
+			}
+		}
+
+		/// <inheritdoc/>
+		public abstract IComputeBuffer AddRef();
+
+		#region Reader Interface
+
+		/// <inheritdoc/>
+		public bool IsComplete(int readerIdx)
+		{
+			ReaderState readerState = _headerPtr.GetReaderStatePtr(readerIdx).Value;
+			ChunkState chunkState = _headerPtr.GetChunkStatePtr(readerState.ChunkIdx).Value;
+			return chunkState.WriteState == WriteState.Complete && readerState.Offset == chunkState.Length;
+		}
+
+		/// <inheritdoc/>
+		public void AdvanceReadPosition(int readerIdx, int length)
+		{
+			ReaderStatePtr readerStatePtr = _headerPtr.GetReaderStatePtr(readerIdx);
+			readerStatePtr.Advance(length);
 		}
 
 		/// <inheritdoc/>
-		public void Release()
+		public ReadOnlyMemory<byte> GetReadBuffer(int readerIdx)
 		{
-			if (Interlocked.Decrement(ref _refCount) == 0)
+			ReaderStatePtr readerStatePtr = _headerPtr.GetReaderStatePtr(readerIdx);
+			ReaderState readerState = readerStatePtr.Value;
+
+			ChunkStatePtr chunkStatePtr = _headerPtr.GetChunkStatePtr(readerState.ChunkIdx);
+			ChunkState chunkState = chunkStatePtr.Value;
+
+			if (chunkState.HasReaderFlag(readerIdx))
 			{
-				Dispose();
+				return _chunks[readerState.ChunkIdx].Slice(readerState.Offset, chunkState.Length - readerState.Offset);
+			}
+			else
+			{
+				return ReadOnlyMemory<byte>.Empty;
 			}
 		}
 
 		/// <inheritdoc/>
-		public void Dispose()
+		public async ValueTask<bool> WaitToReadAsync(int readerIdx, int minLength, CancellationToken cancellationToken = default)
 		{
-			Dispose(true);
-			GC.SuppressFinalize(this);
+			for (; ; )
+			{
+				ReaderStatePtr readerStatePtr = _headerPtr.GetReaderStatePtr(readerIdx);
+				ReaderState readerState = readerStatePtr.Value;
+
+				ChunkStatePtr chunkStatePtr = _headerPtr.GetChunkStatePtr(readerState.ChunkIdx);
+				ChunkState chunkState = chunkStatePtr.Value;
+
+				if (!chunkState.HasReaderFlag(readerIdx))
+				{
+					// Wait until the current chunk is readable
+					_resources.ResetReadEvent(readerIdx);
+					if (!chunkState.HasReaderFlag(readerIdx))
+					{
+						await _resources.WaitForReadEvent(readerIdx, cancellationToken);
+					}
+				}
+				else if (readerState.Offset + minLength <= chunkState.Length)
+				{
+					// We have enough data in the chunk to be able to read a message
+					return true;
+				}
+				else if (chunkState.WriteState == WriteState.Writing)
+				{
+					// Wait until there is more data in the chunk
+					_resources.ResetReadEvent(readerIdx);
+					if (_headerPtr.GetChunkStatePtr(readerState.ChunkIdx).Value == chunkState)
+					{
+						await _resources.WaitForReadEvent(readerIdx, cancellationToken);
+					}
+				}
+				else if (readerState.Offset < chunkState.Length || chunkState.WriteState == WriteState.Complete)
+				{
+					// Cannot read the requested amount of data from this chunk.
+					return false;
+				}
+				else if (chunkState.WriteState == WriteState.MovedToNext)
+				{
+					// Move to the next chunk
+					chunkStatePtr.FinishReading(readerIdx);
+					_resources.SetWriteEvent();
+
+					int chunkIdx = readerStatePtr.ChunkIdx + 1;
+					if (chunkIdx == _chunks.Length)
+					{
+						chunkIdx = 0;
+					}
+
+					readerStatePtr.Value = new ReaderState(chunkIdx, 0);
+				}
+				else
+				{
+					throw new NotImplementedException($"Invalid write state for buffer: {chunkState.WriteState}");
+				}
+			}
 		}
 
-		/// <summary>
-		/// Overriable dispose method
-		/// </summary>
-		protected virtual void Dispose(bool disposing)
+		#endregion
+
+		#region Writer Interface
+
+		/// <inheritdoc/>
+		public void AdvanceWritePosition(int size)
 		{
+			if (size > 0)
+			{
+				ChunkStatePtr chunkStatePtr = _headerPtr.GetChunkStatePtr(_headerPtr.WriteChunkIdx);
+				ChunkState chunkState = chunkStatePtr.Value;
+
+				Debug.Assert(chunkState.WriteState == WriteState.Writing);
+				chunkStatePtr.Append(size);
+
+				SetAllReadEvents();
+			}
 		}
 
-		private void SetAllReadEvents()
+		/// <inheritdoc/>
+		public Memory<byte> GetWriteBuffer()
 		{
-			for (int readerIdx = 0; readerIdx < _headerPtr.NumReaders; readerIdx++)
+			ChunkState state = _headerPtr.GetChunkStatePtr(_headerPtr.WriteChunkIdx).Value;
+			if (state.WriteState == WriteState.Writing)
 			{
-				SetReadEvent(readerIdx);
+				return _chunks[_headerPtr.WriteChunkIdx].Slice(state.Length);
+			}
+			else
+			{
+				return Memory<byte>.Empty;
 			}
 		}
 
-		/// <summary>
-		/// Signals a read event
-		/// </summary>
-		protected abstract void SetReadEvent(int readerIdx);
+		/// <inheritdoc/>
+		public bool MarkComplete()
+		{
+			ChunkStatePtr chunkStatePtr = _headerPtr.GetChunkStatePtr(_headerPtr.WriteChunkIdx);
+			if (chunkStatePtr.WriteState != WriteState.Complete)
+			{
+				chunkStatePtr.MarkComplete();
+				SetAllReadEvents();
+				return true;
+			}
+			return false;
+		}
 
-		/// <summary>
-		/// Resets a read event
-		/// </summary>
-		protected abstract void ResetReadEvent(int readerIdx);
+		/// <inheritdoc/>
+		public async ValueTask WaitToWriteAsync(int minSize, CancellationToken cancellationToken = default)
+		{
+			if (minSize > _headerPtr.ChunkLength)
+			{
+				throw new ArgumentException("Requested read size is larger than chunk size.", nameof(minSize));
+			}
 
-		/// <summary>
-		/// Waits for a read event to be signalled
-		/// </summary>
-		protected abstract Task WaitForReadEvent(int readerIdx, CancellationToken cancellationToken);
+			for (; ; )
+			{
+				ChunkStatePtr chunkStatePtr = _headerPtr.GetChunkStatePtr(_headerPtr.WriteChunkIdx);
 
-		/// <summary>
-		/// Signals the write event
-		/// </summary>
-		protected abstract void SetWriteEvent();
+				ChunkState chunkState = chunkStatePtr.Value;
+				if (chunkState.WriteState == WriteState.Writing)
+				{
+					int length = chunkState.Length;
+					if (length + minSize <= _headerPtr.ChunkLength)
+					{
+						return;
+					}
+					chunkStatePtr.FinishWriting();
+				}
 
-		/// <summary>
-		/// Resets the write event
-		/// </summary>
-		protected abstract void ResetWriteEvent();
+				int nextChunkIdx = _headerPtr.WriteChunkIdx + 1;
+				if (nextChunkIdx == _chunks.Length)
+				{
+					nextChunkIdx = 0;
+				}
 
-		/// <summary>
-		/// Waits for the write event to be signalled
-		/// </summary>
-		protected abstract Task WaitForWriteEvent(CancellationToken cancellationToken);
+				ChunkStatePtr nextChunkStatePtr = _headerPtr.GetChunkStatePtr(nextChunkIdx);
+				while (nextChunkStatePtr.ReaderFlags != 0)
+				{
+					await _resources.WaitForWriteEvent(cancellationToken);
+					_resources.ResetWriteEvent();
+				}
+
+				_headerPtr.WriteChunkIdx = nextChunkIdx;
+				nextChunkStatePtr.StartWriting(_headerPtr.NumReaders);
+
+				SetAllReadEvents();
+			}
+		}
+
+		#endregion
 	}
 }
