Horde: Move reader/writer methods into IComputeBuffer, and turn IComputeBufferReader/IComputeBufferWriter into concrete implementations.

diff --git a/Engine/Source/Programs/Shared/EpicGames.Horde.Tests/BufferTests.cs b/Engine/Source/Programs/Shared/EpicGames.Horde.Tests/BufferTests.cs
--- a/Engine/Source/Programs/Shared/EpicGames.Horde.Tests/BufferTests.cs
+++ b/Engine/Source/Programs/Shared/EpicGames.Horde.Tests/BufferTests.cs
@@ -126,7 +126,7 @@ namespace EpicGames.Horde.Tests
 			await socket.MarkCompleteAsync(ChannelId);
 		}
 
-		static async Task RunConsumerAsync(IComputeBufferReader reader, Memory<byte> output)
+		static async Task RunConsumerAsync(ComputeBufferReader reader, Memory<byte> output)
 		{
 			int offset = 0;
 			while (!reader.IsComplete)
diff --git a/Engine/Source/Programs/Shared/EpicGames.Horde/Compute/ComputeBufferBase.cs b/Engine/Source/Programs/Shared/EpicGames.Horde/Compute/ComputeBufferBase.cs
--- a/Engine/Source/Programs/Shared/EpicGames.Horde/Compute/ComputeBufferBase.cs
+++ b/Engine/Source/Programs/Shared/EpicGames.Horde/Compute/ComputeBufferBase.cs
@@ -228,45 +228,15 @@ namespace EpicGames.Horde.Compute
 			public abstract Task WaitForWriteEvent(CancellationToken cancellationToken);
 		}
 
-		class ReaderImpl : IComputeBufferReader
-		{
-			readonly ComputeBufferBase _buffer;
-
-			public ReaderImpl(ComputeBufferBase buffer) => _buffer = buffer;
-
-			public bool IsComplete => _buffer.IsComplete(0);
-
-			public void AdvanceReadPosition(int size) => _buffer.AdvanceReadPosition(0, size);
-
-			public ReadOnlyMemory<byte> GetReadBuffer() => _buffer.GetReadBuffer(0);
-
-			public ValueTask<bool> WaitToReadAsync(int minLength, CancellationToken cancellationToken = default) => _buffer.WaitToReadAsync(0, minLength, cancellationToken);
-		}
-
-		class WriterImpl : IComputeBufferWriter
-		{
-			readonly ComputeBufferBase _buffer;
-
-			public WriterImpl(ComputeBufferBase buffer) => _buffer = buffer;
-
-			public void AdvanceWritePosition(int size) => _buffer.AdvanceWritePosition(size);
-
-			public Memory<byte> GetWriteBuffer() => _buffer.GetWriteBuffer();
-
-			public bool MarkComplete() => _buffer.MarkComplete();
-
-			public ValueTask WaitToWriteAsync(int minLength, CancellationToken cancellationToken = default) => _buffer.WaitToWriteAsync(minLength, cancellationToken);
-		}
-
 		readonly HeaderPtr _headerPtr;
 		readonly Memory<byte>[] _chunks;
 		ResourcesBase _resources;
 
 		/// <inheritdoc/>
-		public IComputeBufferReader Reader { get; }
+		public ComputeBufferReader Reader { get; }
 
 		/// <inheritdoc/>
-		public IComputeBufferWriter Writer { get; }
+		public ComputeBufferWriter Writer { get; }
 
 #pragma warning disable IDE0051 // Remove unused private members
 		// For debugging purposes only
@@ -287,8 +257,8 @@ namespace EpicGames.Horde.Compute
 			_chunks = chunks;
 			_resources = resources;
 
-			Reader = new ReaderImpl(this);
-			Writer = new WriterImpl(this);
+			Reader = new ComputeBufferReader(this, 0);
+			Writer = new ComputeBufferWriter(this);
 		}
 
 		/// <inheritdoc/>
diff --git a/Engine/Source/Programs/Shared/EpicGames.Horde/Compute/ComputeMessageChannel.cs b/Engine/Source/Programs/Shared/EpicGames.Horde/Compute/ComputeMessageChannel.cs
--- a/Engine/Source/Programs/Shared/EpicGames.Horde/Compute/ComputeMessageChannel.cs
+++ b/Engine/Source/Programs/Shared/EpicGames.Horde/Compute/ComputeMessageChannel.cs
@@ -60,14 +60,14 @@ namespace EpicGames.Horde.Compute
 		class MessageBuilder : IComputeMessageBuilder
 		{
 			readonly ComputeMessageChannel _channel;
-			readonly IComputeBufferWriter _sendBufferWriter;
+			readonly ComputeBufferWriter _sendBufferWriter;
 			readonly ComputeMessageType _type;
 			int _length;
 
 			/// <inheritdoc/>
 			public int Length => _length;
 
-			public MessageBuilder(ComputeMessageChannel channel, IComputeBufferWriter sendBufferWriter, ComputeMessageType type)
+			public MessageBuilder(ComputeMessageChannel channel, ComputeBufferWriter sendBufferWriter, ComputeMessageType type)
 			{
 				_channel = channel;
 				_sendBufferWriter = sendBufferWriter;
@@ -180,7 +180,7 @@ namespace EpicGames.Horde.Compute
 		/// <inheritdoc/>
 		public async ValueTask<IComputeMessage> ReceiveAsync(CancellationToken cancellationToken)
 		{
-			IComputeBufferReader recvBufferReader = _recvBuffer.Reader;
+			ComputeBufferReader recvBufferReader = _recvBuffer.Reader;
 			while (!recvBufferReader.IsComplete)
 			{
 				ReadOnlyMemory<byte> memory = recvBufferReader.GetReadBuffer();
@@ -231,7 +231,7 @@ namespace EpicGames.Horde.Compute
 				throw new InvalidOperationException("Only one writer can be active at a time. Dispose of the previous writer first.");
 			}
 
-			IComputeBufferWriter sendBufferWriter = _sendBuffer.Writer;
+			ComputeBufferWriter sendBufferWriter = _sendBuffer.Writer;
 			await sendBufferWriter.WaitToWriteAsync(maxSize, cancellationToken);
 
 			_currentBuilder = new MessageBuilder(this, sendBufferWriter, type);
diff --git a/Engine/Source/Programs/Shared/EpicGames.Horde/Compute/ComputeWorker.cs b/Engine/Source/Programs/Shared/EpicGames.Horde/Compute/ComputeWorker.cs
--- a/Engine/Source/Programs/Shared/EpicGames.Horde/Compute/ComputeWorker.cs
+++ b/Engine/Source/Programs/Shared/EpicGames.Horde/Compute/ComputeWorker.cs
@@ -190,7 +190,7 @@ namespace EpicGames.Horde.Compute
 			ipcBuffer.Writer.MarkComplete();
 		}
 
-		static async Task ProcessIpcMessagesAsync(IComputeSocket socket, IComputeBufferReader ipcReader, CancellationToken cancellationToken)
+		static async Task ProcessIpcMessagesAsync(IComputeSocket socket, ComputeBufferReader ipcReader, CancellationToken cancellationToken)
 		{
 			while(await ipcReader.WaitToReadAsync(1, cancellationToken))
 			{
diff --git a/Engine/Source/Programs/Shared/EpicGames.Horde/Compute/DefaultComputeSocket.cs b/Engine/Source/Programs/Shared/EpicGames.Horde/Compute/DefaultComputeSocket.cs
--- a/Engine/Source/Programs/Shared/EpicGames.Horde/Compute/DefaultComputeSocket.cs
+++ b/Engine/Source/Programs/Shared/EpicGames.Horde/Compute/DefaultComputeSocket.cs
@@ -117,7 +117,7 @@ namespace EpicGames.Horde.Compute
 			try
 			{
 				// Maintain a local cache of buffers to be able to query for them without having to acquire a global lock
-				Dictionary<int, IComputeBufferWriter> cachedWriters = new Dictionary<int, IComputeBufferWriter>();
+				Dictionary<int, ComputeBufferWriter> cachedWriters = new Dictionary<int, ComputeBufferWriter>();
 
 				Memory<byte> last = Memory<byte>.Empty;
 
@@ -138,7 +138,7 @@ namespace EpicGames.Horde.Compute
 					// Dispatch it to the correct place
 					if (size >= 0)
 					{
-						IComputeBufferWriter writer = GetReceiveBuffer(cachedWriters, id);
+						ComputeBufferWriter writer = GetReceiveBuffer(cachedWriters, id);
 						await ReadPacketAsync(transport, id, size, writer, cancellationToken);
 					}
 					else if (size == (int)ControlMessageType.Attach)
@@ -174,7 +174,7 @@ namespace EpicGames.Horde.Compute
 			_logger.LogTrace("Closing reader");
 		}
 
-		async Task ReadPacketAsync(IComputeTransport transport, int id, int size, IComputeBufferWriter writer, CancellationToken cancellationToken)
+		async Task ReadPacketAsync(IComputeTransport transport, int id, int size, ComputeBufferWriter writer, CancellationToken cancellationToken)
 		{
 			Memory<byte> memory = writer.GetWriteBuffer();
 			while (memory.Length < size)
@@ -193,9 +193,9 @@ namespace EpicGames.Horde.Compute
 			writer.AdvanceWritePosition(size);
 		}
 
-		IComputeBufferWriter GetReceiveBuffer(Dictionary<int, IComputeBufferWriter> cachedWriters, int id)
+		ComputeBufferWriter GetReceiveBuffer(Dictionary<int, ComputeBufferWriter> cachedWriters, int id)
 		{
-			IComputeBufferWriter? writer;
+			ComputeBufferWriter? writer;
 			if (cachedWriters.TryGetValue(id, out writer))
 			{
 				return writer;
@@ -325,7 +325,7 @@ namespace EpicGames.Horde.Compute
 			_remoteBuffersChanged.Pulse();
 		}
 
-		void DetachRecvBuffer(Dictionary<int, IComputeBufferWriter> cachedWriters, int id)
+		void DetachRecvBuffer(Dictionary<int, ComputeBufferWriter> cachedWriters, int id)
 		{
 			cachedWriters.Remove(id);
 
@@ -358,7 +358,7 @@ namespace EpicGames.Horde.Compute
 		{
 			using IComputeBuffer _ = buffer;
 
-			IComputeBufferReader reader = buffer.Reader;
+			ComputeBufferReader reader = buffer.Reader;
 			while (!cancellationToken.IsCancellationRequested)
 			{
 				ReadOnlyMemory<byte> memory = reader.GetReadBuffer();
diff --git a/Engine/Source/Programs/Shared/EpicGames.Horde/Compute/IComputeBuffer.cs b/Engine/Source/Programs/Shared/EpicGames.Horde/Compute/IComputeBuffer.cs
--- a/Engine/Source/Programs/Shared/EpicGames.Horde/Compute/IComputeBuffer.cs
+++ b/Engine/Source/Programs/Shared/EpicGames.Horde/Compute/IComputeBuffer.cs
@@ -14,55 +14,53 @@ namespace EpicGames.Horde.Compute
 		/// <summary>
 		/// Reader for this buffer
 		/// </summary>
-		IComputeBufferReader Reader { get; }
+		ComputeBufferReader Reader { get; }
 
 		/// <summary>
 		/// Writer for this buffer
 		/// </summary>
-		IComputeBufferWriter Writer { get; }
+		ComputeBufferWriter Writer { get; }
 
 		/// <summary>
 		/// Creates a new reference to the underlying buffer. The underlying resources will only be destroyed once all instances are disposed of.
 		/// </summary>
 		IComputeBuffer AddRef();
-	}
 
-	/// <summary>
-	/// Read interface for a compute buffer
-	/// </summary>
-	public interface IComputeBufferReader
-	{
+		#region Reader interface
+
 		/// <summary>
 		/// Whether this buffer is complete (no more data will be added)
 		/// </summary>
-		bool IsComplete { get; }
+		/// <param name="readerIdx">Index of the reader</param>
+		bool IsComplete(int readerIdx);
 
 		/// <summary>
 		/// Updates the read position
 		/// </summary>
+		/// <param name="readerIdx">Index of the reader</param>
 		/// <param name="size">Size of data that was read</param>
-		void AdvanceReadPosition(int size);
+		void AdvanceReadPosition(int readerIdx, int size);
 
 		/// <summary>
 		/// Gets the next data to read
 		/// </summary>
+		/// <param name="readerIdx">Index of the reader</param>
 		/// <returns>Memory to read from</returns>
-		ReadOnlyMemory<byte> GetReadBuffer();
+		ReadOnlyMemory<byte> GetReadBuffer(int readerIdx);
 
 		/// <summary>
 		/// Wait for data to be available, or for the buffer to be marked as complete
 		/// </summary>
+		/// <param name="readerIdx">Index of the reader</param>
 		/// <param name="minLength">Minimum amount of data to read</param>
 		/// <param name="cancellationToken">Cancellation token for the operation</param>
 		/// <returns>True if new data is available, false if the buffer is complete</returns>
-		ValueTask<bool> WaitToReadAsync(int minLength, CancellationToken cancellationToken = default);
-	}
+		ValueTask<bool> WaitToReadAsync(int readerIdx, int minLength, CancellationToken cancellationToken = default);
+
+		#endregion
+
+		#region Writer interface
 
-	/// <summary>
-	/// Buffer that can receive data from a remote machine.
-	/// </summary>
-	public interface IComputeBufferWriter
-	{
 		/// <summary>
 		/// Mark the output to this buffer as complete
 		/// </summary>
@@ -87,6 +85,63 @@ namespace EpicGames.Horde.Compute
 		/// <param name="cancellationToken">Cancellation token for the operation</param>
 		/// <returns>Memory to be written to</returns>
 		ValueTask WaitToWriteAsync(int minLength, CancellationToken cancellationToken = default);
+
+		#endregion
+	}
+
+	/// <summary>
+	/// Wraps the read methods for a compute buffer
+	/// </summary>
+	public sealed class ComputeBufferReader
+	{
+		readonly IComputeBuffer _buffer;
+		readonly int _readerIdx;
+
+		/// <summary>
+		/// Constructor
+		/// </summary>
+		public ComputeBufferReader(IComputeBuffer buffer, int readerIdx)
+		{
+			_buffer = buffer;
+			_readerIdx = readerIdx;
+		}
+
+		/// <inheritdoc cref="IComputeBuffer.IsComplete(Int32)"/>
+		public bool IsComplete => _buffer.IsComplete(_readerIdx);
+
+		/// <inheritdoc cref="IComputeBuffer.AdvanceReadPosition(Int32, Int32)"/>
+		public void AdvanceReadPosition(int size) => _buffer.AdvanceReadPosition(_readerIdx, size);
+
+		/// <inheritdoc cref="IComputeBuffer.GetReadBuffer(Int32)"/>
+		public ReadOnlyMemory<byte> GetReadBuffer() => _buffer.GetReadBuffer(_readerIdx);
+
+		/// <inheritdoc cref="IComputeBuffer.WaitToReadAsync(Int32, Int32, CancellationToken)"/>
+		public ValueTask<bool> WaitToReadAsync(int minLength, CancellationToken cancellationToken = default) => _buffer.WaitToReadAsync(_readerIdx, minLength, cancellationToken);
+	}
+
+	/// <summary>
+	/// Wraps the write methods for a compute buffer
+	/// </summary>
+	public sealed class ComputeBufferWriter
+	{
+		readonly IComputeBuffer _buffer;
+
+		/// <summary>
+		/// Constructor
+		/// </summary>
+		public ComputeBufferWriter(IComputeBuffer buffer) => _buffer = buffer;
+
+		/// <inheritdoc cref="IComputeBuffer.MarkComplete"/>
+		public bool MarkComplete() => _buffer.MarkComplete();
+
+		/// <inheritdoc cref="IComputeBuffer.AdvanceWritePosition(Int32)"/>
+		public void AdvanceWritePosition(int size) => _buffer.AdvanceWritePosition(size);
+
+		/// <inheritdoc cref="IComputeBuffer.GetWriteBuffer()"/>
+		public Memory<byte> GetWriteBuffer() => _buffer.GetWriteBuffer();
+
+		/// <inheritdoc cref="IComputeBuffer.WaitToWriteAsync(Int32, CancellationToken)"/>
+		public ValueTask WaitToWriteAsync(int minLength, CancellationToken cancellationToken = default) => _buffer.WaitToWriteAsync(minLength, cancellationToken);
 	}
 
 	/// <summary>
@@ -101,7 +156,7 @@ namespace EpicGames.Horde.Compute
 		/// <param name="buffer">Memory to receive the read data</param>
 		/// <param name="cancellationToken">Cancellation token for the operation</param>
 		/// <returns>Number of bytes read</returns>
-		public static async ValueTask<int> ReadAsync(this IComputeBufferReader reader, Memory<byte> buffer, CancellationToken cancellationToken = default)
+		public static async ValueTask<int> ReadAsync(this ComputeBufferReader reader, Memory<byte> buffer, CancellationToken cancellationToken = default)
 		{
 			for (; ; )
 			{
@@ -123,7 +178,7 @@ namespace EpicGames.Horde.Compute
 		/// <param name="writer">Writer to output the data to</param>
 		/// <param name="buffer">The data to write</param>
 		/// <param name="cancellationToken">Cancellation token for the operation</param>
-		public static async ValueTask WriteAsync(this IComputeBufferWriter writer, ReadOnlyMemory<byte> buffer, CancellationToken cancellationToken = default)
+		public static async ValueTask WriteAsync(this ComputeBufferWriter writer, ReadOnlyMemory<byte> buffer, CancellationToken cancellationToken = default)
 		{
 			while (buffer.Length > 0)
 			{
