Horde: Allow incrementing ref counts of compute buffer reader/writer instances independently of the owning buffer.

diff --git a/Engine/Source/Programs/Shared/EpicGames.Horde/Compute/Buffers/PooledBuffer.cs b/Engine/Source/Programs/Shared/EpicGames.Horde/Compute/Buffers/PooledBuffer.cs
--- a/Engine/Source/Programs/Shared/EpicGames.Horde/Compute/Buffers/PooledBuffer.cs
+++ b/Engine/Source/Programs/Shared/EpicGames.Horde/Compute/Buffers/PooledBuffer.cs
@@ -1,6 +1,7 @@
 // Copyright Epic Games, Inc. All Rights Reserved.
 
 using System;
+using System.Buffers;
 using System.Runtime.InteropServices;
 using System.Threading;
 using System.Threading.Tasks;
@@ -15,37 +16,42 @@ namespace EpicGames.Horde.Compute.Buffers
 	{
 		class Resources : ResourcesBase
 		{
-			public HeaderPtr _headerPtr;
 			readonly GCHandle _headerHandle;
-
-			public Memory<byte>[] _chunks;
-			readonly GCHandle[] _chunkHandles;
+			readonly IMemoryOwner<byte>[] _chunkOwners;
 
 			readonly AsyncEvent _writerEvent = new AsyncEvent();
 			readonly AsyncEvent _readerEvent = new AsyncEvent();
 
-			public unsafe Resources(int numChunks, int chunkLength, int numReaders)
+			public Resources(HeaderPtr headerPtr, GCHandle headerHandle, Memory<byte>[] chunks, IMemoryOwner<byte>[] chunkOwners)
+				: base(headerPtr, chunks)
+			{
+				_headerHandle = headerHandle;
+				_chunkOwners = chunkOwners;
+			}
+
+			public static unsafe Resources Create(int numChunks, int chunkLength, int numReaders)
 			{
-				_chunks = new Memory<byte>[numChunks];
-				_chunkHandles = new GCHandle[numChunks];
+				byte[] header = new byte[HeaderSize];
+				GCHandle headerHandle = GCHandle.Alloc(header, GCHandleType.Pinned);
+				HeaderPtr headerPtr = new HeaderPtr((ulong*)headerHandle.AddrOfPinnedObject().ToPointer(), numReaders, numChunks, chunkLength);
+
+				Memory<byte>[] chunks = new Memory<byte>[numChunks];
+				IMemoryOwner<byte>[] chunkOwners = new IMemoryOwner<byte>[numChunks];
 
 				for (int idx = 0; idx < numChunks; idx++)
 				{
-					byte[] data = new byte[chunkLength];
-					_chunks[idx] = data;
-					_chunkHandles[idx] = GCHandle.Alloc(data, GCHandleType.Pinned);
+					chunkOwners[idx] = MemoryPool<byte>.Shared.Rent(chunkLength);
+					chunks[idx] = chunkOwners[idx].Memory.Slice(0, chunkLength);
 				}
 
-				byte[] header = new byte[HeaderSize];
-				_headerHandle = GCHandle.Alloc(header, GCHandleType.Pinned);
-				_headerPtr = new HeaderPtr((ulong*)_headerHandle.AddrOfPinnedObject().ToPointer(), numReaders, numChunks, chunkLength);
+				return new Resources(headerPtr, headerHandle, chunks, chunkOwners);
 			}
 
 			public override void Dispose()
 			{
-				for (int idx = 0; idx < _chunkHandles.Length; idx++)
+				for (int idx = 0; idx < _chunkOwners.Length; idx++)
 				{
-					_chunkHandles[idx].Free();
+					_chunkOwners[idx].Dispose();
 				}
 				_headerHandle.Free();
 			}
@@ -69,8 +75,6 @@ namespace EpicGames.Horde.Compute.Buffers
 			public override Task WaitForWriteEvent(CancellationToken cancellationToken) => _writerEvent.Task.WaitAsync(cancellationToken);
 		}
 
-		readonly Resources _resources;
-
 		/// <summary>
 		/// Constructor
 		/// </summary>
@@ -87,14 +91,13 @@ namespace EpicGames.Horde.Compute.Buffers
 		/// <param name="chunkLength">Length of each chunk</param>
 		/// <param name="numReaders">Number of readers for this buffer</param>
 		public PooledBuffer(int numChunks, int chunkLength, int numReaders = 1)
-			: this(new Resources(numChunks, chunkLength, numReaders))
+			: base(Resources.Create(numChunks, chunkLength, numReaders))
 		{
 		}
 
-		private PooledBuffer(Resources resources)
-			: base(resources._headerPtr, resources._chunks, resources)
+		private PooledBuffer(ResourcesBase resources)
+			: base(resources)
 		{
-			_resources = resources;
 		}
 
 		/// <inheritdoc/>
diff --git a/Engine/Source/Programs/Shared/EpicGames.Horde/Compute/Buffers/SharedMemoryBuffer.cs b/Engine/Source/Programs/Shared/EpicGames.Horde/Compute/Buffers/SharedMemoryBuffer.cs
--- a/Engine/Source/Programs/Shared/EpicGames.Horde/Compute/Buffers/SharedMemoryBuffer.cs
+++ b/Engine/Source/Programs/Shared/EpicGames.Horde/Compute/Buffers/SharedMemoryBuffer.cs
@@ -25,13 +25,11 @@ namespace EpicGames.Horde.Compute.Buffers
 			readonly Native.EventHandle _writerEvent;
 
 			public string Name { get; }
-			public HeaderPtr HeaderPtr { get; }
-			public Memory<byte>[] Chunks { get; }
 
-			public Resources(string name, HeaderPtr headerPtr, MemoryMappedFile memoryMappedFile, MemoryMappedViewAccessor memoryMappedViewAccessor, MemoryMappedView memoryMappedView, Native.EventHandle readerEvent, Native.EventHandle writerEvent)
+			public Resources(string name, HeaderPtr headerPtr, Memory<byte>[] chunks, MemoryMappedFile memoryMappedFile, MemoryMappedViewAccessor memoryMappedViewAccessor, MemoryMappedView memoryMappedView, Native.EventHandle readerEvent, Native.EventHandle writerEvent)
+				: base(headerPtr, chunks)
 			{
 				Name = name;
-				HeaderPtr = headerPtr;
 
 				_memoryMappedFile = memoryMappedFile;
 				_memoryMappedViewAccessor = memoryMappedViewAccessor;
@@ -39,13 +37,6 @@ namespace EpicGames.Horde.Compute.Buffers
 
 				_readerEvent = readerEvent;
 				_writerEvent = writerEvent;
-
-				Chunks = new Memory<byte>[headerPtr.NumChunks];
-				for (int chunkIdx = 0; chunkIdx < headerPtr.NumChunks; chunkIdx++)
-				{
-					int chunkOffset = HeaderSize + (headerPtr.ChunkLength * chunkIdx);
-					Chunks[chunkIdx] = memoryMappedView.GetMemory(chunkOffset, headerPtr.ChunkLength);
-				}
 			}
 
 			public override void Dispose()
@@ -78,17 +69,14 @@ namespace EpicGames.Horde.Compute.Buffers
 		}
 
 		/// <inheritdoc/>
-		public string Name => _resources.Name;
-
-		readonly Resources _resources;
+		public string Name => ((Resources)_resources).Name;
 
 		/// <summary>
 		/// Constructor
 		/// </summary>
-		private unsafe SharedMemoryBuffer(Resources resources)
-			: base(resources.HeaderPtr, resources.Chunks, resources)
+		private SharedMemoryBuffer(ResourcesBase resources)
+			: base(resources)
 		{
-			_resources = resources;
 		}
 
 		/// <inheritdoc/>
@@ -129,7 +117,9 @@ namespace EpicGames.Horde.Compute.Buffers
 			Native.EventHandle readerEvent = Native.EventHandle.CreateNew($"{name}_R", EventResetMode.ManualReset, true, HandleInheritability.Inheritable);
 
 			HeaderPtr headerPtr = new HeaderPtr((ulong*)memoryMappedView.GetPointer(), 1, numChunks, chunkLength);
-			return new SharedMemoryBuffer(new Resources(name, headerPtr, memoryMappedFile, memoryMappedViewAccessor, memoryMappedView, readerEvent, writerEvent));
+			Memory<byte>[] chunks = CreateChunks(headerPtr.NumChunks, headerPtr.ChunkLength, memoryMappedView);
+
+			return new SharedMemoryBuffer(new Resources(name, headerPtr, chunks, memoryMappedFile, memoryMappedViewAccessor, memoryMappedView, readerEvent, writerEvent));
 		}
 
 		/// <summary>
@@ -151,7 +141,20 @@ namespace EpicGames.Horde.Compute.Buffers
 			Native.EventHandle writerEvent = Native.EventHandle.OpenExisting($"{name}_W");
 
 			HeaderPtr headerPtr = new HeaderPtr((ulong*)memoryMappedView.GetPointer());
-			return new SharedMemoryBuffer(new Resources(name, headerPtr, memoryMappedFile, memoryMappedViewAccessor, memoryMappedView, readerEvent, writerEvent));
+			Memory<byte>[] chunks = CreateChunks(headerPtr.NumChunks, headerPtr.ChunkLength, memoryMappedView);
+
+			return new SharedMemoryBuffer(new Resources(name, headerPtr, chunks, memoryMappedFile, memoryMappedViewAccessor, memoryMappedView, readerEvent, writerEvent));
+		}
+
+		static Memory<byte>[] CreateChunks(int numChunks, int chunkLength, MemoryMappedView memoryMappedView)
+		{
+			Memory<byte>[] chunks = new Memory<byte>[numChunks];
+			for (int chunkIdx = 0; chunkIdx < numChunks; chunkIdx++)
+			{
+				int chunkOffset = HeaderSize + (chunkLength * chunkIdx);
+				chunks[chunkIdx] = memoryMappedView.GetMemory(chunkOffset, chunkLength);
+			}
+			return chunks;
 		}
 	}
 }
diff --git a/Engine/Source/Programs/Shared/EpicGames.Horde/Compute/ComputeBufferBase.cs b/Engine/Source/Programs/Shared/EpicGames.Horde/Compute/ComputeBufferBase.cs
--- a/Engine/Source/Programs/Shared/EpicGames.Horde/Compute/ComputeBufferBase.cs
+++ b/Engine/Source/Programs/Shared/EpicGames.Horde/Compute/ComputeBufferBase.cs
@@ -177,8 +177,16 @@ namespace EpicGames.Horde.Compute
 		/// </summary>
 		internal abstract class ResourcesBase : IDisposable
 		{
+			public HeaderPtr HeaderPtr { get; }
+			public Memory<byte>[] Chunks { get; }
 			int _refCount = 1;
 
+			public ResourcesBase(HeaderPtr headerPtr, Memory<byte>[] chunks)
+			{
+				HeaderPtr = headerPtr;
+				Chunks = chunks;
+			}
+
 			public void AddRef()
 			{
 				Interlocked.Increment(ref _refCount);
@@ -202,6 +210,17 @@ namespace EpicGames.Horde.Compute
 			/// </summary>
 			public abstract void SetReadEvent(int readerIdx);
 
+			/// <summary>
+			/// Signals read events for every reader
+			/// </summary>
+			public void SetAllReadEvents()
+			{
+				for (int readerIdx = 0; readerIdx < HeaderPtr.NumReaders; readerIdx++)
+				{
+					SetReadEvent(readerIdx);
+				}
+			}
+
 			/// <summary>
 			/// Resets a read event
 			/// </summary>
@@ -226,286 +245,280 @@ namespace EpicGames.Horde.Compute
 			/// Waits for the write event to be signalled
 			/// </summary>
 			public abstract Task WaitForWriteEvent(CancellationToken cancellationToken);
-		}
-
-		class ReaderImpl : IComputeBufferReader
-		{
-			readonly ComputeBufferBase _buffer;
-
-			public ReaderImpl(ComputeBufferBase buffer) => _buffer = buffer;
-
-			public bool IsComplete => _buffer.IsComplete(0);
-
-			public void AdvanceReadPosition(int size) => _buffer.AdvanceReadPosition(0, size);
-
-			public ReadOnlyMemory<byte> GetReadBuffer() => _buffer.GetReadBuffer(0);
-
-			public ValueTask<bool> WaitToReadAsync(int minLength, CancellationToken cancellationToken = default) => _buffer.WaitToReadAsync(0, minLength, cancellationToken);
-		}
-
-		class WriterImpl : IComputeBufferWriter
-		{
-			readonly ComputeBufferBase _buffer;
-
-			public WriterImpl(ComputeBufferBase buffer) => _buffer = buffer;
-
-			public void AdvanceWritePosition(int size) => _buffer.AdvanceWritePosition(size);
-
-			public Memory<byte> GetWriteBuffer() => _buffer.GetWriteBuffer();
-
-			public bool MarkComplete() => _buffer.MarkComplete();
-
-			public ValueTask WaitToWriteAsync(int minLength, CancellationToken cancellationToken = default) => _buffer.WaitToWriteAsync(minLength, cancellationToken);
-		}
-
-		readonly HeaderPtr _headerPtr;
-		readonly Memory<byte>[] _chunks;
-		ResourcesBase _resources;
-
-		/// <inheritdoc/>
-		public IComputeBufferReader Reader { get; }
-
-		/// <inheritdoc/>
-		public IComputeBufferWriter Writer { get; }
 
 #pragma warning disable IDE0051 // Remove unused private members
-		// For debugging purposes only
-		ChunkState[] ChunkStates => Enumerable.Range(0, _headerPtr.NumChunks).Select(x => _headerPtr.GetChunkStatePtr(x).Value).ToArray();
-		ReaderState[] ReaderStates => Enumerable.Range(0, _headerPtr.NumReaders).Select(x => _headerPtr.GetReaderStatePtr(x).Value).ToArray();
-		int WriteChunkIdx => _headerPtr.WriteChunkIdx;
+			// For debugging purposes only
+			ChunkState[] ChunkStates => Enumerable.Range(0, HeaderPtr.NumChunks).Select(x => HeaderPtr.GetChunkStatePtr(x).Value).ToArray();
+			ReaderState[] ReaderStates => Enumerable.Range(0, HeaderPtr.NumReaders).Select(x => HeaderPtr.GetReaderStatePtr(x).Value).ToArray();
+			int WriteChunkIdx => HeaderPtr.WriteChunkIdx;
 #pragma warning restore IDE0051 // Remove unused private members
-
-		/// <summary>
-		/// Constructor
-		/// </summary>
-		/// <param name="headerPtr">Header for the buffer</param>
-		/// <param name="chunks">Data for the buffer</param>
-		/// <param name="resources">Resources shared between instances of the buffer</param>
-		internal ComputeBufferBase(HeaderPtr headerPtr, Memory<byte>[] chunks, ResourcesBase resources)
-		{
-			_headerPtr = headerPtr;
-			_chunks = chunks;
-			_resources = resources;
-
-			Reader = new ReaderImpl(this);
-			Writer = new WriterImpl(this);
 		}
 
-		/// <inheritdoc/>
-		public void Dispose()
+		class ReaderImpl : IComputeBufferReader
 		{
-			Dispose(true);
-			GC.SuppressFinalize(this);
-		}
+			readonly HeaderPtr _headerPtr;
+			readonly Memory<byte>[] _chunks;
+			readonly ResourcesBase _resources;
+			readonly int _readerIdx;
 
-		/// <summary>
-		/// Overriable dispose method
-		/// </summary>
-		protected virtual void Dispose(bool disposing)
-		{
-			if (_resources != null)
+			public ReaderImpl(ResourcesBase resources, int readerIdx)
 			{
-				_resources.Release();
-				_resources = null!;
+				_headerPtr = resources.HeaderPtr;
+				_chunks = resources.Chunks;
+				_resources = resources;
+				_readerIdx = readerIdx;
 			}
-		}
 
-		private void SetAllReadEvents()
-		{
-			for (int readerIdx = 0; readerIdx < _headerPtr.NumReaders; readerIdx++)
+			/// <inheritdoc/>
+			public IComputeBufferReader AddRef()
 			{
-				_resources.SetReadEvent(readerIdx);
+				_resources.AddRef();
+				return new ReaderImpl(_resources, _readerIdx);
 			}
-		}
-
-		/// <inheritdoc/>
-		public abstract IComputeBuffer AddRef();
-
-		#region Reader Interface
-
-		/// <inheritdoc/>
-		public bool IsComplete(int readerIdx)
-		{
-			ReaderState readerState = _headerPtr.GetReaderStatePtr(readerIdx).Value;
-			ChunkState chunkState = _headerPtr.GetChunkStatePtr(readerState.ChunkIdx).Value;
-			return chunkState.WriteState == WriteState.Complete && readerState.Offset == chunkState.Length;
-		}
-
-		/// <inheritdoc/>
-		public void AdvanceReadPosition(int readerIdx, int length)
-		{
-			ReaderStatePtr readerStatePtr = _headerPtr.GetReaderStatePtr(readerIdx);
-			readerStatePtr.Advance(length);
-		}
-
-		/// <inheritdoc/>
-		public ReadOnlyMemory<byte> GetReadBuffer(int readerIdx)
-		{
-			ReaderStatePtr readerStatePtr = _headerPtr.GetReaderStatePtr(readerIdx);
-			ReaderState readerState = readerStatePtr.Value;
 
-			ChunkStatePtr chunkStatePtr = _headerPtr.GetChunkStatePtr(readerState.ChunkIdx);
-			ChunkState chunkState = chunkStatePtr.Value;
-
-			if (chunkState.HasReaderFlag(readerIdx))
+			/// <inheritdoc/>
+			public bool IsComplete
 			{
-				return _chunks[readerState.ChunkIdx].Slice(readerState.Offset, chunkState.Length - readerState.Offset);
+				get
+				{
+					ReaderState readerState = _headerPtr.GetReaderStatePtr(_readerIdx).Value;
+					ChunkState chunkState = _headerPtr.GetChunkStatePtr(readerState.ChunkIdx).Value;
+					return chunkState.WriteState == WriteState.Complete && readerState.Offset == chunkState.Length;
+				}
 			}
-			else
+
+			/// <inheritdoc/>
+			public void AdvanceReadPosition(int length)
 			{
-				return ReadOnlyMemory<byte>.Empty;
+				ReaderStatePtr readerStatePtr = _headerPtr.GetReaderStatePtr(_readerIdx);
+				readerStatePtr.Advance(length);
 			}
-		}
 
-		/// <inheritdoc/>
-		public async ValueTask<bool> WaitToReadAsync(int readerIdx, int minLength, CancellationToken cancellationToken = default)
-		{
-			for (; ; )
+			/// <inheritdoc/>
+			public ReadOnlyMemory<byte> GetReadBuffer()
 			{
-				ReaderStatePtr readerStatePtr = _headerPtr.GetReaderStatePtr(readerIdx);
+				ReaderStatePtr readerStatePtr = _headerPtr.GetReaderStatePtr(_readerIdx);
 				ReaderState readerState = readerStatePtr.Value;
 
 				ChunkStatePtr chunkStatePtr = _headerPtr.GetChunkStatePtr(readerState.ChunkIdx);
 				ChunkState chunkState = chunkStatePtr.Value;
 
-				if (!chunkState.HasReaderFlag(readerIdx))
+				if (chunkState.HasReaderFlag(_readerIdx))
 				{
-					// Wait until the current chunk is readable
-					_resources.ResetReadEvent(readerIdx);
-					if (!chunkState.HasReaderFlag(readerIdx))
-					{
-						await _resources.WaitForReadEvent(readerIdx, cancellationToken);
-					}
+					return _chunks[readerState.ChunkIdx].Slice(readerState.Offset, chunkState.Length - readerState.Offset);
 				}
-				else if (readerState.Offset + minLength <= chunkState.Length)
+				else
 				{
-					// We have enough data in the chunk to be able to read a message
-					return true;
+					return ReadOnlyMemory<byte>.Empty;
 				}
-				else if (chunkState.WriteState == WriteState.Writing)
+			}
+
+			/// <inheritdoc/>
+			public async ValueTask<bool> WaitToReadAsync(int minLength, CancellationToken cancellationToken = default)
+			{
+				for (; ; )
 				{
-					// Wait until there is more data in the chunk
-					_resources.ResetReadEvent(readerIdx);
-					if (_headerPtr.GetChunkStatePtr(readerState.ChunkIdx).Value == chunkState)
+					ReaderStatePtr readerStatePtr = _headerPtr.GetReaderStatePtr(_readerIdx);
+					ReaderState readerState = readerStatePtr.Value;
+
+					ChunkStatePtr chunkStatePtr = _headerPtr.GetChunkStatePtr(readerState.ChunkIdx);
+					ChunkState chunkState = chunkStatePtr.Value;
+
+					if (!chunkState.HasReaderFlag(_readerIdx))
 					{
-						await _resources.WaitForReadEvent(readerIdx, cancellationToken);
+						// Wait until the current chunk is readable
+						_resources.ResetReadEvent(_readerIdx);
+						if (!chunkState.HasReaderFlag(_readerIdx))
+						{
+							await _resources.WaitForReadEvent(_readerIdx, cancellationToken);
+						}
 					}
-				}
-				else if (readerState.Offset < chunkState.Length || chunkState.WriteState == WriteState.Complete)
-				{
-					// Cannot read the requested amount of data from this chunk.
-					return false;
-				}
-				else if (chunkState.WriteState == WriteState.MovedToNext)
-				{
-					// Move to the next chunk
-					chunkStatePtr.FinishReading(readerIdx);
-					_resources.SetWriteEvent();
-
-					int chunkIdx = readerStatePtr.ChunkIdx + 1;
-					if (chunkIdx == _chunks.Length)
+					else if (readerState.Offset + minLength <= chunkState.Length)
+					{
+						// We have enough data in the chunk to be able to read a message
+						return true;
+					}
+					else if (chunkState.WriteState == WriteState.Writing)
 					{
-						chunkIdx = 0;
+						// Wait until there is more data in the chunk
+						_resources.ResetReadEvent(_readerIdx);
+						if (_headerPtr.GetChunkStatePtr(readerState.ChunkIdx).Value == chunkState)
+						{
+							await _resources.WaitForReadEvent(_readerIdx, cancellationToken);
+						}
 					}
+					else if (readerState.Offset < chunkState.Length || chunkState.WriteState == WriteState.Complete)
+					{
+						// Cannot read the requested amount of data from this chunk.
+						return false;
+					}
+					else if (chunkState.WriteState == WriteState.MovedToNext)
+					{
+						// Move to the next chunk
+						chunkStatePtr.FinishReading(_readerIdx);
+						_resources.SetWriteEvent();
 
-					readerStatePtr.Value = new ReaderState(chunkIdx, 0);
-				}
-				else
-				{
-					throw new NotImplementedException($"Invalid write state for buffer: {chunkState.WriteState}");
+						int chunkIdx = readerStatePtr.ChunkIdx + 1;
+						if (chunkIdx == _chunks.Length)
+						{
+							chunkIdx = 0;
+						}
+
+						readerStatePtr.Value = new ReaderState(chunkIdx, 0);
+					}
+					else
+					{
+						throw new NotImplementedException($"Invalid write state for buffer: {chunkState.WriteState}");
+					}
 				}
 			}
 		}
 
-		#endregion
-
-		#region Writer Interface
-
-		/// <inheritdoc/>
-		public void AdvanceWritePosition(int size)
+		class WriterImpl : IComputeBufferWriter
 		{
-			if (size > 0)
-			{
-				ChunkStatePtr chunkStatePtr = _headerPtr.GetChunkStatePtr(_headerPtr.WriteChunkIdx);
-				ChunkState chunkState = chunkStatePtr.Value;
-
-				Debug.Assert(chunkState.WriteState == WriteState.Writing);
-				chunkStatePtr.Append(size);
-
-				SetAllReadEvents();
-			}
-		}
+			readonly HeaderPtr _headerPtr;
+			readonly Memory<byte>[] _chunks;
+			readonly ResourcesBase _resources;
 
-		/// <inheritdoc/>
-		public Memory<byte> GetWriteBuffer()
-		{
-			ChunkState state = _headerPtr.GetChunkStatePtr(_headerPtr.WriteChunkIdx).Value;
-			if (state.WriteState == WriteState.Writing)
+			public WriterImpl(ResourcesBase resources)
 			{
-				return _chunks[_headerPtr.WriteChunkIdx].Slice(state.Length);
+				_headerPtr = resources.HeaderPtr;
+				_chunks = resources.Chunks;
+				_resources = resources;
 			}
-			else
+
+			/// <inheritdoc/>
+			public IComputeBufferWriter AddRef()
 			{
-				return Memory<byte>.Empty;
+				_resources.AddRef();
+				return new WriterImpl(_resources);
 			}
-		}
 
-		/// <inheritdoc/>
-		public bool MarkComplete()
-		{
-			ChunkStatePtr chunkStatePtr = _headerPtr.GetChunkStatePtr(_headerPtr.WriteChunkIdx);
-			if (chunkStatePtr.WriteState != WriteState.Complete)
+			/// <inheritdoc/>
+			public void AdvanceWritePosition(int size)
 			{
-				chunkStatePtr.MarkComplete();
-				SetAllReadEvents();
-				return true;
+				if (size > 0)
+				{
+					ChunkStatePtr chunkStatePtr = _headerPtr.GetChunkStatePtr(_headerPtr.WriteChunkIdx);
+					ChunkState chunkState = chunkStatePtr.Value;
+
+					Debug.Assert(chunkState.WriteState == WriteState.Writing);
+					chunkStatePtr.Append(size);
+
+					_resources.SetAllReadEvents();
+				}
 			}
-			return false;
-		}
 
-		/// <inheritdoc/>
-		public async ValueTask WaitToWriteAsync(int minSize, CancellationToken cancellationToken = default)
-		{
-			if (minSize > _headerPtr.ChunkLength)
+			/// <inheritdoc/>
+			public Memory<byte> GetWriteBuffer()
 			{
-				throw new ArgumentException("Requested read size is larger than chunk size.", nameof(minSize));
+				ChunkState state = _headerPtr.GetChunkStatePtr(_headerPtr.WriteChunkIdx).Value;
+				if (state.WriteState == WriteState.Writing)
+				{
+					return _chunks[_headerPtr.WriteChunkIdx].Slice(state.Length);
+				}
+				else
+				{
+					return Memory<byte>.Empty;
+				}
 			}
 
-			for (; ; )
+			/// <inheritdoc/>
+			public bool MarkComplete()
 			{
 				ChunkStatePtr chunkStatePtr = _headerPtr.GetChunkStatePtr(_headerPtr.WriteChunkIdx);
-
-				ChunkState chunkState = chunkStatePtr.Value;
-				if (chunkState.WriteState == WriteState.Writing)
+				if (chunkStatePtr.WriteState != WriteState.Complete)
 				{
-					int length = chunkState.Length;
-					if (length + minSize <= _headerPtr.ChunkLength)
-					{
-						return;
-					}
-					chunkStatePtr.FinishWriting();
+					chunkStatePtr.MarkComplete();
+					_resources.SetAllReadEvents();
+					return true;
 				}
+				return false;
+			}
 
-				int nextChunkIdx = _headerPtr.WriteChunkIdx + 1;
-				if (nextChunkIdx == _chunks.Length)
+			/// <inheritdoc/>
+			public async ValueTask WaitToWriteAsync(int minSize, CancellationToken cancellationToken = default)
+			{
+				if (minSize > _headerPtr.ChunkLength)
 				{
-					nextChunkIdx = 0;
+					throw new ArgumentException("Requested read size is larger than chunk size.", nameof(minSize));
 				}
 
-				ChunkStatePtr nextChunkStatePtr = _headerPtr.GetChunkStatePtr(nextChunkIdx);
-				while (nextChunkStatePtr.ReaderFlags != 0)
+				for (; ; )
 				{
-					await _resources.WaitForWriteEvent(cancellationToken);
-					_resources.ResetWriteEvent();
+					ChunkStatePtr chunkStatePtr = _headerPtr.GetChunkStatePtr(_headerPtr.WriteChunkIdx);
+
+					ChunkState chunkState = chunkStatePtr.Value;
+					if (chunkState.WriteState == WriteState.Writing)
+					{
+						int length = chunkState.Length;
+						if (length + minSize <= _headerPtr.ChunkLength)
+						{
+							return;
+						}
+						chunkStatePtr.FinishWriting();
+					}
+
+					int nextChunkIdx = _headerPtr.WriteChunkIdx + 1;
+					if (nextChunkIdx == _chunks.Length)
+					{
+						nextChunkIdx = 0;
+					}
+
+					ChunkStatePtr nextChunkStatePtr = _headerPtr.GetChunkStatePtr(nextChunkIdx);
+					while (nextChunkStatePtr.ReaderFlags != 0)
+					{
+						await _resources.WaitForWriteEvent(cancellationToken);
+						_resources.ResetWriteEvent();
+					}
+
+					_headerPtr.WriteChunkIdx = nextChunkIdx;
+					nextChunkStatePtr.StartWriting(_headerPtr.NumReaders);
+
+					_resources.SetAllReadEvents();
 				}
+			}
+		}
+
+		internal ResourcesBase _resources;
+
+		/// <inheritdoc/>
+		public IComputeBufferReader Reader { get; }
+
+		/// <inheritdoc/>
+		public IComputeBufferWriter Writer { get; }
 
-				_headerPtr.WriteChunkIdx = nextChunkIdx;
-				nextChunkStatePtr.StartWriting(_headerPtr.NumReaders);
+		/// <summary>
+		/// Constructor
+		/// </summary>
+		/// <param name="resources">Resources shared between instances of the buffer</param>
+		internal ComputeBufferBase(ResourcesBase resources)
+		{
+			_resources = resources;
+			Reader = new ReaderImpl(resources, 0);
+			Writer = new WriterImpl(resources);
+		}
 
-				SetAllReadEvents();
+		/// <inheritdoc/>
+		public void Dispose()
+		{
+			Dispose(true);
+			GC.SuppressFinalize(this);
+		}
+
+		/// <summary>
+		/// Overriable dispose method
+		/// </summary>
+		protected virtual void Dispose(bool disposing)
+		{
+			if (_resources != null)
+			{
+				_resources.Release();
+				_resources = null!;
 			}
 		}
 
-		#endregion
+		/// <inheritdoc/>
+		public abstract IComputeBuffer AddRef();
 	}
 }
diff --git a/Engine/Source/Programs/Shared/EpicGames.Horde/Compute/IComputeBuffer.cs b/Engine/Source/Programs/Shared/EpicGames.Horde/Compute/IComputeBuffer.cs
--- a/Engine/Source/Programs/Shared/EpicGames.Horde/Compute/IComputeBuffer.cs
+++ b/Engine/Source/Programs/Shared/EpicGames.Horde/Compute/IComputeBuffer.cs
@@ -32,6 +32,11 @@ namespace EpicGames.Horde.Compute
 	/// </summary>
 	public interface IComputeBufferReader
 	{
+		/// <summary>
+		/// Create a new reader instance using the same underlying buffer
+		/// </summary>
+		IComputeBufferReader AddRef();
+
 		/// <summary>
 		/// Whether this buffer is complete (no more data will be added)
 		/// </summary>
@@ -63,6 +68,11 @@ namespace EpicGames.Horde.Compute
 	/// </summary>
 	public interface IComputeBufferWriter
 	{
+		/// <summary>
+		/// Create a new writer instance using the same underlying buffer
+		/// </summary>
+		IComputeBufferWriter AddRef();
+
 		/// <summary>
 		/// Mark the output to this buffer as complete
 		/// </summary>
