Horde: Add dispose methods to compute buffer reader/writer.

diff --git a/Engine/Source/Programs/Shared/EpicGames.Horde/Compute/ComputeBufferBase.cs b/Engine/Source/Programs/Shared/EpicGames.Horde/Compute/ComputeBufferBase.cs
--- a/Engine/Source/Programs/Shared/EpicGames.Horde/Compute/ComputeBufferBase.cs
+++ b/Engine/Source/Programs/Shared/EpicGames.Horde/Compute/ComputeBufferBase.cs
@@ -258,7 +258,7 @@ namespace EpicGames.Horde.Compute
 		{
 			readonly HeaderPtr _headerPtr;
 			readonly Memory<byte>[] _chunks;
-			readonly ResourcesBase _resources;
+			ResourcesBase _resources;
 			readonly int _readerIdx;
 
 			public ReaderImpl(ResourcesBase resources, int readerIdx)
@@ -276,6 +276,16 @@ namespace EpicGames.Horde.Compute
 				return new ReaderImpl(_resources, _readerIdx);
 			}
 
+			/// <inheritdoc/>
+			public void Dispose()
+			{
+				if (_resources != null)
+				{
+					_resources.Dispose();
+					_resources = null!;
+				}
+			}
+
 			/// <inheritdoc/>
 			public bool IsComplete
 			{
@@ -378,7 +388,7 @@ namespace EpicGames.Horde.Compute
 		{
 			readonly HeaderPtr _headerPtr;
 			readonly Memory<byte>[] _chunks;
-			readonly ResourcesBase _resources;
+			ResourcesBase _resources;
 
 			public WriterImpl(ResourcesBase resources)
 			{
@@ -394,6 +404,16 @@ namespace EpicGames.Horde.Compute
 				return new WriterImpl(_resources);
 			}
 
+			/// <inheritdoc/>
+			public void Dispose()
+			{
+				if (_resources != null)
+				{
+					_resources.Dispose();
+					_resources = null!;
+				}
+			}
+
 			/// <inheritdoc/>
 			public void AdvanceWritePosition(int size)
 			{
diff --git a/Engine/Source/Programs/Shared/EpicGames.Horde/Compute/IComputeBuffer.cs b/Engine/Source/Programs/Shared/EpicGames.Horde/Compute/IComputeBuffer.cs
--- a/Engine/Source/Programs/Shared/EpicGames.Horde/Compute/IComputeBuffer.cs
+++ b/Engine/Source/Programs/Shared/EpicGames.Horde/Compute/IComputeBuffer.cs
@@ -30,7 +30,7 @@ namespace EpicGames.Horde.Compute
 	/// <summary>
 	/// Read interface for a compute buffer
 	/// </summary>
-	public interface IComputeBufferReader
+	public interface IComputeBufferReader : IDisposable
 	{
 		/// <summary>
 		/// Create a new reader instance using the same underlying buffer
@@ -66,7 +66,7 @@ namespace EpicGames.Horde.Compute
 	/// <summary>
 	/// Buffer that can receive data from a remote machine.
 	/// </summary>
-	public interface IComputeBufferWriter
+	public interface IComputeBufferWriter : IDisposable
 	{
 		/// <summary>
 		/// Create a new writer instance using the same underlying buffer
