Horde: Use reference counting to extend compute buffer lifetimes in channels and sockets.

diff --git a/Engine/Source/Programs/Shared/EpicGames.Horde.Tests/BufferTests.cs b/Engine/Source/Programs/Shared/EpicGames.Horde.Tests/BufferTests.cs
--- a/Engine/Source/Programs/Shared/EpicGames.Horde.Tests/BufferTests.cs
+++ b/Engine/Source/Programs/Shared/EpicGames.Horde.Tests/BufferTests.cs
@@ -101,7 +101,7 @@ namespace EpicGames.Horde.Tests
 			await using IComputeSocket consumerSocket = ComputeSocket.Create(new PipeTransport(sourceToTargetPipe.Reader, targetToSourcePipe.Writer), NullLogger.Instance);
 
 			using IComputeBuffer consumerBuffer = createBuffer(Length);
-			consumerSocket.AttachRecvBuffer(ChannelId, consumerBuffer);
+			consumerSocket.AttachRecvBuffer(ChannelId, consumerBuffer.Writer);
 
 			byte[] input = RandomNumberGenerator.GetBytes(Length);
 			Task producerTask = RunProducerAsync(producerSocket, input);
diff --git a/Engine/Source/Programs/Shared/EpicGames.Horde/Compute/ComputeMessageChannel.cs b/Engine/Source/Programs/Shared/EpicGames.Horde/Compute/ComputeMessageChannel.cs
--- a/Engine/Source/Programs/Shared/EpicGames.Horde/Compute/ComputeMessageChannel.cs
+++ b/Engine/Source/Programs/Shared/EpicGames.Horde/Compute/ComputeMessageChannel.cs
@@ -110,8 +110,8 @@ namespace EpicGames.Horde.Compute
 
 		readonly IComputeSocket _socket;
 		readonly int _channelId;
-		readonly IComputeBuffer _recvBuffer;
-		readonly IComputeBuffer _sendBuffer;
+		readonly IComputeBufferReader _recvBufferReader;
+		readonly IComputeBufferWriter _sendBufferWriter;
 
 		// Can lock chunked memory writer to acuqire pointer
 		readonly ILogger _logger;
@@ -124,15 +124,15 @@ namespace EpicGames.Horde.Compute
 		/// Constructor
 		/// </summary>
 		/// <param name="channelId"></param>
-		/// <param name="recvBuffer"></param>
-		/// <param name="sendBuffer"></param>
+		/// <param name="recvBufferReader"></param>
+		/// <param name="sendBufferWriter"></param>
 		/// <param name="logger">Logger for diagnostic output</param>
-		public ComputeMessageChannel(int channelId, IComputeBuffer recvBuffer, IComputeBuffer sendBuffer, ILogger logger)
+		public ComputeMessageChannel(int channelId, IComputeBufferReader recvBufferReader, IComputeBufferWriter sendBufferWriter, ILogger logger)
 		{
 			_socket = null!;
 			_channelId = channelId;
-			_recvBuffer = recvBuffer;
-			_sendBuffer = sendBuffer;
+			_recvBufferReader = recvBufferReader.AddRef();
+			_sendBufferWriter = sendBufferWriter.AddRef();
 			_logger = logger;
 		}
 
@@ -148,10 +148,10 @@ namespace EpicGames.Horde.Compute
 		{
 			_socket = socket;
 			_channelId = channelId;
-			_socket.AttachRecvBuffer(channelId, recvBuffer);
-			_socket.AttachSendBuffer(channelId, sendBuffer);
-			_recvBuffer = recvBuffer;
-			_sendBuffer = sendBuffer;
+			_socket.AttachRecvBuffer(channelId, recvBuffer.Writer);
+			_socket.AttachSendBuffer(channelId, sendBuffer.Reader);
+			_recvBufferReader = recvBuffer.Reader.AddRef();
+			_sendBufferWriter = sendBuffer.Writer.AddRef();
 			_logger = logger;
 		}
 
@@ -161,9 +161,11 @@ namespace EpicGames.Horde.Compute
 		public void Dispose()
 		{
 			_currentBuilder?.Dispose();
-			_sendBuffer.Writer.MarkComplete();
-			_sendBuffer.Dispose();
-			_recvBuffer.Dispose();
+
+			_sendBufferWriter.MarkComplete();
+			_sendBufferWriter.Dispose();
+
+			_recvBufferReader.Dispose();
 		}
 
 		/// <summary>
@@ -171,7 +173,7 @@ namespace EpicGames.Horde.Compute
 		/// </summary>
 		public void MarkComplete()
 		{
-			_sendBuffer.Writer.MarkComplete();
+			_sendBufferWriter.MarkComplete();
 		}
 
 		/// <inheritdoc/>
@@ -180,20 +182,19 @@ namespace EpicGames.Horde.Compute
 		/// <inheritdoc/>
 		public async ValueTask<IComputeMessage> ReceiveAsync(CancellationToken cancellationToken)
 		{
-			IComputeBufferReader recvBufferReader = _recvBuffer.Reader;
-			while (!recvBufferReader.IsComplete)
+			while (!_recvBufferReader.IsComplete)
 			{
-				ReadOnlyMemory<byte> memory = recvBufferReader.GetReadBuffer();
+				ReadOnlyMemory<byte> memory = _recvBufferReader.GetReadBuffer();
 				if (memory.Length < HeaderLength)
 				{
-					await recvBufferReader.WaitToReadAsync(HeaderLength, cancellationToken);
+					await _recvBufferReader.WaitToReadAsync(HeaderLength, cancellationToken);
 					continue;
 				}
 
 				int messageLength = BinaryPrimitives.ReadInt32LittleEndian(memory.Span.Slice(1, 4));
 				if (memory.Length < HeaderLength + messageLength)
 				{
-					await recvBufferReader.WaitToReadAsync(HeaderLength + messageLength, cancellationToken);
+					await _recvBufferReader.WaitToReadAsync(HeaderLength + messageLength, cancellationToken);
 					continue;
 				}
 
@@ -203,7 +204,8 @@ namespace EpicGames.Horde.Compute
 				{
 					LogMessageInfo("RECV", message.Type, message.Data.Span);
 				}
-				recvBufferReader.AdvanceReadPosition(HeaderLength + messageLength);
+
+				_recvBufferReader.AdvanceReadPosition(HeaderLength + messageLength);
 				return message;
 			}
 			return new Message(ComputeMessageType.None, ReadOnlyMemory<byte>.Empty);
@@ -231,10 +233,9 @@ namespace EpicGames.Horde.Compute
 				throw new InvalidOperationException("Only one writer can be active at a time. Dispose of the previous writer first.");
 			}
 
-			IComputeBufferWriter sendBufferWriter = _sendBuffer.Writer;
-			await sendBufferWriter.WaitToWriteAsync(maxSize, cancellationToken);
+			await _sendBufferWriter.WaitToWriteAsync(maxSize, cancellationToken);
 
-			_currentBuilder = new MessageBuilder(this, sendBufferWriter, type);
+			_currentBuilder = new MessageBuilder(this, _sendBufferWriter, type);
 			return _currentBuilder;
 		}
 	}
diff --git a/Engine/Source/Programs/Shared/EpicGames.Horde/Compute/ComputeWorker.cs b/Engine/Source/Programs/Shared/EpicGames.Horde/Compute/ComputeWorker.cs
--- a/Engine/Source/Programs/Shared/EpicGames.Horde/Compute/ComputeWorker.cs
+++ b/Engine/Source/Programs/Shared/EpicGames.Horde/Compute/ComputeWorker.cs
@@ -204,18 +204,22 @@ namespace EpicGames.Horde.Compute
 						{
 							int channelId = (int)reader.ReadUnsignedVarInt();
 							string name = reader.ReadString();
-#pragma warning disable CA2000 // Dispose objects before losing scope
-							socket.AttachSendBuffer(channelId, SharedMemoryBuffer.OpenExisting(name));
-#pragma warning restore CA2000 // Dispose objects before losing scope
+
+							using (SharedMemoryBuffer buffer = SharedMemoryBuffer.OpenExisting(name))
+							{
+								socket.AttachSendBuffer(channelId, buffer.Reader);
+							}
 						}
 						break;
 					case IpcMessage.AttachRecvBuffer:
 						{
 							int channelId = (int)reader.ReadUnsignedVarInt();
 							string name = reader.ReadString();
-#pragma warning disable CA2000 // Dispose objects before losing scope
-							socket.AttachRecvBuffer(channelId, SharedMemoryBuffer.OpenExisting(name));
-#pragma warning restore CA2000 // Dispose objects before losing scope
+
+							using (SharedMemoryBuffer buffer = SharedMemoryBuffer.OpenExisting(name))
+							{
+								socket.AttachRecvBuffer(channelId, buffer.Writer);
+							}
 						}
 						break;
 					default:
diff --git a/Engine/Source/Programs/Shared/EpicGames.Horde/Compute/DefaultComputeSocket.cs b/Engine/Source/Programs/Shared/EpicGames.Horde/Compute/DefaultComputeSocket.cs
--- a/Engine/Source/Programs/Shared/EpicGames.Horde/Compute/DefaultComputeSocket.cs
+++ b/Engine/Source/Programs/Shared/EpicGames.Horde/Compute/DefaultComputeSocket.cs
@@ -34,11 +34,11 @@ namespace EpicGames.Horde.Compute
 		readonly CancellationTokenSource _cancellationSource = new CancellationTokenSource();
 
 		readonly BackgroundTask _recvTask;
-		readonly Dictionary<int, IComputeBuffer> _recvBuffers = new Dictionary<int, IComputeBuffer>();
+		readonly Dictionary<int, IComputeBufferWriter> _recvBufferWriters = new Dictionary<int, IComputeBufferWriter>();
 
 		readonly SemaphoreSlim _sendSemaphore = new SemaphoreSlim(1, 1);
 		readonly Dictionary<int, Task> _sendTasks = new Dictionary<int, Task>();
-		readonly Dictionary<IComputeBuffer, int> _sendBuffers = new Dictionary<IComputeBuffer, int>();
+		readonly Dictionary<IComputeBufferReader, int> _sendBufferReaders = new Dictionary<IComputeBufferReader, int>();
 
 		readonly HashSet<int> _remoteRecvBuffers = new HashSet<int>();
 		readonly AsyncEvent _remoteBuffersChanged = new AsyncEvent();
@@ -69,14 +69,7 @@ namespace EpicGames.Horde.Compute
 			lock (_lockObject)
 			{
 				sendTasks = _sendTasks.Values.ToArray();
-				foreach ((IComputeBuffer sendBuffer, int channelId) in _sendBuffers)
-				{
-					if (sendBuffer.Writer.MarkComplete())
-					{
-						_logger.LogWarning("Buffer attached to channel {ChannelId} was not marked as complete before socket was closed.", channelId);
-					}
-				}
-				_sendBuffers.Clear();
+				_sendBufferReaders.Clear();
 			}
 
 			// Wait for all the individual send tasks to complete
@@ -165,7 +158,7 @@ namespace EpicGames.Horde.Compute
 			lock (_lockObject)
 			{
 				_complete = true;
-				foreach (IComputeBuffer buffer in _recvBuffers.Values)
+				foreach (IComputeBuffer buffer in _recvBufferWriters.Values)
 				{
 					buffer.Writer.MarkComplete();
 				}
@@ -203,10 +196,10 @@ namespace EpicGames.Horde.Compute
 
 			lock (_lockObject)
 			{
-				IComputeBuffer? recvBuffer;
-				if (_recvBuffers.TryGetValue(id, out recvBuffer))
+				IComputeBufferWriter? recvBufferWriter;
+				if (_recvBufferWriters.TryGetValue(id, out recvBufferWriter))
 				{
-					writer = recvBuffer.Writer;
+					writer = recvBufferWriter;
 					cachedWriters.Add(id, writer);
 					return writer;
 				}
@@ -278,24 +271,21 @@ namespace EpicGames.Horde.Compute
 		}
 
 		/// <inheritdoc/>
-		public async void AttachRecvBuffer(int channelId, IComputeBuffer recvBuffer)
+		public async void AttachRecvBuffer(int channelId, IComputeBufferWriter recvBufferWriter)
 		{
-			IComputeBuffer shared = recvBuffer.AddRef();
-
 			bool complete;
 			lock (_lockObject)
 			{
 				complete = _complete;
 				if (!complete)
 				{
-					_recvBuffers.Add(channelId, shared);
+					_recvBufferWriters.Add(channelId, recvBufferWriter.AddRef());
 				}
 			}
 
-			if (shared != null && complete)
+			if (recvBufferWriter != null && complete)
 			{
-				shared.Writer.MarkComplete();
-				shared.Dispose();
+				recvBufferWriter.MarkComplete();
 			}
 
 			if (!complete)
@@ -329,36 +319,35 @@ namespace EpicGames.Horde.Compute
 		{
 			cachedWriters.Remove(id);
 
-			IComputeBuffer? buffer;
+			IComputeBufferWriter? recvBufferWriter;
 			lock (_lockObject)
 			{
 #pragma warning disable CA2000 // Dispose objects before losing scope
-				_recvBuffers.Remove(id, out buffer);
+				_recvBufferWriters.Remove(id, out recvBufferWriter);
 #pragma warning restore CA2000 // Dispose objects before losing scope
 			}
-			if (buffer != null)
+			if (recvBufferWriter != null)
 			{
-				buffer.Writer.MarkComplete();
-				buffer.Dispose();
+				recvBufferWriter.MarkComplete();
+				recvBufferWriter.Dispose();
 			}
 		}
 
 		/// <inheritdoc/>
-		public void AttachSendBuffer(int channelId, IComputeBuffer buffer)
+		public void AttachSendBuffer(int channelId, IComputeBufferReader sendBufferReader)
 		{
-			IComputeBuffer shared = buffer.AddRef();
+			sendBufferReader = sendBufferReader.AddRef();
 			lock (_lockObject)
 			{
-				_sendTasks.Add(channelId, Task.Run(() => SendFromBufferAsync(channelId, shared, _cancellationSource.Token), _cancellationSource.Token));
-				_sendBuffers.Add(shared, channelId);
+				_sendTasks.Add(channelId, Task.Run(() => SendFromBufferAsync(channelId, sendBufferReader, _cancellationSource.Token), CancellationToken.None)); // No cancellation token; need to ensure dispose runs
+				_sendBufferReaders.Add(sendBufferReader, channelId);
 			}
 		}
 
-		async Task SendFromBufferAsync(int channelId, IComputeBuffer buffer, CancellationToken cancellationToken)
+		async Task SendFromBufferAsync(int channelId, IComputeBufferReader reader, CancellationToken cancellationToken)
 		{
-			using IComputeBuffer _ = buffer;
+			using IComputeBufferReader _ = reader;
 
-			IComputeBufferReader reader = buffer.Reader;
 			while (!cancellationToken.IsCancellationRequested)
 			{
 				ReadOnlyMemory<byte> memory = reader.GetReadBuffer();
@@ -378,7 +367,7 @@ namespace EpicGames.Horde.Compute
 			lock (_lockObject)
 			{
 				_sendTasks.Remove(channelId);
-				_sendBuffers.Remove(buffer);
+				_sendBufferReaders.Remove(reader);
 			}
 		}
 	}
diff --git a/Engine/Source/Programs/Shared/EpicGames.Horde/Compute/IComputeChannel.cs b/Engine/Source/Programs/Shared/EpicGames.Horde/Compute/IComputeChannel.cs
--- a/Engine/Source/Programs/Shared/EpicGames.Horde/Compute/IComputeChannel.cs
+++ b/Engine/Source/Programs/Shared/EpicGames.Horde/Compute/IComputeChannel.cs
@@ -42,23 +42,23 @@ namespace EpicGames.Horde.Compute
 		{
 			readonly IComputeSocket _socket;
 			readonly int _channelId;
-			readonly IComputeBuffer _recvBuffer;
+			readonly IComputeBufferReader _recvBufferReader;
 
 			public BufferedReaderChannel(IComputeSocket socket, int channelId, IComputeBuffer recvBuffer)
 			{
 				_socket = socket;
 				_channelId = channelId;
-				_recvBuffer = recvBuffer;
+				_recvBufferReader = recvBuffer.Reader.AddRef();
 
-				_socket.AttachRecvBuffer(_channelId, recvBuffer);
+				_socket.AttachRecvBuffer(_channelId, recvBuffer.Writer);
 			}
 
 			public void Dispose()
 			{
-				_recvBuffer.Dispose();
+				_recvBufferReader.Dispose();
 			}
 
-			public ValueTask<int> ReceiveAsync(Memory<byte> buffer, CancellationToken cancellationToken = default) => _recvBuffer.Reader.ReadAsync(buffer, cancellationToken);
+			public ValueTask<int> ReceiveAsync(Memory<byte> buffer, CancellationToken cancellationToken = default) => _recvBufferReader.ReadAsync(buffer, cancellationToken);
 
 			public ValueTask SendAsync(ReadOnlyMemory<byte> memory, CancellationToken cancellationToken = default) => _socket.SendAsync(_channelId, memory, cancellationToken);
 		}
@@ -137,8 +137,8 @@ namespace EpicGames.Horde.Compute
 		/// <param name="sendBuffer">Buffer for sending data</param>
 		public static IComputeChannel CreateChannel(this IComputeSocket socket, int channelId, IComputeBuffer recvBuffer, IComputeBuffer sendBuffer)
 		{
-			socket.AttachRecvBuffer(channelId, recvBuffer);
-			socket.AttachSendBuffer(channelId, sendBuffer);
+			socket.AttachRecvBuffer(channelId, recvBuffer.Writer);
+			socket.AttachSendBuffer(channelId, sendBuffer.Reader);
 			return new BufferedReaderWriterChannel(recvBuffer, sendBuffer);
 		}
 	}
diff --git a/Engine/Source/Programs/Shared/EpicGames.Horde/Compute/IComputeSocket.cs b/Engine/Source/Programs/Shared/EpicGames.Horde/Compute/IComputeSocket.cs
--- a/Engine/Source/Programs/Shared/EpicGames.Horde/Compute/IComputeSocket.cs
+++ b/Engine/Source/Programs/Shared/EpicGames.Horde/Compute/IComputeSocket.cs
@@ -32,15 +32,15 @@ namespace EpicGames.Horde.Compute
 		/// Attaches a buffer to receive data.
 		/// </summary>
 		/// <param name="channelId">Channel to receive data on</param>
-		/// <param name="buffer">Buffer to receive data</param>
-		void AttachRecvBuffer(int channelId, IComputeBuffer buffer);
+		/// <param name="recvBufferWriter">Writer for the buffer to store received data</param>
+		void AttachRecvBuffer(int channelId, IComputeBufferWriter recvBufferWriter);
 
 		/// <summary>
 		/// Attaches a buffer to send data.
 		/// </summary>
 		/// <param name="channelId">Channel to receive data on</param>
-		/// <param name="buffer">Buffer to send data from</param>
-		void AttachSendBuffer(int channelId, IComputeBuffer buffer);
+		/// <param name="sendBufferReader">Reader for the buffer to send data from</param>
+		void AttachSendBuffer(int channelId, IComputeBufferReader sendBufferReader);
 
 		/// <summary>
 		/// Waits until the remote has attached a receive buffer to the given channel
@@ -103,8 +103,8 @@ namespace EpicGames.Horde.Compute
 		/// <param name="logger">Logger for the channel</param>
 		public static IComputeMessageChannel CreateMessageChannel(this IComputeSocket socket, int channelId, int sendBufferSize, int recvBufferSize, ILogger logger)
 		{
-			IComputeBuffer sendBuffer = new PooledBuffer(sendBufferSize);
-			IComputeBuffer recvBuffer = new PooledBuffer(recvBufferSize);
+			using IComputeBuffer sendBuffer = new PooledBuffer(sendBufferSize);
+			using IComputeBuffer recvBuffer = new PooledBuffer(recvBufferSize);
 			return new ComputeMessageChannel(socket, channelId, sendBuffer, recvBuffer, logger);
 		}
 	}
