Horde: Remove methods from IComputeSocket which are only implemented by the transport-based implementation. Expose implementations of IComputeSocket to consumers.

diff --git a/Engine/Source/Programs/Horde/Horde.Agent/Commands/Compute/ComputeCommand.cs b/Engine/Source/Programs/Horde/Horde.Agent/Commands/Compute/ComputeCommand.cs
--- a/Engine/Source/Programs/Horde/Horde.Agent/Commands/Compute/ComputeCommand.cs
+++ b/Engine/Source/Programs/Horde/Horde.Agent/Commands/Compute/ComputeCommand.cs
@@ -127,7 +127,7 @@ namespace Horde.Agent.Commands.Compute
 			using Socket tcpSocket = new Socket(SocketType.Stream, ProtocolType.IP);
 			await tcpSocket.ConnectAsync(IPAddress.Loopback, Port);
 
-			await using (IComputeSocket socket = ComputeSocket.Create(new TcpTransport(tcpSocket), logger))
+			await using (IComputeSocket socket = new RemoteComputeSocket(new TcpTransport(tcpSocket), logger))
 			{
 				logger.LogInformation("Running worker...");
 				await RunWorkerAsync(socket, _memoryCache, logger, CancellationToken.None);
diff --git a/Engine/Source/Programs/Horde/Horde.Agent/Leases/Handlers/ComputeHandler.cs b/Engine/Source/Programs/Horde/Horde.Agent/Leases/Handlers/ComputeHandler.cs
--- a/Engine/Source/Programs/Horde/Horde.Agent/Leases/Handlers/ComputeHandler.cs
+++ b/Engine/Source/Programs/Horde/Horde.Agent/Leases/Handlers/ComputeHandler.cs
@@ -92,7 +92,7 @@ namespace Horde.Agent.Leases.Handlers
 				using (CancellationTokenSource cts = CancellationTokenSource.CreateLinkedTokenSource(cancellationToken))
 				{
 					await using BackgroundTask timeoutTask = BackgroundTask.StartNew(ctx => TickTimeoutAsync(transport, cts, ctx));
-					await using (IComputeSocket socket = ComputeSocket.Create(transport, _logger))
+					await using (IComputeSocket socket = new RemoteComputeSocket(transport, _logger))
 					{
 						DirectoryReference sandboxDir = DirectoryReference.Combine(session.WorkingDir, "Sandbox", leaseId);
 						try
diff --git a/Engine/Source/Programs/Horde/Samples/RemoteClient/RemoteClientApp.cs b/Engine/Source/Programs/Horde/Samples/RemoteClient/RemoteClientApp.cs
--- a/Engine/Source/Programs/Horde/Samples/RemoteClient/RemoteClientApp.cs
+++ b/Engine/Source/Programs/Horde/Samples/RemoteClient/RemoteClientApp.cs
@@ -117,7 +117,7 @@ namespace RemoteClient
 			await lease.CloseAsync();
 		}
 		
-		static async Task WriteNumbersAsync(IComputeSocket socket, ILogger logger, CancellationToken cancellationToken)
+		static async Task WriteNumbersAsync(RemoteComputeSocket socket, ILogger logger, CancellationToken cancellationToken)
 		{
 			// Write data to the child process channel. The remote server will echo them back to us as it receives them, then exit when the channel is complete/closed.
 			await socket.WaitForAttachAsync(ChildProcessChannelId, cancellationToken);
diff --git a/Engine/Source/Programs/Shared/EpicGames.Horde.Tests/BufferTests.cs b/Engine/Source/Programs/Shared/EpicGames.Horde.Tests/BufferTests.cs
--- a/Engine/Source/Programs/Shared/EpicGames.Horde.Tests/BufferTests.cs
+++ b/Engine/Source/Programs/Shared/EpicGames.Horde.Tests/BufferTests.cs
@@ -97,8 +97,8 @@ namespace EpicGames.Horde.Tests
 			Pipe sourceToTargetPipe = new Pipe();
 			Pipe targetToSourcePipe = new Pipe();
 
-			await using IComputeSocket producerSocket = ComputeSocket.Create(new PipeTransport(targetToSourcePipe.Reader, sourceToTargetPipe.Writer), NullLogger.Instance);
-			await using IComputeSocket consumerSocket = ComputeSocket.Create(new PipeTransport(sourceToTargetPipe.Reader, targetToSourcePipe.Writer), NullLogger.Instance);
+			await using RemoteComputeSocket producerSocket = new RemoteComputeSocket(new PipeTransport(targetToSourcePipe.Reader, sourceToTargetPipe.Writer), NullLogger.Instance);
+			await using RemoteComputeSocket consumerSocket = new RemoteComputeSocket(new PipeTransport(sourceToTargetPipe.Reader, targetToSourcePipe.Writer), NullLogger.Instance);
 
 			using IComputeBuffer consumerBuffer = createBuffer(Length);
 			consumerSocket.AttachRecvBuffer(ChannelId, consumerBuffer.Writer);
@@ -113,7 +113,7 @@ namespace EpicGames.Horde.Tests
 			Assert.IsTrue(input.SequenceEqual(output));
 		}
 
-		static async Task RunProducerAsync(IComputeSocket socket, ReadOnlyMemory<byte> input)
+		static async Task RunProducerAsync(RemoteComputeSocket socket, ReadOnlyMemory<byte> input)
 		{
 			int offset = 0;
 			while (offset < input.Length)
diff --git a/Engine/Source/Programs/Shared/EpicGames.Horde/Compute/Clients/AgentComputeClient.cs b/Engine/Source/Programs/Shared/EpicGames.Horde/Compute/Clients/AgentComputeClient.cs
--- a/Engine/Source/Programs/Shared/EpicGames.Horde/Compute/Clients/AgentComputeClient.cs
+++ b/Engine/Source/Programs/Shared/EpicGames.Horde/Compute/Clients/AgentComputeClient.cs
@@ -22,7 +22,7 @@ namespace EpicGames.Horde.Compute.Clients
 	{
 		class LeaseImpl : IComputeLease
 		{
-			readonly IAsyncEnumerator<IComputeSocket> _source;
+			readonly IAsyncEnumerator<RemoteComputeSocket> _source;
 
 			/// <inheritdoc/>
 			public IReadOnlyList<string> Properties { get; } = new List<string>();
@@ -31,9 +31,9 @@ namespace EpicGames.Horde.Compute.Clients
 			public IReadOnlyDictionary<string, int> AssignedResources => new Dictionary<string, int>();
 
 			/// <inheritdoc/>
-			public IComputeSocket Socket => _source.Current;
+			public RemoteComputeSocket Socket => _source.Current;
 
-			public LeaseImpl(IAsyncEnumerator<IComputeSocket> source) => _source = source;
+			public LeaseImpl(IAsyncEnumerator<RemoteComputeSocket> source) => _source = source;
 
 			/// <inheritdoc/>
 			public async ValueTask DisposeAsync()
@@ -73,7 +73,7 @@ namespace EpicGames.Horde.Compute.Clients
 			_logger.LogInformation("Launching {Path} to handle remote", _hordeAgentAssembly);
 
 			// The connection logic is an async enumerator that returns the socket, then shuts down.
-			IAsyncEnumerator<IComputeSocket> source = ConnectAsync(cancellationToken).GetAsyncEnumerator(cancellationToken);
+			IAsyncEnumerator<RemoteComputeSocket> source = ConnectAsync(cancellationToken).GetAsyncEnumerator(cancellationToken);
 			if (!await source.MoveNextAsync())
 			{
 				await source.DisposeAsync();
@@ -83,7 +83,7 @@ namespace EpicGames.Horde.Compute.Clients
 			return new LeaseImpl(source);
 		}
 
-		async IAsyncEnumerable<IComputeSocket> ConnectAsync([EnumeratorCancellation] CancellationToken cancellationToken)
+		async IAsyncEnumerable<RemoteComputeSocket> ConnectAsync([EnumeratorCancellation] CancellationToken cancellationToken)
 		{
 			using Socket listener = new Socket(SocketType.Stream, ProtocolType.IP);
 			listener.Bind(new IPEndPoint(IPAddress.Loopback, _port));
@@ -92,7 +92,7 @@ namespace EpicGames.Horde.Compute.Clients
 			using BackgroundTask agentTask = BackgroundTask.StartNew(ctx => RunAgentAsync(_hordeAgentAssembly, _port, _logger, ctx));
 			using Socket tcpSocket = await listener.AcceptAsync(cancellationToken);
 
-			await using IComputeSocket socket = ComputeSocket.Create(new TcpTransport(tcpSocket), _logger);
+			await using RemoteComputeSocket socket = new RemoteComputeSocket(new TcpTransport(tcpSocket), _logger);
 			yield return socket;
 
 			await socket.CloseAsync(cancellationToken);
diff --git a/Engine/Source/Programs/Shared/EpicGames.Horde/Compute/Clients/LocalComputeClient.cs b/Engine/Source/Programs/Shared/EpicGames.Horde/Compute/Clients/LocalComputeClient.cs
--- a/Engine/Source/Programs/Shared/EpicGames.Horde/Compute/Clients/LocalComputeClient.cs
+++ b/Engine/Source/Programs/Shared/EpicGames.Horde/Compute/Clients/LocalComputeClient.cs
@@ -22,11 +22,11 @@ namespace EpicGames.Horde.Compute.Clients
 		{
 			public IReadOnlyList<string> Properties { get; } = new List<string>();
 			public IReadOnlyDictionary<string, int> AssignedResources => new Dictionary<string, int>();
-			public IComputeSocket Socket => _socket;
+			public RemoteComputeSocket Socket => _socket;
 
-			readonly IComputeSocket _socket;
+			readonly RemoteComputeSocket _socket;
 
-			public LeaseImpl(IComputeSocket socket) => _socket = socket;
+			public LeaseImpl(RemoteComputeSocket socket) => _socket = socket;
 
 			/// <inheritdoc/>
 			public ValueTask DisposeAsync() => _socket.DisposeAsync();
@@ -77,7 +77,7 @@ namespace EpicGames.Horde.Compute.Clients
 
 			using MemoryCache memoryCache = new MemoryCache(new MemoryCacheOptions { SizeLimit = 10 * 1024 * 1024 });
 
-			await using (IComputeSocket socket = ComputeSocket.Create(new TcpTransport(tcpSocket), logger))
+			await using (IComputeSocket socket = new RemoteComputeSocket(new TcpTransport(tcpSocket), logger))
 			{
 				ComputeWorker worker = new ComputeWorker(sandboxDir, memoryCache, logger);
 				await worker.RunAsync(socket, cancellationToken);
@@ -89,7 +89,7 @@ namespace EpicGames.Horde.Compute.Clients
 		public Task<IComputeLease?> TryAssignWorkerAsync(ClusterId clusterId, Requirements? requirements, CancellationToken cancellationToken)
 		{
 #pragma warning disable CA2000 // Dispose objects before losing scope
-			IComputeSocket socket = ComputeSocket.Create(new TcpTransport(_socket), _logger);
+			RemoteComputeSocket socket = new RemoteComputeSocket(new TcpTransport(_socket), _logger);
 			return Task.FromResult<IComputeLease?>(new LeaseImpl(socket));
 #pragma warning restore CA2000 // Dispose objects before losing scope
 		}
diff --git a/Engine/Source/Programs/Shared/EpicGames.Horde/Compute/Clients/ServerComputeClient.cs b/Engine/Source/Programs/Shared/EpicGames.Horde/Compute/Clients/ServerComputeClient.cs
--- a/Engine/Source/Programs/Shared/EpicGames.Horde/Compute/Clients/ServerComputeClient.cs
+++ b/Engine/Source/Programs/Shared/EpicGames.Horde/Compute/Clients/ServerComputeClient.cs
@@ -42,15 +42,15 @@ namespace EpicGames.Horde.Compute.Clients
 			public List<string> Properties { get; set; } = new List<string>();
 		}
 
-		record class LeaseInfo(IReadOnlyList<string> Properties, IReadOnlyDictionary<string, int> AssignedResources, IComputeSocket Socket);
+		record class LeaseInfo(IReadOnlyList<string> Properties, IReadOnlyDictionary<string, int> AssignedResources, RemoteComputeSocket Socket);
 
 		class LeaseImpl : IComputeLease
-			{
+		{
 			readonly IAsyncEnumerator<LeaseInfo> _source;
 
 			public IReadOnlyList<string> Properties => _source.Current.Properties;
 			public IReadOnlyDictionary<string, int> AssignedResources => _source.Current.AssignedResources;
-			public IComputeSocket Socket => _source.Current.Socket;
+			public RemoteComputeSocket Socket => _source.Current.Socket;
 
 			public LeaseImpl(IAsyncEnumerator<LeaseInfo> source)
 			{
@@ -62,11 +62,11 @@ namespace EpicGames.Horde.Compute.Clients
 			{
 				await _source.MoveNextAsync();
 				await _source.DisposeAsync();
-		}
+			}
 
 			/// <inheritdoc/>
 			public ValueTask CloseAsync(CancellationToken cancellationToken) => Socket.CloseAsync(cancellationToken);
-			}
+		}
 
 		readonly HttpClient? _defaultHttpClient;
 		readonly Func<HttpClient> _createHttpClient;
@@ -181,7 +181,7 @@ namespace EpicGames.Horde.Compute.Clients
 			// Pass the rest of the call over to the handler
 			byte[] key = StringUtils.ParseHexString(responseMessage.Key);
 
-			await using IComputeSocket computeSocket = ComputeSocket.Create(new TcpTransport(socket), _logger);
+			await using RemoteComputeSocket computeSocket = new RemoteComputeSocket(new TcpTransport(socket), _logger);
 			yield return new LeaseInfo(responseMessage.Properties, responseMessage.AssignedResources, computeSocket);
 		}
 	}
diff --git a/Engine/Source/Programs/Shared/EpicGames.Horde/Compute/DefaultComputeSocket.cs b/Engine/Source/Programs/Shared/EpicGames.Horde/Compute/DefaultComputeSocket.cs
deleted file mode 100644
diff --git a/Engine/Source/Programs/Shared/EpicGames.Horde/Compute/IComputeChannel.cs b/Engine/Source/Programs/Shared/EpicGames.Horde/Compute/IComputeChannel.cs
--- a/Engine/Source/Programs/Shared/EpicGames.Horde/Compute/IComputeChannel.cs
+++ b/Engine/Source/Programs/Shared/EpicGames.Horde/Compute/IComputeChannel.cs
@@ -40,11 +40,11 @@ namespace EpicGames.Horde.Compute
 
 		class BufferedReaderChannel : IComputeChannel
 		{
-			readonly IComputeSocket _socket;
+			readonly RemoteComputeSocket _socket;
 			readonly int _channelId;
 			readonly IComputeBufferReader _recvBufferReader;
 
-			public BufferedReaderChannel(IComputeSocket socket, int channelId, IComputeBuffer recvBuffer)
+			public BufferedReaderChannel(RemoteComputeSocket socket, int channelId, IComputeBuffer recvBuffer)
 			{
 				_socket = socket;
 				_channelId = channelId;
@@ -126,7 +126,7 @@ namespace EpicGames.Horde.Compute
 		/// <param name="socket">Socket to use for sending data</param>
 		/// <param name="channelId">Channel id to send and receive data</param>
 		/// <param name="recvBuffer">Buffer for receiving data</param>
-		public static IComputeChannel CreateChannel(this IComputeSocket socket, int channelId, IComputeBuffer recvBuffer) => new BufferedReaderChannel(socket, channelId, recvBuffer);
+		public static IComputeChannel CreateChannel(this RemoteComputeSocket socket, int channelId, IComputeBuffer recvBuffer) => new BufferedReaderChannel(socket, channelId, recvBuffer);
 
 		/// <summary>
 		/// Creates a channel using a socket and receive buffer
diff --git a/Engine/Source/Programs/Shared/EpicGames.Horde/Compute/IComputeLease.cs b/Engine/Source/Programs/Shared/EpicGames.Horde/Compute/IComputeLease.cs
--- a/Engine/Source/Programs/Shared/EpicGames.Horde/Compute/IComputeLease.cs
+++ b/Engine/Source/Programs/Shared/EpicGames.Horde/Compute/IComputeLease.cs
@@ -25,7 +25,7 @@ namespace EpicGames.Horde.Compute
 		/// <summary>
 		/// Socket to communicate with the remote
 		/// </summary>
-		IComputeSocket Socket { get; }
+		RemoteComputeSocket Socket { get; }
 
 		/// <summary>
 		/// Relinquish the lease gracefully
diff --git a/Engine/Source/Programs/Shared/EpicGames.Horde/Compute/IComputeSocket.cs b/Engine/Source/Programs/Shared/EpicGames.Horde/Compute/IComputeSocket.cs
--- a/Engine/Source/Programs/Shared/EpicGames.Horde/Compute/IComputeSocket.cs
+++ b/Engine/Source/Programs/Shared/EpicGames.Horde/Compute/IComputeSocket.cs
@@ -13,21 +13,6 @@ namespace EpicGames.Horde.Compute
 	/// </summary>
 	public interface IComputeSocket : IAsyncDisposable
 	{
-		/// <summary>
-		/// Sends data to a remote channel
-		/// </summary>
-		/// <param name="channelId">Channel to write to</param>
-		/// <param name="memory">Memory to write</param>
-		/// <param name="cancellationToken">Cancellation token for the operation</param>
-		ValueTask SendAsync(int channelId, ReadOnlyMemory<byte> memory, CancellationToken cancellationToken = default);
-
-		/// <summary>
-		/// Marks a channel as complete
-		/// </summary>
-		/// <param name="channelId">Channel to mark as complete</param>
-		/// <param name="cancellationToken">Cancellation token for the operation</param>
-		ValueTask MarkCompleteAsync(int channelId, CancellationToken cancellationToken = default);
-
 		/// <summary>
 		/// Attaches a buffer to receive data.
 		/// </summary>
@@ -56,19 +41,6 @@ namespace EpicGames.Horde.Compute
 		ValueTask CloseAsync(CancellationToken cancellationToken = default);
 	}
 
-	/// <summary>
-	/// Static methods for sockets
-	/// </summary>
-	public static class ComputeSocket
-	{
-		/// <summary>
-		/// Create a new compute socket
-		/// </summary>
-		/// <param name="transport">Transport to use for the socket</param>
-		/// <param name="logger">Logger for diagnostic output</param>
-		public static IComputeSocket Create(IComputeTransport transport, ILogger logger) => new DefaultComputeSocket(transport, logger);
-	}
-
 	/// <summary>
 	/// Extension methods for <see cref="IComputeMessageChannel"/>
 	/// </summary>
diff --git a/Engine/Source/Programs/Shared/EpicGames.Horde/Compute/RemoteComputeSocket.cs b/Engine/Source/Programs/Shared/EpicGames.Horde/Compute/RemoteComputeSocket.cs
new file mode 100644
--- /dev/null
+++ b/Engine/Source/Programs/Shared/EpicGames.Horde/Compute/RemoteComputeSocket.cs
@@ -0,0 +1,374 @@
+// Copyright Epic Games, Inc. All Rights Reserved.
+
+using System;
+using System.Buffers;
+using System.Buffers.Binary;
+using System.Collections.Generic;
+using System.Linq;
+using System.Threading;
+using System.Threading.Channels;
+using System.Threading.Tasks;
+using EpicGames.Core;
+using Microsoft.Extensions.Logging;
+
+namespace EpicGames.Horde.Compute
+{
+	/// <summary>
+	/// Manages a set of readers and writers to buffers across a transport layer
+	/// </summary>
+	public class RemoteComputeSocket : IComputeSocket, IAsyncDisposable
+	{
+		enum ControlMessageType
+		{
+			Attach = -1,
+			Detach = -2,
+		}
+
+		readonly object _lockObject = new object();
+
+		bool _complete;
+
+		readonly IComputeTransport _transport;
+		readonly ILogger _logger;
+
+		readonly CancellationTokenSource _cancellationSource = new CancellationTokenSource();
+
+		readonly BackgroundTask _recvTask;
+		readonly Dictionary<int, IComputeBufferWriter> _recvBufferWriters = new Dictionary<int, IComputeBufferWriter>();
+
+		readonly SemaphoreSlim _sendSemaphore = new SemaphoreSlim(1, 1);
+		readonly Dictionary<int, Task> _sendTasks = new Dictionary<int, Task>();
+		readonly Dictionary<IComputeBufferReader, int> _sendBufferReaders = new Dictionary<IComputeBufferReader, int>();
+
+		readonly HashSet<int> _remoteRecvBuffers = new HashSet<int>();
+		readonly AsyncEvent _remoteBuffersChanged = new AsyncEvent();
+
+		readonly BackgroundTask _attachTask;
+		readonly Channel<int> _attachChannel = Channel.CreateUnbounded<int>();
+
+		/// <summary>
+		/// Constructor
+		/// </summary>
+		/// <param name="transport">Transport to communicate with the remote</param>
+		/// <param name="logger">Logger for trace output</param>
+		public RemoteComputeSocket(IComputeTransport transport, ILogger logger)
+		{
+			_transport = transport;
+			_logger = logger;
+			_recvTask = BackgroundTask.StartNew(ctx => RunReaderAsync(transport, ctx));
+			_attachTask = BackgroundTask.StartNew(ctx => SendAttachNotificationsAsync(ctx));
+		}
+
+		/// <summary>
+		/// Attempt to gracefully close the current connection and shutdown both ends of the transport
+		/// </summary>
+		public async ValueTask CloseAsync(CancellationToken cancellationToken)
+		{
+			// Make sure we close all buffers that are attached, otherwise we'll lock up waiting for send tasks to complete
+			Task[] sendTasks;
+			lock (_lockObject)
+			{
+				sendTasks = _sendTasks.Values.ToArray();
+				_sendBufferReaders.Clear();
+			}
+
+			// Wait for all the individual send tasks to complete
+			await Task.WhenAll(sendTasks);
+			cancellationToken.ThrowIfCancellationRequested();
+
+			// Send a final message indicating that the lease is done. This will allow the senders on the remote end to terminate, and trigger
+			// a shutdown event to be sent back to our read task allowing it to shut down gracefully.
+			await _transport.MarkCompleteAsync(cancellationToken);
+
+			// Stop sending 
+			_attachChannel.Writer.Complete();
+			await _attachTask.Task.WaitAsync(TimeSpan.FromSeconds(5), cancellationToken);
+			await _attachTask.StopAsync();
+
+			// Wait for the reader to stop
+			await _recvTask.Task.WaitAsync(TimeSpan.FromSeconds(5), cancellationToken);
+			await _recvTask.StopAsync();
+		}
+
+		/// <inheritdoc/>
+		public async ValueTask DisposeAsync()
+		{
+			_cancellationSource.Cancel();
+			await Task.WhenAll(_sendTasks.Values);
+			await _recvTask.DisposeAsync();
+			await _attachTask.DisposeAsync();
+			_sendSemaphore.Dispose();
+			_cancellationSource.Dispose();
+			GC.SuppressFinalize(this);
+		}
+
+		async Task RunReaderAsync(IComputeTransport transport, CancellationToken cancellationToken)
+		{
+			_logger.LogTrace("Started socket reader");
+
+			byte[] header = new byte[8];
+			try
+			{
+				// Maintain a local cache of buffers to be able to query for them without having to acquire a global lock
+				Dictionary<int, IComputeBufferWriter> cachedWriters = new Dictionary<int, IComputeBufferWriter>();
+
+				Memory<byte> last = Memory<byte>.Empty;
+
+				// Process messages from the remote
+				for (; ; )
+				{
+					// Read the next packet header
+					if (!await transport.ReadOptionalAsync(header, cancellationToken))
+					{
+						_logger.LogTrace("End of socket");
+						break;
+					}
+
+					// Parse the target buffer and packet size
+					int id = BinaryPrimitives.ReadInt32LittleEndian(header);
+					int size = BinaryPrimitives.ReadInt32LittleEndian(header.AsSpan(4));
+
+					// Dispatch it to the correct place
+					if (size >= 0)
+					{
+						IComputeBufferWriter writer = GetReceiveBuffer(cachedWriters, id);
+						await ReadPacketAsync(transport, id, size, writer, cancellationToken);
+					}
+					else if (size == (int)ControlMessageType.Attach)
+					{
+						_logger.LogTrace("Attaching recv buffer {Id}", id);
+						AttachRemoteRecvBuffer(id);
+					}
+					else if (size == (int)ControlMessageType.Detach)
+					{
+						_logger.LogTrace("Detaching recv buffer {Id}", id);
+						DetachRecvBuffer(cachedWriters, id);
+					}
+					else
+					{
+						_logger.LogWarning("Unrecognized control message: {Message}", size);
+					}
+				}
+			}
+			catch (OperationCanceledException)
+			{
+			}
+
+			// Mark all buffers as complete
+			lock (_lockObject)
+			{
+				_complete = true;
+				foreach (IComputeBuffer buffer in _recvBufferWriters.Values)
+				{
+					buffer.Writer.MarkComplete();
+				}
+			}
+
+			_logger.LogTrace("Closing reader");
+		}
+
+		async Task ReadPacketAsync(IComputeTransport transport, int id, int size, IComputeBufferWriter writer, CancellationToken cancellationToken)
+		{
+			Memory<byte> memory = writer.GetWriteBuffer();
+			while (memory.Length < size)
+			{
+				_logger.LogTrace("No space in buffer {Id}, flushing", id);
+				await writer.WaitToWriteAsync(size, cancellationToken);
+				memory = writer.GetWriteBuffer();
+			}
+
+			for (int offset = 0; offset < size;)
+			{
+				int read = await transport.ReadPartialAsync(memory.Slice(offset, size - offset), cancellationToken);
+				offset += read;
+			}
+
+			writer.AdvanceWritePosition(size);
+		}
+
+		IComputeBufferWriter GetReceiveBuffer(Dictionary<int, IComputeBufferWriter> cachedWriters, int id)
+		{
+			IComputeBufferWriter? writer;
+			if (cachedWriters.TryGetValue(id, out writer))
+			{
+				return writer;
+			}
+
+			lock (_lockObject)
+			{
+				IComputeBufferWriter? recvBufferWriter;
+				if (_recvBufferWriters.TryGetValue(id, out recvBufferWriter))
+				{
+					writer = recvBufferWriter;
+					cachedWriters.Add(id, writer);
+					return writer;
+				}
+			}
+
+			throw new ComputeInternalException($"No buffer is attached to channel {id}");
+		}
+
+		/// <inheritdoc/>
+		public async ValueTask WaitForAttachAsync(int channelId, CancellationToken cancellationToken = default)
+		{
+			for (; ; )
+			{
+				Task task = _remoteBuffersChanged.Task;
+				lock (_lockObject)
+				{
+					if (_remoteRecvBuffers.Contains(channelId))
+					{
+						return;
+					}
+				}
+				await task.WaitAsync(cancellationToken);
+			}
+		}
+
+		class SendSegment : ReadOnlySequenceSegment<byte>
+		{
+			public void Set(ReadOnlyMemory<byte> memory, ReadOnlySequenceSegment<byte>? next, long runningIndex)
+			{
+				Memory = memory;
+				Next = next;
+				RunningIndex = runningIndex;
+			}
+		}
+
+		readonly byte[] _header = new byte[8];
+		readonly SendSegment _headerSegment = new SendSegment();
+		readonly SendSegment _bodySegment = new SendSegment();
+
+		/// <inheritdoc/>
+		public async ValueTask SendAsync(int id, ReadOnlyMemory<byte> memory, CancellationToken cancellationToken = default)
+		{
+			if (memory.Length > 0)
+			{
+				await SendInternalAsync(id, memory.Length, memory, cancellationToken);
+			}
+		}
+
+		/// <inheritdoc/>
+		public ValueTask MarkCompleteAsync(int id, CancellationToken cancellationToken = default) => SendInternalAsync(id, (int)ControlMessageType.Detach, ReadOnlyMemory<byte>.Empty, cancellationToken);
+
+		async ValueTask SendInternalAsync(int id, int size, ReadOnlyMemory<byte> memory, CancellationToken cancellationToken)
+		{
+			await _sendSemaphore.WaitAsync(cancellationToken);
+			try
+			{
+				BinaryPrimitives.WriteInt32LittleEndian(_header, id);
+				BinaryPrimitives.WriteInt32LittleEndian(_header.AsSpan(4), size);
+				_headerSegment.Set(_header, _bodySegment, 0);
+				_bodySegment.Set(memory, null, _header.Length);
+
+				ReadOnlySequence<byte> sequence = new ReadOnlySequence<byte>(_headerSegment, 0, _bodySegment, memory.Length);
+				await _transport.WriteAsync(sequence, cancellationToken);
+			}
+			finally
+			{
+				_sendSemaphore.Release();
+			}
+		}
+
+		/// <inheritdoc/>
+		public async void AttachRecvBuffer(int channelId, IComputeBufferWriter recvBufferWriter)
+		{
+			bool complete;
+			lock (_lockObject)
+			{
+				complete = _complete;
+				if (!complete)
+				{
+					_recvBufferWriters.Add(channelId, recvBufferWriter.AddRef());
+				}
+			}
+
+			if (recvBufferWriter != null && complete)
+			{
+				recvBufferWriter.MarkComplete();
+			}
+
+			if (!complete)
+			{
+				await _attachChannel.Writer.WriteAsync(channelId);
+			}
+		}
+
+		async Task SendAttachNotificationsAsync(CancellationToken cancellationToken)
+		{
+			while (await _attachChannel.Reader.WaitToReadAsync(cancellationToken))
+			{
+				int channelId;
+				if (_attachChannel.Reader.TryRead(out channelId))
+				{
+					await SendInternalAsync(channelId, (int)ControlMessageType.Attach, ReadOnlyMemory<byte>.Empty, cancellationToken);
+				}
+			}
+		}
+
+		void AttachRemoteRecvBuffer(int id)
+		{
+			lock (_lockObject)
+			{
+				_remoteRecvBuffers.Add(id);
+			}
+			_remoteBuffersChanged.Pulse();
+		}
+
+		void DetachRecvBuffer(Dictionary<int, IComputeBufferWriter> cachedWriters, int id)
+		{
+			cachedWriters.Remove(id);
+
+			IComputeBufferWriter? recvBufferWriter;
+			lock (_lockObject)
+			{
+#pragma warning disable CA2000 // Dispose objects before losing scope
+				_recvBufferWriters.Remove(id, out recvBufferWriter);
+#pragma warning restore CA2000 // Dispose objects before losing scope
+			}
+			if (recvBufferWriter != null)
+			{
+				recvBufferWriter.MarkComplete();
+				recvBufferWriter.Dispose();
+			}
+		}
+
+		/// <inheritdoc/>
+		public void AttachSendBuffer(int channelId, IComputeBufferReader sendBufferReader)
+		{
+			sendBufferReader = sendBufferReader.AddRef();
+			lock (_lockObject)
+			{
+				_sendTasks.Add(channelId, Task.Run(() => SendFromBufferAsync(channelId, sendBufferReader, _cancellationSource.Token), CancellationToken.None)); // No cancellation token; need to ensure dispose runs
+				_sendBufferReaders.Add(sendBufferReader, channelId);
+			}
+		}
+
+		async Task SendFromBufferAsync(int channelId, IComputeBufferReader reader, CancellationToken cancellationToken)
+		{
+			using IComputeBufferReader _ = reader;
+
+			while (!cancellationToken.IsCancellationRequested)
+			{
+				ReadOnlyMemory<byte> memory = reader.GetReadBuffer();
+				if (memory.Length > 0)
+				{
+					await SendAsync(channelId, memory, cancellationToken);
+					reader.AdvanceReadPosition(memory.Length);
+				}
+				if (reader.IsComplete)
+				{
+					await MarkCompleteAsync(channelId, cancellationToken);
+					break;
+				}
+				await reader.WaitToReadAsync(1, cancellationToken);
+			}
+
+			lock (_lockObject)
+			{
+				_sendTasks.Remove(channelId);
+				_sendBufferReaders.Remove(reader);
+			}
+		}
+	}
+}
