Horde: Fix corruption caused by compute buffers being disposed rather than released.

diff --git a/Engine/Source/Programs/Shared/EpicGames.Horde/Compute/ComputeBufferBase.cs b/Engine/Source/Programs/Shared/EpicGames.Horde/Compute/ComputeBufferBase.cs
--- a/Engine/Source/Programs/Shared/EpicGames.Horde/Compute/ComputeBufferBase.cs
+++ b/Engine/Source/Programs/Shared/EpicGames.Horde/Compute/ComputeBufferBase.cs
@@ -281,7 +281,7 @@ namespace EpicGames.Horde.Compute
 			{
 				if (_resources != null)
 				{
-					_resources.Dispose();
+					_resources.Release();
 					_resources = null!;
 				}
 			}
@@ -409,7 +409,7 @@ namespace EpicGames.Horde.Compute
 			{
 				if (_resources != null)
 				{
-					_resources.Dispose();
+					_resources.Release();
 					_resources = null!;
 				}
 			}
