Horde: Add methods for marking compute channels as complete.

diff --git a/Engine/Source/Programs/Shared/EpicGames.Horde/Compute/IComputeChannel.cs b/Engine/Source/Programs/Shared/EpicGames.Horde/Compute/IComputeChannel.cs
--- a/Engine/Source/Programs/Shared/EpicGames.Horde/Compute/IComputeChannel.cs
+++ b/Engine/Source/Programs/Shared/EpicGames.Horde/Compute/IComputeChannel.cs
@@ -25,6 +25,12 @@ namespace EpicGames.Horde.Compute
 		/// <param name="buffer">Buffer to receive the data</param>
 		/// <param name="cancellationToken">Cancellation token for the operation</param>
 		ValueTask<int> ReceiveAsync(Memory<byte> buffer, CancellationToken cancellationToken = default);
+
+		/// <summary>
+		/// Mark the channel as complete (ie. that no more data will be sent)
+		/// </summary>
+		/// <param name="cancellationToken">Cancellation token for the operation</param>
+		ValueTask MarkCompleteAsync(CancellationToken cancellationToken = default);
 	}
 
 	/// <summary>
@@ -55,6 +61,8 @@ namespace EpicGames.Horde.Compute
 			public ValueTask<int> ReceiveAsync(Memory<byte> buffer, CancellationToken cancellationToken = default) => _recvBufferReader.ReadAsync(buffer, cancellationToken);
 
 			public ValueTask SendAsync(ReadOnlyMemory<byte> memory, CancellationToken cancellationToken = default) => _socket.SendAsync(_channelId, memory, cancellationToken);
+
+			public ValueTask MarkCompleteAsync(CancellationToken cancellationToken = default) => _socket.MarkCompleteAsync(_channelId, cancellationToken);
 		}
 
 		internal sealed class BufferedReaderWriterChannel : IComputeChannel
@@ -79,6 +87,12 @@ namespace EpicGames.Horde.Compute
 			public ValueTask<int> ReceiveAsync(Memory<byte> buffer, CancellationToken cancellationToken = default) => _recvBufferReader.ReadAsync(buffer, cancellationToken);
 
 			public ValueTask SendAsync(ReadOnlyMemory<byte> memory, CancellationToken cancellationToken = default) => _sendBufferWriter.WriteAsync(memory, cancellationToken);
+
+			public ValueTask MarkCompleteAsync(CancellationToken cancellationToken = default)
+			{
+				_sendBufferWriter.MarkComplete();
+				return new ValueTask();
+			}
 		}
 
 		/// <summary>
