Horde: Remove special control messages for notifying on buffer attachment.

This is trivially handled with explicit messaging from the agent, and can be implemented as necessary for direct channel comms.

diff --git a/Engine/Source/Programs/Horde/Samples/RemoteClient/RemoteClientApp.cs b/Engine/Source/Programs/Horde/Samples/RemoteClient/RemoteClientApp.cs
--- a/Engine/Source/Programs/Horde/Samples/RemoteClient/RemoteClientApp.cs
+++ b/Engine/Source/Programs/Horde/Samples/RemoteClient/RemoteClientApp.cs
@@ -6,6 +6,7 @@ using System.Reflection;
 using EpicGames.Core;
 using EpicGames.Horde.Common;
 using EpicGames.Horde.Compute;
+using EpicGames.Horde.Compute.Buffers;
 using EpicGames.Horde.Compute.Clients;
 using EpicGames.Horde.Storage;
 using EpicGames.Horde.Storage.Backends;
@@ -119,11 +120,14 @@ namespace RemoteClient
 		
 		static async Task WriteNumbersAsync(RemoteComputeSocket socket, ILogger logger, CancellationToken cancellationToken)
 		{
-			// Write data to the child process channel. The remote server will echo them back to us as it receives them, then exit when the channel is complete/closed.
-			await socket.WaitForAttachAsync(ChildProcessChannelId, cancellationToken);
-
-			logger.LogInformation("Connected");
+			// Wait until the remote sends a message indicating that it's ready
+			using (PooledBuffer recvBuffer = new PooledBuffer(1, 20))
+			{
+				socket.AttachRecvBuffer(ChildProcessChannelId, recvBuffer.Writer);
+				await recvBuffer.Reader.WaitToReadAsync(1, cancellationToken);
+			}
 
+			// Write data to the child process channel. The remote server will echo them back to us as it receives them, then exit when the channel is complete/closed.
 			byte[] buffer = new byte[4];
 			for (int idx = 0; idx < 3; idx++)
 			{
diff --git a/Engine/Source/Programs/Horde/Samples/RemoteWorker/RemoteWorkerApp.cs b/Engine/Source/Programs/Horde/Samples/RemoteWorker/RemoteWorkerApp.cs
--- a/Engine/Source/Programs/Horde/Samples/RemoteWorker/RemoteWorkerApp.cs
+++ b/Engine/Source/Programs/Horde/Samples/RemoteWorker/RemoteWorkerApp.cs
@@ -17,6 +17,8 @@ namespace RemoteWorker
 			Console.WriteLine("Connected to initiator");
 
 			byte[] data = new byte[4];
+			await channel.SendAsync(data.AsMemory(0, 1)); // Let the remote know we're ready
+
 			while(await channel.TryReceiveMessageAsync(data))
 			{
 				int value = BinaryPrimitives.ReadInt32LittleEndian(data);
diff --git a/Engine/Source/Programs/Horde/Samples/RemoteWorkerCpp/Main.cpp b/Engine/Source/Programs/Horde/Samples/RemoteWorkerCpp/Main.cpp
--- a/Engine/Source/Programs/Horde/Samples/RemoteWorkerCpp/Main.cpp
+++ b/Engine/Source/Programs/Horde/Samples/RemoteWorkerCpp/Main.cpp
@@ -30,6 +30,9 @@ int main()
 	size_t Length = 0;
 	char Buffer[4];
 
+	Buffer[0] = 0;
+	Channel.Send(Buffer, 1); // Send a dummy one-byte message to let the remote know we're listening
+
 	for (;;)
 	{
 		size_t RecvLength = Channel.Recv(Buffer + Length, sizeof(Buffer) - Length);
diff --git a/Engine/Source/Programs/Shared/EpicGames.Horde/Compute/ComputeExceptions.cs b/Engine/Source/Programs/Shared/EpicGames.Horde/Compute/ComputeExceptions.cs
--- a/Engine/Source/Programs/Shared/EpicGames.Horde/Compute/ComputeExceptions.cs
+++ b/Engine/Source/Programs/Shared/EpicGames.Horde/Compute/ComputeExceptions.cs
@@ -64,12 +64,12 @@ namespace EpicGames.Horde.Compute
 	/// <summary>
 	/// Exception thrown when an invalid message is received
 	/// </summary>
-	public sealed class ComputeInvalidMessageException : ComputeException
+	public sealed class InvalidComputeMessageException : ComputeException
 	{
 		/// <summary>
 		/// Constructor
 		/// </summary>
-		public ComputeInvalidMessageException(IComputeMessage message)
+		public InvalidComputeMessageException(IComputeMessage message)
 			: base($"Unexpected message {message.Type}")
 		{
 		}
diff --git a/Engine/Source/Programs/Shared/EpicGames.Horde/Compute/ComputeMessageChannel.cs b/Engine/Source/Programs/Shared/EpicGames.Horde/Compute/ComputeMessageChannel.cs
--- a/Engine/Source/Programs/Shared/EpicGames.Horde/Compute/ComputeMessageChannel.cs
+++ b/Engine/Source/Programs/Shared/EpicGames.Horde/Compute/ComputeMessageChannel.cs
@@ -109,7 +109,6 @@ namespace EpicGames.Horde.Compute
 		}
 
 		readonly IComputeSocket _socket;
-		readonly int _channelId;
 		readonly IComputeBufferReader _recvBufferReader;
 		readonly IComputeBufferWriter _sendBufferWriter;
 
@@ -123,14 +122,12 @@ namespace EpicGames.Horde.Compute
 		/// <summary>
 		/// Constructor
 		/// </summary>
-		/// <param name="channelId"></param>
 		/// <param name="recvBufferReader"></param>
 		/// <param name="sendBufferWriter"></param>
 		/// <param name="logger">Logger for diagnostic output</param>
-		public ComputeMessageChannel(int channelId, IComputeBufferReader recvBufferReader, IComputeBufferWriter sendBufferWriter, ILogger logger)
+		public ComputeMessageChannel(IComputeBufferReader recvBufferReader, IComputeBufferWriter sendBufferWriter, ILogger logger)
 		{
 			_socket = null!;
-			_channelId = channelId;
 			_recvBufferReader = recvBufferReader.AddRef();
 			_sendBufferWriter = sendBufferWriter.AddRef();
 			_logger = logger;
@@ -147,7 +144,6 @@ namespace EpicGames.Horde.Compute
 		public ComputeMessageChannel(IComputeSocket socket, int channelId, IComputeBuffer recvBuffer, IComputeBuffer sendBuffer, ILogger logger)
 		{
 			_socket = socket;
-			_channelId = channelId;
 			_socket.AttachRecvBuffer(channelId, recvBuffer.Writer);
 			_socket.AttachSendBuffer(channelId, sendBuffer.Reader);
 			_recvBufferReader = recvBuffer.Reader.AddRef();
@@ -176,9 +172,6 @@ namespace EpicGames.Horde.Compute
 			_sendBufferWriter.MarkComplete();
 		}
 
-		/// <inheritdoc/>
-		public ValueTask WaitForAttachAsync(CancellationToken cancellationToken = default) => _socket.WaitForAttachAsync(_channelId, cancellationToken);
-
 		/// <inheritdoc/>
 		public async ValueTask<IComputeMessage> ReceiveAsync(CancellationToken cancellationToken)
 		{
diff --git a/Engine/Source/Programs/Shared/EpicGames.Horde/Compute/ComputeProcess.cs b/Engine/Source/Programs/Shared/EpicGames.Horde/Compute/ComputeProcess.cs
--- a/Engine/Source/Programs/Shared/EpicGames.Horde/Compute/ComputeProcess.cs
+++ b/Engine/Source/Programs/Shared/EpicGames.Horde/Compute/ComputeProcess.cs
@@ -81,7 +81,7 @@ namespace EpicGames.Horde.Compute
 						_result.TrySetResult(executeProcessResponse.ExitCode);
 						return;
 					default:
-						throw new ComputeInvalidMessageException(message);
+						throw new InvalidComputeMessageException(message);
 				}
 			}
 		}
diff --git a/Engine/Source/Programs/Shared/EpicGames.Horde/Compute/ComputeWorker.cs b/Engine/Source/Programs/Shared/EpicGames.Horde/Compute/ComputeWorker.cs
--- a/Engine/Source/Programs/Shared/EpicGames.Horde/Compute/ComputeWorker.cs
+++ b/Engine/Source/Programs/Shared/EpicGames.Horde/Compute/ComputeWorker.cs
@@ -51,6 +51,8 @@ namespace EpicGames.Horde.Compute
 		{
 			using (IComputeMessageChannel channel = socket.CreateMessageChannel(channelId, bufferSize, _logger))
 			{
+				await channel.AttachAsync(cancellationToken);
+
 				List<Task> childTasks = new List<Task>();
 				for (; ; )
 				{
@@ -93,7 +95,7 @@ namespace EpicGames.Horde.Compute
 							}
 							break;
 						default:
-							throw new ComputeInvalidMessageException(message);
+							throw new InvalidComputeMessageException(message);
 					}
 				}
 			}
diff --git a/Engine/Source/Programs/Shared/EpicGames.Horde/Compute/IComputeMessage.cs b/Engine/Source/Programs/Shared/EpicGames.Horde/Compute/IComputeMessage.cs
--- a/Engine/Source/Programs/Shared/EpicGames.Horde/Compute/IComputeMessage.cs
+++ b/Engine/Source/Programs/Shared/EpicGames.Horde/Compute/IComputeMessage.cs
@@ -32,6 +32,11 @@ namespace EpicGames.Horde.Compute
 		/// </summary>
 		Fork = 0x03,
 
+		/// <summary>
+		/// Sent as the first message on a channel to notify the remote that the remote end is attached
+		/// </summary>
+		Attach = 0x04,
+
 		#region Process Management
 
 		/// <summary>
@@ -204,7 +209,7 @@ namespace EpicGames.Horde.Compute
 		}
 
 		/// <summary>
-		/// Sends an exception response to the remote
+		/// Requests that the remote message loop be forked
 		/// </summary>
 		public static async ValueTask ForkAsync(this IComputeMessageChannel channel, int channelId, int bufferSize, CancellationToken cancellationToken = default)
 		{
@@ -215,7 +220,7 @@ namespace EpicGames.Horde.Compute
 		}
 
 		/// <summary>
-		/// Parses a message as an <see cref="ExceptionMessage"/>
+		/// Parses a fork request message
 		/// </summary>
 		public static ForkMessage ParseForkMessage(this IComputeMessage message)
 		{
@@ -224,6 +229,30 @@ namespace EpicGames.Horde.Compute
 			return new ForkMessage(channelId, bufferSize);
 		}
 
+		/// <summary>
+		/// Notifies the remote that a buffer has been attached
+		/// </summary>
+		public static async ValueTask AttachAsync(this IComputeMessageChannel channel, CancellationToken cancellationToken = default)
+		{
+			using IComputeMessageBuilder message = await channel.CreateMessageAsync(ComputeMessageType.Attach, cancellationToken);
+			message.Send();
+		}
+
+		/// <summary>
+		/// Waits until an attached notification is received along the channel
+		/// </summary>
+		/// <param name="channel"></param>
+		/// <param name="cancellationToken"></param>
+		/// <returns></returns>
+		public static async ValueTask WaitForAttachAsync(this IComputeMessageChannel channel, CancellationToken cancellationToken = default)
+		{
+			using IComputeMessage message = await channel.ReceiveAsync(cancellationToken);
+			if (message.Type != ComputeMessageType.Attach)
+			{
+				throw new InvalidComputeMessageException(message);
+			}
+		}
+
 		#region Process
 
 		static async Task<IComputeMessage> RunStorageServer(this IComputeMessageChannel channel, IStorageClient storage, CancellationToken cancellationToken = default)
@@ -268,7 +297,7 @@ namespace EpicGames.Horde.Compute
 			using IComputeMessage response = await RunStorageServer(channel, storage, cancellationToken);
 			if (response.Type != ComputeMessageType.WriteFilesResponse)
 			{
-				throw new ComputeInvalidMessageException(response);
+				throw new InvalidComputeMessageException(response);
 			}
 		}
 
@@ -439,7 +468,7 @@ namespace EpicGames.Horde.Compute
 					response = await channel.ReceiveAsync(cancellationToken);
 					if (response.Type != ComputeMessageType.ReadBlobResponse)
 					{
-						throw new ComputeInvalidMessageException(response);
+						throw new InvalidComputeMessageException(response);
 					}
 
 					int chunkOffset = BinaryPrimitives.ReadInt32LittleEndian(response.Data.Span.Slice(0, 4));
diff --git a/Engine/Source/Programs/Shared/EpicGames.Horde/Compute/IComputeMessageChannel.cs b/Engine/Source/Programs/Shared/EpicGames.Horde/Compute/IComputeMessageChannel.cs
--- a/Engine/Source/Programs/Shared/EpicGames.Horde/Compute/IComputeMessageChannel.cs
+++ b/Engine/Source/Programs/Shared/EpicGames.Horde/Compute/IComputeMessageChannel.cs
@@ -12,12 +12,6 @@ namespace EpicGames.Horde.Compute
 	/// </summary>
 	public interface IComputeMessageChannel : IDisposable
 	{
-		/// <summary>
-		/// Wait for the remote end to attach a reader
-		/// </summary>
-		/// <param name="cancellationToken">Cancellation token for the operation</param>
-		ValueTask WaitForAttachAsync(CancellationToken cancellationToken = default);
-
 		/// <summary>
 		/// Reads a message from the channel
 		/// </summary>
@@ -63,7 +57,7 @@ namespace EpicGames.Horde.Compute
 			IComputeMessage message = await channel.ReceiveAsync(cancellationToken);
 			if (message.Type != type)
 			{
-				throw new ComputeInvalidMessageException(message);
+				throw new InvalidComputeMessageException(message);
 			}
 			return message;
 		}
diff --git a/Engine/Source/Programs/Shared/EpicGames.Horde/Compute/IComputeSocket.cs b/Engine/Source/Programs/Shared/EpicGames.Horde/Compute/IComputeSocket.cs
--- a/Engine/Source/Programs/Shared/EpicGames.Horde/Compute/IComputeSocket.cs
+++ b/Engine/Source/Programs/Shared/EpicGames.Horde/Compute/IComputeSocket.cs
@@ -27,13 +27,6 @@ namespace EpicGames.Horde.Compute
 		/// <param name="sendBufferReader">Reader for the buffer to send data from</param>
 		void AttachSendBuffer(int channelId, IComputeBufferReader sendBufferReader);
 
-		/// <summary>
-		/// Waits until the remote has attached a receive buffer to the given channel
-		/// </summary>
-		/// <param name="channelId">Channel id</param>
-		/// <param name="cancellationToken">Cancellation token for the operation</param>
-		ValueTask WaitForAttachAsync(int channelId, CancellationToken cancellationToken = default);
-
 		/// <summary>
 		/// Attempt to gracefully close the current connection and shutdown both ends of the transport
 		/// </summary>
diff --git a/Engine/Source/Programs/Shared/EpicGames.Horde/Compute/RemoteComputeSocket.cs b/Engine/Source/Programs/Shared/EpicGames.Horde/Compute/RemoteComputeSocket.cs
--- a/Engine/Source/Programs/Shared/EpicGames.Horde/Compute/RemoteComputeSocket.cs
+++ b/Engine/Source/Programs/Shared/EpicGames.Horde/Compute/RemoteComputeSocket.cs
@@ -36,7 +36,6 @@ namespace EpicGames.Horde.Compute
 	{
 		enum ControlMessageType
 		{
-			Attach = -1,
 			Detach = -2,
 		}
 
@@ -57,12 +56,6 @@ namespace EpicGames.Horde.Compute
 		readonly Dictionary<int, Task> _sendTasks = new Dictionary<int, Task>();
 		readonly Dictionary<IComputeBufferReader, int> _sendBufferReaders = new Dictionary<IComputeBufferReader, int>();
 
-		readonly HashSet<int> _remoteRecvBuffers = new HashSet<int>();
-		readonly AsyncEvent _remoteBuffersChanged = new AsyncEvent();
-
-		readonly BackgroundTask _attachTask;
-		readonly Channel<int> _attachChannel = Channel.CreateUnbounded<int>();
-
 		string Tag => (_endpoint == ComputeSocketEndpoint.Local)? "LOCAL": "REMOTE";
 
 		/// <summary>
@@ -77,7 +70,6 @@ namespace EpicGames.Horde.Compute
 			_endpoint = endpoint;
 			_logger = logger;
 			_recvTask = BackgroundTask.StartNew(ctx => RunRecvTaskAsync(transport, ctx));
-			_attachTask = BackgroundTask.StartNew(ctx => SendAttachNotificationsAsync(ctx));
 		}
 
 		/// <summary>
@@ -103,11 +95,6 @@ namespace EpicGames.Horde.Compute
 			// a shutdown event to be sent back to our read task allowing it to shut down gracefully.
 			await _transport.MarkCompleteAsync(cancellationToken);
 
-			// Stop sending 
-			_attachChannel.Writer.Complete();
-			await _attachTask.Task.WaitAsync(TimeSpan.FromSeconds(5), cancellationToken);
-			await _attachTask.StopAsync();
-
 			// Wait for the reader to stop
 			await _recvTask.Task.WaitAsync(TimeSpan.FromSeconds(5), cancellationToken);
 			await _recvTask.StopAsync();
@@ -119,7 +106,6 @@ namespace EpicGames.Horde.Compute
 			_cancellationSource.Cancel();
 			await Task.WhenAll(_sendTasks.Values);
 			await _recvTask.DisposeAsync();
-			await _attachTask.DisposeAsync();
 			_sendSemaphore.Dispose();
 			_cancellationSource.Dispose();
 			GC.SuppressFinalize(this);
@@ -157,11 +143,6 @@ namespace EpicGames.Horde.Compute
 						IComputeBufferWriter writer = GetReceiveBuffer(cachedWriters, id);
 						await ReadPacketAsync(transport, id, size, writer, cancellationToken);
 					}
-					else if (size == (int)ControlMessageType.Attach)
-					{
-						_logger.LogTrace("[{Tag}] Attaching recv buffer {Id}", Tag, id);
-						AttachRemoteRecvBuffer(id);
-					}
 					else if (size == (int)ControlMessageType.Detach)
 					{
 						_logger.LogTrace("[{Tag}] Detaching recv buffer {Id}", Tag, id);
@@ -231,23 +212,6 @@ namespace EpicGames.Horde.Compute
 			throw new ComputeInternalException($"No buffer is attached to channel {id}");
 		}
 
-		/// <inheritdoc/>
-		public async ValueTask WaitForAttachAsync(int channelId, CancellationToken cancellationToken = default)
-		{
-			for (; ; )
-			{
-				Task task = _remoteBuffersChanged.Task;
-				lock (_lockObject)
-				{
-					if (_remoteRecvBuffers.Contains(channelId))
-					{
-						return;
-					}
-				}
-				await task.WaitAsync(cancellationToken);
-			}
-		}
-
 		class SendSegment : ReadOnlySequenceSegment<byte>
 		{
 			public void Set(ReadOnlyMemory<byte> memory, ReadOnlySequenceSegment<byte>? next, long runningIndex)
@@ -294,7 +258,7 @@ namespace EpicGames.Horde.Compute
 		}
 
 		/// <inheritdoc/>
-		public async void AttachRecvBuffer(int channelId, IComputeBufferWriter recvBufferWriter)
+		public void AttachRecvBuffer(int channelId, IComputeBufferWriter recvBufferWriter)
 		{
 			bool complete;
 			lock (_lockObject)
@@ -310,32 +274,6 @@ namespace EpicGames.Horde.Compute
 			{
 				recvBufferWriter.MarkComplete();
 			}
-
-			if (!complete)
-			{
-				await _attachChannel.Writer.WriteAsync(channelId);
-			}
-		}
-
-		async Task SendAttachNotificationsAsync(CancellationToken cancellationToken)
-		{
-			while (await _attachChannel.Reader.WaitToReadAsync(cancellationToken))
-			{
-				int channelId;
-				if (_attachChannel.Reader.TryRead(out channelId))
-				{
-					await SendInternalAsync(channelId, (int)ControlMessageType.Attach, ReadOnlyMemory<byte>.Empty, cancellationToken);
-				}
-			}
-		}
-
-		void AttachRemoteRecvBuffer(int id)
-		{
-			lock (_lockObject)
-			{
-				_remoteRecvBuffers.Add(id);
-			}
-			_remoteBuffersChanged.Pulse();
 		}
 
 		void DetachRecvBuffer(Dictionary<int, IComputeBufferWriter> cachedWriters, int id)
