Horde: Remove the 'CloseAsync()' method from IComputeSocket; it's only required by remote sockets.

diff --git a/Engine/Source/Programs/Horde/Horde.Agent/Commands/Compute/ComputeCommand.cs b/Engine/Source/Programs/Horde/Horde.Agent/Commands/Compute/ComputeCommand.cs
--- a/Engine/Source/Programs/Horde/Horde.Agent/Commands/Compute/ComputeCommand.cs
+++ b/Engine/Source/Programs/Horde/Horde.Agent/Commands/Compute/ComputeCommand.cs
@@ -127,7 +127,7 @@ namespace Horde.Agent.Commands.Compute
 			using Socket tcpSocket = new Socket(SocketType.Stream, ProtocolType.IP);
 			await tcpSocket.ConnectAsync(IPAddress.Loopback, Port);
 
-			await using (IComputeSocket socket = new RemoteComputeSocket(new TcpTransport(tcpSocket), ComputeSocketEndpoint.Remote, logger))
+			await using (RemoteComputeSocket socket = new RemoteComputeSocket(new TcpTransport(tcpSocket), ComputeSocketEndpoint.Remote, logger))
 			{
 				logger.LogInformation("Running worker...");
 				await RunWorkerAsync(socket, _memoryCache, logger, CancellationToken.None);
diff --git a/Engine/Source/Programs/Horde/Horde.Agent/Leases/Handlers/ComputeHandler.cs b/Engine/Source/Programs/Horde/Horde.Agent/Leases/Handlers/ComputeHandler.cs
--- a/Engine/Source/Programs/Horde/Horde.Agent/Leases/Handlers/ComputeHandler.cs
+++ b/Engine/Source/Programs/Horde/Horde.Agent/Leases/Handlers/ComputeHandler.cs
@@ -92,7 +92,7 @@ namespace Horde.Agent.Leases.Handlers
 				using (CancellationTokenSource cts = CancellationTokenSource.CreateLinkedTokenSource(cancellationToken))
 				{
 					await using BackgroundTask timeoutTask = BackgroundTask.StartNew(ctx => TickTimeoutAsync(transport, cts, ctx));
-					await using (IComputeSocket socket = new RemoteComputeSocket(transport, ComputeSocketEndpoint.Local, _logger))
+					await using (RemoteComputeSocket socket = new RemoteComputeSocket(transport, ComputeSocketEndpoint.Local, _logger))
 					{
 						DirectoryReference sandboxDir = DirectoryReference.Combine(session.WorkingDir, "Sandbox", leaseId);
 						try
diff --git a/Engine/Source/Programs/Shared/EpicGames.Horde/Compute/Clients/LocalComputeClient.cs b/Engine/Source/Programs/Shared/EpicGames.Horde/Compute/Clients/LocalComputeClient.cs
--- a/Engine/Source/Programs/Shared/EpicGames.Horde/Compute/Clients/LocalComputeClient.cs
+++ b/Engine/Source/Programs/Shared/EpicGames.Horde/Compute/Clients/LocalComputeClient.cs
@@ -77,7 +77,7 @@ namespace EpicGames.Horde.Compute.Clients
 
 			using MemoryCache memoryCache = new MemoryCache(new MemoryCacheOptions { SizeLimit = 10 * 1024 * 1024 });
 
-			await using (IComputeSocket socket = new RemoteComputeSocket(new TcpTransport(tcpSocket), ComputeSocketEndpoint.Remote, logger))
+			await using (RemoteComputeSocket socket = new RemoteComputeSocket(new TcpTransport(tcpSocket), ComputeSocketEndpoint.Remote, logger))
 			{
 				ComputeWorker worker = new ComputeWorker(sandboxDir, memoryCache, logger);
 				await worker.RunAsync(socket, cancellationToken);
diff --git a/Engine/Source/Programs/Shared/EpicGames.Horde/Compute/IComputeSocket.cs b/Engine/Source/Programs/Shared/EpicGames.Horde/Compute/IComputeSocket.cs
--- a/Engine/Source/Programs/Shared/EpicGames.Horde/Compute/IComputeSocket.cs
+++ b/Engine/Source/Programs/Shared/EpicGames.Horde/Compute/IComputeSocket.cs
@@ -1,8 +1,6 @@
 // Copyright Epic Games, Inc. All Rights Reserved.
 
 using System;
-using System.Threading;
-using System.Threading.Tasks;
 using EpicGames.Horde.Compute.Buffers;
 using Microsoft.Extensions.Logging;
 
@@ -26,12 +24,6 @@ namespace EpicGames.Horde.Compute
 		/// <param name="channelId">Channel to receive data on</param>
 		/// <param name="sendBufferReader">Reader for the buffer to send data from</param>
 		void AttachSendBuffer(int channelId, IComputeBufferReader sendBufferReader);
-
-		/// <summary>
-		/// Attempt to gracefully close the current connection and shutdown both ends of the transport
-		/// </summary>
-		/// <param name="cancellationToken">Cancellation token for the operation</param>
-		ValueTask CloseAsync(CancellationToken cancellationToken = default);
 	}
 
 	/// <summary>
