Horde: Implement IComputeSocket in WorkerComputeSocket.

diff --git a/Engine/Source/Programs/Shared/EpicGames.Horde/Compute/Buffers/SharedMemoryBuffer.cs b/Engine/Source/Programs/Shared/EpicGames.Horde/Compute/Buffers/SharedMemoryBuffer.cs
--- a/Engine/Source/Programs/Shared/EpicGames.Horde/Compute/Buffers/SharedMemoryBuffer.cs
+++ b/Engine/Source/Programs/Shared/EpicGames.Horde/Compute/Buffers/SharedMemoryBuffer.cs
@@ -156,5 +156,17 @@ namespace EpicGames.Horde.Compute.Buffers
 			}
 			return chunks;
 		}
+
+		internal static string GetName(IComputeBufferReader reader)
+		{
+			Resources resources = (Resources)((ReaderImpl)reader).GetResources();
+			return resources.Name;
+		}
+
+		internal static string GetName(IComputeBufferWriter writer)
+		{
+			Resources resources = (Resources)((WriterImpl)writer).GetResources();
+			return resources.Name;
+		}
 	}
 }
diff --git a/Engine/Source/Programs/Shared/EpicGames.Horde/Compute/ComputeBufferBase.cs b/Engine/Source/Programs/Shared/EpicGames.Horde/Compute/ComputeBufferBase.cs
--- a/Engine/Source/Programs/Shared/EpicGames.Horde/Compute/ComputeBufferBase.cs
+++ b/Engine/Source/Programs/Shared/EpicGames.Horde/Compute/ComputeBufferBase.cs
@@ -254,7 +254,7 @@ namespace EpicGames.Horde.Compute
 #pragma warning restore IDE0051 // Remove unused private members
 		}
 
-		class ReaderImpl : IComputeBufferReader
+		internal class ReaderImpl : IComputeBufferReader
 		{
 			readonly HeaderPtr _headerPtr;
 			readonly Memory<byte>[] _chunks;
@@ -269,6 +269,8 @@ namespace EpicGames.Horde.Compute
 				_readerIdx = readerIdx;
 			}
 
+			public ResourcesBase GetResources() => _resources;
+
 			/// <inheritdoc/>
 			public IComputeBufferReader AddRef()
 			{
@@ -384,7 +386,7 @@ namespace EpicGames.Horde.Compute
 			}
 		}
 
-		class WriterImpl : IComputeBufferWriter
+		internal class WriterImpl : IComputeBufferWriter
 		{
 			readonly HeaderPtr _headerPtr;
 			readonly Memory<byte>[] _chunks;
@@ -397,6 +399,8 @@ namespace EpicGames.Horde.Compute
 				_resources = resources;
 			}
 
+			public ResourcesBase GetResources() => _resources;
+
 			/// <inheritdoc/>
 			public IComputeBufferWriter AddRef()
 			{
diff --git a/Engine/Source/Programs/Shared/EpicGames.Horde/Compute/IComputeChannel.cs b/Engine/Source/Programs/Shared/EpicGames.Horde/Compute/IComputeChannel.cs
--- a/Engine/Source/Programs/Shared/EpicGames.Horde/Compute/IComputeChannel.cs
+++ b/Engine/Source/Programs/Shared/EpicGames.Horde/Compute/IComputeChannel.cs
@@ -4,6 +4,7 @@ using System;
 using System.IO;
 using System.Threading;
 using System.Threading.Tasks;
+using EpicGames.Horde.Compute.Buffers;
 
 namespace EpicGames.Horde.Compute
 {
@@ -103,6 +104,18 @@ namespace EpicGames.Horde.Compute
 		/// <param name="recvBuffer">Buffer for receiving data</param>
 		public static IComputeChannel CreateChannel(this RemoteComputeSocket socket, int channelId, IComputeBuffer recvBuffer) => new BufferedReaderChannel(socket, channelId, recvBuffer);
 
+		/// <summary>
+		/// Creates a channel using a socket and receive buffer
+		/// </summary>
+		/// <param name="socket">Socket to use for sending data</param>
+		/// <param name="channelId">Channel id to send and receive data</param>
+		public static IComputeChannel CreateChannel(this IComputeSocket socket, int channelId)
+		{
+			using SharedMemoryBuffer recvBuffer = SharedMemoryBuffer.CreateNew(null, 65536);
+			using SharedMemoryBuffer sendBuffer = SharedMemoryBuffer.CreateNew(null, 65536);
+			return CreateChannel(socket, channelId, recvBuffer, sendBuffer);
+		}
+
 		/// <summary>
 		/// Creates a channel using a socket and receive buffer
 		/// </summary>
diff --git a/Engine/Source/Programs/Shared/EpicGames.Horde/Compute/WorkerComputeSocket.cs b/Engine/Source/Programs/Shared/EpicGames.Horde/Compute/WorkerComputeSocket.cs
--- a/Engine/Source/Programs/Shared/EpicGames.Horde/Compute/WorkerComputeSocket.cs
+++ b/Engine/Source/Programs/Shared/EpicGames.Horde/Compute/WorkerComputeSocket.cs
@@ -1,6 +1,7 @@
 // Copyright Epic Games, Inc. All Rights Reserved.
 
 using System;
+using System.Threading.Tasks;
 using EpicGames.Core;
 using EpicGames.Horde.Compute.Buffers;
 
@@ -15,7 +16,7 @@ namespace EpicGames.Horde.Compute
 	/// <summary>
 	/// Provides functionality for attaching buffers for compute workers 
 	/// </summary>
-	public sealed class WorkerComputeSocket : IDisposable
+	public sealed class WorkerComputeSocket : IComputeSocket, IDisposable
 	{
 		/// <summary>
 		/// Name of the environment variable for passing the name of the compute channel
@@ -53,59 +54,34 @@ namespace EpicGames.Horde.Compute
 			_commandBuffer.Dispose();
 		}
 
-		/// <summary>
-		/// Attaches a new buffer for receiving data
-		/// </summary>
-		/// <param name="channelId">Channel id for the buffer</param>
-		/// <param name="buffer">Buffer to attach</param>
-		public void AttachRecvBuffer(int channelId, SharedMemoryBuffer buffer) => AttachBuffer(IpcMessage.AttachRecvBuffer, channelId, buffer);
-
-		/// <summary>
-		/// Attaches a new buffer for sending data
-		/// </summary>
-		/// <param name="channelId">Channel id for the buffer</param>
-		/// <param name="buffer">Buffer to attach</param>
-		public void AttachSendBuffer(int channelId, SharedMemoryBuffer buffer) => AttachBuffer(IpcMessage.AttachSendBuffer, channelId, buffer);
+		/// <inheritdoc/>
+		public ValueTask DisposeAsync()
+		{
+			Dispose();
+			return new ValueTask();
+		}
 
-		void AttachBuffer(IpcMessage message, int channelId, SharedMemoryBuffer buffer)
+		/// <inheritdoc/>
+		public void AttachRecvBuffer(int channelId, IComputeBufferWriter writer)
 		{
-			MemoryWriter writer = new MemoryWriter(_commandBuffer.Writer.GetWriteBuffer());
-			writer.WriteUnsignedVarInt((int)message);
-			writer.WriteUnsignedVarInt(channelId);
-			writer.WriteString(buffer.Name);
-			_commandBuffer.Writer.AdvanceWritePosition(writer.Length);
+			string bufferName = SharedMemoryBuffer.GetName(writer);
+			AttachBuffer(IpcMessage.AttachRecvBuffer, channelId, bufferName);
 		}
-	}
 
-	/// <summary>
-	/// Extension methods for worker sockets
-	/// </summary>
-	public static class WorkerComputeSocketExtensions
-	{
-		/// <summary>
-		/// Creates a channel using a socket and receive buffer
-		/// </summary>
-		/// <param name="socket">Socket to use for sending data</param>
-		/// <param name="channelId">Channel id to send and receive data</param>
-		public static IComputeChannel CreateChannel(this WorkerComputeSocket socket, int channelId)
+		/// <inheritdoc/>
+		public void AttachSendBuffer(int channelId, IComputeBufferReader reader)
 		{
-			using SharedMemoryBuffer recvBuffer = SharedMemoryBuffer.CreateNew(null, 65536);
-			using SharedMemoryBuffer sendBuffer = SharedMemoryBuffer.CreateNew(null, 65536);
-			return CreateChannel(socket, channelId, recvBuffer, sendBuffer);
+			string bufferName = SharedMemoryBuffer.GetName(reader);
+			AttachBuffer(IpcMessage.AttachSendBuffer, channelId, bufferName);
 		}
 
-		/// <summary>
-		/// Creates a channel using a socket and receive buffer
-		/// </summary>
-		/// <param name="socket">Socket to use for sending data</param>
-		/// <param name="channelId">Channel id to send and receive data</param>
-		/// <param name="recvBuffer">Buffer for receiving data</param>
-		/// <param name="sendBuffer">Buffer for sending data</param>
-		public static IComputeChannel CreateChannel(this WorkerComputeSocket socket, int channelId, SharedMemoryBuffer recvBuffer, SharedMemoryBuffer sendBuffer)
+		void AttachBuffer(IpcMessage message, int channelId, string bufferName)
 		{
-			socket.AttachRecvBuffer(channelId, recvBuffer);
-			socket.AttachSendBuffer(channelId, sendBuffer);
-			return new ComputeChannel.BufferedReaderWriterChannel(recvBuffer.Reader, sendBuffer.Writer);
+			MemoryWriter writer = new MemoryWriter(_commandBuffer.Writer.GetWriteBuffer());
+			writer.WriteUnsignedVarInt((int)message);
+			writer.WriteUnsignedVarInt(channelId);
+			writer.WriteString(bufferName);
+			_commandBuffer.Writer.AdvanceWritePosition(writer.Length);
 		}
 	}
 }
