Horde: Rename compute classes designed solely for communicating with the Horde agent to distinguish from the general-purpose compute API.

diff --git a/Engine/Source/Programs/Horde/Horde.Agent/Commands/Compute/ComputeCommand.cs b/Engine/Source/Programs/Horde/Horde.Agent/Commands/Compute/ComputeCommand.cs
--- a/Engine/Source/Programs/Horde/Horde.Agent/Commands/Compute/ComputeCommand.cs
+++ b/Engine/Source/Programs/Horde/Horde.Agent/Commands/Compute/ComputeCommand.cs
@@ -143,7 +143,7 @@ namespace Horde.Agent.Commands.Compute
 		{
 			DirectoryReference sandboxDir = DirectoryReference.Combine(Program.DataDir, "Sandbox");
 
-			ComputeWorker worker = new ComputeWorker(sandboxDir, memoryCache, logger);
+			AgentMessageHandler worker = new AgentMessageHandler(sandboxDir, memoryCache, logger);
 			await worker.RunAsync(socket, cancellationToken);
 		}
 	}
diff --git a/Engine/Source/Programs/Horde/Horde.Agent/Commands/Compute/CppComputeCommand.cs b/Engine/Source/Programs/Horde/Horde.Agent/Commands/Compute/CppComputeCommand.cs
--- a/Engine/Source/Programs/Horde/Horde.Agent/Commands/Compute/CppComputeCommand.cs
+++ b/Engine/Source/Programs/Horde/Horde.Agent/Commands/Compute/CppComputeCommand.cs
@@ -58,11 +58,11 @@ namespace Horde.Agent.Commands.Compute
 			NodeLocator sandbox = await CreateSandboxAsync(TaskFile, storage, cancellationToken);
 
 			// Open a socket and upload the sandbox
-			using (IComputeMessageChannel channel = lease.Socket.CreateMessageChannel(ControlChannelId, 4 * 1024 * 1024, _logger))
+			using (AgentMessageChannel channel = lease.Socket.CreateAgentMessageChannel(ControlChannelId, 4 * 1024 * 1024, _logger))
 			{
 				await channel.UploadFilesAsync("", sandbox, storage, cancellationToken);
 
-				await using (IComputeProcess process = await channel.ExecuteAsync(jsonComputeTask.Executable, jsonComputeTask.Arguments, jsonComputeTask.WorkingDir, jsonComputeTask.EnvVars, cancellationToken))
+				await using (AgentManagedProcess process = await channel.ExecuteAsync(jsonComputeTask.Executable, jsonComputeTask.Arguments, jsonComputeTask.WorkingDir, jsonComputeTask.EnvVars, cancellationToken))
 				{
 					string? line;
 					while ((line = await process.ReadLineAsync(cancellationToken)) != null)
diff --git a/Engine/Source/Programs/Horde/Horde.Agent/Commands/Compute/DdcComputeCommand.cs b/Engine/Source/Programs/Horde/Horde.Agent/Commands/Compute/DdcComputeCommand.cs
--- a/Engine/Source/Programs/Horde/Horde.Agent/Commands/Compute/DdcComputeCommand.cs
+++ b/Engine/Source/Programs/Horde/Horde.Agent/Commands/Compute/DdcComputeCommand.cs
@@ -31,13 +31,13 @@ namespace Horde.Agent.Commands.Compute
 		protected override async Task<bool> HandleRequestAsync(IComputeLease lease, CancellationToken cancellationToken)
 		{
 			IComputeSocket socket = lease.Socket;
-			using (IComputeMessageChannel channel = socket.CreateMessageChannel(0, _logger))
+			using (AgentMessageChannel channel = socket.CreateAgentMessageChannel(0, _logger))
 			{
 				_logger.LogInformation("Sending XOR request");
 				await channel.SendXorRequestAsync(new byte[] { 1, 2, 3, 4, 5 }, (byte)123, cancellationToken);
 
 				_logger.LogInformation("Waiting for response...");
-				IComputeMessage response = await channel.ReceiveAsync(cancellationToken);
+				IAgentMessage response = await channel.ReceiveAsync(cancellationToken);
 
 				byte[] result = response.Data.ToArray();
 				byte[] expectedResult = new byte[] { 1 ^ 123, 2 ^ 123, 3 ^ 123, 4 ^ 123, 5 ^ 123 };
diff --git a/Engine/Source/Programs/Horde/Horde.Agent/Commands/Compute/XorComputeCommand.cs b/Engine/Source/Programs/Horde/Horde.Agent/Commands/Compute/XorComputeCommand.cs
--- a/Engine/Source/Programs/Horde/Horde.Agent/Commands/Compute/XorComputeCommand.cs
+++ b/Engine/Source/Programs/Horde/Horde.Agent/Commands/Compute/XorComputeCommand.cs
@@ -31,17 +31,17 @@ namespace Horde.Agent.Commands.Compute
 		protected override async Task<bool> HandleRequestAsync(IComputeLease lease, CancellationToken cancellationToken)
 		{
 			IComputeSocket socket = lease.Socket;
-			using (IComputeMessageChannel mainChannel = socket.CreateMessageChannel(0, _logger))
+			using (AgentMessageChannel mainChannel = socket.CreateAgentMessageChannel(0, _logger))
 			{
 				_logger.LogInformation("Forking compute channel...");
 
-				IComputeMessageChannel channel = mainChannel;
+				AgentMessageChannel channel = mainChannel;
 
 				_logger.LogInformation("Sending XOR request");
 				await channel.SendXorRequestAsync(new byte[] { 1, 2, 3, 4, 5 }, (byte)123, cancellationToken);
 
 				_logger.LogInformation("Waiting for response...");
-				IComputeMessage response = await channel.ReceiveAsync(cancellationToken);
+				IAgentMessage response = await channel.ReceiveAsync(cancellationToken);
 
 				byte[] result = response.Data.ToArray();
 				byte[] expectedResult = new byte[] { 1 ^ 123, 2 ^ 123, 3 ^ 123, 4 ^ 123, 5 ^ 123 };
diff --git a/Engine/Source/Programs/Horde/Horde.Agent/Leases/Handlers/ComputeHandler.cs b/Engine/Source/Programs/Horde/Horde.Agent/Leases/Handlers/ComputeHandler.cs
--- a/Engine/Source/Programs/Horde/Horde.Agent/Leases/Handlers/ComputeHandler.cs
+++ b/Engine/Source/Programs/Horde/Horde.Agent/Leases/Handlers/ComputeHandler.cs
@@ -99,7 +99,7 @@ namespace Horde.Agent.Leases.Handlers
 						{
 							DirectoryReference.CreateDirectory(sandboxDir);
 
-							ComputeWorker worker = new ComputeWorker(sandboxDir, _memoryCache, _logger);
+							AgentMessageHandler worker = new AgentMessageHandler(sandboxDir, _memoryCache, _logger);
 							await worker.RunAsync(socket, cts.Token);
 							await socket.CloseAsync(cts.Token);
 							return LeaseResult.Success;
diff --git a/Engine/Source/Programs/Horde/Samples/RemoteClient/RemoteClientApp.cs b/Engine/Source/Programs/Horde/Samples/RemoteClient/RemoteClientApp.cs
--- a/Engine/Source/Programs/Horde/Samples/RemoteClient/RemoteClientApp.cs
+++ b/Engine/Source/Programs/Horde/Samples/RemoteClient/RemoteClientApp.cs
@@ -84,12 +84,12 @@ namespace RemoteClient
 		static async Task RunRemoteAsync(IComputeLease lease, DirectoryReference uploadDir, string executable, List<string> arguments, ILogger logger)
 		{
 			// Create a message channel on channel id 0. The Horde Agent always listens on this channel for requests.
-			using (IComputeMessageChannel channel = lease.Socket.CreateMessageChannel(PrimaryChannelId, 4 * 1024 * 1024, logger))
+			using (AgentMessageChannel channel = lease.Socket.CreateAgentMessageChannel(PrimaryChannelId, 4 * 1024 * 1024, logger))
 			{
 				await channel.WaitForAttachAsync();
 
 				// Fork another message loop. We'll use this to run an XOR task in the background.
-				using IComputeMessageChannel backgroundChannel = lease.Socket.CreateMessageChannel(BackgroundChannelId, 4 * 1024 * 1024, logger);
+				using AgentMessageChannel backgroundChannel = lease.Socket.CreateAgentMessageChannel(BackgroundChannelId, 4 * 1024 * 1024, logger);
 				await using BackgroundTask otherChannelTask = BackgroundTask.StartNew(ctx => RunBackgroundXorAsync(backgroundChannel));
 				await channel.ForkAsync(BackgroundChannelId, 4 * 1024 * 1024, default);
 
@@ -104,7 +104,7 @@ namespace RemoteClient
 				}
 
 				// Run the task remotely in the background and echo the output to the console
-				await using (IComputeProcess process = await channel.ExecuteAsync(executable, arguments, null, null))
+				await using (AgentManagedProcess process = await channel.ExecuteAsync(executable, arguments, null, null))
 				{
 					await using BackgroundTask tickTask = BackgroundTask.StartNew(ctx => WriteNumbersAsync(lease.Socket, logger, ctx));
 
@@ -141,14 +141,14 @@ namespace RemoteClient
 			await socket.MarkCompleteAsync(ChildProcessChannelId, cancellationToken);
 		}
 
-		static async Task RunBackgroundXorAsync(IComputeMessageChannel channel)
+		static async Task RunBackgroundXorAsync(AgentMessageChannel channel)
 		{
 			await channel.WaitForAttachAsync();
 
 			byte[] dataToXor = new byte[] { 1, 2, 3, 4, 5 };
 			await channel.SendXorRequestAsync(dataToXor, 123);
 
-			using IComputeMessage response = await channel.ReceiveAsync(ComputeMessageType.XorResponse);
+			using IAgentMessage response = await channel.ReceiveAsync(AgentMessageType.XorResponse);
 			for (int idx = 0; idx < dataToXor.Length; idx++)
 			{
 				if (response.Data.Span[idx] != (byte)(dataToXor[idx] ^ 123))
diff --git a/Engine/Source/Programs/Shared/EpicGames.Horde/Compute/AgentManagedProcess.cs b/Engine/Source/Programs/Shared/EpicGames.Horde/Compute/AgentManagedProcess.cs
new file mode 100644
--- /dev/null
+++ b/Engine/Source/Programs/Shared/EpicGames.Horde/Compute/AgentManagedProcess.cs
@@ -0,0 +1,138 @@
+// Copyright Epic Games, Inc. All Rights Reserved.
+
+using System;
+using System.Text;
+using System.Threading;
+using System.Threading.Channels;
+using System.Threading.Tasks;
+using EpicGames.Core;
+
+namespace EpicGames.Horde.Compute
+{
+	/// <summary>
+	/// Represents a remotely executed process managed by the Horde agent
+	/// </summary>
+	public sealed class AgentManagedProcess : IAsyncDisposable
+	{
+		readonly Channel<string> _output;
+		readonly BackgroundTask _backgroundTask;
+		readonly TaskCompletionSource<int> _result = new TaskCompletionSource<int>();
+
+		byte[] _buffer = new byte[1024];
+		int _bufferLength;
+
+		/// <inheritdoc/>
+		public int ExitCode => _result.Task.Result;
+
+		/// <inheritdoc/>
+		public bool HasExited => _result.Task.IsCompleted;
+
+		/// <summary>
+		/// Constructor
+		/// </summary>
+		public AgentManagedProcess(AgentMessageChannel channel)
+		{
+			_output = Channel.CreateUnbounded<string>();
+			_backgroundTask = BackgroundTask.StartNew(ctx => RunAsync(channel, ctx));
+		}
+
+		/// <inheritdoc/>
+		public ValueTask DisposeAsync() => _backgroundTask.DisposeAsync();
+
+		/// <inheritdoc/>
+		public async ValueTask<string?> ReadLineAsync(CancellationToken cancellationToken = default)
+		{
+			if (!await _output.Reader.WaitToReadAsync(cancellationToken))
+			{
+				return null;
+			}
+			return await _output.Reader.ReadAsync(cancellationToken);
+		}
+
+		/// <inheritdoc/>
+		public Task<int> WaitForExitAsync(CancellationToken cancellationToken) => _result.Task.WaitAsync(cancellationToken);
+
+		async Task RunAsync(AgentMessageChannel channel, CancellationToken cancellationToken)
+		{
+			try
+			{
+				await RunInternalAsync(channel, cancellationToken);
+				_output.Writer.TryComplete();
+			}
+			catch (Exception ex)
+			{
+				_result.TrySetException(ex);
+				_output.Writer.TryComplete(ex);
+			}
+		}
+
+		async Task RunInternalAsync(AgentMessageChannel channel, CancellationToken cancellationToken)
+		{
+			for (; ; )
+			{
+				using IAgentMessage message = await channel.ReceiveAsync(cancellationToken);
+				switch (message.Type)
+				{
+					case AgentMessageType.Exception:
+						ExceptionMessage exception = message.ParseExceptionMessage();
+						throw new ComputeRemoteException(exception);
+					case AgentMessageType.ExecuteOutput:
+						AppendData(message.Data.Span);
+						break;
+					case AgentMessageType.ExecuteResult:
+						ExecuteProcessResponseMessage executeProcessResponse = message.ParseExecuteProcessResponse();
+						_result.TrySetResult(executeProcessResponse.ExitCode);
+						return;
+					default:
+						throw new InvalidAgentMessageException(message);
+				}
+			}
+		}
+
+		void AppendData(ReadOnlySpan<byte> data)
+		{
+			for (; ; )
+			{
+				int lineEnd = data.IndexOf((byte)'\n');
+				if (lineEnd == -1)
+				{
+					AppendToBuffer(data);
+					break;
+				}
+
+				int lineLen = lineEnd;
+				if (lineLen > 0 && data[lineLen - 1] == '\r')
+				{
+					lineLen--;
+				}
+
+				string str;
+				if (_bufferLength == 0)
+				{
+					str = Encoding.UTF8.GetString(data.Slice(0, lineLen));
+				}
+				else
+				{
+					AppendToBuffer(data.Slice(0, lineLen));
+					str = Encoding.UTF8.GetString(_buffer.AsSpan(0, _bufferLength + lineLen));
+					_bufferLength = 0;
+				}
+
+				_output.Writer.TryWrite(str);
+				data = data.Slice(lineEnd + 1);
+			}
+		}
+
+		void AppendToBuffer(ReadOnlySpan<byte> data)
+		{
+			if (_bufferLength + data.Length > _buffer.Length)
+			{
+				byte[] newBuffer = new byte[_bufferLength + data.Length + 256];
+				_buffer.AsSpan(0, _bufferLength).CopyTo(newBuffer);
+				_buffer = newBuffer;
+			}
+			data.CopyTo(_buffer.AsSpan(_bufferLength));
+			_bufferLength += data.Length;
+		}
+	}
+}
diff --git a/Engine/Source/Programs/Shared/EpicGames.Horde/Compute/AgentMessage.cs b/Engine/Source/Programs/Shared/EpicGames.Horde/Compute/AgentMessage.cs
new file mode 100644
--- /dev/null
+++ b/Engine/Source/Programs/Shared/EpicGames.Horde/Compute/AgentMessage.cs
@@ -0,0 +1,587 @@
+// Copyright Epic Games, Inc. All Rights Reserved.
+
+using System;
+using System.Buffers.Binary;
+using System.Collections.Generic;
+using System.IO;
+using System.Threading;
+using System.Threading.Tasks;
+using EpicGames.Core;
+using EpicGames.Horde.Storage;
+using EpicGames.Horde.Storage.Nodes;
+
+namespace EpicGames.Horde.Compute
+{
+	/// <summary>
+	/// Type of a compute message
+	/// </summary>
+	public enum AgentMessageType
+	{
+		/// <summary>
+		/// No message was received (end of stream)
+		/// </summary>
+		None = 0x00,
+
+		/// <summary>
+		/// Sent in place of a regular response if an error occurs on the remote
+		/// </summary>
+		Exception = 0x02,
+
+		/// <summary>
+		/// Fork the message loop into a new channel
+		/// </summary>
+		Fork = 0x03,
+
+		/// <summary>
+		/// Sent as the first message on a channel to notify the remote that the remote end is attached
+		/// </summary>
+		Attach = 0x04,
+
+		#region Process Management
+
+		/// <summary>
+		/// Extract files on the remote machine (Initiator -> Remote)
+		/// </summary>
+		WriteFiles = 0x10,
+
+		/// <summary>
+		/// Notification that files have been extracted (Remote -> Initiator)
+		/// </summary>
+		WriteFilesResponse = 0x11,
+
+		/// <summary>
+		/// Deletes files on the remote machine (Initiator -> Remote)
+		/// </summary>
+		DeleteFiles = 0x12,
+
+		/// <summary>
+		/// Execute a process in a sandbox (Initiator -> Remote)
+		/// </summary>
+		Execute = 0x16,
+
+		/// <summary>
+		/// Returns output from the child process to the caller (Remote -> Initiator)
+		/// </summary>
+		ExecuteOutput = 0x17,
+
+		/// <summary>
+		/// Returns the process exit code (Remote -> Initiator)
+		/// </summary>
+		ExecuteResult = 0x18,
+
+		#endregion
+
+		#region Storage
+
+		/// <summary>
+		/// Reads a blob from storage
+		/// </summary>
+		ReadBlob = 0x20,
+
+		/// <summary>
+		/// Response to a <see cref="ReadBlob"/> request.
+		/// </summary>
+		ReadBlobResponse = 0x21,
+
+		#endregion
+
+		#region Test Requests
+
+		/// <summary>
+		/// Xor a block of data with a value
+		/// </summary>
+		XorRequest = 0xf0,
+
+		/// <summary>
+		/// Result from an <see cref="XorRequest"/> request.
+		/// </summary>
+		XorResponse = 0xf1,
+
+		#endregion
+	}
+
+	/// <summary>
+	/// Information about a received compute message
+	/// </summary>
+	public interface IAgentMessage : IMemoryReader, IDisposable
+	{
+		/// <summary>
+		/// Type of the message
+		/// </summary>
+		public AgentMessageType Type { get; }
+
+		/// <summary>
+		/// Data that was read
+		/// </summary>
+		public ReadOnlyMemory<byte> Data { get; }
+	}
+
+	/// <summary>
+	/// Message for reporting an error
+	/// </summary>
+	public record struct ExceptionMessage(string Message, string Description);
+
+	/// <summary>
+	/// Message requesting that the message loop be forked
+	/// </summary>
+	/// <param name="channelId">New channel to communicate on</param>
+	/// <param name="bufferSize">Size of the buffer</param>
+	public record struct ForkMessage(int channelId, int bufferSize);
+
+	/// <summary>
+	/// Extract files from a bundle to a path in the remote sandbox
+	/// </summary>
+	/// <param name="Name">Path to extract the files to</param>
+	/// <param name="Locator">Locator for the tree to extract</param>
+	public record struct UploadFilesMessage(string Name, NodeLocator Locator);
+
+	/// <summary>
+	/// Deletes files or directories in the remote
+	/// </summary>
+	/// <param name="Filter">Filter for files to delete</param>
+	public record struct DeleteFilesMessage(IReadOnlyList<string> Filter);
+
+	/// <summary>
+	/// Message to execute a new child process
+	/// </summary>
+	/// <param name="Executable">Executable path</param>
+	/// <param name="Arguments">Arguments for the executable</param>
+	/// <param name="WorkingDir">Working directory to execute in</param>
+	/// <param name="EnvVars">Environment variables for the child process. Null values unset variables.</param>
+	public record struct ExecuteProcessMessage(string Executable, IReadOnlyList<string> Arguments, string? WorkingDir, IReadOnlyDictionary<string, string?> EnvVars);
+
+	/// <summary>
+	/// Response from executing a child process
+	/// </summary>
+	/// <param name="ExitCode">Exit code for the process</param>
+	public record struct ExecuteProcessResponseMessage(int ExitCode);
+
+	/// <summary>
+	/// Creates a blob read request
+	/// </summary>
+	public record struct ReadBlobMessage(BlobLocator Locator, int Offset, int Length);
+
+	/// <summary>
+	/// Message for running an XOR command
+	/// </summary>
+	/// <param name="Data">Data to xor</param>
+	/// <param name="Value">Value to XOR with</param>
+	public record struct XorRequestMessage(ReadOnlyMemory<byte> Data, byte Value);
+
+	/// <summary>
+	/// Wraps various requests across compute channels
+	/// </summary>
+	public static class AgentMessageExtensions
+	{
+		/// <summary>
+		/// Closes the remote message loop
+		/// </summary>
+		public static async ValueTask CloseAsync(this AgentMessageChannel channel, CancellationToken cancellationToken = default)
+		{
+			using IAgentMessageBuilder message = await channel.CreateMessageAsync(AgentMessageType.None, cancellationToken);
+			message.Send();
+		}
+
+		/// <summary>
+		/// Sends an exception response to the remote
+		/// </summary>
+		public static ValueTask SendExceptionAsync(this AgentMessageChannel channel, Exception ex, CancellationToken cancellationToken = default) => SendExceptionAsync(channel, ex.Message, ex.ToString(), cancellationToken);
+
+		/// <summary>
+		/// Sends an exception response to the remote
+		/// </summary>
+		public static async ValueTask SendExceptionAsync(this AgentMessageChannel channel, string description, string trace, CancellationToken cancellationToken = default)
+		{
+			using IAgentMessageBuilder message = await channel.CreateMessageAsync(AgentMessageType.Exception, cancellationToken);
+			message.WriteString(description);
+			message.WriteString(trace);
+			message.Send();
+		}
+
+		/// <summary>
+		/// Parses a message as an <see cref="ExceptionMessage"/>
+		/// </summary>
+		public static ExceptionMessage ParseExceptionMessage(this IAgentMessage message)
+		{
+			string msg = message.ReadString();
+			string description = message.ReadString();
+			return new ExceptionMessage(msg, description);
+		}
+
+		/// <summary>
+		/// Requests that the remote message loop be forked
+		/// </summary>
+		public static async ValueTask ForkAsync(this AgentMessageChannel channel, int channelId, int bufferSize, CancellationToken cancellationToken = default)
+		{
+			using IAgentMessageBuilder message = await channel.CreateMessageAsync(AgentMessageType.Fork, cancellationToken);
+			message.WriteInt32(channelId);
+			message.WriteInt32(bufferSize);
+			message.Send();
+		}
+
+		/// <summary>
+		/// Parses a fork request message
+		/// </summary>
+		public static ForkMessage ParseForkMessage(this IAgentMessage message)
+		{
+			int channelId = message.ReadInt32();
+			int bufferSize = message.ReadInt32();
+			return new ForkMessage(channelId, bufferSize);
+		}
+
+		/// <summary>
+		/// Notifies the remote that a buffer has been attached
+		/// </summary>
+		public static async ValueTask AttachAsync(this AgentMessageChannel channel, CancellationToken cancellationToken = default)
+		{
+			using IAgentMessageBuilder message = await channel.CreateMessageAsync(AgentMessageType.Attach, cancellationToken);
+			message.Send();
+		}
+
+		/// <summary>
+		/// Waits until an attached notification is received along the channel
+		/// </summary>
+		/// <param name="channel"></param>
+		/// <param name="cancellationToken"></param>
+		/// <returns></returns>
+		public static async ValueTask WaitForAttachAsync(this AgentMessageChannel channel, CancellationToken cancellationToken = default)
+		{
+			using IAgentMessage message = await channel.ReceiveAsync(cancellationToken);
+			if (message.Type != AgentMessageType.Attach)
+			{
+				throw new InvalidAgentMessageException(message);
+			}
+		}
+
+		#region Process
+
+		static async Task<IAgentMessage> RunStorageServer(this AgentMessageChannel channel, IStorageClient storage, CancellationToken cancellationToken = default)
+		{
+			for (; ; )
+			{
+				IAgentMessage message = await channel.ReceiveAsync(cancellationToken);
+				if (message.Type != AgentMessageType.ReadBlob)
+				{
+					return message;
+				}
+
+				try
+				{
+					ReadBlobMessage readBlob = message.ParseReadBlobRequest();
+					await SendBlobDataAsync(channel, readBlob, storage, cancellationToken);
+				}
+				finally
+				{
+					message.Dispose();
+				}
+			}
+		}
+
+		/// <summary>
+		/// Creates a sandbox on the remote machine
+		/// </summary>
+		/// <param name="channel">Current channel</param>
+		/// <param name="path">Root directory to extract files within the sandbox</param>
+		/// <param name="locator">Location of a <see cref="DirectoryNode"/> describing contents of the sandbox</param>
+		/// <param name="storage">Storage for the sandbox data</param>
+		/// <param name="cancellationToken">Cancellation token for the operation</param>
+		public static async Task UploadFilesAsync(this AgentMessageChannel channel, string path, NodeLocator locator, IStorageClient storage, CancellationToken cancellationToken = default)
+		{
+			using (IAgentMessageBuilder request = await channel.CreateMessageAsync(AgentMessageType.WriteFiles, cancellationToken))
+			{
+				request.WriteString(path);
+				request.WriteNodeLocator(locator);
+				request.Send();
+			}
+
+			using IAgentMessage response = await RunStorageServer(channel, storage, cancellationToken);
+			if (response.Type != AgentMessageType.WriteFilesResponse)
+			{
+				throw new InvalidAgentMessageException(response);
+			}
+		}
+
+		/// <summary>
+		/// Parses a message as a <see cref="UploadFilesMessage"/>
+		/// </summary>
+		public static UploadFilesMessage ParseUploadFilesMessage(this IAgentMessage message)
+		{
+			string name = message.ReadString();
+			NodeLocator locator = message.ReadNodeLocator();
+			return new UploadFilesMessage(name, locator);
+		}
+
+		/// <summary>
+		/// Destroys a sandbox on the remote machine
+		/// </summary>
+		/// <param name="channel">Current channel</param>
+		/// <param name="paths">Paths of files or directories to delete</param>
+		/// <param name="cancellationToken">Cancellation token for the operation</param>
+		public static async ValueTask DeleteFilesAsync(this AgentMessageChannel channel, IReadOnlyList<string> paths, CancellationToken cancellationToken)
+		{
+			IAgentMessageBuilder request = await channel.CreateMessageAsync(AgentMessageType.DeleteFiles, cancellationToken);
+			request.WriteList(paths, MemoryWriterExtensions.WriteString);
+			request.Send();
+		}
+
+		/// <summary>
+		/// Parses a message as a <see cref="DeleteFilesMessage"/>
+		/// </summary>
+		public static DeleteFilesMessage ParseDeleteFilesMessage(this IAgentMessage message)
+		{
+			List<string> files = message.ReadList(MemoryReaderExtensions.ReadString);
+			return new DeleteFilesMessage(files);
+		}
+
+		/// <summary>
+		/// Executes a remote process
+		/// </summary>
+		/// <param name="channel">Current channel</param>
+		/// <param name="executable">Executable to run, relative to the sandbox root</param>
+		/// <param name="arguments">Arguments for the child process</param>
+		/// <param name="workingDir">Working directory for the process</param>
+		/// <param name="envVars">Environment variables for the child process</param>
+		/// <param name="cancellationToken">Cancellation token for the operation</param>
+		public static async Task<AgentManagedProcess> ExecuteAsync(this AgentMessageChannel channel, string executable, IReadOnlyList<string> arguments, string? workingDir, IReadOnlyDictionary<string, string?>? envVars, CancellationToken cancellationToken = default)
+		{
+			using (IAgentMessageBuilder request = await channel.CreateMessageAsync(AgentMessageType.Execute, cancellationToken))
+			{
+				request.WriteString(executable);
+				request.WriteList(arguments, MemoryWriterExtensions.WriteString);
+				request.WriteOptionalString(workingDir);
+				request.WriteDictionary(envVars ?? new Dictionary<string, string?>(), MemoryWriterExtensions.WriteString, MemoryWriterExtensions.WriteOptionalString);
+				request.Send();
+			}
+			return new AgentManagedProcess(channel);
+		}
+
+		/// <summary>
+		/// Parses a message as a <see cref="ExecuteProcessMessage"/>
+		/// </summary>
+		public static ExecuteProcessMessage ParseExecuteProcessMessage(this IAgentMessage message)
+		{
+			string executable = message.ReadString();
+			List<string> arguments = message.ReadList(MemoryReaderExtensions.ReadString);
+			string? workingDir = message.ReadOptionalString();
+			Dictionary<string, string?> envVars = message.ReadDictionary(MemoryReaderExtensions.ReadString, MemoryReaderExtensions.ReadOptionalString);
+			return new ExecuteProcessMessage(executable, arguments, workingDir, envVars);
+		}
+
+		/// <summary>
+		/// Sends output from a child process
+		/// </summary>
+		public static async ValueTask SendExecuteOutputAsync(this AgentMessageChannel channel, ReadOnlyMemory<byte> data, CancellationToken cancellationToken = default)
+		{
+			using IAgentMessageBuilder message = await channel.CreateMessageAsync(AgentMessageType.ExecuteOutput, cancellationToken);
+			message.WriteFixedLengthBytes(data.Span);
+			message.Send();
+		}
+
+		/// <summary>
+		/// Sends a response from executing a child process
+		/// </summary>
+		/// <param name="channel"></param>
+		/// <param name="exitCode">Exit code from the process</param>
+		/// <param name="cancellationToken">Cancellation token for the operation</param>
+		public static async ValueTask SendExecuteResultAsync(this AgentMessageChannel channel, int exitCode, CancellationToken cancellationToken = default)
+		{
+			using IAgentMessageBuilder builder = await channel.CreateMessageAsync(AgentMessageType.ExecuteResult, cancellationToken);
+			builder.WriteInt32(exitCode);
+			builder.Send();
+		}
+
+		/// <summary>
+		/// Parses a message as a <see cref="ExecuteProcessMessage"/>
+		/// </summary>
+		public static ExecuteProcessResponseMessage ParseExecuteProcessResponse(this IAgentMessage message)
+		{
+			int exitCode = message.ReadInt32();
+			return new ExecuteProcessResponseMessage(exitCode);
+		}
+
+		#endregion
+
+		#region Storage
+
+		/// <summary>
+		/// 
+		/// </summary>
+		/// <param name="message"></param>
+		/// <returns></returns>
+		public static ReadBlobMessage ParseReadBlobRequest(this IAgentMessage message)
+		{
+			BlobLocator locator = message.ReadBlobLocator();
+			int offset = (int)message.ReadUnsignedVarInt();
+			int length = (int)message.ReadUnsignedVarInt();
+			return new ReadBlobMessage(locator, offset, length);
+		}
+
+		/// <summary>
+		/// Wraps a compute message containing blob data
+		/// </summary>
+		sealed class BlobDataStream : ReadOnlyMemoryStream
+		{
+			readonly IAgentMessage _message;
+
+			public BlobDataStream(IAgentMessage message)
+				: base(message.Data.Slice(8))
+			{
+				_message = message;
+			}
+
+			protected override void Dispose(bool disposing)
+			{
+				base.Dispose(disposing);
+
+				if (disposing)
+				{
+					_message.Dispose();
+				}
+			}
+		}
+
+		/// <summary>
+		/// Reads a blob from the remote
+		/// </summary>
+		/// <param name="channel">Channel to write to</param>
+		/// <param name="locator">Locator for the blob</param>
+		/// <param name="offset">Offset within the blob</param>
+		/// <param name="length">Length of data to return</param>
+		/// <param name="cancellationToken">Cancellation token for the operation</param>
+		/// <returns>Stream containing the blob data</returns>
+		public static async Task<Stream> ReadBlobAsync(this AgentMessageChannel channel, BlobLocator locator, int offset, int length, CancellationToken cancellationToken = default)
+		{
+			using (IAgentMessageBuilder request = await channel.CreateMessageAsync(AgentMessageType.ReadBlob, cancellationToken))
+			{
+				request.WriteBlobLocator(locator);
+				request.WriteUnsignedVarInt(offset);
+				request.WriteUnsignedVarInt(length);
+				request.Send();
+			}
+
+			byte[]? buffer = null;
+			for(; ;)
+			{
+				IAgentMessage? response = null;
+				try
+				{
+					response = await channel.ReceiveAsync(cancellationToken);
+					if (response.Type != AgentMessageType.ReadBlobResponse)
+					{
+						throw new InvalidAgentMessageException(response);
+					}
+
+					int chunkOffset = BinaryPrimitives.ReadInt32LittleEndian(response.Data.Span.Slice(0, 4));
+					int chunkLength = response.Data.Length - 8;
+					int totalLength = BinaryPrimitives.ReadInt32LittleEndian(response.Data.Span.Slice(4, 4));
+
+					if (chunkOffset == 0 && chunkLength == totalLength)
+					{
+						BlobDataStream stream = new BlobDataStream(response);
+						response = null;
+						return stream;
+					}
+
+					buffer ??= new byte[totalLength];
+					response.Data.Slice(8).CopyTo(buffer.AsMemory(chunkOffset));
+
+					if (chunkOffset + chunkLength == totalLength)
+					{
+						break;
+					}
+				}
+				catch
+				{
+					response?.Dispose();
+					throw;
+				}
+			}
+
+			return new ReadOnlyMemoryStream(buffer);
+		}
+
+		/// <summary>
+		/// Writes blob data to a compute channel
+		/// </summary>
+		/// <param name="channel">Channel to write to</param>
+		/// <param name="message">The read request</param>
+		/// <param name="storage">Storage client to retrieve the blob from</param>
+		/// <param name="cancellationToken">Cancellation token for the operation</param>
+		public static Task SendBlobDataAsync(this AgentMessageChannel channel, ReadBlobMessage message, IStorageClient storage, CancellationToken cancellationToken = default)
+		{
+			return SendBlobDataAsync(channel, message.Locator, message.Offset, message.Length, storage, cancellationToken);
+		}
+
+		/// <summary>
+		/// Writes blob data to a compute channel
+		/// </summary>
+		/// <param name="channel">Channel to write to</param>
+		/// <param name="locator">Locator for the blob to send</param>
+		/// <param name="offset">Starting offset of the data</param>
+		/// <param name="length">Length of the data</param>
+		/// <param name="storage">Storage client to retrieve the blob from</param>
+		/// <param name="cancellationToken">Cancellation token for the operation</param>
+		public static async Task SendBlobDataAsync(this AgentMessageChannel channel, BlobLocator locator, int offset, int length, IStorageClient storage, CancellationToken cancellationToken = default)
+		{
+			byte[] data;
+			if (offset == 0 && length == 0)
+			{
+				using (Stream stream = await storage.ReadBlobAsync(locator, cancellationToken))
+				{
+					using MemoryStream target = new MemoryStream();
+					await stream.CopyToAsync(target, cancellationToken);
+					data = target.ToArray();
+				}
+			}
+			else
+			{
+				using (Stream stream = await storage.ReadBlobRangeAsync(locator, offset, length, cancellationToken))
+				{
+					using MemoryStream target = new MemoryStream();
+					await stream.CopyToAsync(target, cancellationToken);
+					data = target.ToArray();
+				}
+			}
+
+			const int MaxChunkSize = 512 * 1024;
+			for (int chunkOffset = 0; chunkOffset < data.Length;)
+			{
+				int chunkLength = Math.Min(data.Length - chunkOffset, MaxChunkSize);
+				using (IAgentMessageBuilder response = await channel.CreateMessageAsync(AgentMessageType.ReadBlobResponse, chunkLength + 128, cancellationToken))
+				{
+					response.WriteInt32(chunkOffset);
+					response.WriteInt32(data.Length);
+					response.WriteFixedLengthBytes(data.AsSpan(chunkOffset, chunkLength));
+					response.Send();
+				}
+				chunkOffset += chunkLength;
+			}
+		}
+
+		#endregion
+
+		#region Test Messages
+
+		/// <summary>
+		/// Send a message to request that a byte string be xor'ed with a particular value
+		/// </summary>
+		public static async ValueTask SendXorRequestAsync(this AgentMessageChannel channel, ReadOnlyMemory<byte> data, byte value, CancellationToken cancellationToken = default)
+		{
+			using IAgentMessageBuilder message = await channel.CreateMessageAsync(AgentMessageType.XorRequest, cancellationToken);
+			message.WriteFixedLengthBytes(data.Span);
+			message.WriteUInt8(value);
+			message.Send();
+		}
+
+		/// <summary>
+		/// Parse a message as an XOR request
+		/// </summary>
+		public static XorRequestMessage AsXorRequest(this IAgentMessage message)
+		{
+			ReadOnlyMemory<byte> data = message.Data;
+			return new XorRequestMessage(data[0..^1], data.Span[^1]);
+		}
+
+		#endregion
+	}
+}
diff --git a/Engine/Source/Programs/Shared/EpicGames.Horde/Compute/AgentMessageChannel.cs b/Engine/Source/Programs/Shared/EpicGames.Horde/Compute/AgentMessageChannel.cs
new file mode 100644
--- /dev/null
+++ b/Engine/Source/Programs/Shared/EpicGames.Horde/Compute/AgentMessageChannel.cs
@@ -0,0 +1,321 @@
+// Copyright Epic Games, Inc. All Rights Reserved.
+
+using System;
+using System.Buffers;
+using System.Buffers.Binary;
+using System.Text;
+using System.Threading;
+using System.Threading.Tasks;
+using EpicGames.Core;
+using EpicGames.Horde.Compute.Buffers;
+using Microsoft.Extensions.Logging;
+
+namespace EpicGames.Horde.Compute
+{
+	/// <summary>
+	/// Implementation of a compute channel
+	/// </summary>
+	public sealed class AgentMessageChannel : IDisposable
+	{
+		// Length of a message header. Consists of a 1 byte type field, followed by 4 byte length field.
+		const int HeaderLength = 5;
+
+		/// <summary>
+		/// Standard implementation of a message
+		/// </summary>
+		sealed class Message : IAgentMessage
+		{
+			/// <inheritdoc/>
+			public AgentMessageType Type { get; }
+
+			/// <inheritdoc/>
+			public ReadOnlyMemory<byte> Data { get; }
+
+			readonly IMemoryOwner<byte> _memoryOwner;
+			int _position;
+
+			public Message(AgentMessageType type, ReadOnlyMemory<byte> data)
+			{
+				_memoryOwner = MemoryPool<byte>.Shared.Rent(data.Length);
+				data.CopyTo(_memoryOwner.Memory);
+
+				Type = type;
+				Data = _memoryOwner.Memory.Slice(0, data.Length);
+			}
+
+			/// <inheritdoc/>
+			public void Dispose()
+			{
+				_memoryOwner.Dispose();
+			}
+
+			/// <inheritdoc/>
+			public ReadOnlyMemory<byte> GetMemory(int minSize = 1) => Data.Slice(_position);
+
+			/// <inheritdoc/>
+			public void Advance(int length) => _position += length;
+		}
+
+		/// <summary>
+		/// Allows creating new messages in rented memory
+		/// </summary>
+		class MessageBuilder : IAgentMessageBuilder
+		{
+			readonly AgentMessageChannel _channel;
+			readonly IComputeBufferWriter _sendBufferWriter;
+			readonly AgentMessageType _type;
+			int _length;
+
+			/// <inheritdoc/>
+			public int Length => _length;
+
+			public MessageBuilder(AgentMessageChannel channel, IComputeBufferWriter sendBufferWriter, AgentMessageType type)
+			{
+				_channel = channel;
+				_sendBufferWriter = sendBufferWriter;
+				_type = type;
+				_length = 0;
+			}
+
+			public void Dispose()
+			{
+				if (_channel._currentBuilder == this)
+				{
+					_channel._currentBuilder = null;
+				}
+			}
+
+			public void Send()
+			{
+				Span<byte> header = _sendBufferWriter.GetWriteBuffer().Span;
+				header[0] = (byte)_type;
+				BinaryPrimitives.WriteInt32LittleEndian(header.Slice(1, 4), _length);
+
+				if (_channel._logger.IsEnabled(LogLevel.Trace))
+				{
+					_channel.LogMessageInfo("SEND", _type, _sendBufferWriter.GetWriteBuffer().Slice(HeaderLength, _length).Span);
+				}
+
+				_sendBufferWriter.AdvanceWritePosition(HeaderLength + _length);
+				_length = 0;
+			}
+
+			/// <inheritdoc/>
+			public void Advance(int count) => _length += count;
+
+			/// <inheritdoc/>
+			public Memory<byte> GetMemory(int sizeHint = 0) => _sendBufferWriter.GetWriteBuffer().Slice(HeaderLength + _length);
+
+			/// <inheritdoc/>
+			public Span<byte> GetSpan(int sizeHint = 0) => GetMemory(sizeHint).Span;
+		}
+
+		readonly IComputeSocket _socket;
+		readonly IComputeBufferReader _recvBufferReader;
+		readonly IComputeBufferWriter _sendBufferWriter;
+
+		// Can lock chunked memory writer to acuqire pointer
+		readonly ILogger _logger;
+
+#pragma warning disable CA2213
+		MessageBuilder? _currentBuilder;
+#pragma warning restore CA2213
+
+		/// <summary>
+		/// Constructor
+		/// </summary>
+		/// <param name="recvBufferReader"></param>
+		/// <param name="sendBufferWriter"></param>
+		/// <param name="logger">Logger for diagnostic output</param>
+		public AgentMessageChannel(IComputeBufferReader recvBufferReader, IComputeBufferWriter sendBufferWriter, ILogger logger)
+		{
+			_socket = null!;
+			_recvBufferReader = recvBufferReader.AddRef();
+			_sendBufferWriter = sendBufferWriter.AddRef();
+			_logger = logger;
+		}
+
+		/// <summary>
+		/// Constructor
+		/// </summary>
+		/// <param name="socket"></param>
+		/// <param name="channelId"></param>
+		/// <param name="recvBuffer"></param>
+		/// <param name="sendBuffer"></param>
+		/// <param name="logger">Logger for diagnostic output</param>
+		public AgentMessageChannel(IComputeSocket socket, int channelId, IComputeBuffer recvBuffer, IComputeBuffer sendBuffer, ILogger logger)
+		{
+			_socket = socket;
+			_socket.AttachRecvBuffer(channelId, recvBuffer.Writer);
+			_socket.AttachSendBuffer(channelId, sendBuffer.Reader);
+			_recvBufferReader = recvBuffer.Reader.AddRef();
+			_sendBufferWriter = sendBuffer.Writer.AddRef();
+			_logger = logger;
+		}
+
+		/// <summary>
+		/// Overridable dispose method
+		/// </summary>
+		public void Dispose()
+		{
+			_currentBuilder?.Dispose();
+
+			_sendBufferWriter.MarkComplete();
+			_sendBufferWriter.Dispose();
+
+			_recvBufferReader.Dispose();
+		}
+
+		/// <summary>
+		/// Mark the send buffer as complete
+		/// </summary>
+		public void MarkComplete()
+		{
+			_sendBufferWriter.MarkComplete();
+		}
+
+		/// <inheritdoc/>
+		public async ValueTask<IAgentMessage> ReceiveAsync(CancellationToken cancellationToken)
+		{
+			while (!_recvBufferReader.IsComplete)
+			{
+				ReadOnlyMemory<byte> memory = _recvBufferReader.GetReadBuffer();
+				if (memory.Length < HeaderLength)
+				{
+					await _recvBufferReader.WaitToReadAsync(HeaderLength, cancellationToken);
+					continue;
+				}
+
+				int messageLength = BinaryPrimitives.ReadInt32LittleEndian(memory.Span.Slice(1, 4));
+				if (memory.Length < HeaderLength + messageLength)
+				{
+					await _recvBufferReader.WaitToReadAsync(HeaderLength + messageLength, cancellationToken);
+					continue;
+				}
+
+				AgentMessageType type = (AgentMessageType)memory.Span[0];
+				Message message = new Message(type, memory.Slice(HeaderLength, messageLength));
+				if (_logger.IsEnabled(LogLevel.Trace))
+				{
+					LogMessageInfo("RECV", message.Type, message.Data.Span);
+				}
+
+				_recvBufferReader.AdvanceReadPosition(HeaderLength + messageLength);
+				return message;
+			}
+			return new Message(AgentMessageType.None, ReadOnlyMemory<byte>.Empty);
+		}
+
+		void LogMessageInfo(string verb, AgentMessageType type, ReadOnlySpan<byte> data)
+		{
+			StringBuilder bytes = new StringBuilder();
+			for (int offset = 0; offset < 16 && offset < data.Length; offset++)
+			{
+				bytes.Append($" {data[offset]:X2}");
+			}
+			if (data.Length > 16)
+			{
+				bytes.Append("..");
+			}
+			_logger.LogTrace("{Verb} {Type,-22} [{Length,10:n0}] = {Bytes}", verb, type, data.Length, bytes.ToString());
+		}
+
+		/// <inheritdoc/>
+		public async ValueTask<IAgentMessageBuilder> CreateMessageAsync(AgentMessageType type, int maxSize, CancellationToken cancellationToken)
+		{
+			if (_currentBuilder != null)
+			{
+				throw new InvalidOperationException("Only one writer can be active at a time. Dispose of the previous writer first.");
+			}
+
+			await _sendBufferWriter.WaitToWriteAsync(maxSize, cancellationToken);
+
+			_currentBuilder = new MessageBuilder(this, _sendBufferWriter, type);
+			return _currentBuilder;
+		}
+	}
+
+	/// <summary>
+	/// Extension methods to allow creating channels from leases
+	/// </summary>
+	public static class AgentMessageChannelExtensions
+	{
+		/// <summary>
+		/// Creates a message channel with the given identifier
+		/// </summary>
+		/// <param name="socket">Socket to create a channel for</param>
+		/// <param name="channelId">Identifier for the channel</param>
+		/// <param name="logger">Logger for the channel</param>
+		public static AgentMessageChannel CreateAgentMessageChannel(this IComputeSocket socket, int channelId, ILogger logger)
+			=> socket.CreateAgentMessageChannel(channelId, 65536, logger);
+
+		/// <summary>
+		/// Creates a message channel with the given identifier
+		/// </summary>
+		/// <param name="socket">Socket to create a channel for</param>
+		/// <param name="channelId">Identifier for the channel</param>
+		/// <param name="bufferSize">Size of the send and receive buffer</param>
+		/// <param name="logger">Logger for the channel</param>
+		public static AgentMessageChannel CreateAgentMessageChannel(this IComputeSocket socket, int channelId, int bufferSize, ILogger logger)
+			=> socket.CreateAgentMessageChannel(channelId, bufferSize, bufferSize, logger);
+
+		/// <summary>
+		/// Creates a message channel with the given identifier
+		/// </summary>
+		/// <param name="socket">Socket to create a channel for</param>
+		/// <param name="channelId">Identifier for the channel</param>
+		/// <param name="sendBufferSize">Size of the send buffer</param>
+		/// <param name="recvBufferSize">Size of the recieve buffer</param>
+		/// <param name="logger">Logger for the channel</param>
+		public static AgentMessageChannel CreateAgentMessageChannel(this IComputeSocket socket, int channelId, int sendBufferSize, int recvBufferSize, ILogger logger)
+		{
+			using IComputeBuffer sendBuffer = new PooledBuffer(sendBufferSize);
+			using IComputeBuffer recvBuffer = new PooledBuffer(recvBufferSize);
+			return new AgentMessageChannel(socket, channelId, sendBuffer, recvBuffer, logger);
+		}
+
+		/// <summary>
+		/// Reads a message from the channel
+		/// </summary>
+		/// <param name="channel">Channel to receive on</param>
+		/// <param name="type">Expected type of the message</param>
+		/// <param name="cancellationToken">Cancellation token for the operation</param>
+		/// <returns>Data for a message that was read. Must be disposed.</returns>
+		public static async ValueTask<IAgentMessage> ReceiveAsync(this AgentMessageChannel channel, AgentMessageType type, CancellationToken cancellationToken = default)
+		{
+			IAgentMessage message = await channel.ReceiveAsync(cancellationToken);
+			if (message.Type != type)
+			{
+				throw new InvalidAgentMessageException(message);
+			}
+			return message;
+		}
+
+		/// <summary>
+		/// Creates a new builder for a message
+		/// </summary>
+		/// <param name="channel">Channel to send on</param>
+		/// <param name="type">Type of the message</param>
+		/// <param name="cancellationToken">Cancellation token for the operation</param>
+		/// <returns>New builder for messages</returns>
+		public static ValueTask<IAgentMessageBuilder> CreateMessageAsync(this AgentMessageChannel channel, AgentMessageType type, CancellationToken cancellationToken)
+		{
+			return channel.CreateMessageAsync(type, 1024, cancellationToken);
+		}
+
+		/// <summary>
+		/// Forwards an existing message across a channel
+		/// </summary>
+		/// <param name="channel">Channel to send on</param>
+		/// <param name="message">The message to be sent</param>
+		/// <param name="cancellationToken">Cancellation token for the operation</param>
+		public static async ValueTask SendAsync(this AgentMessageChannel channel, IAgentMessage message, CancellationToken cancellationToken)
+		{
+			using (IAgentMessageBuilder builder = await channel.CreateMessageAsync(message.Type, message.Data.Length, cancellationToken))
+			{
+				builder.WriteFixedLengthBytes(message.Data.Span);
+				builder.Send();
+			}
+		}
+	}
+}
diff --git a/Engine/Source/Programs/Shared/EpicGames.Horde/Compute/AgentMessageHandler.cs b/Engine/Source/Programs/Shared/EpicGames.Horde/Compute/AgentMessageHandler.cs
new file mode 100644
--- /dev/null
+++ b/Engine/Source/Programs/Shared/EpicGames.Horde/Compute/AgentMessageHandler.cs
@@ -0,0 +1,303 @@
+// Copyright Epic Games, Inc. All Rights Reserved.
+
+using System;
+using System.Collections.Generic;
+using System.Diagnostics;
+using System.Runtime.InteropServices;
+using System.Threading;
+using System.Threading.Tasks;
+using EpicGames.Core;
+using EpicGames.Horde.Compute.Buffers;
+using EpicGames.Horde.Storage;
+using EpicGames.Horde.Storage.Nodes;
+using Microsoft.Extensions.Caching.Memory;
+using Microsoft.Extensions.Logging;
+
+namespace EpicGames.Horde.Compute
+{
+	/// <summary>
+	/// Implements the remote end of a compute worker. 
+	/// </summary>
+	public class AgentMessageHandler
+	{
+		readonly DirectoryReference _sandboxDir;
+		readonly IMemoryCache _memoryCache;
+		readonly ILogger _logger;
+
+		/// <summary>
+		/// Constructor
+		/// </summary>
+		/// <param name="sandboxDir">Directory to use for reading/writing files</param>
+		/// <param name="memoryCache">Cache for nodes read from storage</param>
+		/// <param name="logger">Logger for diagnostics</param>
+		public AgentMessageHandler(DirectoryReference sandboxDir, IMemoryCache memoryCache, ILogger logger)
+		{
+			_sandboxDir = sandboxDir;
+			_memoryCache = memoryCache;
+			_logger = logger;
+		}
+
+		/// <summary>
+		/// Runs the worker using commands sent along the given socket
+		/// </summary>
+		/// <param name="socket">Socket to read from</param>
+		/// <param name="cancellationToken">Cancellation token for the operation</param>
+		public async Task RunAsync(IComputeSocket socket, CancellationToken cancellationToken)
+		{
+			await RunAsync(socket, 0, 4 * 1024 * 1024, cancellationToken);
+		}
+
+		async Task RunAsync(IComputeSocket socket, int channelId, int bufferSize, CancellationToken cancellationToken)
+		{
+			using (AgentMessageChannel channel = socket.CreateAgentMessageChannel(channelId, bufferSize, _logger))
+			{
+				await channel.AttachAsync(cancellationToken);
+
+				List<Task> childTasks = new List<Task>();
+				for (; ; )
+				{
+					using IAgentMessage message = await channel.ReceiveAsync(cancellationToken);
+					_logger.LogTrace("Compute Channel {ChannelId}: {MessageType}", channelId, message.Type);
+
+					switch (message.Type)
+					{
+						case AgentMessageType.None:
+							await Task.WhenAll(childTasks);
+							return;
+						case AgentMessageType.Fork:
+							{
+								ForkMessage fork = message.ParseForkMessage();
+								childTasks.Add(Task.Run(() => RunAsync(socket, fork.channelId, fork.bufferSize, cancellationToken), cancellationToken));
+							}
+							break;
+						case AgentMessageType.WriteFiles:
+							{
+								UploadFilesMessage writeFiles = message.ParseUploadFilesMessage();
+								await WriteFilesAsync(channel, writeFiles.Name, writeFiles.Locator, cancellationToken);
+							}
+							break;
+						case AgentMessageType.DeleteFiles:
+							{
+								DeleteFilesMessage deleteFiles = message.ParseDeleteFilesMessage();
+								DeleteFiles(deleteFiles.Filter);
+							}
+							break;
+						case AgentMessageType.Execute:
+							{
+								ExecuteProcessMessage executeProcess = message.ParseExecuteProcessMessage();
+								await ExecuteProcessAsync(socket, channel, executeProcess.Executable, executeProcess.Arguments, executeProcess.WorkingDir, executeProcess.EnvVars, cancellationToken);
+							}
+							break;
+						case AgentMessageType.XorRequest:
+							{
+								XorRequestMessage xorRequest = message.AsXorRequest();
+								await RunXor(channel, xorRequest.Data, xorRequest.Value, cancellationToken);
+							}
+							break;
+						default:
+							throw new InvalidAgentMessageException(message);
+					}
+				}
+			}
+		}
+
+		static async ValueTask RunXor(AgentMessageChannel channel, ReadOnlyMemory<byte> source, byte value, CancellationToken cancellationToken)
+		{
+			using IAgentMessageBuilder response = await channel.CreateMessageAsync(AgentMessageType.XorResponse, source.Length, cancellationToken);
+			XorData(source.Span, response.GetSpanAndAdvance(source.Length), value);
+			response.Send();
+		}
+
+		static void XorData(ReadOnlySpan<byte> source, Span<byte> target, byte value)
+		{
+			for (int idx = 0; idx < source.Length; idx++)
+			{
+				target[idx] = (byte)(source[idx] ^ value);
+			}
+		}
+
+		async Task WriteFilesAsync(AgentMessageChannel channel, string path, NodeLocator locator, CancellationToken cancellationToken)
+		{
+			using ComputeStorageClient store = new ComputeStorageClient(channel);
+			TreeReader reader = new TreeReader(store, _memoryCache, _logger);
+
+			DirectoryNode directoryNode = await reader.ReadNodeAsync<DirectoryNode>(locator, cancellationToken);
+
+			DirectoryReference outputDir = DirectoryReference.Combine(_sandboxDir, path);
+			if (!outputDir.IsUnderDirectory(_sandboxDir))
+			{
+				throw new InvalidOperationException("Cannot write files outside sandbox");
+			}
+
+			await directoryNode.CopyToDirectoryAsync(reader, outputDir.ToDirectoryInfo(), _logger, cancellationToken);
+
+			using (IAgentMessageBuilder message = await channel.CreateMessageAsync(AgentMessageType.WriteFilesResponse, cancellationToken))
+			{
+				message.Send();
+			}
+		}
+
+		void DeleteFiles(IReadOnlyList<string> deleteFiles)
+		{
+			FileFilter filter = new FileFilter(deleteFiles);
+
+			List<FileReference> files = filter.ApplyToDirectory(_sandboxDir, false);
+			foreach (FileReference file in files)
+			{
+				FileUtils.ForceDeleteFile(file);
+			}
+		}
+
+		async Task ExecuteProcessAsync(IComputeSocket socket, AgentMessageChannel channel, string executable, IReadOnlyList<string> arguments, string? workingDir, IReadOnlyDictionary<string, string?>? envVars, CancellationToken cancellationToken)
+		{
+			try
+			{
+				if (RuntimeInformation.IsOSPlatform(OSPlatform.Windows))
+				{
+					await ExecuteProcessWindowsAsync(socket, channel, executable, arguments, workingDir, envVars, cancellationToken);
+				}
+				else
+				{
+					await ExecuteProcessInternalAsync(channel, executable, arguments, workingDir, envVars, cancellationToken);
+				}
+			}
+			catch (Exception ex)
+			{
+				await channel.SendExceptionAsync(ex, cancellationToken);
+			}
+		}
+
+		async Task ExecuteProcessWindowsAsync(IComputeSocket socket, AgentMessageChannel channel, string executable, IReadOnlyList<string> arguments, string? workingDir, IReadOnlyDictionary<string, string?>? envVars, CancellationToken cancellationToken)
+		{
+			Dictionary<string, string?> newEnvVars = new Dictionary<string, string?>();
+			if (envVars != null)
+			{
+				foreach ((string name, string? value) in envVars)
+				{
+					newEnvVars.Add(name, value);
+				}
+			}
+
+			using (SharedMemoryBuffer ipcBuffer = SharedMemoryBuffer.CreateNew(null, 1, 64 * 1024))
+			{
+				newEnvVars[WorkerComputeSocket.IpcEnvVar] = ipcBuffer.Name;
+
+				using (BackgroundTask backgroundTask = BackgroundTask.StartNew(ctx => ProcessIpcMessagesAsync(socket, ipcBuffer.Reader, cancellationToken)))
+				{
+					_logger.LogInformation("Launching {Executable} {Arguments}", CommandLineArguments.Quote(executable), CommandLineArguments.Join(arguments));
+					await ExecuteProcessInternalAsync(channel, executable, arguments, workingDir, newEnvVars, cancellationToken);
+					_logger.LogInformation("Finished executing process");
+					ipcBuffer.Writer.MarkComplete();
+				}
+			}
+
+			_logger.LogInformation("Child process has shut down");
+		}
+
+		async Task ProcessIpcMessagesAsync(IComputeSocket socket, IComputeBufferReader ipcReader, CancellationToken cancellationToken)
+		{
+			List<(IpcMessage, int, SharedMemoryBuffer)> buffers = new();
+			try
+			{
+				List<(int, IComputeBufferWriter)> writers = new List<(int, IComputeBufferWriter)>();
+				while (await ipcReader.WaitToReadAsync(1, cancellationToken))
+				{
+					ReadOnlyMemory<byte> memory = ipcReader.GetReadBuffer();
+					MemoryReader reader = new MemoryReader(memory);
+
+					IpcMessage message = (IpcMessage)reader.ReadUnsignedVarInt();
+					try
+					{
+						switch (message)
+						{
+							case IpcMessage.AttachSendBuffer:
+								{
+									int channelId = (int)reader.ReadUnsignedVarInt();
+									string name = reader.ReadString();
+									_logger.LogDebug("Attaching send buffer for channel {ChannelId} to {Name}", channelId, name);
+
+									SharedMemoryBuffer buffer = SharedMemoryBuffer.OpenExisting(name);
+									buffers.Add((message, channelId, buffer));
+
+									socket.AttachSendBuffer(channelId, buffer.Reader);
+								}
+								break;
+							case IpcMessage.AttachRecvBuffer:
+								{
+									int channelId = (int)reader.ReadUnsignedVarInt();
+									string name = reader.ReadString();
+									_logger.LogDebug("Attaching recv buffer for channel {ChannelId} to {Name}", channelId, name);
+
+									SharedMemoryBuffer buffer = SharedMemoryBuffer.OpenExisting(name);
+									buffers.Add((message, channelId, buffer));
+
+									socket.AttachRecvBuffer(channelId, buffer.Writer);
+								}
+								break;
+							default:
+								throw new InvalidOperationException($"Invalid IPC message: {message}");
+						}
+					}
+					catch (Exception ex)
+					{
+						_logger.LogError(ex, "Exception while processing messages from child process: {Message}", ex.Message);
+					}
+
+					ipcReader.AdvanceReadPosition(memory.Length - reader.RemainingMemory.Length);
+				}
+			}
+			finally
+			{
+				foreach ((IpcMessage message, int channelId, SharedMemoryBuffer buffer) in buffers)
+				{
+					if (buffer.Writer.MarkComplete())
+					{
+						_logger.LogWarning("Buffer added via {Message} on channel {ChannelId} was not marked complete", message, channelId);
+					}
+					buffer.Dispose();
+				}
+			}
+		}
+
+		async Task ExecuteProcessInternalAsync(AgentMessageChannel channel, string executable, IReadOnlyList<string> arguments, string? workingDir, IReadOnlyDictionary<string, string?>? envVars, CancellationToken cancellationToken)
+		{
+			string resolvedExecutable = FileReference.Combine(_sandboxDir, executable).FullName;
+			string resolvedCommandLine = CommandLineArguments.Join(arguments);
+			string resolvedWorkingDir = DirectoryReference.Combine(_sandboxDir, workingDir ?? String.Empty).FullName;
+
+			Dictionary<string, string> resolvedEnvVars = ManagedProcess.GetCurrentEnvVars();
+			if (envVars != null)
+			{
+				foreach ((string key, string? value) in envVars)
+				{
+					if (value == null)
+					{
+						resolvedEnvVars.Remove(key);
+					}
+					else
+					{
+						resolvedEnvVars[key] = value;
+					}
+				}
+			}
+
+			using (ManagedProcessGroup group = new ManagedProcessGroup())
+			{
+				using (ManagedProcess process = new ManagedProcess(group, resolvedExecutable, resolvedCommandLine, resolvedWorkingDir, resolvedEnvVars, null, ProcessPriorityClass.Normal))
+				{
+					byte[] buffer = new byte[1024];
+					for (; ; )
+					{
+						int length = await process.ReadAsync(buffer, 0, buffer.Length, cancellationToken);
+						if (length == 0)
+						{
+							await channel.SendExecuteResultAsync(process.ExitCode, cancellationToken);
+							return;
+						}
+						await channel.SendExecuteOutputAsync(buffer.AsMemory(0, length), cancellationToken);
+					}
+				}
+			}
+		}
+	}
+}
diff --git a/Engine/Source/Programs/Shared/EpicGames.Horde/Compute/AgentStorageClient.cs b/Engine/Source/Programs/Shared/EpicGames.Horde/Compute/AgentStorageClient.cs
new file mode 100644
--- /dev/null
+++ b/Engine/Source/Programs/Shared/EpicGames.Horde/Compute/AgentStorageClient.cs
@@ -0,0 +1,111 @@
+// Copyright Epic Games, Inc. All Rights Reserved.
+
+using System;
+using System.Collections.Generic;
+using System.IO;
+using System.Threading;
+using System.Threading.Tasks;
+using EpicGames.Core;
+using EpicGames.Horde.Storage;
+
+namespace EpicGames.Horde.Compute
+{
+	/// <summary>
+	/// Storage client which can read bundles over a compute channel
+	/// </summary>
+	public sealed class ComputeStorageClient : StorageClientBase, IDisposable
+	{
+		readonly AgentMessageChannel _channel;
+		readonly SemaphoreSlim _semaphore;
+
+		/// <summary>
+		/// Constructor
+		/// </summary>
+		/// <param name="channel"></param>
+		public ComputeStorageClient(AgentMessageChannel channel)
+		{
+			_channel = channel;
+			_semaphore = new SemaphoreSlim(1);
+		}
+
+		/// <inheritdoc/>
+		public void Dispose()
+		{
+			_semaphore.Dispose();
+		}
+
+		#region Blobs
+
+		/// <inheritdoc/>
+		public override async Task<Stream> ReadBlobAsync(BlobLocator locator, CancellationToken cancellationToken = default)
+		{
+			return await ReadBlobRangeAsync(locator, 0, 0, cancellationToken);
+		}
+
+		/// <inheritdoc/>
+		public override async Task<Stream> ReadBlobRangeAsync(BlobLocator locator, int offset, int length, CancellationToken cancellationToken = default)
+		{
+			await _semaphore.WaitAsync(cancellationToken);
+			try
+			{
+				return await _channel.ReadBlobAsync(locator, offset, length, cancellationToken);
+			}
+			finally
+			{
+				_semaphore.Release();
+			}
+		}
+
+		/// <inheritdoc/>
+		public override Task<BlobLocator> WriteBlobAsync(Stream stream, Utf8String prefix = default, CancellationToken cancellationToken = default)
+		{
+			throw new NotSupportedException();
+		}
+
+		#endregion
+
+		#region Nodes
+
+		/// <inheritdoc/>
+		public override Task AddAliasAsync(Utf8String name, NodeHandle locator, CancellationToken cancellationToken = default)
+		{
+			throw new NotSupportedException();
+		}
+
+		/// <inheritdoc/>
+		public override Task RemoveAliasAsync(Utf8String name, NodeHandle locator, CancellationToken cancellationToken = default)
+		{
+			throw new NotSupportedException();
+		}
+
+		/// <inheritdoc/>
+		public override IAsyncEnumerable<NodeHandle> FindNodesAsync(Utf8String name, CancellationToken cancellationToken = default)
+		{
+			throw new NotSupportedException();
+		}
+
+		#endregion
+
+		#region Refs
+
+		/// <inheritdoc/>
+		public override Task DeleteRefAsync(RefName name, CancellationToken cancellationToken = default)
+		{
+			throw new NotSupportedException();
+		}
+
+		/// <inheritdoc/>
+		public override Task<NodeHandle?> TryReadRefTargetAsync(RefName name, DateTime cacheTime = default, CancellationToken cancellationToken = default)
+		{
+			throw new NotSupportedException();
+		}
+
+		/// <inheritdoc/>
+		public override Task WriteRefTargetAsync(RefName name, NodeHandle target, RefOptions? options = null, CancellationToken cancellationToken = default)
+		{
+			throw new NotSupportedException();
+		}
+
+		#endregion
+	}
+}
diff --git a/Engine/Source/Programs/Shared/EpicGames.Horde/Compute/Buffers/PooledBuffer.cs b/Engine/Source/Programs/Shared/EpicGames.Horde/Compute/Buffers/PooledBuffer.cs
--- a/Engine/Source/Programs/Shared/EpicGames.Horde/Compute/Buffers/PooledBuffer.cs
+++ b/Engine/Source/Programs/Shared/EpicGames.Horde/Compute/Buffers/PooledBuffer.cs
@@ -12,7 +12,7 @@ namespace EpicGames.Horde.Compute.Buffers
 	/// <summary>
 	/// In-process buffer used to store compute messages
 	/// </summary>
-	public sealed class PooledBuffer : ComputeBufferBase
+	public sealed class PooledBuffer : ComputeBuffer
 	{
 		class Resources : ResourcesBase
 		{
diff --git a/Engine/Source/Programs/Shared/EpicGames.Horde/Compute/Buffers/SharedMemoryBuffer.cs b/Engine/Source/Programs/Shared/EpicGames.Horde/Compute/Buffers/SharedMemoryBuffer.cs
--- a/Engine/Source/Programs/Shared/EpicGames.Horde/Compute/Buffers/SharedMemoryBuffer.cs
+++ b/Engine/Source/Programs/Shared/EpicGames.Horde/Compute/Buffers/SharedMemoryBuffer.cs
@@ -13,7 +13,7 @@ namespace EpicGames.Horde.Compute.Buffers
 	/// <summary>
 	/// Core implementation of <see cref="SharedMemoryBuffer"/>
 	/// </summary>
-	public sealed class SharedMemoryBuffer : ComputeBufferBase
+	public sealed class SharedMemoryBuffer : ComputeBuffer
 	{
 		class Resources : ResourcesBase
 		{
diff --git a/Engine/Source/Programs/Shared/EpicGames.Horde/Compute/Clients/LocalComputeClient.cs b/Engine/Source/Programs/Shared/EpicGames.Horde/Compute/Clients/LocalComputeClient.cs
--- a/Engine/Source/Programs/Shared/EpicGames.Horde/Compute/Clients/LocalComputeClient.cs
+++ b/Engine/Source/Programs/Shared/EpicGames.Horde/Compute/Clients/LocalComputeClient.cs
@@ -79,7 +79,7 @@ namespace EpicGames.Horde.Compute.Clients
 
 			await using (RemoteComputeSocket socket = new RemoteComputeSocket(new TcpTransport(tcpSocket), ComputeSocketEndpoint.Remote, logger))
 			{
-				ComputeWorker worker = new ComputeWorker(sandboxDir, memoryCache, logger);
+				AgentMessageHandler worker = new AgentMessageHandler(sandboxDir, memoryCache, logger);
 				await worker.RunAsync(socket, cancellationToken);
 				await socket.CloseAsync(cancellationToken);
 			}
diff --git a/Engine/Source/Programs/Shared/EpicGames.Horde/Compute/ComputeBuffer.cs b/Engine/Source/Programs/Shared/EpicGames.Horde/Compute/ComputeBuffer.cs
new file mode 100644
--- /dev/null
+++ b/Engine/Source/Programs/Shared/EpicGames.Horde/Compute/ComputeBuffer.cs
@@ -0,0 +1,548 @@
+// Copyright Epic Games, Inc. All Rights Reserved.
+
+using System;
+using System.Diagnostics;
+using System.Linq;
+using System.Threading;
+using System.Threading.Tasks;
+
+namespace EpicGames.Horde.Compute
+{
+	/// <summary>
+	/// In-process buffer used to store compute messages
+	/// </summary>
+	public abstract class ComputeBuffer : IComputeBuffer
+	{
+		const int MaxChunks = 16;
+		const int MaxReaders = 16;
+
+		internal const int HeaderSize = (2 + MaxChunks + MaxReaders) * sizeof(ulong);
+
+		/// <summary>
+		/// Tracked state of the buffer
+		/// </summary>
+		internal readonly unsafe struct HeaderPtr
+		{
+			readonly ulong* _data;
+
+			public HeaderPtr(ulong* data) => _data = data;
+
+			public HeaderPtr(ulong* data, int numReaders, int numChunks, int chunkLength)
+			{
+				_data = data;
+
+				data[0] = ((ulong)(uint)numChunks << 32) | (uint)numReaders;
+				data[1] = (uint)chunkLength;
+
+				GetChunkStatePtr(0).StartWriting(numReaders);
+			}
+
+			public int NumReaders => (int)_data[0];
+			public int NumChunks => (int)(_data[0] >> 32);
+			public int ChunkLength => (int)_data[1];
+			public int WriteChunkIdx
+			{
+				get => (int)(_data[1] >> 32);
+				set => _data[1] = ((ulong)value << 32) | (uint)ChunkLength;
+			}
+
+			public ChunkStatePtr GetChunkStatePtr(int chunkIdx) => new ChunkStatePtr(_data + 2 + chunkIdx);
+
+			public ReaderStatePtr GetReaderStatePtr(int readerIdx) => new ReaderStatePtr(_data + 2 + MaxChunks + readerIdx);
+		}
+
+		/// <summary>
+		/// Write state for a chunk
+		/// </summary>
+		internal enum WriteState
+		{
+			/// <summary>
+			/// Writer has moved to the next chunk
+			/// </summary>
+			MovedToNext = 0,
+
+			/// <summary>
+			/// Chunk is still being appended to
+			/// </summary>
+			Writing = 2,
+
+			/// <summary>
+			/// This chunk marks the end of the stream
+			/// </summary>
+			Complete = 3,
+		}
+
+		/// <summary>
+		/// Stores the state of a chunk in a 64-bit value, which can be updated atomically
+		/// </summary>
+		internal record struct ChunkState(ulong Value)
+		{
+			// Written length of this chunk
+			public int Length => (int)(Value & 0x7fffffff);
+
+			// Set of flags which are set for each reader that still has to read from a chunk
+			public int ReaderFlags => (int)((Value >> 31) & 0x7fffffff);
+
+			// State of the writer
+			public WriteState WriteState => (WriteState)(Value >> 62);
+
+			// Constructor
+			public ChunkState(WriteState writerState, int readerFlags, int length) : this(((ulong)writerState << 62) | ((ulong)readerFlags << 31) | (uint)length) { }
+
+			// Test whether a particular reader is still referencing the chunk
+			public bool HasReaderFlag(int readerIdx) => (Value & (1UL << (readerIdx + 31))) != 0;
+
+			/// <inheritdoc/>
+			public override string ToString() => $"{WriteState}, Length: {Length}, Readers: {ReaderFlags}";
+		}
+
+		/// <summary>
+		/// Wraps a pointer to the state of a chunk
+		/// </summary>
+		internal readonly unsafe struct ChunkStatePtr
+		{
+			readonly ulong* _data;
+
+			public ChunkStatePtr(ulong* data) => _data = data;
+
+			// Current value of the chunk state
+			public ChunkState Value => new ChunkState(Interlocked.CompareExchange(ref *_data, 0, 0));
+
+			// Written length of this chunk
+			public int Length => Value.Length;
+
+			// Set of flags which are set for each reader that still has to read from a chunk
+			public int ReaderFlags => Value.ReaderFlags;
+
+			// State of the writer
+			public WriteState WriteState => Value.WriteState;
+
+			// Append data to the chunk
+			public void Append(int length) => Interlocked.Add(ref *_data, (uint)length);
+
+			// Mark the chunk as being written to
+			public void StartWriting(int numReaders) => Interlocked.Exchange(ref *_data, new ChunkState(WriteState.Writing, (1 << numReaders) - 1, 0).Value);
+
+			// Move to the next chunk
+			public void MarkComplete() => Interlocked.Or(ref *_data, (ulong)WriteState.Complete << 62);
+
+			// Clear the reader flag
+			public void FinishReading(int readerIdx) => Interlocked.And(ref *_data, ~(1UL << (readerIdx + 31)));
+
+			// Move to the next chunk
+			public void FinishWriting() => Interlocked.And(ref *_data, ~((ulong)WriteState.Writing << 62));
+
+			/// <inheritdoc/>
+			public override string ToString() => Value.ToString();
+		}
+
+		/// <summary>
+		/// Encodes the state of a reader in a 64-bit value
+		/// </summary>
+		internal record struct ReaderState(ulong Value)
+		{
+			public ReaderState(int chunkIdx, int offset) : this(((ulong)chunkIdx << 32) | (uint)offset) { }
+
+			public int ChunkIdx => (int)(Value >> 32);
+			public int Offset => (int)Value;
+
+			public override string ToString() => $"Chunk {ChunkIdx}, Offset {Offset}";
+		}
+
+		/// <summary>
+		/// Pointer to a reader state value
+		/// </summary>
+		internal readonly unsafe struct ReaderStatePtr
+		{
+			readonly ulong* _data;
+
+			public ReaderStatePtr(ulong* data) => _data = data;
+
+			public ReaderState Value
+			{
+				get => new ReaderState(Interlocked.CompareExchange(ref *_data, 0, 0));
+				set => Interlocked.Exchange(ref *_data, value.Value);
+			}
+
+			public int ChunkIdx => Value.ChunkIdx;
+			public int Offset => Value.Offset;
+
+			public void Advance(int length) => Interlocked.Add(ref *_data, (ulong)length);
+
+			public override string ToString() => Value.ToString();
+		}
+
+		/// <summary>
+		/// State shared between buffer instances
+		/// </summary>
+		internal abstract class ResourcesBase : IDisposable
+		{
+			public HeaderPtr HeaderPtr { get; }
+			public Memory<byte>[] Chunks { get; }
+			int _refCount = 1;
+
+			public ResourcesBase(HeaderPtr headerPtr, Memory<byte>[] chunks)
+			{
+				HeaderPtr = headerPtr;
+				Chunks = chunks;
+			}
+
+			public void AddRef()
+			{
+				Interlocked.Increment(ref _refCount);
+			}
+
+			public void Release()
+			{
+				if (Interlocked.Decrement(ref _refCount) == 0)
+				{
+					Dispose();
+				}
+			}
+
+			/// <summary>
+			/// Dispose of the 
+			/// </summary>
+			public abstract void Dispose();
+
+			/// <summary>
+			/// Signals a read event
+			/// </summary>
+			public abstract void SetReadEvent(int readerIdx);
+
+			/// <summary>
+			/// Signals read events for every reader
+			/// </summary>
+			public void SetAllReadEvents()
+			{
+				for (int readerIdx = 0; readerIdx < HeaderPtr.NumReaders; readerIdx++)
+				{
+					SetReadEvent(readerIdx);
+				}
+			}
+
+			/// <summary>
+			/// Resets a read event
+			/// </summary>
+			public abstract void ResetReadEvent(int readerIdx);
+
+			/// <summary>
+			/// Waits for a read event to be signalled
+			/// </summary>
+			public abstract Task WaitForReadEvent(int readerIdx, CancellationToken cancellationToken);
+
+			/// <summary>
+			/// Signals the write event
+			/// </summary>
+			public abstract void SetWriteEvent();
+
+			/// <summary>
+			/// Resets the write event
+			/// </summary>
+			public abstract void ResetWriteEvent();
+
+			/// <summary>
+			/// Waits for the write event to be signalled
+			/// </summary>
+			public abstract Task WaitForWriteEvent(CancellationToken cancellationToken);
+
+#pragma warning disable IDE0051 // Remove unused private members
+			// For debugging purposes only
+			ChunkState[] ChunkStates => Enumerable.Range(0, HeaderPtr.NumChunks).Select(x => HeaderPtr.GetChunkStatePtr(x).Value).ToArray();
+			ReaderState[] ReaderStates => Enumerable.Range(0, HeaderPtr.NumReaders).Select(x => HeaderPtr.GetReaderStatePtr(x).Value).ToArray();
+			int WriteChunkIdx => HeaderPtr.WriteChunkIdx;
+#pragma warning restore IDE0051 // Remove unused private members
+		}
+
+		internal class ReaderImpl : IComputeBufferReader
+		{
+			readonly HeaderPtr _headerPtr;
+			readonly Memory<byte>[] _chunks;
+			ResourcesBase _resources;
+			readonly int _readerIdx;
+
+			public ReaderImpl(ResourcesBase resources, int readerIdx)
+			{
+				_headerPtr = resources.HeaderPtr;
+				_chunks = resources.Chunks;
+				_resources = resources;
+				_readerIdx = readerIdx;
+			}
+
+			public ResourcesBase GetResources() => _resources;
+
+			/// <inheritdoc/>
+			public IComputeBufferReader AddRef()
+			{
+				_resources.AddRef();
+				return new ReaderImpl(_resources, _readerIdx);
+			}
+
+			/// <inheritdoc/>
+			public void Dispose()
+			{
+				if (_resources != null)
+				{
+					_resources.Release();
+					_resources = null!;
+				}
+			}
+
+			/// <inheritdoc/>
+			public bool IsComplete
+			{
+				get
+				{
+					ReaderState readerState = _headerPtr.GetReaderStatePtr(_readerIdx).Value;
+					ChunkState chunkState = _headerPtr.GetChunkStatePtr(readerState.ChunkIdx).Value;
+					return chunkState.WriteState == WriteState.Complete && readerState.Offset == chunkState.Length;
+				}
+			}
+
+			/// <inheritdoc/>
+			public void AdvanceReadPosition(int length)
+			{
+				ReaderStatePtr readerStatePtr = _headerPtr.GetReaderStatePtr(_readerIdx);
+				readerStatePtr.Advance(length);
+			}
+
+			/// <inheritdoc/>
+			public ReadOnlyMemory<byte> GetReadBuffer()
+			{
+				ReaderStatePtr readerStatePtr = _headerPtr.GetReaderStatePtr(_readerIdx);
+				ReaderState readerState = readerStatePtr.Value;
+
+				ChunkStatePtr chunkStatePtr = _headerPtr.GetChunkStatePtr(readerState.ChunkIdx);
+				ChunkState chunkState = chunkStatePtr.Value;
+
+				if (chunkState.HasReaderFlag(_readerIdx))
+				{
+					return _chunks[readerState.ChunkIdx].Slice(readerState.Offset, chunkState.Length - readerState.Offset);
+				}
+				else
+				{
+					return ReadOnlyMemory<byte>.Empty;
+				}
+			}
+
+			/// <inheritdoc/>
+			public async ValueTask<bool> WaitToReadAsync(int minLength, CancellationToken cancellationToken = default)
+			{
+				for (; ; )
+				{
+					ReaderStatePtr readerStatePtr = _headerPtr.GetReaderStatePtr(_readerIdx);
+					ReaderState readerState = readerStatePtr.Value;
+
+					ChunkStatePtr chunkStatePtr = _headerPtr.GetChunkStatePtr(readerState.ChunkIdx);
+					ChunkState chunkState = chunkStatePtr.Value;
+
+					if (!chunkState.HasReaderFlag(_readerIdx))
+					{
+						// Wait until the current chunk is readable
+						_resources.ResetReadEvent(_readerIdx);
+						if (!chunkState.HasReaderFlag(_readerIdx))
+						{
+							await _resources.WaitForReadEvent(_readerIdx, cancellationToken);
+						}
+					}
+					else if (readerState.Offset + minLength <= chunkState.Length)
+					{
+						// We have enough data in the chunk to be able to read a message
+						return true;
+					}
+					else if (chunkState.WriteState == WriteState.Writing)
+					{
+						// Wait until there is more data in the chunk
+						_resources.ResetReadEvent(_readerIdx);
+						if (_headerPtr.GetChunkStatePtr(readerState.ChunkIdx).Value == chunkState)
+						{
+							await _resources.WaitForReadEvent(_readerIdx, cancellationToken);
+						}
+					}
+					else if (readerState.Offset < chunkState.Length || chunkState.WriteState == WriteState.Complete)
+					{
+						// Cannot read the requested amount of data from this chunk.
+						return false;
+					}
+					else if (chunkState.WriteState == WriteState.MovedToNext)
+					{
+						// Move to the next chunk
+						chunkStatePtr.FinishReading(_readerIdx);
+						_resources.SetWriteEvent();
+
+						int chunkIdx = readerStatePtr.ChunkIdx + 1;
+						if (chunkIdx == _chunks.Length)
+						{
+							chunkIdx = 0;
+						}
+
+						readerStatePtr.Value = new ReaderState(chunkIdx, 0);
+					}
+					else
+					{
+						throw new NotImplementedException($"Invalid write state for buffer: {chunkState.WriteState}");
+					}
+				}
+			}
+		}
+
+		internal class WriterImpl : IComputeBufferWriter
+		{
+			readonly HeaderPtr _headerPtr;
+			readonly Memory<byte>[] _chunks;
+			ResourcesBase _resources;
+
+			public WriterImpl(ResourcesBase resources)
+			{
+				_headerPtr = resources.HeaderPtr;
+				_chunks = resources.Chunks;
+				_resources = resources;
+			}
+
+			public ResourcesBase GetResources() => _resources;
+
+			/// <inheritdoc/>
+			public IComputeBufferWriter AddRef()
+			{
+				_resources.AddRef();
+				return new WriterImpl(_resources);
+			}
+
+			/// <inheritdoc/>
+			public void Dispose()
+			{
+				if (_resources != null)
+				{
+					_resources.Release();
+					_resources = null!;
+				}
+			}
+
+			/// <inheritdoc/>
+			public void AdvanceWritePosition(int size)
+			{
+				if (size > 0)
+				{
+					ChunkStatePtr chunkStatePtr = _headerPtr.GetChunkStatePtr(_headerPtr.WriteChunkIdx);
+					ChunkState chunkState = chunkStatePtr.Value;
+
+					Debug.Assert(chunkState.WriteState == WriteState.Writing);
+					chunkStatePtr.Append(size);
+
+					_resources.SetAllReadEvents();
+				}
+			}
+
+			/// <inheritdoc/>
+			public Memory<byte> GetWriteBuffer()
+			{
+				ChunkState state = _headerPtr.GetChunkStatePtr(_headerPtr.WriteChunkIdx).Value;
+				if (state.WriteState == WriteState.Writing)
+				{
+					return _chunks[_headerPtr.WriteChunkIdx].Slice(state.Length);
+				}
+				else
+				{
+					return Memory<byte>.Empty;
+				}
+			}
+
+			/// <inheritdoc/>
+			public bool MarkComplete()
+			{
+				ChunkStatePtr chunkStatePtr = _headerPtr.GetChunkStatePtr(_headerPtr.WriteChunkIdx);
+				if (chunkStatePtr.WriteState != WriteState.Complete)
+				{
+					chunkStatePtr.MarkComplete();
+					_resources.SetAllReadEvents();
+					return true;
+				}
+				return false;
+			}
+
+			/// <inheritdoc/>
+			public async ValueTask WaitToWriteAsync(int minSize, CancellationToken cancellationToken = default)
+			{
+				if (minSize > _headerPtr.ChunkLength)
+				{
+					throw new ArgumentException("Requested read size is larger than chunk size.", nameof(minSize));
+				}
+
+				for (; ; )
+				{
+					ChunkStatePtr chunkStatePtr = _headerPtr.GetChunkStatePtr(_headerPtr.WriteChunkIdx);
+
+					ChunkState chunkState = chunkStatePtr.Value;
+					if (chunkState.WriteState == WriteState.Writing)
+					{
+						int length = chunkState.Length;
+						if (length + minSize <= _headerPtr.ChunkLength)
+						{
+							return;
+						}
+						chunkStatePtr.FinishWriting();
+					}
+
+					int nextChunkIdx = _headerPtr.WriteChunkIdx + 1;
+					if (nextChunkIdx == _chunks.Length)
+					{
+						nextChunkIdx = 0;
+					}
+
+					ChunkStatePtr nextChunkStatePtr = _headerPtr.GetChunkStatePtr(nextChunkIdx);
+					while (nextChunkStatePtr.ReaderFlags != 0)
+					{
+						await _resources.WaitForWriteEvent(cancellationToken);
+						_resources.ResetWriteEvent();
+					}
+
+					_headerPtr.WriteChunkIdx = nextChunkIdx;
+					nextChunkStatePtr.StartWriting(_headerPtr.NumReaders);
+
+					_resources.SetAllReadEvents();
+				}
+			}
+		}
+
+		internal ResourcesBase _resources;
+
+		/// <inheritdoc/>
+		public IComputeBufferReader Reader { get; }
+
+		/// <inheritdoc/>
+		public IComputeBufferWriter Writer { get; }
+
+		/// <summary>
+		/// Constructor
+		/// </summary>
+		/// <param name="resources">Resources shared between instances of the buffer</param>
+		internal ComputeBuffer(ResourcesBase resources)
+		{
+			_resources = resources;
+			Reader = new ReaderImpl(resources, 0);
+			Writer = new WriterImpl(resources);
+		}
+
+		/// <inheritdoc/>
+		public void Dispose()
+		{
+			Dispose(true);
+			GC.SuppressFinalize(this);
+		}
+
+		/// <summary>
+		/// Overriable dispose method
+		/// </summary>
+		protected virtual void Dispose(bool disposing)
+		{
+			if (_resources != null)
+			{
+				_resources.Release();
+				_resources = null!;
+			}
+		}
+
+		/// <inheritdoc/>
+		public abstract IComputeBuffer AddRef();
+	}
+}
diff --git a/Engine/Source/Programs/Shared/EpicGames.Horde/Compute/ComputeBufferBase.cs b/Engine/Source/Programs/Shared/EpicGames.Horde/Compute/ComputeBufferBase.cs
deleted file mode 100644
diff --git a/Engine/Source/Programs/Shared/EpicGames.Horde/Compute/ComputeExceptions.cs b/Engine/Source/Programs/Shared/EpicGames.Horde/Compute/ComputeExceptions.cs
--- a/Engine/Source/Programs/Shared/EpicGames.Horde/Compute/ComputeExceptions.cs
+++ b/Engine/Source/Programs/Shared/EpicGames.Horde/Compute/ComputeExceptions.cs
@@ -64,12 +64,12 @@ namespace EpicGames.Horde.Compute
 	/// <summary>
 	/// Exception thrown when an invalid message is received
 	/// </summary>
-	public sealed class InvalidComputeMessageException : ComputeException
+	public sealed class InvalidAgentMessageException : ComputeException
 	{
 		/// <summary>
 		/// Constructor
 		/// </summary>
-		public InvalidComputeMessageException(IComputeMessage message)
+		public InvalidAgentMessageException(IAgentMessage message)
 			: base($"Unexpected message {message.Type}")
 		{
 		}
diff --git a/Engine/Source/Programs/Shared/EpicGames.Horde/Compute/ComputeMessageChannel.cs b/Engine/Source/Programs/Shared/EpicGames.Horde/Compute/ComputeMessageChannel.cs
deleted file mode 100644
diff --git a/Engine/Source/Programs/Shared/EpicGames.Horde/Compute/ComputeProcess.cs b/Engine/Source/Programs/Shared/EpicGames.Horde/Compute/ComputeProcess.cs
deleted file mode 100644
diff --git a/Engine/Source/Programs/Shared/EpicGames.Horde/Compute/ComputeStorageClient.cs b/Engine/Source/Programs/Shared/EpicGames.Horde/Compute/ComputeStorageClient.cs
deleted file mode 100644
diff --git a/Engine/Source/Programs/Shared/EpicGames.Horde/Compute/ComputeWorker.cs b/Engine/Source/Programs/Shared/EpicGames.Horde/Compute/ComputeWorker.cs
deleted file mode 100644
diff --git a/Engine/Source/Programs/Shared/EpicGames.Horde/Compute/IComputeMessage.cs b/Engine/Source/Programs/Shared/EpicGames.Horde/Compute/IComputeMessage.cs
deleted file mode 100644
diff --git a/Engine/Source/Programs/Shared/EpicGames.Horde/Compute/IComputeMessageChannel.cs b/Engine/Source/Programs/Shared/EpicGames.Horde/Compute/IComputeMessageChannel.cs
deleted file mode 100644
diff --git a/Engine/Source/Programs/Shared/EpicGames.Horde/Compute/IComputeProcess.cs b/Engine/Source/Programs/Shared/EpicGames.Horde/Compute/IComputeProcess.cs
deleted file mode 100644
diff --git a/Engine/Source/Programs/Shared/EpicGames.Horde/Compute/IComputeSocket.cs b/Engine/Source/Programs/Shared/EpicGames.Horde/Compute/IComputeSocket.cs
--- a/Engine/Source/Programs/Shared/EpicGames.Horde/Compute/IComputeSocket.cs
+++ b/Engine/Source/Programs/Shared/EpicGames.Horde/Compute/IComputeSocket.cs
@@ -1,8 +1,6 @@
 // Copyright Epic Games, Inc. All Rights Reserved.
 
 using System;
-using EpicGames.Horde.Compute.Buffers;
-using Microsoft.Extensions.Logging;
 
 namespace EpicGames.Horde.Compute
 {
@@ -25,44 +23,4 @@ namespace EpicGames.Horde.Compute
 		/// <param name="sendBufferReader">Reader for the buffer to send data from</param>
 		void AttachSendBuffer(int channelId, IComputeBufferReader sendBufferReader);
 	}
-
-	/// <summary>
-	/// Extension methods for <see cref="IComputeMessageChannel"/>
-	/// </summary>
-	public static class ComputeSocketExtensions
-	{
-		/// <summary>
-		/// Creates a message channel with the given identifier
-		/// </summary>
-		/// <param name="socket">Socket to create a channel for</param>
-		/// <param name="channelId">Identifier for the channel</param>
-		/// <param name="logger">Logger for the channel</param>
-		public static IComputeMessageChannel CreateMessageChannel(this IComputeSocket socket, int channelId, ILogger logger) 
-			=> socket.CreateMessageChannel(channelId, 65536, logger);
-
-		/// <summary>
-		/// Creates a message channel with the given identifier
-		/// </summary>
-		/// <param name="socket">Socket to create a channel for</param>
-		/// <param name="channelId">Identifier for the channel</param>
-		/// <param name="bufferSize">Size of the send and receive buffer</param>
-		/// <param name="logger">Logger for the channel</param>
-		public static IComputeMessageChannel CreateMessageChannel(this IComputeSocket socket, int channelId, int bufferSize, ILogger logger)
-			=> socket.CreateMessageChannel(channelId, bufferSize, bufferSize, logger);
-
-		/// <summary>
-		/// Creates a message channel with the given identifier
-		/// </summary>
-		/// <param name="socket">Socket to create a channel for</param>
-		/// <param name="channelId">Identifier for the channel</param>
-		/// <param name="sendBufferSize">Size of the send buffer</param>
-		/// <param name="recvBufferSize">Size of the recieve buffer</param>
-		/// <param name="logger">Logger for the channel</param>
-		public static IComputeMessageChannel CreateMessageChannel(this IComputeSocket socket, int channelId, int sendBufferSize, int recvBufferSize, ILogger logger)
-		{
-			using IComputeBuffer sendBuffer = new PooledBuffer(sendBufferSize);
-			using IComputeBuffer recvBuffer = new PooledBuffer(recvBufferSize);
-			return new ComputeMessageChannel(socket, channelId, sendBuffer, recvBuffer, logger);
-		}
-	}
 }
