Horde: Fix compile error.

diff --git a/Engine/Source/Programs/Shared/EpicGames.Horde/Compute/AgentMessage.cs b/Engine/Source/Programs/Shared/EpicGames.Horde/Compute/AgentMessage.cs
--- a/Engine/Source/Programs/Shared/EpicGames.Horde/Compute/AgentMessage.cs
+++ b/Engine/Source/Programs/Shared/EpicGames.Horde/Compute/AgentMessage.cs
@@ -116,6 +116,17 @@ namespace EpicGames.Horde.Compute
 		public ReadOnlyMemory<byte> Data { get; }
 	}
 
+	/// <summary>
+	/// Writer for compute messages
+	/// </summary>
+	public interface IAgentMessageBuilder : IMemoryWriter, IDisposable
+	{
+		/// <summary>
+		/// Sends the current message
+		/// </summary>
+		void Send();
+	}
+
 	/// <summary>
 	/// Message for reporting an error
 	/// </summary>
