Horde: Remove IAgentMessage interface; expose the concrete type instead.

diff --git a/Engine/Source/Programs/Horde/Horde.Agent/Commands/Compute/DdcComputeCommand.cs b/Engine/Source/Programs/Horde/Horde.Agent/Commands/Compute/DdcComputeCommand.cs
--- a/Engine/Source/Programs/Horde/Horde.Agent/Commands/Compute/DdcComputeCommand.cs
+++ b/Engine/Source/Programs/Horde/Horde.Agent/Commands/Compute/DdcComputeCommand.cs
@@ -37,7 +37,7 @@ namespace Horde.Agent.Commands.Compute
 				await channel.SendXorRequestAsync(new byte[] { 1, 2, 3, 4, 5 }, (byte)123, cancellationToken);
 
 				_logger.LogInformation("Waiting for response...");
-				IAgentMessage response = await channel.ReceiveAsync(cancellationToken);
+				AgentMessage response = await channel.ReceiveAsync(cancellationToken);
 
 				byte[] result = response.Data.ToArray();
 				byte[] expectedResult = new byte[] { 1 ^ 123, 2 ^ 123, 3 ^ 123, 4 ^ 123, 5 ^ 123 };
diff --git a/Engine/Source/Programs/Horde/Horde.Agent/Commands/Compute/XorComputeCommand.cs b/Engine/Source/Programs/Horde/Horde.Agent/Commands/Compute/XorComputeCommand.cs
--- a/Engine/Source/Programs/Horde/Horde.Agent/Commands/Compute/XorComputeCommand.cs
+++ b/Engine/Source/Programs/Horde/Horde.Agent/Commands/Compute/XorComputeCommand.cs
@@ -41,7 +41,7 @@ namespace Horde.Agent.Commands.Compute
 				await channel.SendXorRequestAsync(new byte[] { 1, 2, 3, 4, 5 }, (byte)123, cancellationToken);
 
 				_logger.LogInformation("Waiting for response...");
-				IAgentMessage response = await channel.ReceiveAsync(cancellationToken);
+				AgentMessage response = await channel.ReceiveAsync(cancellationToken);
 
 				byte[] result = response.Data.ToArray();
 				byte[] expectedResult = new byte[] { 1 ^ 123, 2 ^ 123, 3 ^ 123, 4 ^ 123, 5 ^ 123 };
diff --git a/Engine/Source/Programs/Horde/Samples/RemoteClient/RemoteClientApp.cs b/Engine/Source/Programs/Horde/Samples/RemoteClient/RemoteClientApp.cs
--- a/Engine/Source/Programs/Horde/Samples/RemoteClient/RemoteClientApp.cs
+++ b/Engine/Source/Programs/Horde/Samples/RemoteClient/RemoteClientApp.cs
@@ -148,7 +148,7 @@ namespace RemoteClient
 			byte[] dataToXor = new byte[] { 1, 2, 3, 4, 5 };
 			await channel.SendXorRequestAsync(dataToXor, 123);
 
-			using IAgentMessage response = await channel.ReceiveAsync(AgentMessageType.XorResponse);
+			using AgentMessage response = await channel.ReceiveAsync(AgentMessageType.XorResponse);
 			for (int idx = 0; idx < dataToXor.Length; idx++)
 			{
 				if (response.Data.Span[idx] != (byte)(dataToXor[idx] ^ 123))
diff --git a/Engine/Source/Programs/Shared/EpicGames.Horde/Compute/AgentManagedProcess.cs b/Engine/Source/Programs/Shared/EpicGames.Horde/Compute/AgentManagedProcess.cs
--- a/Engine/Source/Programs/Shared/EpicGames.Horde/Compute/AgentManagedProcess.cs
+++ b/Engine/Source/Programs/Shared/EpicGames.Horde/Compute/AgentManagedProcess.cs
@@ -70,7 +70,7 @@ namespace EpicGames.Horde.Compute
 		{
 			for (; ; )
 			{
-				using IAgentMessage message = await channel.ReceiveAsync(cancellationToken);
+				using AgentMessage message = await channel.ReceiveAsync(cancellationToken);
 				switch (message.Type)
 				{
 					case AgentMessageType.Exception:
diff --git a/Engine/Source/Programs/Shared/EpicGames.Horde/Compute/AgentMessage.cs b/Engine/Source/Programs/Shared/EpicGames.Horde/Compute/AgentMessage.cs
--- a/Engine/Source/Programs/Shared/EpicGames.Horde/Compute/AgentMessage.cs
+++ b/Engine/Source/Programs/Shared/EpicGames.Horde/Compute/AgentMessage.cs
@@ -1,6 +1,7 @@
 // Copyright Epic Games, Inc. All Rights Reserved.
 
 using System;
+using System.Buffers;
 using System.Buffers.Binary;
 using System.Collections.Generic;
 using System.IO;
@@ -101,9 +102,9 @@ namespace EpicGames.Horde.Compute
 	}
 
 	/// <summary>
-	/// Information about a received compute message
+	/// Standard implementation of a message
 	/// </summary>
-	public interface IAgentMessage : IMemoryReader, IDisposable
+	public sealed class AgentMessage : IMemoryReader, IDisposable
 	{
 		/// <summary>
 		/// Type of the message
@@ -114,6 +115,47 @@ namespace EpicGames.Horde.Compute
 		/// Data that was read
 		/// </summary>
 		public ReadOnlyMemory<byte> Data { get; }
+
+		readonly IMemoryOwner<byte> _memoryOwner;
+		int _position;
+
+		/// <summary>
+		/// Constructor
+		/// </summary>
+		public AgentMessage(AgentMessageType type, ReadOnlyMemory<byte> data)
+		{
+			_memoryOwner = MemoryPool<byte>.Shared.Rent(data.Length);
+			data.CopyTo(_memoryOwner.Memory);
+
+			Type = type;
+			Data = _memoryOwner.Memory.Slice(0, data.Length);
+		}
+
+		/// <inheritdoc/>
+		public void Dispose()
+		{
+			_memoryOwner.Dispose();
+		}
+
+		/// <inheritdoc/>
+		public ReadOnlyMemory<byte> GetMemory(int minSize = 1) => Data.Slice(_position);
+
+		/// <inheritdoc/>
+		public void Advance(int length) => _position += length;
+	}
+
+	/// <summary>
+	/// Exception thrown when an invalid message is received
+	/// </summary>
+	public sealed class InvalidAgentMessageException : ComputeException
+	{
+		/// <summary>
+		/// Constructor
+		/// </summary>
+		public InvalidAgentMessageException(AgentMessage message)
+			: base($"Unexpected message {message.Type}")
+		{
+		}
 	}
 
 	/// <summary>
@@ -212,7 +254,7 @@ namespace EpicGames.Horde.Compute
 		/// <summary>
 		/// Parses a message as an <see cref="ExceptionMessage"/>
 		/// </summary>
-		public static ExceptionMessage ParseExceptionMessage(this IAgentMessage message)
+		public static ExceptionMessage ParseExceptionMessage(this AgentMessage message)
 		{
 			string msg = message.ReadString();
 			string description = message.ReadString();
@@ -233,7 +275,7 @@ namespace EpicGames.Horde.Compute
 		/// <summary>
 		/// Parses a fork request message
 		/// </summary>
-		public static ForkMessage ParseForkMessage(this IAgentMessage message)
+		public static ForkMessage ParseForkMessage(this AgentMessage message)
 		{
 			int channelId = message.ReadInt32();
 			int bufferSize = message.ReadInt32();
@@ -257,7 +299,7 @@ namespace EpicGames.Horde.Compute
 		/// <returns></returns>
 		public static async ValueTask WaitForAttachAsync(this AgentMessageChannel channel, CancellationToken cancellationToken = default)
 		{
-			using IAgentMessage message = await channel.ReceiveAsync(cancellationToken);
+			using AgentMessage message = await channel.ReceiveAsync(cancellationToken);
 			if (message.Type != AgentMessageType.Attach)
 			{
 				throw new InvalidAgentMessageException(message);
@@ -266,11 +308,11 @@ namespace EpicGames.Horde.Compute
 
 		#region Process
 
-		static async Task<IAgentMessage> RunStorageServer(this AgentMessageChannel channel, IStorageClient storage, CancellationToken cancellationToken = default)
+		static async Task<AgentMessage> RunStorageServer(this AgentMessageChannel channel, IStorageClient storage, CancellationToken cancellationToken = default)
 		{
 			for (; ; )
 			{
-				IAgentMessage message = await channel.ReceiveAsync(cancellationToken);
+				AgentMessage message = await channel.ReceiveAsync(cancellationToken);
 				if (message.Type != AgentMessageType.ReadBlob)
 				{
 					return message;
@@ -305,7 +347,7 @@ namespace EpicGames.Horde.Compute
 				request.Send();
 			}
 
-			using IAgentMessage response = await RunStorageServer(channel, storage, cancellationToken);
+			using AgentMessage response = await RunStorageServer(channel, storage, cancellationToken);
 			if (response.Type != AgentMessageType.WriteFilesResponse)
 			{
 				throw new InvalidAgentMessageException(response);
@@ -315,7 +357,7 @@ namespace EpicGames.Horde.Compute
 		/// <summary>
 		/// Parses a message as a <see cref="UploadFilesMessage"/>
 		/// </summary>
-		public static UploadFilesMessage ParseUploadFilesMessage(this IAgentMessage message)
+		public static UploadFilesMessage ParseUploadFilesMessage(this AgentMessage message)
 		{
 			string name = message.ReadString();
 			NodeLocator locator = message.ReadNodeLocator();
@@ -338,7 +380,7 @@ namespace EpicGames.Horde.Compute
 		/// <summary>
 		/// Parses a message as a <see cref="DeleteFilesMessage"/>
 		/// </summary>
-		public static DeleteFilesMessage ParseDeleteFilesMessage(this IAgentMessage message)
+		public static DeleteFilesMessage ParseDeleteFilesMessage(this AgentMessage message)
 		{
 			List<string> files = message.ReadList(MemoryReaderExtensions.ReadString);
 			return new DeleteFilesMessage(files);
@@ -369,7 +411,7 @@ namespace EpicGames.Horde.Compute
 		/// <summary>
 		/// Parses a message as a <see cref="ExecuteProcessMessage"/>
 		/// </summary>
-		public static ExecuteProcessMessage ParseExecuteProcessMessage(this IAgentMessage message)
+		public static ExecuteProcessMessage ParseExecuteProcessMessage(this AgentMessage message)
 		{
 			string executable = message.ReadString();
 			List<string> arguments = message.ReadList(MemoryReaderExtensions.ReadString);
@@ -404,7 +446,7 @@ namespace EpicGames.Horde.Compute
 		/// <summary>
 		/// Parses a message as a <see cref="ExecuteProcessMessage"/>
 		/// </summary>
-		public static ExecuteProcessResponseMessage ParseExecuteProcessResponse(this IAgentMessage message)
+		public static ExecuteProcessResponseMessage ParseExecuteProcessResponse(this AgentMessage message)
 		{
 			int exitCode = message.ReadInt32();
 			return new ExecuteProcessResponseMessage(exitCode);
@@ -419,7 +461,7 @@ namespace EpicGames.Horde.Compute
 		/// </summary>
 		/// <param name="message"></param>
 		/// <returns></returns>
-		public static ReadBlobMessage ParseReadBlobRequest(this IAgentMessage message)
+		public static ReadBlobMessage ParseReadBlobRequest(this AgentMessage message)
 		{
 			BlobLocator locator = message.ReadBlobLocator();
 			int offset = (int)message.ReadUnsignedVarInt();
@@ -432,9 +474,9 @@ namespace EpicGames.Horde.Compute
 		/// </summary>
 		sealed class BlobDataStream : ReadOnlyMemoryStream
 		{
-			readonly IAgentMessage _message;
+			readonly AgentMessage _message;
 
-			public BlobDataStream(IAgentMessage message)
+			public BlobDataStream(AgentMessage message)
 				: base(message.Data.Slice(8))
 			{
 				_message = message;
@@ -473,7 +515,7 @@ namespace EpicGames.Horde.Compute
 			byte[]? buffer = null;
 			for(; ;)
 			{
-				IAgentMessage? response = null;
+				AgentMessage? response = null;
 				try
 				{
 					response = await channel.ReceiveAsync(cancellationToken);
@@ -587,7 +629,7 @@ namespace EpicGames.Horde.Compute
 		/// <summary>
 		/// Parse a message as an XOR request
 		/// </summary>
-		public static XorRequestMessage AsXorRequest(this IAgentMessage message)
+		public static XorRequestMessage AsXorRequest(this AgentMessage message)
 		{
 			ReadOnlyMemory<byte> data = message.Data;
 			return new XorRequestMessage(data[0..^1], data.Span[^1]);
diff --git a/Engine/Source/Programs/Shared/EpicGames.Horde/Compute/AgentMessageChannel.cs b/Engine/Source/Programs/Shared/EpicGames.Horde/Compute/AgentMessageChannel.cs
--- a/Engine/Source/Programs/Shared/EpicGames.Horde/Compute/AgentMessageChannel.cs
+++ b/Engine/Source/Programs/Shared/EpicGames.Horde/Compute/AgentMessageChannel.cs
@@ -20,42 +20,6 @@ namespace EpicGames.Horde.Compute
 		// Length of a message header. Consists of a 1 byte type field, followed by 4 byte length field.
 		const int HeaderLength = 5;
 
-		/// <summary>
-		/// Standard implementation of a message
-		/// </summary>
-		sealed class Message : IAgentMessage
-		{
-			/// <inheritdoc/>
-			public AgentMessageType Type { get; }
-
-			/// <inheritdoc/>
-			public ReadOnlyMemory<byte> Data { get; }
-
-			readonly IMemoryOwner<byte> _memoryOwner;
-			int _position;
-
-			public Message(AgentMessageType type, ReadOnlyMemory<byte> data)
-			{
-				_memoryOwner = MemoryPool<byte>.Shared.Rent(data.Length);
-				data.CopyTo(_memoryOwner.Memory);
-
-				Type = type;
-				Data = _memoryOwner.Memory.Slice(0, data.Length);
-			}
-
-			/// <inheritdoc/>
-			public void Dispose()
-			{
-				_memoryOwner.Dispose();
-			}
-
-			/// <inheritdoc/>
-			public ReadOnlyMemory<byte> GetMemory(int minSize = 1) => Data.Slice(_position);
-
-			/// <inheritdoc/>
-			public void Advance(int length) => _position += length;
-		}
-
 		/// <summary>
 		/// Allows creating new messages in rented memory
 		/// </summary>
@@ -175,7 +139,7 @@ namespace EpicGames.Horde.Compute
 		}
 
 		/// <inheritdoc/>
-		public async ValueTask<IAgentMessage> ReceiveAsync(CancellationToken cancellationToken)
+		public async ValueTask<AgentMessage> ReceiveAsync(CancellationToken cancellationToken)
 		{
 			while (!_recvBufferReader.IsComplete)
 			{
@@ -194,7 +158,7 @@ namespace EpicGames.Horde.Compute
 				}
 
 				AgentMessageType type = (AgentMessageType)memory.Span[0];
-				Message message = new Message(type, memory.Slice(HeaderLength, messageLength));
+				AgentMessage message = new AgentMessage(type, memory.Slice(HeaderLength, messageLength));
 				if (_logger.IsEnabled(LogLevel.Trace))
 				{
 					LogMessageInfo("RECV", message.Type, message.Data.Span);
@@ -203,7 +167,7 @@ namespace EpicGames.Horde.Compute
 				_recvBufferReader.AdvanceReadPosition(HeaderLength + messageLength);
 				return message;
 			}
-			return new Message(AgentMessageType.None, ReadOnlyMemory<byte>.Empty);
+			return new AgentMessage(AgentMessageType.None, ReadOnlyMemory<byte>.Empty);
 		}
 
 		void LogMessageInfo(string verb, AgentMessageType type, ReadOnlySpan<byte> data)
@@ -281,9 +245,9 @@ namespace EpicGames.Horde.Compute
 		/// <param name="type">Expected type of the message</param>
 		/// <param name="cancellationToken">Cancellation token for the operation</param>
 		/// <returns>Data for a message that was read. Must be disposed.</returns>
-		public static async ValueTask<IAgentMessage> ReceiveAsync(this AgentMessageChannel channel, AgentMessageType type, CancellationToken cancellationToken = default)
+		public static async ValueTask<AgentMessage> ReceiveAsync(this AgentMessageChannel channel, AgentMessageType type, CancellationToken cancellationToken = default)
 		{
-			IAgentMessage message = await channel.ReceiveAsync(cancellationToken);
+			AgentMessage message = await channel.ReceiveAsync(cancellationToken);
 			if (message.Type != type)
 			{
 				throw new InvalidAgentMessageException(message);
@@ -309,7 +273,7 @@ namespace EpicGames.Horde.Compute
 		/// <param name="channel">Channel to send on</param>
 		/// <param name="message">The message to be sent</param>
 		/// <param name="cancellationToken">Cancellation token for the operation</param>
-		public static async ValueTask SendAsync(this AgentMessageChannel channel, IAgentMessage message, CancellationToken cancellationToken)
+		public static async ValueTask SendAsync(this AgentMessageChannel channel, AgentMessage message, CancellationToken cancellationToken)
 		{
 			using (IAgentMessageBuilder builder = await channel.CreateMessageAsync(message.Type, message.Data.Length, cancellationToken))
 			{
diff --git a/Engine/Source/Programs/Shared/EpicGames.Horde/Compute/AgentMessageHandler.cs b/Engine/Source/Programs/Shared/EpicGames.Horde/Compute/AgentMessageHandler.cs
--- a/Engine/Source/Programs/Shared/EpicGames.Horde/Compute/AgentMessageHandler.cs
+++ b/Engine/Source/Programs/Shared/EpicGames.Horde/Compute/AgentMessageHandler.cs
@@ -56,7 +56,7 @@ namespace EpicGames.Horde.Compute
 				List<Task> childTasks = new List<Task>();
 				for (; ; )
 				{
-					using IAgentMessage message = await channel.ReceiveAsync(cancellationToken);
+					using AgentMessage message = await channel.ReceiveAsync(cancellationToken);
 					_logger.LogTrace("Compute Channel {ChannelId}: {MessageType}", channelId, message.Type);
 
 					switch (message.Type)
diff --git a/Engine/Source/Programs/Shared/EpicGames.Horde/Compute/ComputeExceptions.cs b/Engine/Source/Programs/Shared/EpicGames.Horde/Compute/ComputeExceptions.cs
--- a/Engine/Source/Programs/Shared/EpicGames.Horde/Compute/ComputeExceptions.cs
+++ b/Engine/Source/Programs/Shared/EpicGames.Horde/Compute/ComputeExceptions.cs
@@ -60,18 +60,4 @@ namespace EpicGames.Horde.Compute
 		/// <inheritdoc/>
 		public override string ToString() => _description;
 	}
-
-	/// <summary>
-	/// Exception thrown when an invalid message is received
-	/// </summary>
-	public sealed class InvalidAgentMessageException : ComputeException
-	{
-		/// <summary>
-		/// Constructor
-		/// </summary>
-		public InvalidAgentMessageException(IAgentMessage message)
-			: base($"Unexpected message {message.Type}")
-		{
-		}
-	}
 }
