Horde: Rename ComputeStorageClient -> AgentStorageClient.

diff --git a/Engine/Source/Programs/Shared/EpicGames.Horde/Compute/AgentMessageHandler.cs b/Engine/Source/Programs/Shared/EpicGames.Horde/Compute/AgentMessageHandler.cs
--- a/Engine/Source/Programs/Shared/EpicGames.Horde/Compute/AgentMessageHandler.cs
+++ b/Engine/Source/Programs/Shared/EpicGames.Horde/Compute/AgentMessageHandler.cs
@@ -118,7 +118,7 @@ namespace EpicGames.Horde.Compute
 
 		async Task WriteFilesAsync(AgentMessageChannel channel, string path, NodeLocator locator, CancellationToken cancellationToken)
 		{
-			using ComputeStorageClient store = new ComputeStorageClient(channel);
+			using AgentStorageClient store = new AgentStorageClient(channel);
 			TreeReader reader = new TreeReader(store, _memoryCache, _logger);
 
 			DirectoryNode directoryNode = await reader.ReadNodeAsync<DirectoryNode>(locator, cancellationToken);
diff --git a/Engine/Source/Programs/Shared/EpicGames.Horde/Compute/AgentStorageClient.cs b/Engine/Source/Programs/Shared/EpicGames.Horde/Compute/AgentStorageClient.cs
--- a/Engine/Source/Programs/Shared/EpicGames.Horde/Compute/AgentStorageClient.cs
+++ b/Engine/Source/Programs/Shared/EpicGames.Horde/Compute/AgentStorageClient.cs
@@ -13,7 +13,7 @@ namespace EpicGames.Horde.Compute
 	/// <summary>
 	/// Storage client which can read bundles over a compute channel
 	/// </summary>
-	public sealed class ComputeStorageClient : StorageClientBase, IDisposable
+	public sealed class AgentStorageClient : StorageClientBase, IDisposable
 	{
 		readonly AgentMessageChannel _channel;
 		readonly SemaphoreSlim _semaphore;
@@ -22,7 +22,7 @@ namespace EpicGames.Horde.Compute
 		/// Constructor
 		/// </summary>
 		/// <param name="channel"></param>
-		public ComputeStorageClient(AgentMessageChannel channel)
+		public AgentStorageClient(AgentMessageChannel channel)
 		{
 			_channel = channel;
 			_semaphore = new SemaphoreSlim(1);
