Horde: Move methods for building file nodes into FileNodeWriter.

diff --git a/Engine/Source/Programs/Shared/EpicGames.Horde/Storage/Nodes/FileNode.cs b/Engine/Source/Programs/Shared/EpicGames.Horde/Storage/Nodes/FileNode.cs
--- a/Engine/Source/Programs/Shared/EpicGames.Horde/Storage/Nodes/FileNode.cs
+++ b/Engine/Source/Programs/Shared/EpicGames.Horde/Storage/Nodes/FileNode.cs
@@ -11,71 +11,6 @@ using EpicGames.Core;
 
 namespace EpicGames.Horde.Storage.Nodes
 {
-	/// <summary>
-	/// Options for creating file nodes
-	/// </summary>
-	public class ChunkingOptions
-	{
-		/// <summary>
-		/// Options for creating leaf nodes
-		/// </summary>
-		public ChunkingOptionsForNodeType LeafOptions { get; set; }
-
-		/// <summary>
-		/// Options for creating interior nodes
-		/// </summary>
-		public ChunkingOptionsForNodeType InteriorOptions { get; set; }
-
-		/// <summary>
-		/// Constructor
-		/// </summary>
-		public ChunkingOptions()
-		{
-			LeafOptions = new ChunkingOptionsForNodeType(32 * 1024, 256 * 1024, 64 * 1024);
-			InteriorOptions = new ChunkingOptionsForNodeType(32 * 1024, 256 * 1024, 64 * 1024);
-		}
-	}
-
-	/// <summary>
-	/// Options for creating a specific type of file nodes
-	/// </summary>
-	public class ChunkingOptionsForNodeType
-	{
-		/// <summary>
-		/// Minimum chunk size
-		/// </summary>
-		public int MinSize { get; set; }
-
-		/// <summary>
-		/// Maximum chunk size. Chunks will be split on this boundary if another match is not found.
-		/// </summary>
-		public int MaxSize { get; set; }
-
-		/// <summary>
-		/// Target chunk size for content-slicing
-		/// </summary>
-		public int TargetSize { get; set; }
-
-		/// <summary>
-		/// Constructor
-		/// </summary>
-		/// <param name="size">Fixed size chunks to use</param>
-		public ChunkingOptionsForNodeType(int size)
-			: this(size, size, size)
-		{
-		}
-
-		/// <summary>
-		/// Default constructor
-		/// </summary>
-		public ChunkingOptionsForNodeType(int minSize, int maxSize, int targetSize)
-		{
-			MinSize = minSize;
-			MaxSize = maxSize;
-			TargetSize = targetSize;
-		}
-	}
-
 	/// <summary>
 	/// Representation of a data stream, split into chunks along content-aware boundaries using a rolling hash (<see cref="BuzHash"/>).
 	/// Chunks are pushed into a tree hierarchy as data is appended to the root, with nodes of the tree also split along content-aware boundaries with <see cref="IoHash.NumBytes"/> granularity.
@@ -190,83 +125,6 @@ namespace EpicGames.Horde.Storage.Nodes
 		/// <inheritdoc/>
 		public override IEnumerable<NodeRef> EnumerateRefs() => Enumerable.Empty<NodeRef>();
 
-		/// <summary>
-		/// Determines how much data to append to an existing leaf node
-		/// </summary>
-		/// <param name="currentData">Current data in the leaf node</param>
-		/// <param name="appendData">Data to be appended</param>
-		/// <param name="rollingHash">Current BuzHash of the data</param>
-		/// <param name="options">Options for chunking the data</param>
-		/// <returns>The number of bytes to append</returns>
-		internal static int AppendData(ReadOnlySpan<byte> currentData, ReadOnlySpan<byte> appendData, ref uint rollingHash, ChunkingOptionsForNodeType options)
-		{
-			// If the target option sizes are fixed, just chunk the data along fixed boundaries
-			if (options.MinSize == options.TargetSize && options.MaxSize == options.TargetSize)
-			{
-				return Math.Min(appendData.Length, options.MaxSize - (int)currentData.Length);
-			}
-
-			// Cap the append data span to the maximum amount we can add
-			int maxAppendLength = options.MaxSize - currentData.Length;
-			if (maxAppendLength < appendData.Length)
-			{
-				appendData = appendData.Slice(0, maxAppendLength);
-			}
-
-			// Length of the data to be appended
-			int appendLength = 0;
-
-			// Fast path for appending data to the buffer up to the chunk window size
-			int windowSize = options.MinSize;
-			if (currentData.Length < windowSize)
-			{
-				appendLength = Math.Min(windowSize - (int)currentData.Length, appendData.Length);
-				rollingHash = BuzHash.Add(rollingHash, appendData.Slice(0, appendLength));
-			}
-
-			// Get the threshold for the rolling hash
-			uint rollingHashThreshold = (uint)((1L << 32) / options.TargetSize);
-
-			// Step through the part of the data where the tail of the window is in currentData, and the head of the window is in appendData.
-			if(appendLength < appendData.Length && windowSize > appendLength)
-			{
-				int overlap = windowSize - appendLength;
-				int overlapLength = Math.Min(appendData.Length - appendLength, overlap);
-
-				ReadOnlySpan<byte> tailSpan = currentData.Slice(currentData.Length - overlap, overlapLength);
-				ReadOnlySpan<byte> headSpan = appendData.Slice(appendLength, overlapLength);
-
-				int count = BuzHash.Update(tailSpan, headSpan, rollingHashThreshold, ref rollingHash);
-				if (count != -1)
-				{
-					appendLength += count;
-					return appendLength;
-				}
-
-				appendLength += headSpan.Length;
-			}
-
-			// Step through the rest of the data which is completely contained in appendData.
-			if (appendLength < appendData.Length)
-			{
-				Debug.Assert(appendLength >= windowSize);
-					
-				ReadOnlySpan<byte> tailSpan = appendData.Slice(appendLength - windowSize, appendData.Length - windowSize);
-				ReadOnlySpan<byte> headSpan = appendData.Slice(appendLength);
-
-				int count = BuzHash.Update(tailSpan, headSpan, rollingHashThreshold, ref rollingHash);
-				if (count != -1)
-				{
-					appendLength += count;
-					return appendLength;
-				}
-
-				appendLength += headSpan.Length;
-			}
-
-			return appendLength;
-		}
-
 		/// <inheritdoc/>
 		public override async Task CopyToStreamAsync(TreeReader reader, Stream outputStream, CancellationToken cancellationToken)
 		{
@@ -332,55 +190,6 @@ namespace EpicGames.Horde.Storage.Nodes
 		/// <inheritdoc/>
 		public override IEnumerable<NodeRef> EnumerateRefs() => Children;
 
-		/// <summary>
-		/// Test whether the current node is complete
-		/// </summary>
-		/// <param name="currentData"></param>
-		/// <param name="rollingHash"></param>
-		/// <param name="options"></param>
-		/// <returns></returns>
-		internal static bool IsComplete(ReadOnlySpan<byte> currentData, uint rollingHash, ChunkingOptionsForNodeType options)
-		{
-			if (currentData.Length + IoHash.NumBytes > options.MaxSize)
-			{
-				return true;
-			}
-
-			if (currentData.Length >= options.MinSize)
-			{
-				uint rollingHashThreshold = BuzHash.GetThreshold(options.TargetSize);
-				if (rollingHash < rollingHashThreshold)
-				{
-					return true;
-				}
-			}
-
-			return false;
-		}
-
-		/// <summary>
-		/// Append a new hash to this interior node
-		/// </summary>
-		/// <param name="currentData">Current data for the node</param>
-		/// <param name="hash">Hash of the child node</param>
-		/// <param name="rollingHash">Current rolling hash for the node</param>
-		/// <param name="options">Options for chunking the node</param>
-		/// <returns>True if the hash could be appended, false otherwise</returns>
-		internal static void AppendData(ReadOnlySpan<byte> currentData, IoHash hash, ref uint rollingHash, ChunkingOptionsForNodeType options)
-		{
-			Span<byte> hashData = stackalloc byte[IoHash.NumBytes];
-			hash.CopyTo(hashData);
-
-			rollingHash = BuzHash.Add(rollingHash, hashData);
-
-			int windowSize = options.MinSize - (options.MinSize % IoHash.NumBytes);
-			if (currentData.Length > windowSize)
-			{
-				ReadOnlySpan<byte> removeData = currentData.Slice(currentData.Length - windowSize, IoHash.NumBytes);
-				rollingHash = BuzHash.Sub(rollingHash, removeData, windowSize + IoHash.NumBytes);
-			}
-		}
-
 		/// <inheritdoc/>
 		public override async Task CopyToStreamAsync(TreeReader reader, Stream outputStream, CancellationToken cancellationToken)
 		{
diff --git a/Engine/Source/Programs/Shared/EpicGames.Horde/Storage/Nodes/FileNodeWriter.cs b/Engine/Source/Programs/Shared/EpicGames.Horde/Storage/Nodes/FileNodeWriter.cs
--- a/Engine/Source/Programs/Shared/EpicGames.Horde/Storage/Nodes/FileNodeWriter.cs
+++ b/Engine/Source/Programs/Shared/EpicGames.Horde/Storage/Nodes/FileNodeWriter.cs
@@ -3,6 +3,7 @@
 using System;
 using System.Buffers;
 using System.Collections.Generic;
+using System.Diagnostics;
 using System.IO;
 using System.Threading;
 using System.Threading.Tasks;
@@ -10,12 +11,76 @@ using EpicGames.Core;
 
 namespace EpicGames.Horde.Storage.Nodes
 {
+	/// <summary>
+	/// Options for creating file nodes
+	/// </summary>
+	public class ChunkingOptions
+	{
+		/// <summary>
+		/// Options for creating leaf nodes
+		/// </summary>
+		public ChunkingOptionsForNodeType LeafOptions { get; set; }
+
+		/// <summary>
+		/// Options for creating interior nodes
+		/// </summary>
+		public ChunkingOptionsForNodeType InteriorOptions { get; set; }
+
+		/// <summary>
+		/// Constructor
+		/// </summary>
+		public ChunkingOptions()
+		{
+			LeafOptions = new ChunkingOptionsForNodeType(32 * 1024, 256 * 1024, 64 * 1024);
+			InteriorOptions = new ChunkingOptionsForNodeType(32 * 1024, 256 * 1024, 64 * 1024);
+		}
+	}
+
+	/// <summary>
+	/// Options for creating a specific type of file nodes
+	/// </summary>
+	public class ChunkingOptionsForNodeType
+	{
+		/// <summary>
+		/// Minimum chunk size
+		/// </summary>
+		public int MinSize { get; set; }
+
+		/// <summary>
+		/// Maximum chunk size. Chunks will be split on this boundary if another match is not found.
+		/// </summary>
+		public int MaxSize { get; set; }
+
+		/// <summary>
+		/// Target chunk size for content-slicing
+		/// </summary>
+		public int TargetSize { get; set; }
+
+		/// <summary>
+		/// Constructor
+		/// </summary>
+		/// <param name="size">Fixed size chunks to use</param>
+		public ChunkingOptionsForNodeType(int size)
+			: this(size, size, size)
+		{
+		}
+
+		/// <summary>
+		/// Default constructor
+		/// </summary>
+		public ChunkingOptionsForNodeType(int minSize, int maxSize, int targetSize)
+		{
+			MinSize = minSize;
+			MaxSize = maxSize;
+			TargetSize = targetSize;
+		}
+	}
+
 	/// <summary>
 	/// Utility class for generating FileNode data directly into <see cref="TreeWriter"/> instances, without constructing node representations first.
 	/// </summary>
 	public class FileNodeWriter
 	{
-
 		/// <summary>
 		/// Default buffer length when calling CreateAsync/AppendAsync
 		/// </summary>
@@ -217,7 +282,7 @@ namespace EpicGames.Horde.Storage.Nodes
 			for (; ; )
 			{
 				// Append data to the current leaf node
-				int appendLength = LeafFileNode.AppendData(buffer.Span.Slice(0, _leafLength), data.Span, ref _leafHash, _options.LeafOptions);
+				int appendLength = AppendToLeafNode(buffer.Span.Slice(0, _leafLength), data.Span, ref _leafHash, _options.LeafOptions);
 
 				buffer = _writer.GetOutputBuffer(_leafLength, _leafLength + appendLength);
 				data.Slice(0, appendLength).CopyTo(buffer.Slice(_leafLength));
@@ -239,6 +304,83 @@ namespace EpicGames.Horde.Storage.Nodes
 			}
 		}
 
+		/// <summary>
+		/// Determines how much data to append to an existing leaf node
+		/// </summary>
+		/// <param name="currentData">Current data in the leaf node</param>
+		/// <param name="appendData">Data to be appended</param>
+		/// <param name="rollingHash">Current BuzHash of the data</param>
+		/// <param name="options">Options for chunking the data</param>
+		/// <returns>The number of bytes to append</returns>
+		static int AppendToLeafNode(ReadOnlySpan<byte> currentData, ReadOnlySpan<byte> appendData, ref uint rollingHash, ChunkingOptionsForNodeType options)
+		{
+			// If the target option sizes are fixed, just chunk the data along fixed boundaries
+			if (options.MinSize == options.TargetSize && options.MaxSize == options.TargetSize)
+			{
+				return Math.Min(appendData.Length, options.MaxSize - (int)currentData.Length);
+			}
+
+			// Cap the append data span to the maximum amount we can add
+			int maxAppendLength = options.MaxSize - currentData.Length;
+			if (maxAppendLength < appendData.Length)
+			{
+				appendData = appendData.Slice(0, maxAppendLength);
+			}
+
+			// Length of the data to be appended
+			int appendLength = 0;
+
+			// Fast path for appending data to the buffer up to the chunk window size
+			int windowSize = options.MinSize;
+			if (currentData.Length < windowSize)
+			{
+				appendLength = Math.Min(windowSize - (int)currentData.Length, appendData.Length);
+				rollingHash = BuzHash.Add(rollingHash, appendData.Slice(0, appendLength));
+			}
+
+			// Get the threshold for the rolling hash
+			uint rollingHashThreshold = (uint)((1L << 32) / options.TargetSize);
+
+			// Step through the part of the data where the tail of the window is in currentData, and the head of the window is in appendData.
+			if (appendLength < appendData.Length && windowSize > appendLength)
+			{
+				int overlap = windowSize - appendLength;
+				int overlapLength = Math.Min(appendData.Length - appendLength, overlap);
+
+				ReadOnlySpan<byte> tailSpan = currentData.Slice(currentData.Length - overlap, overlapLength);
+				ReadOnlySpan<byte> headSpan = appendData.Slice(appendLength, overlapLength);
+
+				int count = BuzHash.Update(tailSpan, headSpan, rollingHashThreshold, ref rollingHash);
+				if (count != -1)
+				{
+					appendLength += count;
+					return appendLength;
+				}
+
+				appendLength += headSpan.Length;
+			}
+
+			// Step through the rest of the data which is completely contained in appendData.
+			if (appendLength < appendData.Length)
+			{
+				Debug.Assert(appendLength >= windowSize);
+
+				ReadOnlySpan<byte> tailSpan = appendData.Slice(appendLength - windowSize, appendData.Length - windowSize);
+				ReadOnlySpan<byte> headSpan = appendData.Slice(appendLength);
+
+				int count = BuzHash.Update(tailSpan, headSpan, rollingHashThreshold, ref rollingHash);
+				if (count != -1)
+				{
+					appendLength += count;
+					return appendLength;
+				}
+
+				appendLength += headSpan.Length;
+			}
+
+			return appendLength;
+		}
+
 		async Task AddToInteriorNodeAsync(NodeHandle handle, CancellationToken cancellationToken)
 		{
 			_topInteriorNode ??= CreateInteriorNode();
@@ -248,7 +390,7 @@ namespace EpicGames.Horde.Storage.Nodes
 		async Task AddToInteriorNodeAsync(InteriorNodeState interiorNode, NodeHandle handle, CancellationToken cancellationToken)
 		{
 			// If the node is already full, flush it
-			if (InteriorFileNode.IsComplete(interiorNode._data.WrittenSpan, interiorNode._rollingHash, _options.InteriorOptions))
+			if (IsInteriorNodeComplete(interiorNode._data.WrittenSpan, interiorNode._rollingHash, _options.InteriorOptions))
 			{
 				NodeHandle interiorNodeHandle = await WriteInteriorNodeAndResetAsync(interiorNode, cancellationToken);
 				interiorNode._parent ??= CreateInteriorNode();
@@ -256,10 +398,60 @@ namespace EpicGames.Horde.Storage.Nodes
 			}
 
 			// Add this handle
-			InteriorFileNode.AppendData(interiorNode._data.WrittenSpan, handle.Hash, ref interiorNode._rollingHash, _options.InteriorOptions);
+			AppendToInteriorNode(interiorNode._data.WrittenSpan, handle.Hash, ref interiorNode._rollingHash, _options.InteriorOptions);
 			interiorNode.Write(handle);
 		}
 
+
+		/// <summary>
+		/// Test whether the current node is complete
+		/// </summary>
+		/// <param name="currentData"></param>
+		/// <param name="rollingHash"></param>
+		/// <param name="options"></param>
+		/// <returns></returns>
+		static bool IsInteriorNodeComplete(ReadOnlySpan<byte> currentData, uint rollingHash, ChunkingOptionsForNodeType options)
+		{
+			if (currentData.Length + IoHash.NumBytes > options.MaxSize)
+			{
+				return true;
+			}
+
+			if (currentData.Length >= options.MinSize)
+			{
+				uint rollingHashThreshold = BuzHash.GetThreshold(options.TargetSize);
+				if (rollingHash < rollingHashThreshold)
+				{
+					return true;
+				}
+			}
+
+			return false;
+		}
+
+		/// <summary>
+		/// Append a new hash to this interior node
+		/// </summary>
+		/// <param name="currentData">Current data for the node</param>
+		/// <param name="hash">Hash of the child node</param>
+		/// <param name="rollingHash">Current rolling hash for the node</param>
+		/// <param name="options">Options for chunking the node</param>
+		/// <returns>True if the hash could be appended, false otherwise</returns>
+		static void AppendToInteriorNode(ReadOnlySpan<byte> currentData, IoHash hash, ref uint rollingHash, ChunkingOptionsForNodeType options)
+		{
+			Span<byte> hashData = stackalloc byte[IoHash.NumBytes];
+			hash.CopyTo(hashData);
+
+			rollingHash = BuzHash.Add(rollingHash, hashData);
+
+			int windowSize = options.MinSize - (options.MinSize % IoHash.NumBytes);
+			if (currentData.Length > windowSize)
+			{
+				ReadOnlySpan<byte> removeData = currentData.Slice(currentData.Length - windowSize, IoHash.NumBytes);
+				rollingHash = BuzHash.Sub(rollingHash, removeData, windowSize + IoHash.NumBytes);
+			}
+		}
+
 		/// <summary>
 		/// Complete the current file, and write all open nodes to the underlying writer
 		/// </summary>
