Horde: Fix agent mutex not being disposed on the same thread that created it.

diff --git a/Engine/Source/Programs/Horde/Horde.Agent/Services/WorkerService.cs b/Engine/Source/Programs/Horde/Horde.Agent/Services/WorkerService.cs
--- a/Engine/Source/Programs/Horde/Horde.Agent/Services/WorkerService.cs
+++ b/Engine/Source/Programs/Horde/Horde.Agent/Services/WorkerService.cs
@@ -6,8 +6,8 @@ using System.Diagnostics;
 using System.Linq;
 using System.Threading;
 using System.Threading.Tasks;
+using EpicGames.Core;
 using Horde.Agent.Leases;
-using Horde.Agent.Utility;
 using Microsoft.Extensions.DependencyInjection;
 using Microsoft.Extensions.Hosting;
 using Microsoft.Extensions.Logging;
@@ -94,8 +94,15 @@ namespace Horde.Agent.Services
 				{
 					try
 					{
-						using Mutex singleInstanceMutex = new(false, "Global\\HordeAgent-DB828ACB-0AA5-4D32-A62A-21D4429B1014");
-						await WaitForMutexAsync(singleInstanceMutex, stoppingToken);
+						Task<IDisposable> mutexTask = SingleInstanceMutex.AcquireAsync("Global\\HordeAgent-DB828ACB-0AA5-4D32-A62A-21D4429B1014", stoppingToken);
+
+						Task delayTask = Task.Delay(TimeSpan.FromSeconds(1.0), stoppingToken);
+						if (Task.WhenAny(mutexTask, delayTask) == delayTask)
+						{
+							_logger.LogInformation("Another agent instance is already running. Waiting for it to terminate.");
+						}
+
+						using IDisposable mutex = await mutexTask;
 
 						await using (ISession session = await _sessionFactory.CreateAsync(stoppingToken))
 						{
diff --git a/Engine/Source/Programs/Shared/EpicGames.Core/SingleInstanceMutex.cs b/Engine/Source/Programs/Shared/EpicGames.Core/SingleInstanceMutex.cs
new file mode 100644
--- /dev/null
+++ b/Engine/Source/Programs/Shared/EpicGames.Core/SingleInstanceMutex.cs
@@ -0,0 +1,102 @@
+// Copyright Epic Games, Inc. All Rights Reserved.
+
+using System;
+using System.Threading;
+using System.Threading.Tasks;
+
+namespace EpicGames.Core
+{
+	/// <summary>
+	/// Allows taking a lock on a global mutex from async code.
+	/// </summary>
+	public static class SingleInstanceMutex
+	{
+		class SingleInstanceMutexImpl : IDisposable
+		{
+			readonly string? _name;
+			readonly TaskCompletionSource _readyTcs = new TaskCompletionSource();
+			readonly ManualResetEvent _disposing = new ManualResetEvent(false);
+			Thread? _thread;
+
+			public SingleInstanceMutexImpl(string name)
+			{
+				_name = name;
+
+				_thread = new Thread(() => ThreadFunc());
+				_thread.Start();
+			}
+
+			public async ValueTask WaitAsync(CancellationToken cancellationToken)
+			{
+				using (CancellationTokenRegistration registration = cancellationToken.Register(() => _readyTcs.TrySetCanceled()))
+				{
+					await _readyTcs.Task;
+				}
+			}
+
+			public void Dispose()
+			{
+				if (_thread != null)
+				{
+					_disposing.Set();
+
+					_thread.Join();
+					_thread = null;
+				}
+
+				_disposing.Dispose();
+			}
+
+			void ThreadFunc()
+			{
+				using Mutex mutex = new Mutex(false, _name);
+				try
+				{
+					WaitHandle[] handles = new WaitHandle[2];
+					handles[0] = _disposing;
+					handles[1] = mutex;
+
+					if (WaitHandle.WaitAny(handles) == 0)
+					{
+						return;
+					}
+				}
+				catch (AbandonedMutexException)
+				{
+				}
+
+				_readyTcs.TrySetResult();
+				_disposing.WaitOne();
+
+				mutex.ReleaseMutex();
+			}
+		}
+
+		/// <summary>
+		/// Acquires a global mutex with the given name.
+		/// </summary>
+		/// <param name="name">Name of the mutex to acquire</param>
+		/// <param name="cancellationToken">Cancellation token for the wait</param>
+		/// <returns>Handle to the mutex. Must be disposed when complete.</returns>
+		public static async Task<IDisposable> AcquireAsync(string name, CancellationToken cancellationToken = default)
+		{
+			SingleInstanceMutexImpl? impl = new SingleInstanceMutexImpl(name);
+			try
+			{
+				await impl.WaitAsync(cancellationToken);
+				if (cancellationToken.IsCancellationRequested)
+				{
+					throw new OperationCanceledException(cancellationToken);
+				}
+
+				SingleInstanceMutexImpl result = impl;
+				impl = null;
+				return result;
+			}
+			finally
+			{
+				impl?.Dispose();
+			}
+		}
+	}
+}
