EpicGames.Core: Respect OutputLevel when writing to the log, and reduce verbosity of output to the default tracer (since it hampers debugger performance significantly).

diff --git a/Engine/Source/Programs/Shared/EpicGames.Core/Log.cs b/Engine/Source/Programs/Shared/EpicGames.Core/Log.cs
--- a/Engine/Source/Programs/Shared/EpicGames.Core/Log.cs
+++ b/Engine/Source/Programs/Shared/EpicGames.Core/Log.cs
@@ -1181,6 +1181,11 @@ namespace EpicGames.Core
 
 		public void Log<TState>(LogLevel logLevel, EventId eventId, TState state, Exception? exception, Func<TState, Exception?, string> formatter)
 		{
+			if (!IsEnabled(logLevel))
+			{
+				return;
+			}
+
 			string[] lines = formatter(state, exception).Split('\n');
 			lock (_syncObject)
 			{
@@ -1190,6 +1195,11 @@ namespace EpicGames.Core
 				{
 					if (listener != null)
 					{
+						if (listener is DefaultTraceListener && logLevel < LogLevel.Information)
+						{
+							continue;
+						}
+
 						string timePrefixActual =
 							IncludeTimestamps &&
 							!(listener is DefaultTraceListener) // no timestamps when writing to the Visual Studio debug window
