UBT: Reduce verbosity of worker request logging.

diff --git a/Engine/Source/Programs/Shared/EpicGames.Horde/Compute/Clients/ServerComputeClient.cs b/Engine/Source/Programs/Shared/EpicGames.Horde/Compute/Clients/ServerComputeClient.cs
--- a/Engine/Source/Programs/Shared/EpicGames.Horde/Compute/Clients/ServerComputeClient.cs
+++ b/Engine/Source/Programs/Shared/EpicGames.Horde/Compute/Clients/ServerComputeClient.cs
@@ -131,7 +131,7 @@ namespace EpicGames.Horde.Compute.Clients
 		{
 			IAsyncEnumerator<LeaseInfo> source = ConnectAsync(clusterId, requirements, cancellationToken).GetAsyncEnumerator(cancellationToken);
 			if (!await source.MoveNextAsync())
-		{
+			{
 				await source.DisposeAsync();
 				return null;
 			}
@@ -141,7 +141,7 @@ namespace EpicGames.Horde.Compute.Clients
 		/// <inheritdoc/>
 		async IAsyncEnumerable<LeaseInfo> ConnectAsync(ClusterId clusterId, Requirements? requirements, [EnumeratorCancellation] CancellationToken cancellationToken)
 		{
-			_logger.LogInformation("Requesting compute resource");
+			_logger.LogDebug("Requesting compute resource");
 
 			// Assign a compute worker
 			HttpClient client = _createHttpClient();
@@ -154,7 +154,7 @@ namespace EpicGames.Horde.Compute.Clients
 			{
 				if (response.StatusCode == HttpStatusCode.ServiceUnavailable)
 				{
-					_logger.LogInformation("No compute resource is available.");
+					_logger.LogDebug("No compute resource is available.");
 					yield break;
 				}
 
@@ -163,11 +163,11 @@ namespace EpicGames.Horde.Compute.Clients
 				responseMessage = await response.Content.ReadFromJsonAsync<AssignComputeResponse>(cancellationToken: cancellationToken);
 				if (responseMessage == null)
 				{
-						throw new InvalidOperationException();
+					throw new InvalidOperationException();
 				}
 			}
 
-			_logger.LogInformation("Connecting to {Ip} with nonce {Nonce}...", responseMessage.Ip, responseMessage.Nonce);
+			_logger.LogDebug("Connecting to {Ip} with nonce {Nonce}...", responseMessage.Ip, responseMessage.Nonce);
 
 			// Connect to the remote machine
 			using Socket socket = new Socket(SocketType.Stream, ProtocolType.Tcp);
@@ -176,7 +176,7 @@ namespace EpicGames.Horde.Compute.Clients
 			// Send the nonce
 			byte[] nonce = StringUtils.ParseHexString(responseMessage.Nonce);
 			await socket.SendMessageAsync(nonce, SocketFlags.None, cancellationToken);
-			_logger.LogInformation("Connection established.");
+			_logger.LogDebug("Connection established.");
 
 			// Pass the rest of the call over to the handler
 			byte[] key = StringUtils.ParseHexString(responseMessage.Key);
