Upgrade EpicGames.OIDC to NET 6.

diff --git a/Engine/Source/Programs/Shared/EpicGames.OIDC/EpicGames.OIDC.csproj b/Engine/Source/Programs/Shared/EpicGames.OIDC/EpicGames.OIDC.csproj
--- a/Engine/Source/Programs/Shared/EpicGames.OIDC/EpicGames.OIDC.csproj
+++ b/Engine/Source/Programs/Shared/EpicGames.OIDC/EpicGames.OIDC.csproj
@@ -3,12 +3,14 @@
   <PropertyGroup>
     <OutputType>Library</OutputType>
     <RootNamespace>EpicGames.OIDC</RootNamespace>
-    <TargetFramework>netcoreapp3.1</TargetFramework>
+    <TargetFramework>net6.0</TargetFramework>
     <DebugType>pdbonly</DebugType>
     <Nullable>enable</Nullable>
     <GenerateAssemblyInfo>false</GenerateAssemblyInfo>
     <GenerateTargetFrameworkAttribute>false</GenerateTargetFrameworkAttribute>
-    <Configurations>Debug;Release;Development</Configurations>
+    <Configurations>Debug;Release;Development;Analyze</Configurations>
+    <Deterministic>true</Deterministic>
+    <RunAnalyzersDuringBuild>false</RunAnalyzersDuringBuild>
   </PropertyGroup>
 
   <PropertyGroup Condition=" '$(Configuration)|$(Platform)' == 'Debug|AnyCPU' ">
@@ -19,6 +21,12 @@
   <PropertyGroup Condition="'$(Configuration)|$(Platform)' == 'Release|AnyCPU'">
     <Optimize>true</Optimize>
   </PropertyGroup>
+  <PropertyGroup Condition=" '$(Configuration)|$(Platform)' == 'Analyze|AnyCPU' ">
+    <RunAnalyzersDuringBuild>True</RunAnalyzersDuringBuild>
+    <TreatWarningsAsErrors>false</TreatWarningsAsErrors>
+    <GenerateDocumentationFile>True</GenerateDocumentationFile>
+    <DocumentationFile></DocumentationFile>
+  </PropertyGroup>
 
   <ItemGroup>
     <Compile Include="..\MetaData.cs">
diff --git a/Engine/Source/Programs/Shared/EpicGames.OIDC/OidcTokenManager.cs b/Engine/Source/Programs/Shared/EpicGames.OIDC/OidcTokenManager.cs
--- a/Engine/Source/Programs/Shared/EpicGames.OIDC/OidcTokenManager.cs
+++ b/Engine/Source/Programs/Shared/EpicGames.OIDC/OidcTokenManager.cs
@@ -260,7 +260,7 @@ namespace EpicGames.OIDC
 					AuthorizeState loginState = await oidcClient!.PrepareLoginAsync(cancellationToken: cancellationToken);
 
 					// start the user browser
-					using (Process process = OpenBrowser(loginState.StartUrl))
+					using (Process? process = OpenBrowser(loginState.StartUrl))
 					{
 						using (IDisposable _ = cancellationToken.Register(() => http.Stop()))
 						{
@@ -275,7 +275,7 @@ namespace EpicGames.OIDC
 							}
 
 							// wait a few seconds before shutting down the http server to give the browser time to actually load everything it needs
-							await Task.Delay(2000);
+							await Task.Delay(2000, cancellationToken);
 						}
 					}
 
@@ -390,7 +390,7 @@ namespace EpicGames.OIDC
 			return httpPageFailure;
 		}
 
-		private Process OpenBrowser(string url)
+		private static Process? OpenBrowser(string url)
 		{
 			if (RuntimeInformation.IsOSPlatform(OSPlatform.Windows))
 			{
@@ -401,7 +401,7 @@ namespace EpicGames.OIDC
 					{
 						if (progId.Equals("ChromeHTML", StringComparison.Ordinal) || progId.Equals("MSEdgeHTM", StringComparison.Ordinal))
 						{
-							string exe = command.Replace("--single-argument", "").Replace("%1", "").TrimEnd();
+							string exe = command.Replace("--single-argument", "", StringComparison.OrdinalIgnoreCase).Replace("%1", "", StringComparison.Ordinal).TrimEnd();
 							exe = $"{exe} --new-window \"{url}\"";
 							return Process.Start(new ProcessStartInfo(exe));
 						}
