Horde: Fix SingleInstanceMutex on Mac/Linux. Waiting for multiple wait handles is only supported on Windows.

diff --git a/Engine/Source/Programs/Shared/EpicGames.Core/SingleInstanceMutex.cs b/Engine/Source/Programs/Shared/EpicGames.Core/SingleInstanceMutex.cs
--- a/Engine/Source/Programs/Shared/EpicGames.Core/SingleInstanceMutex.cs
+++ b/Engine/Source/Programs/Shared/EpicGames.Core/SingleInstanceMutex.cs
@@ -1,6 +1,7 @@
 // Copyright Epic Games, Inc. All Rights Reserved.
 
 using System;
+using System.Runtime.InteropServices;
 using System.Threading;
 using System.Threading.Tasks;
 
@@ -52,13 +53,27 @@ namespace EpicGames.Core
 				using Mutex mutex = new Mutex(false, _name);
 				try
 				{
-					WaitHandle[] handles = new WaitHandle[2];
-					handles[0] = _disposing;
-					handles[1] = mutex;
+					// Waiting for multiple handles is only supported on Windows
+					if (RuntimeInformation.IsOSPlatform(OSPlatform.Windows))
+					{
+						WaitHandle[] handles = new WaitHandle[2];
+						handles[0] = _disposing;
+						handles[1] = mutex;
 
-					if (WaitHandle.WaitAny(handles) == 0)
+						if (WaitHandle.WaitAny(handles) == 0)
+						{
+							return;
+						}
+					}
+					else
 					{
-						return;
+						while(!mutex.WaitOne(100))
+						{
+							if(_disposing.WaitOne(0))
+							{
+								return;
+							}
+						}
 					}
 				}
 				catch (AbandonedMutexException)
