Skip over isSparse fields in Perforce output.

diff --git a/Engine/Source/Programs/Shared/EpicGames.Perforce/IPerforceOutput.cs b/Engine/Source/Programs/Shared/EpicGames.Perforce/IPerforceOutput.cs
--- a/Engine/Source/Programs/Shared/EpicGames.Perforce/IPerforceOutput.cs
+++ b/Engine/Source/Programs/Shared/EpicGames.Perforce/IPerforceOutput.cs
@@ -54,6 +54,7 @@ namespace EpicGames.Perforce
 			public static readonly Utf8String Error = "error";
 			public static readonly Utf8String Io = "io";
 			public static readonly Utf8String Func = "func";
+			public static readonly Utf8String IsSparse = "isSparse";
 		}
 
 		/// <summary>
@@ -674,7 +675,7 @@ namespace EpicGames.Perforce
 						}
 					}
 				}
-				else if (tag == ReadOnlyUtf8StringConstants.Func)
+				else if (tag == ReadOnlyUtf8StringConstants.Func || tag == ReadOnlyUtf8StringConstants.IsSparse)
 				{
 					// Not sure why these fields are in the client output, but they are peppered into filelog results without an element index breaking the parser.
 					if (!TryReadValue(buffer, ref bufferPos, newRecord, null))
