Improved method of getting Xcode version

diff --git a/Engine/Source/Programs/Shared/EpicGames.Build/Platform/Mac/ApplePlatformSDK.cs b/Engine/Source/Programs/Shared/EpicGames.Build/Platform/Mac/ApplePlatformSDK.cs
--- a/Engine/Source/Programs/Shared/EpicGames.Build/Platform/Mac/ApplePlatformSDK.cs
+++ b/Engine/Source/Programs/Shared/EpicGames.Build/Platform/Mac/ApplePlatformSDK.cs
@@ -49,20 +49,21 @@ namespace UnrealBuildBase
 			// get Xcode version on Mac
 			if (OperatingSystem.IsMacOS())
 			{
-				int ExitCode = 0;
-				string Output = RunLocalProcessAndReturnStdOut("sh", "-c 'xcodebuild -version'", out ExitCode);
+				int ExitCode;
+				// xcode-select -p gives the currently selected Xcode location (xcodebuild -version may fail if Xcode.app is broken)
+				// Example output: /Applications/Xcode.app/Contents/Developer
+				string Output = RunLocalProcessAndReturnStdOut("sh", "-c 'xcode-select -p'", out ExitCode);
 
-				// For macOS, only an ExitCode of 0 means there was no issues running the local process
 				if (ExitCode == 0)
 				{
-					Match Result = Regex.Match(Output, @"Xcode (\S*)");
-
-					// If Xcode is not installed (or xcode-select is pointing to an invalid dir), "Output" can return random text
-					// (generally an error message with the word "Xcode " within it that can successfully Regex match), 
-					// so verify we got back an Int that is in String format.
-					if (Result.Success && TryConvertVersionToInt(Result.Groups[1].Value, out _))
+					DirectoryReference DeveloperDir = new DirectoryReference(Output);
+					FileReference Plist = FileReference.Combine(DeveloperDir.ParentDirectory!, "Info.plist");
+					// Find out the version number in Xcode.app/Contents/Info.plist
+					Output = RunLocalProcessAndReturnStdOut("sh",
+						  $"-c 'plutil -extract CFBundleShortVersionString raw {Plist}'", out ExitCode);
+					if (ExitCode == 0)
 					{
-						return Result.Groups[1].Value;
+						return Output;
 					}
 				}
 
