Horde: Rename FileNode to ChunkedDataNode.

diff --git a/Engine/Source/Programs/Horde/Horde.Agent/Utility/TempStorage.cs b/Engine/Source/Programs/Horde/Horde.Agent/Utility/TempStorage.cs
--- a/Engine/Source/Programs/Horde/Horde.Agent/Utility/TempStorage.cs
+++ b/Engine/Source/Programs/Horde/Horde.Agent/Utility/TempStorage.cs
@@ -601,7 +601,7 @@ namespace Horde.Storage.Utility
 			FileReference localFileListLocation = GetTagManifestLocation(manifestDir, nodeName, tagName);
 			fileList.Save(localFileListLocation);
 
-			FileNodeWriter fileNodeWriter = new FileNodeWriter(writer, new ChunkingOptions());
+			ChunkedDataWriter fileNodeWriter = new ChunkedDataWriter(writer, new ChunkingOptions());
 			NodeHandle handle = await fileNodeWriter.CreateAsync(localFileListLocation.ToFileInfo(), cancellationToken);
 
 			return new FileEntry(localFileListLocation.GetFileName(), FileEntryFlags.None, fileNodeWriter.Length, handle);
diff --git a/Engine/Source/Programs/Horde/Horde.Server/Commands/Bundles/BundleWriteTest.cs b/Engine/Source/Programs/Horde/Horde.Server/Commands/Bundles/BundleWriteTest.cs
--- a/Engine/Source/Programs/Horde/Horde.Server/Commands/Bundles/BundleWriteTest.cs
+++ b/Engine/Source/Programs/Horde/Horde.Server/Commands/Bundles/BundleWriteTest.cs
@@ -42,7 +42,7 @@ namespace Horde.Server.Commands.Bundles
 			ChunkingOptions chunkingOptions = new ChunkingOptions();
 //			chunkingOptions.LeafOptions = new ChunkingOptionsForNodeType(64 * 1024);
 
-			FileNode node = new LeafFileNode();
+			ChunkedDataNode node = new LeafChunkedDataNode();
 
 			byte[] buffer = new byte[64 * 1024];
 			RandomNumberGenerator.Fill(buffer);
@@ -51,7 +51,7 @@ namespace Horde.Server.Commands.Bundles
 			long length = 0;
 			double nextTime = 2.0;
 
-			FileNodeWriter fileNodeWriter = new FileNodeWriter(writer, chunkingOptions);
+			ChunkedDataWriter fileNodeWriter = new ChunkedDataWriter(writer, chunkingOptions);
 			for (; ; )
 			{
 				await fileNodeWriter.AppendAsync(buffer, default);
diff --git a/Engine/Source/Programs/Horde/Horde.Server/Perforce/PerforceReplicator.cs b/Engine/Source/Programs/Horde/Horde.Server/Perforce/PerforceReplicator.cs
--- a/Engine/Source/Programs/Horde/Horde.Server/Perforce/PerforceReplicator.cs
+++ b/Engine/Source/Programs/Horde/Horde.Server/Perforce/PerforceReplicator.cs
@@ -115,14 +115,14 @@ namespace Horde.Server.Perforce
 			{
 				public Utf8String _path;
 				public FileEntryFlags _flags;
-				public readonly FileNodeWriter _fileWriter;
+				public readonly ChunkedDataWriter _fileWriter;
 				public long _size;
 				public long _sizeWritten;
 				public readonly IncrementalHash _hash = IncrementalHash.CreateHash(HashAlgorithmName.MD5);
 
 				public Handle(TreeWriter writer, ChunkingOptions options)
 				{
-					_fileWriter = new FileNodeWriter(writer, options);
+					_fileWriter = new ChunkedDataWriter(writer, options);
 				}
 			}
 
diff --git a/Engine/Source/Programs/Horde/Horde/Commands/Vcs/VcsCheckout.cs b/Engine/Source/Programs/Horde/Horde/Commands/Vcs/VcsCheckout.cs
--- a/Engine/Source/Programs/Horde/Horde/Commands/Vcs/VcsCheckout.cs
+++ b/Engine/Source/Programs/Horde/Horde/Commands/Vcs/VcsCheckout.cs
@@ -124,7 +124,7 @@ namespace Horde.Commands.Vcs
 		static async Task<FileState> CheckoutFileAsync(TreeReader reader, FileEntry fileRef, FileInfo fileInfo, ILogger logger)
 		{
 			logger.LogInformation("Updating {File} to {Hash}", fileInfo, fileRef.Hash);
-			FileNode fileNode = await fileRef.ExpandAsync(reader);
+			ChunkedDataNode fileNode = await fileRef.ExpandAsync(reader);
 			await fileNode.CopyToFileAsync(reader, fileInfo, CancellationToken.None);
 			fileInfo.Refresh();
 			return new FileState(fileInfo, fileRef.Hash);
diff --git a/Engine/Source/Programs/Shared/EpicGames.Horde.Tests/BundleStoreTests.cs b/Engine/Source/Programs/Shared/EpicGames.Horde.Tests/BundleStoreTests.cs
--- a/Engine/Source/Programs/Shared/EpicGames.Horde.Tests/BundleStoreTests.cs
+++ b/Engine/Source/Programs/Shared/EpicGames.Horde.Tests/BundleStoreTests.cs
@@ -259,7 +259,7 @@ namespace EpicGames.Horde.Tests
 				DirectoryNode root = new DirectoryNode(DirectoryFlags.None);
 				DirectoryNode hello = root.AddDirectory("hello");
 
-				FileNodeWriter fileWriter = new FileNodeWriter(writer, new ChunkingOptions());
+				ChunkedDataWriter fileWriter = new ChunkedDataWriter(writer, new ChunkingOptions());
 				NodeHandle fileHandle = await fileWriter.CreateAsync(Encoding.UTF8.GetBytes("world"), CancellationToken.None);
 				hello.AddFile("world", FileEntryFlags.None, fileWriter.Length, fileHandle);
 
@@ -285,7 +285,7 @@ namespace EpicGames.Horde.Tests
 			Assert.AreEqual(1, hello.Files.Count);
 			Assert.AreEqual("world", hello.Files.First().Name);
 
-			FileNode world = await hello.Files.First().ExpandAsync(reader);
+			ChunkedDataNode world = await hello.Files.First().ExpandAsync(reader);
 
 			byte[] worldData = await GetFileDataAsync(reader, world);
 			Assert.IsTrue(worldData.SequenceEqual(Encoding.UTF8.GetBytes("world")));
@@ -312,7 +312,7 @@ namespace EpicGames.Horde.Tests
 				options.LeafOptions.TargetSize = 256;
 				options.LeafOptions.MaxSize = 64 * 1024;
 
-				FileNodeWriter fileWriter = new FileNodeWriter(writer, options);
+				ChunkedDataWriter fileWriter = new ChunkedDataWriter(writer, options);
 				for (int idx = 0; idx < chunk.Length / 16; idx++)
 				{
 					await fileWriter.AppendAsync(chunk.AsMemory(idx * 16, 16), CancellationToken.None);
@@ -326,7 +326,7 @@ namespace EpicGames.Horde.Tests
 			{
 				TreeReader reader = new TreeReader(store, cache, NullLogger.Instance);
 
-				FileNode newRoot = await reader.ReadNodeAsync<FileNode>(locator);
+				ChunkedDataNode newRoot = await reader.ReadNodeAsync<ChunkedDataNode>(locator);
 
 				using MemoryStream stream = new MemoryStream();
 				await newRoot.CopyToStreamAsync(reader, stream, CancellationToken.None);
@@ -366,7 +366,7 @@ namespace EpicGames.Horde.Tests
 				options.LeafOptions.TargetSize = 256;
 				options.LeafOptions.MaxSize = 64 * 1024;
 
-				FileNodeWriter fileWriter = new FileNodeWriter(writer, options);
+				ChunkedDataWriter fileWriter = new ChunkedDataWriter(writer, options);
 				NodeHandle handle = await fileWriter.CreateAsync(data, CancellationToken.None);
 				root.AddFile("test", FileEntryFlags.None, fileWriter.Length, handle);
 
@@ -384,7 +384,7 @@ namespace EpicGames.Horde.Tests
 				await CompareTrees(reader, root, newRoot);
 				await CheckLargeFileTreeAsync(reader, root, data);
 
-				NodeRef<FileNode> file = root.GetFileEntry("test");
+				NodeRef<ChunkedDataNode> file = root.GetFileEntry("test");
 
 				long uniqueSize = store.Blobs.Values.SelectMany(x => x.Header.Packets).Sum(x => x.DecodedLength);
 				Assert.IsTrue(uniqueSize < data.Length / 3); // random fraction meaning "lots of dedupe happened"
@@ -399,31 +399,31 @@ namespace EpicGames.Horde.Tests
 
 			foreach ((FileEntry oldFileEntry, FileEntry newFileEntry) in oldNode.Files.Zip(newNode.Files))
 			{
-				FileNode oldFile = await oldFileEntry.ExpandAsync(reader);
-				FileNode newFile = await newFileEntry.ExpandAsync(reader);
+				ChunkedDataNode oldFile = await oldFileEntry.ExpandAsync(reader);
+				ChunkedDataNode newFile = await newFileEntry.ExpandAsync(reader);
 				await CompareTrees(reader, oldFile, newFile);
 			}
 		}
 
-		static async Task CompareTrees(TreeReader reader, FileNode oldNode, FileNode newNode)
+		static async Task CompareTrees(TreeReader reader, ChunkedDataNode oldNode, ChunkedDataNode newNode)
 		{
-			if (oldNode is InteriorFileNode oldInteriorNode)
+			if (oldNode is InteriorChunkedDataNode oldInteriorNode)
 			{
-				InteriorFileNode newInteriorNode = (InteriorFileNode)newNode;
+				InteriorChunkedDataNode newInteriorNode = (InteriorChunkedDataNode)newNode;
 				Assert.AreEqual(oldInteriorNode.Children.Count, newInteriorNode.Children.Count);
 
 				int index = 0;
-				foreach ((NodeRef<FileNode> oldFileRef, NodeRef<FileNode> newFileRef) in oldInteriorNode.Children.Zip(newInteriorNode.Children))
+				foreach ((NodeRef<ChunkedDataNode> oldFileRef, NodeRef<ChunkedDataNode> newFileRef) in oldInteriorNode.Children.Zip(newInteriorNode.Children))
 				{
-					FileNode oldFile = await oldFileRef.ExpandAsync(reader);
-					FileNode newFile = await newFileRef.ExpandAsync(reader);
+					ChunkedDataNode oldFile = await oldFileRef.ExpandAsync(reader);
+					ChunkedDataNode newFile = await newFileRef.ExpandAsync(reader);
 					await CompareTrees(reader, oldFile, newFile);
 					index++;
 				}
 			}
-			else if (oldNode is LeafFileNode oldLeafNode)
+			else if (oldNode is LeafChunkedDataNode oldLeafNode)
 			{
-				LeafFileNode newLeafNode = (LeafFileNode)newNode;
+				LeafChunkedDataNode newLeafNode = (LeafChunkedDataNode)newNode;
 				Assert.IsTrue(oldLeafNode.Data.Span.SequenceEqual(newLeafNode.Data.Span));
 			}
 			else
@@ -437,25 +437,25 @@ namespace EpicGames.Horde.Tests
 			Assert.AreEqual(0, root.Directories.Count);
 			Assert.AreEqual(1, root.Files.Count);
 
-			FileNode world = await root.Files.First().ExpandAsync(reader, CancellationToken.None);
+			ChunkedDataNode world = await root.Files.First().ExpandAsync(reader, CancellationToken.None);
 
 			int length = await CheckFileDataAsync(reader, world, data);
 			Assert.AreEqual(data.Length, length);
 		}
 
-		static async Task<int> CheckFileDataAsync(TreeReader reader, FileNode fileNode, ReadOnlyMemory<byte> data)
+		static async Task<int> CheckFileDataAsync(TreeReader reader, ChunkedDataNode fileNode, ReadOnlyMemory<byte> data)
 		{
 			int offset = 0;
-			if (fileNode is LeafFileNode leafNode)
+			if (fileNode is LeafChunkedDataNode leafNode)
 			{
 				Assert.IsTrue(leafNode.Data.Span.SequenceEqual(data.Span.Slice(offset, leafNode.Data.Length)));
 				offset += leafNode.Data.Length;
 			}
-			else if (fileNode is InteriorFileNode interiorFileNode)
+			else if (fileNode is InteriorChunkedDataNode interiorFileNode)
 			{
-				foreach (NodeRef<FileNode> childRef in interiorFileNode.Children)
+				foreach (NodeRef<ChunkedDataNode> childRef in interiorFileNode.Children)
 				{
-					FileNode child = await childRef.ExpandAsync(reader);
+					ChunkedDataNode child = await childRef.ExpandAsync(reader);
 					offset += await CheckFileDataAsync(reader, child, data.Slice(offset));
 				}
 			}
@@ -466,7 +466,7 @@ namespace EpicGames.Horde.Tests
 			return offset;
 		}
 
-		static async Task<byte[]> GetFileDataAsync(TreeReader reader, FileNode fileNode)
+		static async Task<byte[]> GetFileDataAsync(TreeReader reader, ChunkedDataNode fileNode)
 		{
 			using (MemoryStream stream = new MemoryStream())
 			{
diff --git a/Engine/Source/Programs/Shared/EpicGames.Horde.Tests/BundleTests.cs b/Engine/Source/Programs/Shared/EpicGames.Horde.Tests/BundleTests.cs
--- a/Engine/Source/Programs/Shared/EpicGames.Horde.Tests/BundleTests.cs
+++ b/Engine/Source/Programs/Shared/EpicGames.Horde.Tests/BundleTests.cs
@@ -74,43 +74,43 @@ namespace EpicGames.Horde.Tests
 			ChunkingOptions options = new ChunkingOptions();
 			options.LeafOptions = new ChunkingOptionsForNodeType(8, 8, 8);
 
-			FileNodeWriter fileNodeWriter = new FileNodeWriter(writer, options);
+			ChunkedDataWriter fileNodeWriter = new ChunkedDataWriter(writer, options);
 
 			NodeHandle handle;
-			FileNode node;
+			ChunkedDataNode node;
 			byte[] data = CreateBuffer(1024);
 
 			handle = await fileNodeWriter.CreateAsync(data.AsMemory(0, 7), CancellationToken.None);
-			node = await reader.ReadNodeAsync<FileNode>(handle.Locator);
-			Assert.IsTrue(node is LeafFileNode);
-			Assert.AreEqual(7, ((LeafFileNode)node).Data.Length);
+			node = await reader.ReadNodeAsync<ChunkedDataNode>(handle.Locator);
+			Assert.IsTrue(node is LeafChunkedDataNode);
+			Assert.AreEqual(7, ((LeafChunkedDataNode)node).Data.Length);
 			await TestBufferlessReadsAsync(reader, handle.Locator, data.AsMemory(0, 7));
 
 			handle = await fileNodeWriter.CreateAsync(data.AsMemory(0, 8), CancellationToken.None);
-			node = await reader.ReadNodeAsync<FileNode>(handle.Locator);
-			Assert.IsTrue(node is LeafFileNode);
-			Assert.AreEqual(8, ((LeafFileNode)node).Data.Length);
+			node = await reader.ReadNodeAsync<ChunkedDataNode>(handle.Locator);
+			Assert.IsTrue(node is LeafChunkedDataNode);
+			Assert.AreEqual(8, ((LeafChunkedDataNode)node).Data.Length);
 			await TestBufferlessReadsAsync(reader, handle.Locator, data.AsMemory(0, 8));
 
 			handle = await fileNodeWriter.CreateAsync(data.AsMemory(0, 9), CancellationToken.None);
-			node = await reader.ReadNodeAsync<FileNode>(handle.Locator);
-			Assert.IsTrue(node is InteriorFileNode);
-			Assert.AreEqual(2, ((InteriorFileNode)node).Children.Count);
+			node = await reader.ReadNodeAsync<ChunkedDataNode>(handle.Locator);
+			Assert.IsTrue(node is InteriorChunkedDataNode);
+			Assert.AreEqual(2, ((InteriorChunkedDataNode)node).Children.Count);
 			await TestBufferlessReadsAsync(reader, handle.Locator, data.AsMemory(0, 9));
 
-			FileNode? childNode1 = await ((InteriorFileNode)node).Children[0].ExpandAsync(reader);
+			ChunkedDataNode? childNode1 = await ((InteriorChunkedDataNode)node).Children[0].ExpandAsync(reader);
 			Assert.IsNotNull(childNode1);
-			Assert.IsTrue(childNode1 is LeafFileNode);
-			Assert.AreEqual(8, ((LeafFileNode)childNode1!).Data.Length);
+			Assert.IsTrue(childNode1 is LeafChunkedDataNode);
+			Assert.AreEqual(8, ((LeafChunkedDataNode)childNode1!).Data.Length);
 
-			FileNode? childNode2 = await ((InteriorFileNode)node).Children[1].ExpandAsync(reader);
+			ChunkedDataNode? childNode2 = await ((InteriorChunkedDataNode)node).Children[1].ExpandAsync(reader);
 			Assert.IsNotNull(childNode2);
-			Assert.IsTrue(childNode2 is LeafFileNode);
-			Assert.AreEqual(1, ((LeafFileNode)childNode2!).Data.Length);
+			Assert.IsTrue(childNode2 is LeafChunkedDataNode);
+			Assert.AreEqual(1, ((LeafChunkedDataNode)childNode2!).Data.Length);
 
 			handle = await fileNodeWriter.CreateAsync(data, CancellationToken.None);
-			node = await reader.ReadNodeAsync<FileNode>(handle.Locator);
-			Assert.IsTrue(node is InteriorFileNode);
+			node = await reader.ReadNodeAsync<ChunkedDataNode>(handle.Locator);
+			Assert.IsTrue(node is InteriorChunkedDataNode);
 			await TestBufferlessReadsAsync(reader, handle.Locator, data);
 		}
 
@@ -127,7 +127,7 @@ namespace EpicGames.Horde.Tests
 		private static async Task TestBufferlessReadsAsync(TreeReader reader, NodeLocator locator, ReadOnlyMemory<byte> expected)
 		{
 			using MemoryStream memoryStream = new MemoryStream();
-			await FileNode.CopyToStreamAsync(reader, locator, memoryStream, default);
+			await ChunkedDataNode.CopyToStreamAsync(reader, locator, memoryStream, default);
 			ReadOnlyMemory<byte> read = memoryStream.ToArray().AsMemory();
 			Assert.IsTrue(read.Span.SequenceEqual(expected.Span));			
 		}
diff --git a/Engine/Source/Programs/Shared/EpicGames.Horde.Tests/ChunkingTests.cs b/Engine/Source/Programs/Shared/EpicGames.Horde.Tests/ChunkingTests.cs
--- a/Engine/Source/Programs/Shared/EpicGames.Horde.Tests/ChunkingTests.cs
+++ b/Engine/Source/Programs/Shared/EpicGames.Horde.Tests/ChunkingTests.cs
@@ -87,7 +87,7 @@ namespace EpicGames.Horde.Tests
 
 			const int NumIterations = 100;
 			{
-				FileNodeWriter fileWriter = new FileNodeWriter(writer, options);
+				ChunkedDataWriter fileWriter = new ChunkedDataWriter(writer, options);
 
 				for (int idx = 0; idx < NumIterations; idx++)
 				{
@@ -97,7 +97,7 @@ namespace EpicGames.Horde.Tests
 				handle = await fileWriter.FlushAsync(CancellationToken.None);
 			}
 
-			FileNode root = await reader.ReadNodeAsync<FileNode>(handle.Locator);
+			ChunkedDataNode root = await reader.ReadNodeAsync<ChunkedDataNode>(handle.Locator);
 
 			byte[] result;
 			using (MemoryStream stream = new MemoryStream())
@@ -117,16 +117,16 @@ namespace EpicGames.Horde.Tests
 			await CheckSizes(reader, root, options, true);
 		}
 
-		async Task CheckSizes(TreeReader reader, FileNode node, ChunkingOptions options, bool rightmost)
+		async Task CheckSizes(TreeReader reader, ChunkedDataNode node, ChunkingOptions options, bool rightmost)
 		{
-			if (node is LeafFileNode leafNode)
+			if (node is LeafChunkedDataNode leafNode)
 			{
 				Assert.IsTrue(rightmost || leafNode.Data.Length >= options.LeafOptions.MinSize);
 				Assert.IsTrue(leafNode.Data.Length <= options.LeafOptions.MaxSize);
 			}
 			else
 			{
-				InteriorFileNode interiorNode = (InteriorFileNode)node;
+				InteriorChunkedDataNode interiorNode = (InteriorChunkedDataNode)node;
 
 				Assert.IsTrue(rightmost || interiorNode.Children.Count * IoHash.NumBytes >= options.InteriorOptions.MinSize);
 				Assert.IsTrue(interiorNode.Children.Count <= options.InteriorOptions.MaxSize);
@@ -134,7 +134,7 @@ namespace EpicGames.Horde.Tests
 				int childCount = interiorNode.Children.Count;
 				for (int idx = 0; idx < childCount; idx++)
 				{
-					FileNode childNode = await interiorNode.Children[idx].ExpandAsync(reader, CancellationToken.None);
+					ChunkedDataNode childNode = await interiorNode.Children[idx].ExpandAsync(reader, CancellationToken.None);
 					await CheckSizes(reader, childNode, options, idx == childCount - 1);
 				}
 			}
diff --git a/Engine/Source/Programs/Shared/EpicGames.Horde/Storage/Nodes/DirectoryNode.cs b/Engine/Source/Programs/Shared/EpicGames.Horde/Storage/Nodes/DirectoryNode.cs
--- a/Engine/Source/Programs/Shared/EpicGames.Horde/Storage/Nodes/DirectoryNode.cs
+++ b/Engine/Source/Programs/Shared/EpicGames.Horde/Storage/Nodes/DirectoryNode.cs
@@ -612,7 +612,7 @@ namespace EpicGames.Horde.Storage.Nodes
 		/// <param name="cancellationToken">Cancellation token for the operation</param>
 		public async Task CopyFromZipStreamAsync(Stream stream, TreeWriter writer, CancellationToken cancellationToken = default)
 		{
-			FileNodeWriter fileWriter = new FileNodeWriter(writer, new ChunkingOptions());
+			ChunkedDataWriter fileWriter = new ChunkedDataWriter(writer, new ChunkingOptions());
 			using (ZipArchive archive = new ZipArchive(stream, ZipArchiveMode.Read, true))
 			{
 				foreach (ZipArchiveEntry entry in archive.Entries)
@@ -750,7 +750,7 @@ namespace EpicGames.Horde.Storage.Nodes
 					writer = new TreeWriter(baseWriter);
 				}
 
-				FileNodeWriter fileNodeWriter = new FileNodeWriter(writer, options);
+				ChunkedDataWriter fileNodeWriter = new ChunkedDataWriter(writer, options);
 				for (int idx = minIdx; idx < maxIdx; idx++)
 				{
 					FileInfo fileInfo = files[idx].FileInfo;
@@ -789,7 +789,7 @@ namespace EpicGames.Horde.Storage.Nodes
 			foreach (FileEntry fileEntry in _nameToFileEntry.Values)
 			{
 				FileInfo fileInfo = new FileInfo(Path.Combine(directoryInfo.FullName, fileEntry.Name.ToString()));
-				FileNode fileNode = await fileEntry.ExpandAsync(reader, cancellationToken);
+				ChunkedDataNode fileNode = await fileEntry.ExpandAsync(reader, cancellationToken);
 				tasks.Add(Task.Run(() => fileNode.CopyToFileAsync(reader, fileInfo, cancellationToken), cancellationToken));
 			}
 			foreach (DirectoryEntry directoryEntry in _nameToDirectoryEntry.Values)
diff --git a/Engine/Source/Programs/Shared/EpicGames.Horde/Storage/Nodes/FileEntry.cs b/Engine/Source/Programs/Shared/EpicGames.Horde/Storage/Nodes/FileEntry.cs
--- a/Engine/Source/Programs/Shared/EpicGames.Horde/Storage/Nodes/FileEntry.cs
+++ b/Engine/Source/Programs/Shared/EpicGames.Horde/Storage/Nodes/FileEntry.cs
@@ -49,7 +49,7 @@ namespace EpicGames.Horde.Storage.Nodes
 	/// <summary>
 	/// Entry for a file within a directory node
 	/// </summary>
-	public sealed class FileEntry : NodeRef<FileNode>
+	public sealed class FileEntry : NodeRef<ChunkedDataNode>
 	{
 		/// <summary>
 		/// Name of this file
@@ -143,7 +143,7 @@ namespace EpicGames.Horde.Storage.Nodes
 		/// <param name="cancellationToken">Cancellation token for the operation</param>
 		public async Task CopyToStreamAsync(TreeReader reader, Stream outputStream, CancellationToken cancellationToken)
 		{
-			FileNode node = await ExpandAsync(reader, cancellationToken);
+			ChunkedDataNode node = await ExpandAsync(reader, cancellationToken);
 			await node.CopyToStreamAsync(reader, outputStream, cancellationToken);
 		}
 
@@ -156,7 +156,7 @@ namespace EpicGames.Horde.Storage.Nodes
 		/// <returns></returns>
 		public async Task CopyToFileAsync(TreeReader reader, FileInfo file, CancellationToken cancellationToken)
 		{
-			FileNode node = await ExpandAsync(reader, cancellationToken);
+			ChunkedDataNode node = await ExpandAsync(reader, cancellationToken);
 			await node.CopyToFileAsync(reader, file, cancellationToken);
 		}
 
diff --git a/Engine/Source/Programs/Shared/EpicGames.Horde/Storage/TreeReader.cs b/Engine/Source/Programs/Shared/EpicGames.Horde/Storage/TreeReader.cs
--- a/Engine/Source/Programs/Shared/EpicGames.Horde/Storage/TreeReader.cs
+++ b/Engine/Source/Programs/Shared/EpicGames.Horde/Storage/TreeReader.cs
@@ -61,8 +61,8 @@ namespace EpicGames.Horde.Storage
 		{
 			typeof(Nodes.CommitNode),
 			typeof(Nodes.DirectoryNode),
-			typeof(Nodes.LeafFileNode),
-			typeof(Nodes.InteriorFileNode),
+			typeof(Nodes.LeafChunkedDataNode),
+			typeof(Nodes.InteriorChunkedDataNode),
 			typeof(Logs.LogNode),
 			typeof(Logs.LogChunkNode),
 			typeof(Logs.LogIndexNode),
