Horde: Missing files.

diff --git a/Engine/Source/Programs/Shared/EpicGames.Horde/Storage/Nodes/ChunkedDataNode.cs b/Engine/Source/Programs/Shared/EpicGames.Horde/Storage/Nodes/ChunkedDataNode.cs
new file mode 100644
--- /dev/null
+++ b/Engine/Source/Programs/Shared/EpicGames.Horde/Storage/Nodes/ChunkedDataNode.cs
@@ -0,0 +1,235 @@
+// Copyright Epic Games, Inc. All Rights Reserved.
+
+using System;
+using System.Collections.Generic;
+using System.IO;
+using System.Linq;
+using System.Threading;
+using System.Threading.Tasks;
+using EpicGames.Core;
+
+namespace EpicGames.Horde.Storage.Nodes
+{
+	/// <summary>
+	/// Representation of a data stream, split into chunks along content-aware boundaries using a rolling hash (<see cref="BuzHash"/>).
+	/// Chunks are pushed into a tree hierarchy as data is appended to the root, with nodes of the tree also split along content-aware boundaries with <see cref="IoHash.NumBytes"/> granularity.
+	/// Once a chunk has been written to storage, it is treated as immutable.
+	/// </summary>
+	public abstract class ChunkedDataNode : Node
+	{
+		/// <summary>
+		/// Copies the contents of this node and its children to the given output stream
+		/// </summary>
+		/// <param name="reader">Reader for nodes in the tree</param>
+		/// <param name="outputStream">The output stream to receive the data</param>
+		/// <param name="cancellationToken">Cancellation token for the operation</param>
+		public abstract Task CopyToStreamAsync(TreeReader reader, Stream outputStream, CancellationToken cancellationToken);
+
+		/// <summary>
+		/// Copy the contents of the node to the output stream without creating the intermediate FileNodes
+		/// </summary>
+		/// <param name="reader">Reader for nodes in the tree</param>
+		/// <param name="locator">File node to be copied</param>
+		/// <param name="outputStream">The output stream to receive the data</param>
+		/// <param name="cancellationToken">Cancellation token for the operation</param>
+		public static async Task CopyToStreamAsync(TreeReader reader, NodeLocator locator, Stream outputStream, CancellationToken cancellationToken)
+		{
+			async Task CopyFunc(ReadNodeAsyncCallbackParams parms, ITreeNodeReader nodeReader)
+			{
+				if (parms.Type == typeof(LeafChunkedDataNode))
+				{
+					await LeafChunkedDataNode.CopyToStreamAsync(outputStream, nodeReader, cancellationToken);
+				}
+				else if (parms.Type == typeof(InteriorChunkedDataNode))
+				{
+					await InteriorChunkedDataNode.CopyToStreamAsync(reader, outputStream, nodeReader, cancellationToken);
+				}
+				else
+				{
+					throw new ArgumentException("Unexpected FileNode type found.");
+				}
+			}
+
+			await reader.ReadNodeAsync(locator, CopyFunc, cancellationToken);
+		}
+
+		/// <summary>
+		/// Extracts the contents of this node to a file
+		/// </summary>
+		/// <param name="reader">Reader for nodes in the tree</param>
+		/// <param name="file">File to write with the contents of this node</param>
+		/// <param name="cancellationToken">Cancellation token for the operation</param>
+		/// <returns></returns>
+		public async Task CopyToFileAsync(TreeReader reader, FileInfo file, CancellationToken cancellationToken)
+		{
+			if(file.Exists && (file.Attributes & FileAttributes.ReadOnly) != 0)
+			{
+				file.Attributes &= ~FileAttributes.ReadOnly;
+			}
+			using (FileStream stream = file.Open(FileMode.Create, FileAccess.Write, FileShare.Read))
+			{
+				await CopyToStreamAsync(reader, stream, cancellationToken);
+			}
+		}
+
+		/// <summary>
+		/// Serialize this node and its children into a byte array
+		/// </summary>
+		/// <param name="reader">Reader for nodes in the tree</param>
+		/// <param name="cancellationToken"></param>
+		/// <returns>Array of data stored by the tree</returns>
+		public async Task<byte[]> ToByteArrayAsync(TreeReader reader, CancellationToken cancellationToken)
+		{
+			using (MemoryStream stream = new MemoryStream())
+			{
+				await CopyToStreamAsync(reader, stream, cancellationToken);
+				return stream.ToArray();
+			}
+		}
+	}
+
+	/// <summary>
+	/// File node that contains a chunk of data
+	/// </summary>
+	[NodeType("{B27AFB68-9E20-4A4B-A4D8-788A4098D439}", 1)]
+	public sealed class LeafChunkedDataNode : ChunkedDataNode
+	{
+		/// <summary>
+		/// Data for this node
+		/// </summary>
+		public ReadOnlyMemory<byte> Data { get; }
+
+		/// <summary>
+		/// Create an empty leaf node
+		/// </summary>
+		public LeafChunkedDataNode()
+		{
+		}
+
+		/// <summary>
+		/// Create a leaf node from the given serialized data
+		/// </summary>
+		public LeafChunkedDataNode(ITreeNodeReader reader)
+		{
+			// Keep this code in sync with CopyToStreamAsync
+			Data = reader.ReadFixedLengthBytes(reader.Length);
+		}
+
+		/// <inheritdoc/>
+		public override void Serialize(ITreeNodeWriter writer)
+		{
+			writer.WriteFixedLengthBytes(Data.Span);
+		}
+
+		/// <inheritdoc/>
+		public override IEnumerable<NodeRef> EnumerateRefs() => Enumerable.Empty<NodeRef>();
+
+		/// <inheritdoc/>
+		public override async Task CopyToStreamAsync(TreeReader reader, Stream outputStream, CancellationToken cancellationToken)
+		{
+			await outputStream.WriteAsync(Data, cancellationToken);
+		}
+
+		/// <summary>
+		/// Copy the contents of the node to the output stream without creating the intermediate FileNodes
+		/// </summary>
+		/// <param name="reader">Reader for nodes in the tree</param>
+		/// <param name="outputStream">The output stream to receive the data</param>
+		/// <param name="cancellationToken">Cancellation token for the operation</param>
+		public static async Task CopyToStreamAsync(Stream outputStream, ITreeNodeReader reader, CancellationToken cancellationToken)
+		{
+			// Keep this code in sync with the constructor
+			await outputStream.WriteAsync(reader.ReadFixedLengthBytes(reader.Length), cancellationToken);
+		}
+	}
+
+	/// <summary>
+	/// An interior file node
+	/// </summary>
+	[NodeType("{F4DEDDBC-70CB-4C7A-8347-F011AFCCCDB9}", 1)]
+	public class InteriorChunkedDataNode : ChunkedDataNode
+	{
+		/// <summary>
+		/// Child nodes
+		/// </summary>
+		public IReadOnlyList<NodeRef<ChunkedDataNode>> Children { get; }
+
+		/// <summary>
+		/// Constructor
+		/// </summary>
+		/// <param name="children"></param>
+		public InteriorChunkedDataNode(IReadOnlyList<NodeRef<ChunkedDataNode>> children)
+		{
+			Children = children;
+		}
+
+		/// <summary>
+		/// Constructor
+		/// </summary>
+		public InteriorChunkedDataNode(ITreeNodeReader reader)
+		{
+			// Keep this code in sync with CopyToStreamAsync
+			NodeRef<ChunkedDataNode>[] children = new NodeRef<ChunkedDataNode>[reader.Length / IoHash.NumBytes];
+			for (int idx = 0; idx < children.Length; idx++)
+			{
+				children[idx] = reader.ReadRef<ChunkedDataNode>();
+			}
+			Children = children;
+		}
+
+		/// <inheritdoc/>
+		public override void Serialize(ITreeNodeWriter writer)
+		{
+			foreach (NodeRef<ChunkedDataNode> child in Children)
+			{
+				writer.WriteRef(child);
+			}
+		}
+
+		/// <inheritdoc/>
+		public override IEnumerable<NodeRef> EnumerateRefs() => Children;
+
+		/// <inheritdoc/>
+		public override async Task CopyToStreamAsync(TreeReader reader, Stream outputStream, CancellationToken cancellationToken)
+		{
+			foreach (NodeRef<ChunkedDataNode> childNodeRef in Children)
+			{
+				ChunkedDataNode childNode = await childNodeRef.ExpandAsync(reader, cancellationToken);
+				await childNode.CopyToStreamAsync(reader, outputStream, cancellationToken);
+			}
+		}
+
+		/// <summary>
+		/// Copy the contents of the node to the output stream without creating the intermediate FileNodes
+		/// </summary>
+		/// <param name="reader">Reader for nodes in the tree</param>
+		/// <param name="outputStream">The output stream to receive the data</param>
+		/// <param name="nodeReader">Source data</param>
+		/// <param name="cancellationToken">Cancellation token for the operation</param>
+		public static async Task CopyToStreamAsync(TreeReader reader, Stream outputStream, ITreeNodeReader nodeReader, CancellationToken cancellationToken)
+		{
+			async Task CopyFunc(ReadNodeAsyncCallbackParams parms, ITreeNodeReader nodeReader)
+			{
+				if (parms.Type == typeof(LeafChunkedDataNode))
+				{
+					await LeafChunkedDataNode.CopyToStreamAsync(outputStream, nodeReader, cancellationToken);
+				}
+				else
+				{
+					throw new ArgumentException("Unexpected FileNode type found.");
+				}
+			}
+
+			// Keep this code in sync with the constructor
+			NodeRef<ChunkedDataNode>[] children = new NodeRef<ChunkedDataNode>[nodeReader.Length / IoHash.NumBytes];
+			for (int idx = 0; idx < children.Length; idx++)
+			{
+				children[idx] = nodeReader.ReadRef<ChunkedDataNode>();
+			}
+			foreach (NodeRef<ChunkedDataNode> node in children)
+			{
+				await reader.ReadNodeAsync(node.Handle!.Locator, CopyFunc, cancellationToken);
+			}
+		}
+	}
+}
diff --git a/Engine/Source/Programs/Shared/EpicGames.Horde/Storage/Nodes/ChunkedDataWriter.cs b/Engine/Source/Programs/Shared/EpicGames.Horde/Storage/Nodes/ChunkedDataWriter.cs
new file mode 100644
--- /dev/null
+++ b/Engine/Source/Programs/Shared/EpicGames.Horde/Storage/Nodes/ChunkedDataWriter.cs
@@ -0,0 +1,513 @@
+// Copyright Epic Games, Inc. All Rights Reserved.
+
+using System;
+using System.Collections.Generic;
+using System.Diagnostics;
+using System.IO;
+using System.Threading;
+using System.Threading.Tasks;
+using EpicGames.Core;
+
+namespace EpicGames.Horde.Storage.Nodes
+{
+	/// <summary>
+	/// Options for creating file nodes
+	/// </summary>
+	public class ChunkingOptions
+	{
+		/// <summary>
+		/// Options for creating leaf nodes
+		/// </summary>
+		public ChunkingOptionsForNodeType LeafOptions { get; set; }
+
+		/// <summary>
+		/// Options for creating interior nodes
+		/// </summary>
+		public ChunkingOptionsForNodeType InteriorOptions { get; set; }
+
+		/// <summary>
+		/// Constructor
+		/// </summary>
+		public ChunkingOptions()
+		{
+			LeafOptions = new ChunkingOptionsForNodeType(32 * 1024, 256 * 1024, 64 * 1024);
+			InteriorOptions = new ChunkingOptionsForNodeType(32 * 1024, 256 * 1024, 64 * 1024);
+		}
+	}
+
+	/// <summary>
+	/// Options for creating a specific type of file nodes
+	/// </summary>
+	public class ChunkingOptionsForNodeType
+	{
+		/// <summary>
+		/// Minimum chunk size
+		/// </summary>
+		public int MinSize { get; set; }
+
+		/// <summary>
+		/// Maximum chunk size. Chunks will be split on this boundary if another match is not found.
+		/// </summary>
+		public int MaxSize { get; set; }
+
+		/// <summary>
+		/// Target chunk size for content-slicing
+		/// </summary>
+		public int TargetSize { get; set; }
+
+		/// <summary>
+		/// Constructor
+		/// </summary>
+		/// <param name="size">Fixed size chunks to use</param>
+		public ChunkingOptionsForNodeType(int size)
+			: this(size, size, size)
+		{
+		}
+
+		/// <summary>
+		/// Default constructor
+		/// </summary>
+		public ChunkingOptionsForNodeType(int minSize, int maxSize, int targetSize)
+		{
+			MinSize = minSize;
+			MaxSize = maxSize;
+			TargetSize = targetSize;
+		}
+	}
+
+	/// <summary>
+	/// Utility class for generating FileNode data directly into <see cref="TreeWriter"/> instances, without constructing node representations first.
+	/// </summary>
+	public class ChunkedDataWriter
+	{
+		/// <summary>
+		/// Default buffer length when calling CreateAsync/AppendAsync
+		/// </summary>
+		public const int DefaultBufferLength = 32 * 1024;
+
+		class InteriorNodeState
+		{
+			public InteriorNodeState? _parent;
+			public readonly ArrayMemoryWriter _data;
+			public readonly List<NodeHandle> _children = new List<NodeHandle>();
+
+			public uint _rollingHash;
+
+			public InteriorNodeState(int maxSize)
+			{
+				_data = new ArrayMemoryWriter(maxSize);
+			}
+
+			public void Reset()
+			{
+				_rollingHash = 0;
+				_data.Clear();
+				_children.Clear();
+			}
+
+			public void Write(NodeHandle handle)
+			{
+				_children.Add(handle);
+				_data.WriteIoHash(handle.Hash);
+			}
+		}
+
+		static readonly NodeType s_leafNodeType = NodeType.Get<LeafChunkedDataNode>();
+		static readonly NodeType s_interiorNodeType = NodeType.Get<InteriorChunkedDataNode>();
+
+		readonly TreeWriter _writer;
+		readonly ChunkingOptions _options;
+
+		// Tree state
+		InteriorNodeState? _topInteriorNode = null;
+		long _totalLength;
+		readonly Stack<InteriorNodeState> _freeInteriorNodes = new Stack<InteriorNodeState>();
+
+		// Leaf node state
+		uint _leafHash;
+		int _leafLength;
+
+		/// <summary>
+		/// Length of the file so far
+		/// </summary>
+		public long Length => _totalLength;
+
+		/// <summary>
+		/// Constructor
+		/// </summary>
+		/// <param name="writer">Writer for new nodes</param>
+		/// <param name="options">Chunking options</param>
+		public ChunkedDataWriter(TreeWriter writer, ChunkingOptions options)
+		{
+			_writer = writer;
+			_options = options;
+		}
+
+		/// <summary>
+		/// Reset the current state
+		/// </summary>
+		public void Reset()
+		{
+			FreeInteriorNodes();
+			ResetLeafState();
+			_totalLength = 0;
+		}
+
+		/// <summary>
+		/// Resets the state of the current leaf node
+		/// </summary>
+		void ResetLeafState()
+		{
+			_leafHash = 0;
+			_leafLength = 0;
+		}
+
+		/// <summary>
+		/// Creates a new interior node state
+		/// </summary>
+		/// <returns>State object</returns>
+		InteriorNodeState CreateInteriorNode()
+		{
+			InteriorNodeState? result;
+			if (!_freeInteriorNodes.TryPop(out result))
+			{
+				result = new InteriorNodeState(_options.InteriorOptions.MaxSize);
+			}
+			return result;
+		}
+
+		/// <summary>
+		/// Free all the current interior nodes
+		/// </summary>
+		void FreeInteriorNodes()
+		{
+			while (_topInteriorNode != null)
+			{
+				InteriorNodeState current = _topInteriorNode;
+				_topInteriorNode = _topInteriorNode._parent;
+				current.Reset();
+				_freeInteriorNodes.Push(current);
+			}
+		}
+
+		/// <summary>
+		/// Creates data for the given file
+		/// </summary>
+		/// <param name="fileInfo">File to append</param>
+		/// <param name="cancellationToken">Cancellation token for the operation</param>
+		public async Task<NodeHandle> CreateAsync(FileInfo fileInfo, CancellationToken cancellationToken)
+		{
+			return await CreateAsync(fileInfo, DefaultBufferLength, cancellationToken);
+		}
+
+		/// <summary>
+		/// Creates data for the given file
+		/// </summary>
+		/// <param name="fileInfo">File to append</param>
+		/// <param name="bufferLength">Size of the read buffer</param>
+		/// <param name="cancellationToken">Cancellation token for the operation</param>
+		public async Task<NodeHandle> CreateAsync(FileInfo fileInfo, int bufferLength, CancellationToken cancellationToken)
+		{
+			using (FileStream stream = fileInfo.OpenRead())
+			{
+				return await CreateAsync(stream, bufferLength, cancellationToken);
+			}
+		}
+
+		/// <summary>
+		/// Creates data from the given stream
+		/// </summary>
+		/// <param name="stream">Stream to append</param>
+		/// <param name="cancellationToken">Cancellation token for the operation</param>
+		public async Task<NodeHandle> CreateAsync(Stream stream, CancellationToken cancellationToken)
+		{
+			return await CreateAsync(stream, DefaultBufferLength, cancellationToken);
+		}
+
+		/// <summary>
+		/// Creates data from the given stream
+		/// </summary>
+		/// <param name="stream">Stream to append</param>
+		/// <param name="bufferLength">Size of the read buffer</param>
+		/// <param name="cancellationToken">Cancellation token for the operation</param>
+		public async Task<NodeHandle> CreateAsync(Stream stream, int bufferLength, CancellationToken cancellationToken)
+		{
+			Reset();
+			await AppendAsync(stream, bufferLength, cancellationToken);
+			return await CompleteAsync(cancellationToken);
+		}
+
+		/// <summary>
+		/// Creates data from the given data
+		/// </summary>
+		/// <param name="data">Stream to append</param>
+		/// <param name="cancellationToken">Cancellation token for the operation</param>
+		public async Task<NodeHandle> CreateAsync(ReadOnlyMemory<byte> data, CancellationToken cancellationToken)
+		{
+			Reset();
+			await AppendAsync(data, cancellationToken);
+			return await FlushAsync(cancellationToken);
+		}
+
+		/// <summary>
+		/// Appends data to the current file
+		/// </summary>
+		/// <param name="stream">Stream containing data to append</param>
+		/// <param name="cancellationToken">Cancellation token for the operation</param>
+		public async Task AppendAsync(Stream stream, CancellationToken cancellationToken)
+		{
+			await AppendAsync(stream, DefaultBufferLength, cancellationToken);
+		}
+
+		/// <summary>
+		/// Appends data to the current file
+		/// </summary>
+		/// <param name="stream">Stream containing data to append</param>
+		/// <param name="bufferLength">Size of the read buffer</param>
+		/// <param name="cancellationToken">Cancellation token for the operation</param>
+		public async Task AppendAsync(Stream stream, int bufferLength, CancellationToken cancellationToken)
+		{
+			await stream.ReadAllBytesAsync(bufferLength, async (x) => await AppendAsync(x, cancellationToken), cancellationToken);
+		}
+
+		/// <summary>
+		/// Appends data to the current file
+		/// </summary>
+		/// <param name="data">Data to append</param>
+		/// <param name="cancellationToken">Cancellation token for the operation</param>
+		public async Task AppendAsync(ReadOnlyMemory<byte> data, CancellationToken cancellationToken)
+		{
+			Memory<byte> buffer = _writer.GetOutputBuffer(_leafLength, _leafLength);
+			for (; ; )
+			{
+				// Append data to the current leaf node
+				int appendLength = AppendToLeafNode(buffer.Span.Slice(0, _leafLength), data.Span, ref _leafHash, _options.LeafOptions);
+
+				buffer = _writer.GetOutputBuffer(_leafLength, _leafLength + appendLength);
+				data.Slice(0, appendLength).CopyTo(buffer.Slice(_leafLength));
+
+				_leafLength += appendLength;
+				data = data.Slice(appendLength);
+
+				_totalLength += appendLength;
+
+				if (data.Length == 0)
+				{
+					break;
+				}
+
+				// Flush the leaf node and any interior nodes that are full
+				NodeHandle handle = await WriteLeafNodeAsync(cancellationToken);
+				ResetLeafState();
+				await AddToInteriorNodeAsync(handle, cancellationToken);
+			}
+		}
+
+		/// <summary>
+		/// Determines how much data to append to an existing leaf node
+		/// </summary>
+		/// <param name="currentData">Current data in the leaf node</param>
+		/// <param name="appendData">Data to be appended</param>
+		/// <param name="rollingHash">Current BuzHash of the data</param>
+		/// <param name="options">Options for chunking the data</param>
+		/// <returns>The number of bytes to append</returns>
+		static int AppendToLeafNode(ReadOnlySpan<byte> currentData, ReadOnlySpan<byte> appendData, ref uint rollingHash, ChunkingOptionsForNodeType options)
+		{
+			// If the target option sizes are fixed, just chunk the data along fixed boundaries
+			if (options.MinSize == options.TargetSize && options.MaxSize == options.TargetSize)
+			{
+				return Math.Min(appendData.Length, options.MaxSize - (int)currentData.Length);
+			}
+
+			// Cap the append data span to the maximum amount we can add
+			int maxAppendLength = options.MaxSize - currentData.Length;
+			if (maxAppendLength < appendData.Length)
+			{
+				appendData = appendData.Slice(0, maxAppendLength);
+			}
+
+			// Length of the data to be appended
+			int appendLength = 0;
+
+			// Fast path for appending data to the buffer up to the chunk window size
+			int windowSize = options.MinSize;
+			if (currentData.Length < windowSize)
+			{
+				appendLength = Math.Min(windowSize - (int)currentData.Length, appendData.Length);
+				rollingHash = BuzHash.Add(rollingHash, appendData.Slice(0, appendLength));
+			}
+
+			// Get the threshold for the rolling hash
+			uint rollingHashThreshold = (uint)((1L << 32) / options.TargetSize);
+
+			// Step through the part of the data where the tail of the window is in currentData, and the head of the window is in appendData.
+			if (appendLength < appendData.Length && windowSize > appendLength)
+			{
+				int overlap = windowSize - appendLength;
+				int overlapLength = Math.Min(appendData.Length - appendLength, overlap);
+
+				ReadOnlySpan<byte> tailSpan = currentData.Slice(currentData.Length - overlap, overlapLength);
+				ReadOnlySpan<byte> headSpan = appendData.Slice(appendLength, overlapLength);
+
+				int count = BuzHash.Update(tailSpan, headSpan, rollingHashThreshold, ref rollingHash);
+				if (count != -1)
+				{
+					appendLength += count;
+					return appendLength;
+				}
+
+				appendLength += headSpan.Length;
+			}
+
+			// Step through the rest of the data which is completely contained in appendData.
+			if (appendLength < appendData.Length)
+			{
+				Debug.Assert(appendLength >= windowSize);
+
+				ReadOnlySpan<byte> tailSpan = appendData.Slice(appendLength - windowSize, appendData.Length - windowSize);
+				ReadOnlySpan<byte> headSpan = appendData.Slice(appendLength);
+
+				int count = BuzHash.Update(tailSpan, headSpan, rollingHashThreshold, ref rollingHash);
+				if (count != -1)
+				{
+					appendLength += count;
+					return appendLength;
+				}
+
+				appendLength += headSpan.Length;
+			}
+
+			return appendLength;
+		}
+
+		async Task AddToInteriorNodeAsync(NodeHandle handle, CancellationToken cancellationToken)
+		{
+			_topInteriorNode ??= CreateInteriorNode();
+			await AddToInteriorNodeAsync(_topInteriorNode, handle, cancellationToken);
+		}
+
+		async Task AddToInteriorNodeAsync(InteriorNodeState interiorNode, NodeHandle handle, CancellationToken cancellationToken)
+		{
+			// If the node is already full, flush it
+			if (IsInteriorNodeComplete(interiorNode._data.WrittenSpan, interiorNode._rollingHash, _options.InteriorOptions))
+			{
+				NodeHandle interiorNodeHandle = await WriteInteriorNodeAndResetAsync(interiorNode, cancellationToken);
+				interiorNode._parent ??= CreateInteriorNode();
+				await AddToInteriorNodeAsync(interiorNode._parent, interiorNodeHandle, cancellationToken);
+			}
+
+			// Add this handle
+			AppendToInteriorNode(interiorNode._data.WrittenSpan, handle.Hash, ref interiorNode._rollingHash, _options.InteriorOptions);
+			interiorNode.Write(handle);
+		}
+
+
+		/// <summary>
+		/// Test whether the current node is complete
+		/// </summary>
+		/// <param name="currentData"></param>
+		/// <param name="rollingHash"></param>
+		/// <param name="options"></param>
+		/// <returns></returns>
+		static bool IsInteriorNodeComplete(ReadOnlySpan<byte> currentData, uint rollingHash, ChunkingOptionsForNodeType options)
+		{
+			if (currentData.Length + IoHash.NumBytes > options.MaxSize)
+			{
+				return true;
+			}
+
+			if (currentData.Length >= options.MinSize)
+			{
+				uint rollingHashThreshold = BuzHash.GetThreshold(options.TargetSize);
+				if (rollingHash < rollingHashThreshold)
+				{
+					return true;
+				}
+			}
+
+			return false;
+		}
+
+		/// <summary>
+		/// Append a new hash to this interior node
+		/// </summary>
+		/// <param name="currentData">Current data for the node</param>
+		/// <param name="hash">Hash of the child node</param>
+		/// <param name="rollingHash">Current rolling hash for the node</param>
+		/// <param name="options">Options for chunking the node</param>
+		/// <returns>True if the hash could be appended, false otherwise</returns>
+		static void AppendToInteriorNode(ReadOnlySpan<byte> currentData, IoHash hash, ref uint rollingHash, ChunkingOptionsForNodeType options)
+		{
+			Span<byte> hashData = stackalloc byte[IoHash.NumBytes];
+			hash.CopyTo(hashData);
+
+			rollingHash = BuzHash.Add(rollingHash, hashData);
+
+			int windowSize = options.MinSize - (options.MinSize % IoHash.NumBytes);
+			if (currentData.Length > windowSize)
+			{
+				ReadOnlySpan<byte> removeData = currentData.Slice(currentData.Length - windowSize, IoHash.NumBytes);
+				rollingHash = BuzHash.Sub(rollingHash, removeData, windowSize + IoHash.NumBytes);
+			}
+		}
+
+		/// <summary>
+		/// Complete the current file, and write all open nodes to the underlying writer
+		/// </summary>
+		/// <param name="cancellationToken">Cancellation token for the operation</param>
+		/// <returns>Handle to the root node</returns>
+		public async Task<NodeHandle> CompleteAsync(CancellationToken cancellationToken)
+		{
+			NodeHandle handle = await WriteLeafNodeAsync(cancellationToken);
+			ResetLeafState();
+
+			for (InteriorNodeState? state = _topInteriorNode; state != null; state = state._parent)
+			{
+				await AddToInteriorNodeAsync(state, handle, cancellationToken);
+				handle = await WriteInteriorNodeAndResetAsync(state, cancellationToken);
+			}
+
+			FreeInteriorNodes();
+			return handle;
+		}
+
+		/// <summary>
+		/// Flush the state of the writer
+		/// </summary>
+		/// <param name="cancellationToken">Cancellation token for the operation</param>
+		/// <returns>Handle to the root FileNode</returns>
+		public async Task<NodeHandle> FlushAsync(CancellationToken cancellationToken)
+		{
+			NodeHandle handle = await CompleteAsync(cancellationToken);
+			await _writer.FlushAsync(cancellationToken);
+			return handle;
+		}
+
+		/// <summary>
+		/// Writes the state of the given interior node to storage
+		/// </summary>
+		/// <param name="state"></param>
+		/// <param name="cancellationToken"></param>
+		/// <returns></returns>
+		async ValueTask<NodeHandle> WriteInteriorNodeAndResetAsync(InteriorNodeState state, CancellationToken cancellationToken)
+		{
+			Memory<byte> buffer = _writer.GetOutputBuffer(0, state._data.Length);
+			state._data.WrittenMemory.CopyTo(buffer);
+
+			NodeHandle handle = await _writer.WriteNodeAsync(state._data.Length, state._children, s_interiorNodeType, cancellationToken);
+			state.Reset();
+
+			return handle;
+		}
+
+		/// <summary>
+		/// Writes the contents of the current leaf node to storage
+		/// </summary>
+		/// <param name="cancellationToken">Cancellation token for the operation</param>
+		/// <returns>Handle to the written leaf node</returns>
+		async ValueTask<NodeHandle> WriteLeafNodeAsync(CancellationToken cancellationToken)
+		{
+			return await _writer.WriteNodeAsync(_leafLength, Array.Empty<NodeHandle>(), s_leafNodeType, cancellationToken);
+		}
+	}
+}
diff --git a/Engine/Source/Programs/Shared/EpicGames.Horde/Storage/Nodes/FileNode.cs b/Engine/Source/Programs/Shared/EpicGames.Horde/Storage/Nodes/FileNode.cs
deleted file mode 100644
diff --git a/Engine/Source/Programs/Shared/EpicGames.Horde/Storage/Nodes/FileNodeWriter.cs b/Engine/Source/Programs/Shared/EpicGames.Horde/Storage/Nodes/FileNodeWriter.cs
deleted file mode 100644
