Horde: Separate C#-specific typed node serialization/deserialization from core bundle logic.

diff --git a/Engine/Source/Programs/Horde/Horde.Server/Perforce/PerforceReplicator.cs b/Engine/Source/Programs/Horde/Horde.Server/Perforce/PerforceReplicator.cs
--- a/Engine/Source/Programs/Horde/Horde.Server/Perforce/PerforceReplicator.cs
+++ b/Engine/Source/Programs/Horde/Horde.Server/Perforce/PerforceReplicator.cs
@@ -53,7 +53,7 @@ namespace Horde.Server.Perforce
 				Paths = new List<Utf8String>();
 			}
 
-			public SyncNode(ITreeNodeReader reader)
+			public SyncNode(NodeReader reader)
 			{
 				Change = (int)reader.ReadUnsignedVarInt();
 				ParentChange = (int)reader.ReadUnsignedVarInt();
@@ -224,6 +224,8 @@ namespace Horde.Server.Perforce
 		/// </summary>
 		public PerforceReplicator(IPerforceService perforceService, StorageService storageService, IMemoryCache memoryCache, ILogger<PerforceReplicator> logger)
 		{
+			Node.RegisterType<SyncNode>();
+
 			_perforceService = perforceService;
 			_storageService = storageService;
 			_memoryCache = memoryCache;
@@ -255,10 +257,7 @@ namespace Horde.Server.Perforce
 		{
 			IStorageClient store = await _storageService.GetClientAsync(Namespace.Perforce, cancellationToken);
 
-			TreeReaderOptions readerOptions = new TreeReaderOptions();
-			readerOptions.Types.Add(typeof(SyncNode));
-
-			TreeReader reader = new TreeReader(store, _memoryCache, readerOptions, _logger);
+			TreeReader reader = new TreeReader(store, _memoryCache, _logger);
 
 			// Find the parent node
 			RefName refName = GetRefName(streamConfig.Id);
diff --git a/Engine/Source/Programs/Horde/Horde.Server/Storage/StorageController.cs b/Engine/Source/Programs/Horde/Horde.Server/Storage/StorageController.cs
--- a/Engine/Source/Programs/Horde/Horde.Server/Storage/StorageController.cs
+++ b/Engine/Source/Programs/Horde/Horde.Server/Storage/StorageController.cs
@@ -592,25 +592,9 @@ namespace Horde.Server.Storage
 
 		static string GetNodeLink(string linkBase, NodeRef treeNodeRef) => $"{linkBase}/nodes/{treeNodeRef.Handle!.Locator.Blob}?export={treeNodeRef.Handle!.Locator.ExportIdx}";
 
-		static readonly Dictionary<Guid, Type> s_guidToType = GetGuidToTypeMap();
-
-		static Dictionary<Guid, Type> GetGuidToTypeMap()
-		{
-			Dictionary<Guid, Type> guidToType = new Dictionary<Guid, Type>();
-			foreach (Type nodeType in new TreeReaderOptions().Types)
-			{
-				NodeTypeAttribute? attribute = nodeType.GetCustomAttribute<NodeTypeAttribute>();
-				if (attribute != null)
-				{
-					guidToType.Add(Guid.Parse(attribute.Guid), nodeType);
-				}
-			}
-			return guidToType;
-		}
-
 		static Type? GetNodeType(Guid typeGuid)
 		{
-			s_guidToType.TryGetValue(typeGuid, out Type? type);
+			Node.TryGetConcreteType(typeGuid, out Type? type);
 			return type;
 		}
 	}
diff --git a/Engine/Source/Programs/Horde/Horde/Commands/Bundles/BundleDump.cs b/Engine/Source/Programs/Horde/Horde/Commands/Bundles/BundleDump.cs
--- a/Engine/Source/Programs/Horde/Horde/Commands/Bundles/BundleDump.cs
+++ b/Engine/Source/Programs/Horde/Horde/Commands/Bundles/BundleDump.cs
@@ -27,23 +27,15 @@ namespace Horde.Commands.Bundles
 
 			BundleHeader header = bundle.Header;
 
-			Dictionary<Guid, string> typeIdToName = new Dictionary<Guid, string>();
-
-			TreeReaderOptions options = new TreeReaderOptions();
-			foreach (Type type in options.Types)
-			{
-				NodeTypeAttribute? attribute = type.GetCustomAttribute<NodeTypeAttribute>();
-				if (attribute != null)
-				{
-					typeIdToName.Add(Guid.Parse(attribute.Guid), type.Name);
-				}
-			}
-
 			string[] types = new string[header.Types.Count];
 			for (int idx = 0; idx < header.Types.Count; idx++)
 			{
-				string? name;
-				if (!typeIdToName.TryGetValue(header.Types[idx].Guid, out name))
+				string name;
+				if (Node.TryGetConcreteType(header.Types[idx].Guid, out Type? type))
+				{
+					name = type.Name;
+				}
+				else
 				{
 					name = header.Types[idx].Guid.ToString();
 				}
diff --git a/Engine/Source/Programs/Shared/EpicGames.Horde.Tests/BundleStoreTests.cs b/Engine/Source/Programs/Shared/EpicGames.Horde.Tests/BundleStoreTests.cs
--- a/Engine/Source/Programs/Shared/EpicGames.Horde.Tests/BundleStoreTests.cs
+++ b/Engine/Source/Programs/Shared/EpicGames.Horde.Tests/BundleStoreTests.cs
@@ -15,14 +15,13 @@ using System.IO;
 using EpicGames.Core;
 using Microsoft.Extensions.Caching.Memory;
 using Microsoft.Extensions.Logging.Abstractions;
+using System.Reflection;
 
 namespace EpicGames.Horde.Tests
 {
 	[TestClass]
 	public class BundleStoreTests
 	{
-		static readonly TreeReaderOptions s_readOptions = new TreeReaderOptions(typeof(SimpleNode));
-
 		[TestMethod]
 		public async Task CreateBundlesManuallyAsync()
 		{
@@ -41,7 +40,7 @@ namespace EpicGames.Horde.Tests
 
 			public TextNode(string text) => Text = text;
 
-			public TextNode(ITreeNodeReader reader)
+			public TextNode(NodeReader reader)
 			{
 				Text = reader.ReadString();
 			}
@@ -54,6 +53,11 @@ namespace EpicGames.Horde.Tests
 			public override IEnumerable<NodeRef> EnumerateRefs() => Array.Empty<NodeRef>();
 		}
 
+		public BundleStoreTests()
+		{
+			Node.RegisterTypesFromAssembly(Assembly.GetExecutingAssembly());
+		}
+
 		static async Task<Bundle> CreateBundleNormalAsync()
 		{
 			MemoryStorageClient store = new MemoryStorageClient();
@@ -120,7 +124,7 @@ namespace EpicGames.Horde.Tests
 				Refs = refs;
 			}
 
-			public SimpleNode(ITreeNodeReader reader)
+			public SimpleNode(NodeReader reader)
 			{
 				Data = new ReadOnlySequence<byte>(reader.ReadVariableLengthBytes());
 				Refs = reader.ReadVariableLengthArray(() => reader.ReadRef<SimpleNode>());
@@ -150,13 +154,13 @@ namespace EpicGames.Horde.Tests
 
 				await writer.WriteAsync(new RefName("test"), root);
 
-				TreeReader reader = new TreeReader(store, null, s_readOptions, NullLogger.Instance);
+				TreeReader reader = new TreeReader(store, null, NullLogger.Instance);
 				await CheckTree(reader, root);
 			}
 
 			// Check we can read it back in
 			{
-				TreeReader reader = new TreeReader(store, null, s_readOptions, NullLogger.Instance);
+				TreeReader reader = new TreeReader(store, null, NullLogger.Instance);
 				SimpleNode root = await reader.ReadNodeAsync<SimpleNode>(new RefName("test"));
 				await CheckTree(reader, root);
 			}
@@ -203,7 +207,7 @@ namespace EpicGames.Horde.Tests
 			RefName refName = new RefName("test");
 			await store.WriteNodeAsync(refName, new SimpleNode(new ReadOnlySequence<byte>(new byte[] { (byte)123 }), Array.Empty<NodeRef<SimpleNode>>()));
 
-			TreeReader reader = new TreeReader(store, null, s_readOptions, NullLogger.Instance);
+			TreeReader reader = new TreeReader(store, null, NullLogger.Instance);
 			SimpleNode node = await reader.ReadNodeAsync<SimpleNode>(refName);
 
 			Assert.AreEqual(123, node.Data.FirstSpan[0]);
@@ -213,7 +217,7 @@ namespace EpicGames.Horde.Tests
 		public async Task DirectoryNodesAsync()
 		{
 			MemoryStorageClient store = new MemoryStorageClient();
-			TreeReader reader = new TreeReader(store, null, s_readOptions, NullLogger.Instance);
+			TreeReader reader = new TreeReader(store, null, NullLogger.Instance);
 
 			// Generate a tree
 			{
diff --git a/Engine/Source/Programs/Shared/EpicGames.Horde/Logs/LogChunkNode.cs b/Engine/Source/Programs/Shared/EpicGames.Horde/Logs/LogChunkNode.cs
--- a/Engine/Source/Programs/Shared/EpicGames.Horde/Logs/LogChunkNode.cs
+++ b/Engine/Source/Programs/Shared/EpicGames.Horde/Logs/LogChunkNode.cs
@@ -113,7 +113,7 @@ namespace EpicGames.Horde.Logs
 		/// Deserializing constructor
 		/// </summary>
 		/// <param name="reader">Reader to pull data from</param>
-		public LogChunkNode(ITreeNodeReader reader)
+		public LogChunkNode(NodeReader reader)
 			: this(reader.ReadVariableLengthBytes())
 		{
 		}
@@ -231,7 +231,7 @@ namespace EpicGames.Horde.Logs
 		/// Deserializing constructor
 		/// </summary>
 		/// <param name="reader"></param>
-		public LogChunkRef(ITreeNodeReader reader)
+		public LogChunkRef(NodeReader reader)
 			: base(reader)
 		{
 			LineIndex = (int)reader.ReadUnsignedVarInt();
diff --git a/Engine/Source/Programs/Shared/EpicGames.Horde/Logs/LogIndexNode.cs b/Engine/Source/Programs/Shared/EpicGames.Horde/Logs/LogIndexNode.cs
--- a/Engine/Source/Programs/Shared/EpicGames.Horde/Logs/LogIndexNode.cs
+++ b/Engine/Source/Programs/Shared/EpicGames.Horde/Logs/LogIndexNode.cs
@@ -87,7 +87,7 @@ namespace EpicGames.Horde.Logs
 		/// Deserialization constructor
 		/// </summary>
 		/// <param name="reader">Reader for data</param>
-		public LogIndexNode(ITreeNodeReader reader)
+		public LogIndexNode(NodeReader reader)
 		{
 			int version = (int)reader.ReadUnsignedVarInt();
 			if (version != CurrentVersion)
diff --git a/Engine/Source/Programs/Shared/EpicGames.Horde/Logs/LogNode.cs b/Engine/Source/Programs/Shared/EpicGames.Horde/Logs/LogNode.cs
--- a/Engine/Source/Programs/Shared/EpicGames.Horde/Logs/LogNode.cs
+++ b/Engine/Source/Programs/Shared/EpicGames.Horde/Logs/LogNode.cs
@@ -90,7 +90,7 @@ namespace EpicGames.Horde.Logs
 		/// Deserializing constructor
 		/// </summary>
 		/// <param name="reader">Reader to draw data from</param>
-		public LogNode(ITreeNodeReader reader)
+		public LogNode(NodeReader reader)
 		{
 			Format = (LogFormat)reader.ReadUInt8();
 			LineCount = (int)reader.ReadUnsignedVarInt();
diff --git a/Engine/Source/Programs/Shared/EpicGames.Horde/Storage/Node.cs b/Engine/Source/Programs/Shared/EpicGames.Horde/Storage/Node.cs
--- a/Engine/Source/Programs/Shared/EpicGames.Horde/Storage/Node.cs
+++ b/Engine/Source/Programs/Shared/EpicGames.Horde/Storage/Node.cs
@@ -3,13 +3,41 @@
 using System;
 using System.Collections.Concurrent;
 using System.Collections.Generic;
+using System.Diagnostics.CodeAnalysis;
 using System.Reflection;
+using System.Reflection.Emit;
 using System.Threading;
 using System.Threading.Tasks;
 using EpicGames.Core;
 
 namespace EpicGames.Horde.Storage
 {
+	/// <summary>
+	/// Attribute used to define a factory for a particular node type
+	/// </summary>
+	[AttributeUsage(AttributeTargets.Class)]
+	public sealed class NodeTypeAttribute : Attribute // TODO: THIS SHOULD BE NODEATTRIBUTE, not NODETYPEATTRIBUTE
+	{
+		/// <summary>
+		/// Name of the type to store in the bundle header
+		/// </summary>
+		public string Guid { get; }
+
+		/// <summary>
+		/// Version number of the serializer
+		/// </summary>
+		public int Version { get; }
+
+		/// <summary>
+		/// Constructor
+		/// </summary>
+		public NodeTypeAttribute(string guid, int version = 1)
+		{
+			Guid = guid;
+			Version = version;
+		}
+	}
+
 	/// <summary>
 	/// Base class for user-defined types that are stored in a tree
 	/// </summary>
@@ -29,7 +57,7 @@ namespace EpicGames.Horde.Storage
 		/// <summary>
 		/// Accessor for the bundle type definition associated with this node
 		/// </summary>
-		public NodeType NodeType => NodeType.Get(GetType());
+		public NodeType NodeType => GetNodeType(GetType());
 
 		/// <summary>
 		/// Default constructor
@@ -42,7 +70,7 @@ namespace EpicGames.Horde.Storage
 		/// Serialization constructor. Leaves the revision number zeroed by default.
 		/// </summary>
 		/// <param name="reader"></param>
-		protected Node(ITreeNodeReader reader)
+		protected Node(NodeReader reader)
 		{
 			Hash = reader.Hash;
 		}
@@ -67,6 +95,154 @@ namespace EpicGames.Horde.Storage
 		/// </summary>
 		/// <returns>References to other nodes</returns>
 		public abstract IEnumerable<NodeRef> EnumerateRefs();
+
+		#region Static methods
+
+		static readonly object s_writeLock = new object();
+		static readonly ConcurrentDictionary<Type, NodeType> s_typeToNodeType = new ConcurrentDictionary<Type, NodeType>();
+		static readonly ConcurrentDictionary<Guid, Type> s_guidToType = new ConcurrentDictionary<Guid, Type>();
+		static readonly ConcurrentDictionary<Guid, Func<NodeReader, Node>> s_guidToDeserializer = new ConcurrentDictionary<Guid, Func<NodeReader, Node>>();
+
+		static NodeType CreateNodeType(NodeTypeAttribute attribute)
+		{
+			return new NodeType(Guid.Parse(attribute.Guid), attribute.Version);
+		}
+
+		/// <summary>
+		/// Attempts to get the concrete type with the given node. The type must have been registered via a previous call to <see cref="RegisterType(Type)"/>.
+		/// </summary>
+		/// <param name="guid">Guid specified in the <see cref="NodeTypeAttribute"/></param>
+		/// <param name="type">On success, receives the C# type associated with this GUID</param>
+		/// <returns>True if the type was found</returns>
+		public static bool TryGetConcreteType(Guid guid, [NotNullWhen(true)] out Type? type) => s_guidToType.TryGetValue(guid, out type);
+
+		/// <summary>
+		/// Gets the node type corresponding to the given C# type
+		/// </summary>
+		/// <param name="type"></param>
+		/// <returns></returns>
+		public static NodeType GetNodeType(Type type)
+		{
+			NodeType nodeType;
+			if (!s_typeToNodeType.TryGetValue(type, out nodeType))
+			{
+				lock (s_writeLock)
+				{
+					if (!s_typeToNodeType.TryGetValue(type, out nodeType))
+					{
+						NodeTypeAttribute? attribute = type.GetCustomAttribute<NodeTypeAttribute>();
+						if (attribute == null)
+						{
+							throw new InvalidOperationException($"Missing {nameof(NodeTypeAttribute)} from type {type.Name}");
+						}
+						nodeType = s_typeToNodeType.GetOrAdd(type, CreateNodeType(attribute));
+					}
+				}
+			}
+			return nodeType;
+		}
+
+		/// <summary>
+		/// Gets the type descriptor for the given type
+		/// </summary>
+		/// <typeparam name="T">Type to get a <see cref="NodeType"/> for</typeparam>
+		/// <returns></returns>
+		public static NodeType GetNodeType<T>() where T : Node => GetNodeType(typeof(T));
+
+		/// <summary>
+		/// Deserialize a node from the given reader
+		/// </summary>
+		/// <param name="nodeData">Data to deserialize from</param>
+		/// <returns>New node instance</returns>
+		public static Node Deserialize(NodeData nodeData)
+		{
+			NodeReader reader = new NodeReader(nodeData);
+			return s_guidToDeserializer[reader.Type.Guid](reader);
+		}
+
+		/// <summary>
+		/// Static constructor. Registers all the types in the current assembly.
+		/// </summary>
+		static Node()
+		{
+			RegisterTypesFromAssembly(Assembly.GetExecutingAssembly());
+		}
+
+		/// <summary>
+		/// Registers a single deserializer
+		/// </summary>
+		/// <param name="type"></param>
+		/// <exception cref="NotImplementedException"></exception>
+		public static void RegisterType(Type type)
+		{
+			NodeTypeAttribute? attribute = type.GetCustomAttribute<NodeTypeAttribute>();
+			if (attribute == null)
+			{
+				throw new InvalidOperationException($"Missing {nameof(NodeTypeAttribute)} from type {type.Name}");
+			}
+			RegisterType(type, attribute);
+		}
+
+		/// <summary>
+		/// 
+		/// </summary>
+		/// <typeparam name="T"></typeparam>
+		public static void RegisterType<T>() where T : Node => RegisterType(typeof(T));
+
+		/// <summary>
+		/// Register all node types with the <see cref="NodeTypeAttribute"/> from the given assembly
+		/// </summary>
+		/// <param name="assembly">Assembly to register types from</param>
+		public static void RegisterTypesFromAssembly(Assembly assembly)
+		{
+			Type[] types = assembly.GetTypes();
+			lock (s_writeLock)
+			{
+				foreach (Type type in types)
+				{
+					if (type.IsClass)
+					{
+						NodeTypeAttribute? attribute = type.GetCustomAttribute<NodeTypeAttribute>();
+						if (attribute != null)
+						{
+							RegisterType(type, attribute);
+						}
+					}
+				}
+			}
+		}
+
+		static void RegisterType(Type type, NodeTypeAttribute attribute)
+		{
+			NodeType nodeType = CreateNodeType(attribute);
+			s_typeToNodeType.TryAdd(type, nodeType);
+
+			Func<NodeReader, Node> deserializer = CreateDeserializer(type);
+			s_guidToType.TryAdd(nodeType.Guid, type);
+			s_guidToDeserializer.TryAdd(nodeType.Guid, deserializer);
+		}
+
+		static Func<NodeReader, Node> CreateDeserializer(Type type)
+		{
+			Type[] signature = new[] { typeof(NodeReader) };
+
+			ConstructorInfo? constructorInfo = type.GetConstructor(signature);
+			if (constructorInfo == null)
+			{
+				throw new InvalidOperationException($"Type {type.Name} does not have a constructor taking an {typeof(NodeReader).Name} as parameter.");
+			}
+
+			DynamicMethod method = new DynamicMethod($"Create_{type.Name}", type, signature, true);
+
+			ILGenerator generator = method.GetILGenerator();
+			generator.Emit(OpCodes.Ldarg_0);
+			generator.Emit(OpCodes.Newobj, constructorInfo);
+			generator.Emit(OpCodes.Ret);
+
+			return (Func<NodeReader, Node>)method.CreateDelegate(typeof(Func<NodeReader, Node>));
+		}
+
+		#endregion
 	}
 
 	/// <summary>
@@ -91,7 +267,7 @@ namespace EpicGames.Horde.Storage
 		/// </summary>
 		/// <param name="reader">Reader to deserialize from</param>
 		/// <returns>New untyped ref</returns>
-		public static NodeRef ReadRef(this ITreeNodeReader reader)
+		public static NodeRef ReadRef(this NodeReader reader)
 		{
 			return new NodeRef(reader);
 		}
@@ -102,7 +278,7 @@ namespace EpicGames.Horde.Storage
 		/// <typeparam name="T">Type of the referenced node</typeparam>
 		/// <param name="reader">Reader to deserialize from</param>
 		/// <returns>New strongly typed ref</returns>
-		public static NodeRef<T> ReadRef<T>(this ITreeNodeReader reader) where T : Node
+		public static NodeRef<T> ReadRef<T>(this NodeReader reader) where T : Node
 		{
 			return new NodeRef<T>(reader);
 		}
@@ -112,7 +288,7 @@ namespace EpicGames.Horde.Storage
 		/// </summary>
 		/// <param name="reader">Reader to deserialize from</param>
 		/// <returns>New untyped ref</returns>
-		public static NodeRef? ReadOptionalRef(this ITreeNodeReader reader)
+		public static NodeRef? ReadOptionalRef(this NodeReader reader)
 		{
 			if (reader.ReadBoolean())
 			{
@@ -129,7 +305,7 @@ namespace EpicGames.Horde.Storage
 		/// </summary>
 		/// <param name="reader">Reader to deserialize from</param>
 		/// <returns>New strongly typed ref</returns>
-		public static NodeRef<T>? ReadOptionalRef<T>(this ITreeNodeReader reader) where T : Node
+		public static NodeRef<T>? ReadOptionalRef<T>(this NodeReader reader) where T : Node
 		{
 			if (reader.ReadBoolean())
 			{
diff --git a/Engine/Source/Programs/Shared/EpicGames.Horde/Storage/NodeRef.cs b/Engine/Source/Programs/Shared/EpicGames.Horde/Storage/NodeRef.cs
--- a/Engine/Source/Programs/Shared/EpicGames.Horde/Storage/NodeRef.cs
+++ b/Engine/Source/Programs/Shared/EpicGames.Horde/Storage/NodeRef.cs
@@ -33,7 +33,7 @@ namespace EpicGames.Horde.Storage
 	/// Each ref must have EXACTLY one owner; sharing of refs between objects is not permitted and will break change tracking.
 	/// Multiple refs MAY point to the same target object.
 	/// 
-	/// To read an untracked object that can be added to a new ref, call <see cref="TreeReader.ReadNodeAsync(NodeLocator, CancellationToken)"/> 
+	/// To read an untracked object that can be added to a new ref, call <see cref="TreeReader.ReadNodeDataAsync(NodeLocator, CancellationToken)"/> 
 	/// directly, or use <see cref="NodeRef{T}.ExpandCopyAsync(TreeReader, CancellationToken)"/>.
 	/// </summary>
 	public class NodeRef
@@ -70,7 +70,7 @@ namespace EpicGames.Horde.Storage
 		/// Deserialization constructor
 		/// </summary>
 		/// <param name="reader"></param>
-		public NodeRef(ITreeNodeReader reader) : this(reader.ReadNodeHandle())
+		public NodeRef(NodeReader reader) : this(reader.ReadNodeHandle())
 		{
 		}
 
@@ -134,7 +134,8 @@ namespace EpicGames.Horde.Storage
 		{
 			if (Target == null)
 			{
-				Target = await reader.ReadNodeAsync(Handle!.Locator, cancellationToken);
+				NodeData nodeData = await reader.ReadNodeDataAsync(Handle!.Locator, cancellationToken);
+				Target = Node.Deserialize(nodeData);
 				OnExpand();
 			}
 			return Target;
@@ -198,7 +199,7 @@ namespace EpicGames.Horde.Storage
 		/// Constructor
 		/// </summary>
 		/// <param name="reader">The reader to deserialize from</param>
-		public NodeRef(ITreeNodeReader reader) : base(reader)
+		public NodeRef(NodeReader reader) : base(reader)
 		{
 		}
 
diff --git a/Engine/Source/Programs/Shared/EpicGames.Horde/Storage/NodeType.cs b/Engine/Source/Programs/Shared/EpicGames.Horde/Storage/NodeType.cs
--- a/Engine/Source/Programs/Shared/EpicGames.Horde/Storage/NodeType.cs
+++ b/Engine/Source/Programs/Shared/EpicGames.Horde/Storage/NodeType.cs
@@ -1,94 +1,17 @@
 // Copyright Epic Games, Inc. All Rights Reserved.
 
 using System;
-using System.Collections.Concurrent;
-using System.Globalization;
-using System.Reflection;
 
 namespace EpicGames.Horde.Storage
 {
-	/// <summary>
-	/// Attribute used to define a factory for a particular node type
-	/// </summary>
-	[AttributeUsage(AttributeTargets.Class)]
-	public sealed class NodeTypeAttribute : Attribute
-	{
-		/// <summary>
-		/// Name of the type to store in the bundle header
-		/// </summary>
-		public string Guid { get; }
-
-		/// <summary>
-		/// Version number of the serializer
-		/// </summary>
-		public int Version { get; }
-
-		/// <summary>
-		/// Constructor
-		/// </summary>
-		public NodeTypeAttribute(string guid, int version = 1)
-		{
-			Guid = guid;
-			Version = version;
-		}
-	}
-
 	/// <summary>
 	/// Information about a type within a bundle
 	/// </summary>
 	/// <param name="Guid">Nominal identifier for the type</param>
 	/// <param name="Version">Version number for the serializer</param>
-	public record class NodeType(Guid Guid, int Version)
+	public record struct NodeType(Guid Guid, int Version)
 	{
-		/// <summary>
-		/// Parse a span of characters as a bundle type
-		/// </summary>
-		/// <param name="text">Text to parse</param>
-		/// <returns>The parsed bundle type</returns>
-		public static NodeType Parse(ReadOnlySpan<char> text)
-		{
-			int hashIdx = text.IndexOf('#');
-			Guid guid = Guid.Parse(text.Slice(0, hashIdx));
-			int version = Int32.Parse(text.Slice(hashIdx + 1), NumberStyles.None, CultureInfo.InvariantCulture);
-			return new NodeType(guid, version);
-		}
-
 		/// <inheritdoc/>
 		public override string ToString() => $"{Guid}#{Version}";
-
-		#region Static Methods
-
-		/// <summary>
-		/// Cache of constructed <see cref="NodeType"/> instances.
-		/// </summary>
-		static readonly ConcurrentDictionary<Type, NodeType> s_typeToBundleType = new ConcurrentDictionary<Type, NodeType>();
-
-		/// <summary>
-		/// Gets the <see cref="NodeType"/> instance for a particular node type
-		/// </summary>
-		/// <param name="type"></param>
-		/// <returns>Bundle</returns>
-		public static NodeType Get(Type type)
-		{
-			NodeType? bundleType;
-			if (!s_typeToBundleType.TryGetValue(type, out bundleType))
-			{
-				NodeTypeAttribute? attribute = type.GetCustomAttribute<NodeTypeAttribute>();
-				if (attribute == null)
-				{
-					throw new InvalidOperationException($"Missing {nameof(NodeTypeAttribute)} from type {type.Name}");
-				}
-				bundleType = s_typeToBundleType.GetOrAdd(type, new NodeType(Guid.Parse(attribute.Guid), attribute.Version));
-			}
-			return bundleType;
-		}
-
-		/// <summary>
-		/// Gets the <see cref="NodeType"/> instance for a particular node type
-		/// </summary>
-		/// <typeparam name="T">Type of the node</typeparam>
-		public static NodeType Get<T>() where T : Node => Get(typeof(T));
-
-		#endregion
 	}
 }
diff --git a/Engine/Source/Programs/Shared/EpicGames.Horde/Storage/Nodes/CbNode.cs b/Engine/Source/Programs/Shared/EpicGames.Horde/Storage/Nodes/CbNode.cs
--- a/Engine/Source/Programs/Shared/EpicGames.Horde/Storage/Nodes/CbNode.cs
+++ b/Engine/Source/Programs/Shared/EpicGames.Horde/Storage/Nodes/CbNode.cs
@@ -71,7 +71,7 @@ namespace EpicGames.Horde.Storage.Nodes
 		/// Deserialization constructor
 		/// </summary>
 		/// <param name="reader">Reader to deserialize from</param>
-		public CbNode(ITreeNodeReader reader)
+		public CbNode(NodeReader reader)
 		{
 			Object = new CbObject(reader.ReadFixedLengthBytes(reader.Length));
 
diff --git a/Engine/Source/Programs/Shared/EpicGames.Horde/Storage/Nodes/ChunkedDataNode.cs b/Engine/Source/Programs/Shared/EpicGames.Horde/Storage/Nodes/ChunkedDataNode.cs
--- a/Engine/Source/Programs/Shared/EpicGames.Horde/Storage/Nodes/ChunkedDataNode.cs
+++ b/Engine/Source/Programs/Shared/EpicGames.Horde/Storage/Nodes/ChunkedDataNode.cs
@@ -25,6 +25,9 @@ namespace EpicGames.Horde.Storage.Nodes
 		/// <param name="cancellationToken">Cancellation token for the operation</param>
 		public abstract Task CopyToStreamAsync(TreeReader reader, Stream outputStream, CancellationToken cancellationToken);
 
+		static readonly Guid s_leafNodeGuid = Node.GetNodeType<LeafChunkedDataNode>().Guid;
+		static readonly Guid s_interiorNodeGuid = Node.GetNodeType<InteriorChunkedDataNode>().Guid;
+
 		/// <summary>
 		/// Copy the contents of the node to the output stream without creating the intermediate FileNodes
 		/// </summary>
@@ -34,23 +37,19 @@ namespace EpicGames.Horde.Storage.Nodes
 		/// <param name="cancellationToken">Cancellation token for the operation</param>
 		public static async Task CopyToStreamAsync(TreeReader reader, NodeLocator locator, Stream outputStream, CancellationToken cancellationToken)
 		{
-			async Task CopyFunc(ReadNodeAsyncCallbackParams parms, ITreeNodeReader nodeReader)
+			NodeData nodeData = await reader.ReadNodeDataAsync(locator, cancellationToken);
+			if (nodeData.Type.Guid == s_leafNodeGuid)
 			{
-				if (parms.Type == typeof(LeafChunkedDataNode))
-				{
-					await LeafChunkedDataNode.CopyToStreamAsync(outputStream, nodeReader, cancellationToken);
-				}
-				else if (parms.Type == typeof(InteriorChunkedDataNode))
-				{
-					await InteriorChunkedDataNode.CopyToStreamAsync(reader, outputStream, nodeReader, cancellationToken);
-				}
-				else
-				{
-					throw new ArgumentException("Unexpected FileNode type found.");
-				}
+				await LeafChunkedDataNode.CopyToStreamAsync(nodeData, outputStream, cancellationToken);
+			}
+			else if (nodeData.Type.Guid == s_interiorNodeGuid)
+			{
+				await InteriorChunkedDataNode.CopyToStreamAsync(reader, nodeData, outputStream, cancellationToken);
+			}
+			else
+			{
+				throw new ArgumentException($"Unexpected {nameof(ChunkedDataNode)} type found.");
 			}
-
-			await reader.ReadNodeAsync(locator, CopyFunc, cancellationToken);
 		}
 
 		/// <summary>
@@ -109,7 +108,7 @@ namespace EpicGames.Horde.Storage.Nodes
 		/// <summary>
 		/// Create a leaf node from the given serialized data
 		/// </summary>
-		public LeafChunkedDataNode(ITreeNodeReader reader)
+		public LeafChunkedDataNode(NodeReader reader)
 		{
 			// Keep this code in sync with CopyToStreamAsync
 			Data = reader.ReadFixedLengthBytes(reader.Length);
@@ -133,13 +132,13 @@ namespace EpicGames.Horde.Storage.Nodes
 		/// <summary>
 		/// Copy the contents of the node to the output stream without creating the intermediate FileNodes
 		/// </summary>
-		/// <param name="reader">Reader for nodes in the tree</param>
+		/// <param name="nodeData">The raw node data</param>
 		/// <param name="outputStream">The output stream to receive the data</param>
 		/// <param name="cancellationToken">Cancellation token for the operation</param>
-		public static async Task CopyToStreamAsync(Stream outputStream, ITreeNodeReader reader, CancellationToken cancellationToken)
+		public static async Task CopyToStreamAsync(NodeData nodeData, Stream outputStream, CancellationToken cancellationToken)
 		{
 			// Keep this code in sync with the constructor
-			await outputStream.WriteAsync(reader.ReadFixedLengthBytes(reader.Length), cancellationToken);
+			await outputStream.WriteAsync(nodeData.Data, cancellationToken);
 		}
 	}
 
@@ -166,7 +165,7 @@ namespace EpicGames.Horde.Storage.Nodes
 		/// <summary>
 		/// Constructor
 		/// </summary>
-		public InteriorChunkedDataNode(ITreeNodeReader reader)
+		public InteriorChunkedDataNode(NodeReader reader)
 		{
 			// Keep this code in sync with CopyToStreamAsync
 			NodeRef<ChunkedDataNode>[] children = new NodeRef<ChunkedDataNode>[reader.Length / IoHash.NumBytes];
@@ -203,32 +202,16 @@ namespace EpicGames.Horde.Storage.Nodes
 		/// Copy the contents of the node to the output stream without creating the intermediate FileNodes
 		/// </summary>
 		/// <param name="reader">Reader for nodes in the tree</param>
+		/// <param name="nodeData">Source data</param>
 		/// <param name="outputStream">The output stream to receive the data</param>
-		/// <param name="nodeReader">Source data</param>
 		/// <param name="cancellationToken">Cancellation token for the operation</param>
-		public static async Task CopyToStreamAsync(TreeReader reader, Stream outputStream, ITreeNodeReader nodeReader, CancellationToken cancellationToken)
+		public static async Task CopyToStreamAsync(TreeReader reader, NodeData nodeData, Stream outputStream, CancellationToken cancellationToken)
 		{
-			async Task CopyFunc(ReadNodeAsyncCallbackParams parms, ITreeNodeReader nodeReader)
-			{
-				if (parms.Type == typeof(LeafChunkedDataNode))
-				{
-					await LeafChunkedDataNode.CopyToStreamAsync(outputStream, nodeReader, cancellationToken);
-				}
-				else
-				{
-					throw new ArgumentException("Unexpected FileNode type found.");
-				}
-			}
-
-			// Keep this code in sync with the constructor
-			NodeRef<ChunkedDataNode>[] children = new NodeRef<ChunkedDataNode>[nodeReader.Length / IoHash.NumBytes];
-			for (int idx = 0; idx < children.Length; idx++)
-			{
-				children[idx] = nodeReader.ReadRef<ChunkedDataNode>();
-			}
-			foreach (NodeRef<ChunkedDataNode> node in children)
+			NodeReader nodeReader = new NodeReader(nodeData);
+			while (nodeReader.GetMemory(0).Length > 0)
 			{
-				await reader.ReadNodeAsync(node.Handle!.Locator, CopyFunc, cancellationToken);
+				NodeRef nodeRef = nodeReader.ReadRef();
+				await ChunkedDataNode.CopyToStreamAsync(reader, nodeRef.Handle!.Locator, outputStream, cancellationToken);
 			}
 		}
 	}
diff --git a/Engine/Source/Programs/Shared/EpicGames.Horde/Storage/Nodes/ChunkedDataWriter.cs b/Engine/Source/Programs/Shared/EpicGames.Horde/Storage/Nodes/ChunkedDataWriter.cs
--- a/Engine/Source/Programs/Shared/EpicGames.Horde/Storage/Nodes/ChunkedDataWriter.cs
+++ b/Engine/Source/Programs/Shared/EpicGames.Horde/Storage/Nodes/ChunkedDataWriter.cs
@@ -112,8 +112,8 @@ namespace EpicGames.Horde.Storage.Nodes
 			}
 		}
 
-		static readonly NodeType s_leafNodeType = NodeType.Get<LeafChunkedDataNode>();
-		static readonly NodeType s_interiorNodeType = NodeType.Get<InteriorChunkedDataNode>();
+		static readonly NodeType s_leafNodeType = Node.GetNodeType<LeafChunkedDataNode>();
+		static readonly NodeType s_interiorNodeType = Node.GetNodeType<InteriorChunkedDataNode>();
 
 		readonly TreeWriter _writer;
 		readonly ChunkingOptions _options;
diff --git a/Engine/Source/Programs/Shared/EpicGames.Horde/Storage/Nodes/CommitNode.cs b/Engine/Source/Programs/Shared/EpicGames.Horde/Storage/Nodes/CommitNode.cs
--- a/Engine/Source/Programs/Shared/EpicGames.Horde/Storage/Nodes/CommitNode.cs
+++ b/Engine/Source/Programs/Shared/EpicGames.Horde/Storage/Nodes/CommitNode.cs
@@ -85,7 +85,7 @@ namespace EpicGames.Horde.Storage.Nodes
 		/// Deserializing constructor
 		/// </summary>
 		/// <param name="reader"></param>
-		public CommitNode(ITreeNodeReader reader)
+		public CommitNode(NodeReader reader)
 		{
 			Number = (int)reader.ReadUnsignedVarInt();
 			Parent = reader.ReadOptionalRef<CommitNode>();
diff --git a/Engine/Source/Programs/Shared/EpicGames.Horde/Storage/Nodes/DirectoryEntry.cs b/Engine/Source/Programs/Shared/EpicGames.Horde/Storage/Nodes/DirectoryEntry.cs
--- a/Engine/Source/Programs/Shared/EpicGames.Horde/Storage/Nodes/DirectoryEntry.cs
+++ b/Engine/Source/Programs/Shared/EpicGames.Horde/Storage/Nodes/DirectoryEntry.cs
@@ -35,7 +35,7 @@ namespace EpicGames.Horde.Storage.Nodes
 		/// Deserializing constructor
 		/// </summary>
 		/// <param name="reader"></param>
-		public DirectoryEntry(ITreeNodeReader reader)
+		public DirectoryEntry(NodeReader reader)
 			: base(reader)
 		{
 			Name = reader.ReadUtf8String();
diff --git a/Engine/Source/Programs/Shared/EpicGames.Horde/Storage/Nodes/DirectoryNode.cs b/Engine/Source/Programs/Shared/EpicGames.Horde/Storage/Nodes/DirectoryNode.cs
--- a/Engine/Source/Programs/Shared/EpicGames.Horde/Storage/Nodes/DirectoryNode.cs
+++ b/Engine/Source/Programs/Shared/EpicGames.Horde/Storage/Nodes/DirectoryNode.cs
@@ -124,7 +124,7 @@ namespace EpicGames.Horde.Storage.Nodes
 		/// Deserialization constructor
 		/// </summary>
 		/// <param name="reader">Reader to deserialize from</param>
-		public DirectoryNode(ITreeNodeReader reader)
+		public DirectoryNode(NodeReader reader)
 		{
 			Flags = (DirectoryFlags)reader.ReadUnsignedVarInt();
 
diff --git a/Engine/Source/Programs/Shared/EpicGames.Horde/Storage/Nodes/DirectoryNodeRef.cs b/Engine/Source/Programs/Shared/EpicGames.Horde/Storage/Nodes/DirectoryNodeRef.cs
--- a/Engine/Source/Programs/Shared/EpicGames.Horde/Storage/Nodes/DirectoryNodeRef.cs
+++ b/Engine/Source/Programs/Shared/EpicGames.Horde/Storage/Nodes/DirectoryNodeRef.cs
@@ -41,7 +41,7 @@ namespace EpicGames.Horde.Storage.Nodes
 		/// Constructor
 		/// </summary>
 		/// <param name="reader"></param>
-		public DirectoryNodeRef(ITreeNodeReader reader)
+		public DirectoryNodeRef(NodeReader reader)
 			: base(reader)
 		{
 			_cachedLength = (long)reader.ReadUnsignedVarInt();
diff --git a/Engine/Source/Programs/Shared/EpicGames.Horde/Storage/Nodes/FileEntry.cs b/Engine/Source/Programs/Shared/EpicGames.Horde/Storage/Nodes/FileEntry.cs
--- a/Engine/Source/Programs/Shared/EpicGames.Horde/Storage/Nodes/FileEntry.cs
+++ b/Engine/Source/Programs/Shared/EpicGames.Horde/Storage/Nodes/FileEntry.cs
@@ -92,7 +92,7 @@ namespace EpicGames.Horde.Storage.Nodes
 		/// Deserialize from a buffer
 		/// </summary>
 		/// <param name="reader"></param>
-		public FileEntry(ITreeNodeReader reader)
+		public FileEntry(NodeReader reader)
 			: base(reader)
 		{
 			Name = reader.ReadUtf8String();
diff --git a/Engine/Source/Programs/Shared/EpicGames.Horde/Storage/TreeReader.cs b/Engine/Source/Programs/Shared/EpicGames.Horde/Storage/TreeReader.cs
--- a/Engine/Source/Programs/Shared/EpicGames.Horde/Storage/TreeReader.cs
+++ b/Engine/Source/Programs/Shared/EpicGames.Horde/Storage/TreeReader.cs
@@ -4,140 +4,78 @@ using System;
 using System.Buffers;
 using System.Collections.Generic;
 using System.Linq;
-using System.Reflection.Emit;
-using System.Reflection;
 using System.Threading;
 using System.Threading.Tasks;
 using EpicGames.Core;
 using Microsoft.Extensions.Caching.Memory;
 using Microsoft.Extensions.Logging;
-using System.Collections.Concurrent;
 using Microsoft.CodeAnalysis;
-using Blake3;
 using System.Diagnostics;
 
 namespace EpicGames.Horde.Storage
 {
+	/// <summary>
+	/// Data for an individual node
+	/// </summary>
+	public record struct NodeData(NodeType Type, IoHash Hash, ReadOnlyMemory<byte> Data, IReadOnlyList<NodeLocator> Refs);
+
 	/// <summary>
 	/// Reader for tree nodes
 	/// </summary>
-	public interface ITreeNodeReader : IMemoryReader
+	public class NodeReader : MemoryReader
 	{
+		/// <summary>
+		/// Type to deserialize
+		/// </summary>
+		public NodeType Type => _nodeData.Type;
+
 		/// <summary>
 		/// Version of the current node, as specified via <see cref="NodeTypeAttribute"/>
 		/// </summary>
-		int Version { get; }
+		public int Version => Type.Version;
 
 		/// <summary>
 		/// Total length of the data in this node
 		/// </summary>
-		int Length { get; }
+		public int Length => _nodeData.Data.Length;
 
 		/// <summary>
 		/// Hash of the node being deserialized
 		/// </summary>
-		IoHash Hash { get; }
+		public IoHash Hash => _nodeData.Hash;
 
 		/// <summary>
 		/// Locations of all referenced nodes.
 		/// </summary>
-		IReadOnlyList<NodeLocator> References { get; }
+		public IReadOnlyList<NodeLocator> References => _nodeData.Refs;
 
-		/// <summary>
-		/// Reads a reference to another node
-		/// </summary>
-		NodeHandle ReadNodeHandle();
-	}
+		readonly NodeData _nodeData;
+		int _refIdx;
 
-	/// <summary>
-	/// Options for configuring a bundle serializer
-	/// </summary>
-	public class TreeReaderOptions
-	{
 		/// <summary>
-		/// Known node types. Each node type should have a <see cref="NodeTypeAttribute"/> indicating the guid and latest supported version number.
+		/// Constructor
 		/// </summary>
-		public List<Type> Types { get; } = new List<Type>
+		public NodeReader(NodeData nodeData)
+			: base(nodeData.Data)
 		{
-			typeof(Nodes.CommitNode),
-			typeof(Nodes.DirectoryNode),
-			typeof(Nodes.LeafChunkedDataNode),
-			typeof(Nodes.InteriorChunkedDataNode),
-			typeof(Logs.LogNode),
-			typeof(Logs.LogChunkNode),
-			typeof(Logs.LogIndexNode),
-		};
+			_nodeData = nodeData;
+		}
 
 		/// <summary>
-		/// Constructor
+		/// Reads the next reference to another node
 		/// </summary>
-		/// <param name="types"></param>
-		public TreeReaderOptions(params Type[] types)
+		public NodeHandle ReadNodeHandle()
 		{
-			Types.AddRange(types);
+			IoHash hash = this.ReadIoHash();
+			return new NodeHandle(hash, _nodeData.Refs[_refIdx++]);
 		}
 	}
 
-	/// <summary>
-	/// Callback parameters for ReadNodeAsync calls. 
-	/// </summary>
-	/// <param name="Type">TreeNode type being read</param>
-	/// <param name="Deserialize">Deserialization function commonly </param>
-	public record struct ReadNodeAsyncCallbackParams(Type Type, Func<ITreeNodeReader, Node> Deserialize);
-
 	/// <summary>
 	/// Writes nodes from bundles in an <see cref="IStorageClient"/> instance.
 	/// </summary>
 	public class TreeReader
 	{
-		/// <summary>
-		/// Information about a known type that can be deserialized from a bundle
-		/// </summary>
-		class TypeInfo
-		{
-			static readonly ConcurrentDictionary<Type, TypeInfo> s_cachedTypeInfo = new ConcurrentDictionary<Type, TypeInfo>();
-
-			public Type Type { get; }
-			public NodeType BundleType { get; }
-			public Func<ITreeNodeReader, Node> Deserialize { get; }
-
-			private TypeInfo(Type type, NodeType bundleType, Func<ITreeNodeReader, Node> deserialize)
-			{
-				Type = type;
-				BundleType = bundleType;
-				Deserialize = deserialize;
-			}
-
-			public static TypeInfo Create(Type type)
-			{
-				TypeInfo? typeInfo;
-				if (!s_cachedTypeInfo.TryGetValue(type, out typeInfo))
-				{
-					NodeType bundleType = NodeType.Get(type);
-
-					Type[] signature = new[] { typeof(ITreeNodeReader) };
-
-					ConstructorInfo? constructorInfo = type.GetConstructor(signature);
-					if (constructorInfo == null)
-					{
-						throw new InvalidOperationException($"Type {type.Name} does not have a constructor taking an {typeof(ITreeNodeReader).Name} as parameter.");
-					}
-
-					DynamicMethod method = new DynamicMethod($"Create_{type.Name}", type, signature, true);
-
-					ILGenerator generator = method.GetILGenerator();
-					generator.Emit(OpCodes.Ldarg_0);
-					generator.Emit(OpCodes.Newobj, constructorInfo);
-					generator.Emit(OpCodes.Ret);
-
-					Func<ITreeNodeReader, Node> deserialize = (Func<ITreeNodeReader, Node>)method.CreateDelegate(typeof(Func<ITreeNodeReader, Node>));
-
-					typeInfo = s_cachedTypeInfo.GetOrAdd(type, new TypeInfo(type, bundleType, deserialize));
-				}
-				return typeInfo;
-			}
-		}
-
 		/// <summary>
 		/// Computed information about a bundle
 		/// </summary>
@@ -146,14 +84,12 @@ namespace EpicGames.Horde.Storage
 			public readonly BlobLocator Locator;
 			public readonly BundleHeader Header;
 			public readonly int HeaderLength;
-			public readonly TypeInfo?[] Types;
 
-			public BundleInfo(BlobLocator locator, BundleHeader header, int headerLength, TypeInfo?[] types)
+			public BundleInfo(BlobLocator locator, BundleHeader header, int headerLength)
 			{
 				Locator = locator;
 				Header = header;
 				HeaderLength = headerLength;
-				Types = types;
 			}
 		}
 
@@ -191,42 +127,6 @@ namespace EpicGames.Horde.Storage
 			}
 		}
 
-		/// <summary>
-		/// Implementation of <see cref="ITreeNodeReader"/>
-		/// </summary>
-		class NodeReader : MemoryReader, ITreeNodeReader
-		{
-			readonly IReadOnlyList<NodeLocator> _refs;
-			readonly IoHash _hash;
-			readonly NodeType _type;
-			readonly int _length;
-
-			int _refIdx;
-
-			public NodeReader(ReadOnlyMemory<byte> data, IoHash hash, IReadOnlyList<NodeLocator> refs, NodeType type)
-				: base(data)
-			{
-				_refs = refs;
-				_hash = hash;
-				_type = type;
-				_length = data.Length;
-			}
-
-			public int Version => _type.Version;
-
-			public int Length => _length;
-
-			public IoHash Hash => _hash;
-
-			public IReadOnlyList<NodeLocator> References => _refs;
-
-			public NodeHandle ReadNodeHandle()
-			{
-				IoHash hash = this.ReadIoHash();
-				return new NodeHandle(hash, _refs[_refIdx++]);
-			}
-		}
-
 		// Size of data to fetch by default. This is larger than the minimum request size to reduce number of reads.
 		const int DefaultFetchSize = 15 * 1024 * 1024;
 
@@ -235,7 +135,6 @@ namespace EpicGames.Horde.Storage
 
 		readonly IStorageClient _store;
 		readonly IMemoryCache? _cache;
-		readonly Dictionary<Guid, TypeInfo> _types;
 		readonly ILogger _logger;
 
 		readonly object _queueLock = new object();
@@ -244,9 +143,6 @@ namespace EpicGames.Horde.Storage
 		readonly Dictionary<string, Task<ReadOnlyMemory<byte>>> _decodeTasks = new Dictionary<string, Task<ReadOnlyMemory<byte>>>(StringComparer.Ordinal);
 		Task? _readTask;
 
-		// Default options object, if unspecified
-		static readonly TreeReaderOptions s_defaultOptions = new TreeReaderOptions();
-
 		/// <summary>
 		/// Constructor
 		/// </summary>
@@ -254,22 +150,9 @@ namespace EpicGames.Horde.Storage
 		/// <param name="cache">Cache for data</param>
 		/// <param name="logger">Logger for output</param>
 		public TreeReader(IStorageClient store, IMemoryCache? cache, ILogger logger)
-			: this(store, cache, s_defaultOptions, logger)
-		{
-		}
-
-		/// <summary>
-		/// Constructor
-		/// </summary>
-		/// <param name="store"></param>
-		/// <param name="cache">Cache for data</param>
-		/// <param name="options"></param>
-		/// <param name="logger">Logger for output</param>
-		public TreeReader(IStorageClient store, IMemoryCache? cache, TreeReaderOptions options, ILogger logger)
 		{
 			_store = store;
 			_cache = cache;
-			_types = options.Types.Select(x => TypeInfo.Create(x)).ToDictionary(x => x.BundleType.Guid, x => x);
 			_logger = logger;
 		}
 
@@ -419,15 +302,8 @@ namespace EpicGames.Horde.Storage
 					// Parse the header and construct the bundle info from it
 					BundleHeader header = BundleHeader.Read(memory.ToArray());
 
-					// Get the types within this bundle
-					TypeInfo?[] types = new TypeInfo?[header.Types.Count];
-					for (int idx = 0; idx < header.Types.Count; idx++)
-					{
-						_types.TryGetValue(header.Types[idx].Guid, out types[idx]);
-					}
-
 					// Construct the bundle info
-					BundleInfo bundleInfo = new BundleInfo(queuedHeader.Blob, header, headerSize, types);
+					BundleInfo bundleInfo = new BundleInfo(queuedHeader.Blob, header, headerSize);
 
 					if (_cache != null)
 					{
@@ -720,25 +596,7 @@ namespace EpicGames.Horde.Storage
 		/// <param name="locator">Locator for the node</param>
 		/// <param name="cancellationToken">Cancellation token for the operation</param>
 		/// <returns>Node data read from the given bundle</returns>
-		public async ValueTask<Node> ReadNodeAsync(NodeLocator locator, CancellationToken cancellationToken = default)
-		{
-			Node? treeNode = null;
-			await ReadNodeAsync(locator, (ReadNodeAsyncCallbackParams parms, ITreeNodeReader reader) =>
-			{
-				treeNode = parms.Deserialize(reader);
-				return Task.CompletedTask;
-			}, cancellationToken);
-			return treeNode!;
-		}
-
-		/// <summary>
-		/// Reads a node from a bundle
-		/// </summary>
-		/// <param name="locator">Locator for the node</param>
-		/// <param name="readFunc">Action to be invoked</param>
-		/// <param name="cancellationToken">Cancellation token for the operation</param>
-		/// <returns>Node data read from the given bundle</returns>
-		public async ValueTask ReadNodeAsync(NodeLocator locator, Func<ReadNodeAsyncCallbackParams /*parms*/, ITreeNodeReader /*reader*/, Task> readFunc, CancellationToken cancellationToken = default)
+		public async ValueTask<NodeData> ReadNodeDataAsync(NodeLocator locator, CancellationToken cancellationToken = default)
 		{
 			BundleInfo bundleInfo = await GetBundleInfoAsync(locator.Blob, cancellationToken);
 			BundleExport export = bundleInfo.Header.Exports[locator.ExportIdx];
@@ -766,13 +624,8 @@ namespace EpicGames.Horde.Storage
 				nodeData = packetData.Slice(export.Offset, export.Length);
 			}
 
-			TypeInfo? typeInfo = bundleInfo.Types[export.TypeIdx];
-			if (typeInfo == null)
-			{
-				throw new InvalidOperationException($"No registered serializer for type {bundleInfo.Header.Types[export.TypeIdx].Guid}");
-			}
-
-			await readFunc(new ReadNodeAsyncCallbackParams(typeInfo.Type, typeInfo.Deserialize), new NodeReader(nodeData, export.Hash, refs, typeInfo.BundleType));
+			NodeType nodeType = bundleInfo.Header.Types[export.TypeIdx];
+			return new NodeData(nodeType, export.Hash, nodeData, refs);
 		}
 
 		/// <summary>
@@ -781,11 +634,20 @@ namespace EpicGames.Horde.Storage
 		/// <param name="locator">Locator for the node</param>
 		/// <param name="cancellationToken">Cancellation token for the operation</param>
 		/// <returns>Node data read from the given bundle</returns>
-		public async ValueTask<TNode> ReadNodeAsync<TNode>(NodeLocator locator, CancellationToken cancellationToken = default) where TNode : Node
+		public async ValueTask<Node> ReadNodeAsync(NodeLocator locator, CancellationToken cancellationToken = default)
 		{
-			return (TNode)await ReadNodeAsync(locator, cancellationToken);
+			NodeData nodeData = await ReadNodeDataAsync(locator, cancellationToken);
+			return Node.Deserialize(nodeData);
 		}
 
+		/// <summary>
+		/// Reads a node from a bundle
+		/// </summary>
+		/// <param name="locator">Locator for the node</param>
+		/// <param name="cancellationToken">Cancellation token for the operation</param>
+		/// <returns>Node data read from the given bundle</returns>
+		public async ValueTask<TNode> ReadNodeAsync<TNode>(NodeLocator locator, CancellationToken cancellationToken = default) where TNode : Node => (TNode)await ReadNodeAsync(locator, cancellationToken);
+
 		/// <summary>
 		/// Reads data for a ref from the store, along with the node's contents.
 		/// </summary>
