Managed p/invoke libray for box rather than use a managed c++ library

#rnx

diff --git a/Engine/Source/Programs/Shared/EpicGames.Box/EpicGames.Box.csproj b/Engine/Source/Programs/Shared/EpicGames.Box/EpicGames.Box.csproj
new file mode 100644
--- /dev/null
+++ b/Engine/Source/Programs/Shared/EpicGames.Box/EpicGames.Box.csproj
@@ -0,0 +1,71 @@
+﻿<Project Sdk="Microsoft.NET.Sdk">
+  <Import Project="..\UnrealEngine.csproj.props" />
+
+  <PropertyGroup>
+    <TargetFramework>net6.0</TargetFramework>
+    <Configuration Condition=" '$(Configuration)' == '' ">Development</Configuration>
+    <OutputType>Library</OutputType>
+    <RootNamespace>EpicGames.Box</RootNamespace>
+    <GenerateAssemblyInfo>false</GenerateAssemblyInfo>
+    <GenerateTargetFrameworkAttribute>false</GenerateTargetFrameworkAttribute>
+    <Configurations>Debug;Release;Development</Configurations>
+    <Nullable>enable</Nullable>
+    <AllowUnsafeBlocks>true</AllowUnsafeBlocks>
+    <GenerateDocumentationFile>True</GenerateDocumentationFile>
+    <DebugType>pdbonly</DebugType>
+    <EnforceCodeStyleInBuild>False</EnforceCodeStyleInBuild>
+    <ProduceReferenceAssemblyInOutDir>true</ProduceReferenceAssemblyInOutDir>
+    <RunAnalyzersDuringBuild>false</RunAnalyzersDuringBuild>
+  </PropertyGroup>
+  <PropertyGroup Condition=" '$(Configuration)|$(Platform)' == 'Debug|AnyCPU' ">
+    <DocumentationFile></DocumentationFile>
+  </PropertyGroup>
+  <PropertyGroup Condition="'$(Configuration)|$(Platform)' == 'Development|AnyCPU'">
+    <DocumentationFile></DocumentationFile>
+    <Optimize>true</Optimize>
+  </PropertyGroup>
+  <PropertyGroup Condition="'$(Configuration)|$(Platform)' == 'Release|AnyCPU'">
+    <DocumentationFile></DocumentationFile>
+    <Optimize>true</Optimize>
+  </PropertyGroup>
+
+  <ItemGroup>
+    <Compile Include="..\MetaData.cs">
+      <Link>Properties\MetaData.cs</Link>
+    </Compile>
+  </ItemGroup>
+
+  <ItemGroup>
+    <Compile Include="..\..\..\..\Restricted\NotForLicensees\Source\Programs\Shared\EpicGames.Box\*.cs">
+      <Link>$([System.Text.RegularExpressions.Regex]::Replace(%(FullPath), '^.+?[\\/]EpicGames.Box[\\/]', ''))</Link>
+    </Compile>
+  </ItemGroup>
+
+  <Choose>
+    <When Condition="'$([System.Runtime.InteropServices.RuntimeInformation]::IsOSPlatform($([System.Runtime.InteropServices.OSPlatform]::Windows)))' == 'true'">
+      <ItemGroup Condition="Exists('..\..\..\..\Restricted\NotForLicensees\Binaries\Win64\Box\BoxExport.dll')">
+        <Content Include="..\..\..\..\Restricted\NotForLicensees\Binaries\Win64\Box\BoxExport.dll">
+          <CopyToOutputDirectory>PreserveNewest</CopyToOutputDirectory>
+        </Content>
+        <Content Include="..\..\..\..\Restricted\NotForLicensees\Binaries\Win64\Box\BoxDetours.dll">
+          <CopyToOutputDirectory>PreserveNewest</CopyToOutputDirectory>
+        </Content>
+        <Content Include="..\..\..\..\Restricted\NotForLicensees\Binaries\Win64\Box\BoxVisualizer.exe">
+          <CopyToOutputDirectory>PreserveNewest</CopyToOutputDirectory>
+        </Content>
+        <Content Include="..\..\..\..\Restricted\NotForLicensees\Binaries\Win64\Box\Ijwhost.dll">
+          <CopyToOutputDirectory>PreserveNewest</CopyToOutputDirectory>
+        </Content>
+      </ItemGroup>
+    </When>
+  </Choose>
+
+  <ItemGroup>
+    <ProjectReference Include="..\..\Shared\EpicGames.Core\EpicGames.Core.csproj" />
+  </ItemGroup>
+  <ItemGroup>
+    <PackageReference Include="Microsoft.CSharp" Version="4.7.0" />
+    <PackageReference Include="Microsoft.Win32.Registry" Version="5.0.0" />
+    <PackageReference Include="System.Data.DataSetExtensions" Version="4.5.0" />
+  </ItemGroup>
+</Project>
\ No newline at end of file
diff --git a/Engine/Source/Programs/Shared/EpicGames.Box/Utils.cs b/Engine/Source/Programs/Shared/EpicGames.Box/Utils.cs
new file mode 100644
--- /dev/null
+++ b/Engine/Source/Programs/Shared/EpicGames.Box/Utils.cs
@@ -0,0 +1,19 @@
+﻿// Copyright Epic Games, Inc. All Rights Reserved.
+
+using System;
+
+namespace EpicGames.Box
+{
+	/// <summary>
+	/// Box utils
+	/// </summary>
+	public static partial class Utils
+	{
+		static Lazy<bool> Available { get; set; } = new Lazy<bool>(() => false);
+
+		/// <summary>
+		/// Is box available?
+		/// </summary>
+		public static bool IsAvailable() => Available.Value;
+	}
+}
diff --git a/Engine/Source/Programs/UnrealBuildTool/System/ActionGraph.cs b/Engine/Source/Programs/UnrealBuildTool/System/ActionGraph.cs
--- a/Engine/Source/Programs/UnrealBuildTool/System/ActionGraph.cs
+++ b/Engine/Source/Programs/UnrealBuildTool/System/ActionGraph.cs
@@ -356,7 +356,7 @@ namespace UnrealBuildTool
 		private static ActionExecutor SelectExecutor(BuildConfiguration BuildConfiguration, int ActionCount, List<TargetDescriptor> TargetDescriptors, ILogger Logger)
 		{
 #if __BOXEXECUTOR_AVAILABLE__
-			if (BuildConfiguration.bAllowBoxExecutor)
+			if (BuildConfiguration.bAllowBoxExecutor && BoxExecutor.IsAvailable(Logger))
 			{
 				return new BoxExecutor(BuildConfiguration.MaxParallelActions, BuildConfiguration.bAllCores, BuildConfiguration.bCompactOutput, Logger);
 			}
diff --git a/Engine/Source/Programs/UnrealBuildTool/UnrealBuildTool.cs b/Engine/Source/Programs/UnrealBuildTool/UnrealBuildTool.cs
--- a/Engine/Source/Programs/UnrealBuildTool/UnrealBuildTool.cs
+++ b/Engine/Source/Programs/UnrealBuildTool/UnrealBuildTool.cs
@@ -490,6 +490,13 @@ namespace UnrealBuildTool
 				// Ensure we can resolve any external assemblies that are not in the same folder as our assembly.
 				AssemblyUtils.InstallAssemblyResolver(Path.GetDirectoryName(Assembly.GetEntryAssembly()!.GetOriginalLocation())!);
 
+				// Add the application directory to PATH
+				string path = Environment.GetEnvironmentVariable("PATH") ?? string.Empty;
+				if (!path.Split(';').Any(x => string.Compare(x, Unreal.UnrealBuildToolDllPath.Directory.FullName, true) == 0))
+				{
+					Environment.SetEnvironmentVariable("PATH", $"{path};{Unreal.UnrealBuildToolDllPath.Directory.FullName}");
+				}
+
 				// Change the working directory to be the Engine/Source folder. We are likely running from Engine/Binaries/DotNET
 				// This is critical to be done early so any code that relies on the current directory being Engine/Source will work.
 				DirectoryReference.CreateDirectory(Unreal.EngineSourceDirectory);
diff --git a/Engine/Source/Programs/UnrealBuildTool/UnrealBuildTool.csproj b/Engine/Source/Programs/UnrealBuildTool/UnrealBuildTool.csproj
--- a/Engine/Source/Programs/UnrealBuildTool/UnrealBuildTool.csproj
+++ b/Engine/Source/Programs/UnrealBuildTool/UnrealBuildTool.csproj
@@ -73,21 +73,11 @@
 
   <Choose>
     <When Condition="'$([System.Runtime.InteropServices.RuntimeInformation]::IsOSPlatform($([System.Runtime.InteropServices.OSPlatform]::Windows)))' == 'true' And '$([System.Runtime.InteropServices.RuntimeInformation]::OSArchitecture)' != 'Arm64'">
-      <ItemGroup Condition="Exists('..\..\..\Restricted\NotForLicensees\Binaries\Win64\Box\BoxManaged.dll')">
-        <Reference Include="..\..\..\Restricted\NotForLicensees\Binaries\Win64\Box\BoxManaged.dll" />
-        <Content Include="..\..\..\Restricted\NotForLicensees\Binaries\Win64\Box\BoxDetours.dll">
-          <CopyToOutputDirectory>PreserveNewest</CopyToOutputDirectory>
-        </Content>
-        <Content Include="..\..\..\Restricted\NotForLicensees\Binaries\Win64\Box\BoxVisualizer.exe">
-          <CopyToOutputDirectory>PreserveNewest</CopyToOutputDirectory>
-        </Content>
-        <Content Include="..\..\..\Restricted\NotForLicensees\Binaries\Win64\Box\Ijwhost.dll">
-          <CopyToOutputDirectory>PreserveNewest</CopyToOutputDirectory>
-        </Content>
-      </ItemGroup>
-      <PropertyGroup Condition="Exists('..\..\..\Restricted\NotForLicensees\Binaries\Win64\Box\BoxManaged.dll')">
+      <PropertyGroup Condition="Exists('..\..\..\Restricted\NotForLicensees\Source\Programs\UnrealBuildTool\Executors\Experimental\BoxExecutor.cs')">
         <DefineConstants>$(DefineConstants);__BOXEXECUTOR_AVAILABLE__</DefineConstants>
-        <ResolveAssemblyWarnOrErrorOnTargetArchitectureMismatch>None</ResolveAssemblyWarnOrErrorOnTargetArchitectureMismatch>
+      </PropertyGroup>
+      <PropertyGroup Condition="Exists('..\..\..\Restricted\NotForLicensees\Source\Programs\Shared\EpicGames.Box\BoxServer.cs')">
+        <DefineConstants>$(DefineConstants);__BOX_AVAILABLE__</DefineConstants>
       </PropertyGroup>
     </When>
   </Choose>
@@ -124,6 +114,7 @@
 
   <ItemGroup>
     <!-- If you add a project reference, BuildUAT.bat/.sh and BuildUBT.bat/.sh need to be updated to also test for those projects.-->
+    <ProjectReference Include="..\Shared\EpicGames.Box\EpicGames.Box.csproj" />
     <ProjectReference Include="..\Shared\EpicGames.Build\EpicGames.Build.csproj" />
     <ProjectReference Include="..\Shared\EpicGames.Core\EpicGames.Core.csproj" />
     <ProjectReference Include="..\Shared\EpicGames.Horde\EpicGames.Horde.csproj" />
