[Box]
* New binaries with bug fixes
* Removed BoxManaged code now when ubt is using pinvoke instead

diff --git a/Engine/Source/Programs/Shared/EpicGames.Box/EpicGames.Box.csproj b/Engine/Source/Programs/Shared/EpicGames.Box/EpicGames.Box.csproj
--- a/Engine/Source/Programs/Shared/EpicGames.Box/EpicGames.Box.csproj
+++ b/Engine/Source/Programs/Shared/EpicGames.Box/EpicGames.Box.csproj
@@ -53,9 +53,6 @@
         <Content Include="..\..\..\..\Restricted\NotForLicensees\Binaries\Win64\Box\BoxVisualizer.exe">
           <CopyToOutputDirectory>PreserveNewest</CopyToOutputDirectory>
         </Content>
-        <Content Include="..\..\..\..\Restricted\NotForLicensees\Binaries\Win64\Box\Ijwhost.dll">
-          <CopyToOutputDirectory>PreserveNewest</CopyToOutputDirectory>
-        </Content>
       </ItemGroup>
     </When>
   </Choose>
