Horde: Remove generic storage client interface. We don't currently use this.

diff --git a/Engine/Source/Programs/Shared/EpicGames.Horde/Storage/IStorageClient.cs b/Engine/Source/Programs/Shared/EpicGames.Horde/Storage/IStorageClient.cs
--- a/Engine/Source/Programs/Shared/EpicGames.Horde/Storage/IStorageClient.cs
+++ b/Engine/Source/Programs/Shared/EpicGames.Horde/Storage/IStorageClient.cs
@@ -180,13 +180,6 @@ namespace EpicGames.Horde.Storage
 		#endregion
 	}
 
-	/// <summary>
-	/// Typed implementation of <see cref="IStorageClient"/> for use with dependency injection
-	/// </summary>
-	public interface IStorageClient<T> : IStorageClient
-	{
-	}
-
 	/// <summary>
 	/// Allows creating storage clients for different namespaces
 	/// </summary>
@@ -206,62 +199,6 @@ namespace EpicGames.Horde.Storage
 	/// </summary>
 	public static class StorageClientExtensions
 	{
-		class TypedStorageClient<T> : IStorageClient<T>
-		{
-			readonly IStorageClient _inner;
-
-			public TypedStorageClient(IStorageClient inner) => _inner = inner;
-
-			#region Blobs
-
-			/// <inheritdoc/>
-			public Task<Stream> ReadBlobAsync(BlobLocator locator, CancellationToken cancellationToken = default) => _inner.ReadBlobAsync(locator, cancellationToken);
-
-			/// <inheritdoc/>
-			public Task<Stream> ReadBlobRangeAsync(BlobLocator locator, int offset, int length, CancellationToken cancellationToken = default) => _inner.ReadBlobRangeAsync(locator, offset, length, cancellationToken);
-
-			/// <inheritdoc/>
-			public Task<BlobLocator> WriteBlobAsync(Stream stream, Utf8String prefix = default, CancellationToken cancellationToken = default) => _inner.WriteBlobAsync(stream, prefix, cancellationToken);
-
-			#endregion
-
-			#region Aliases
-
-			/// <inheritdoc/>
-			public Task AddAliasAsync(Utf8String name, NodeHandle handle, CancellationToken cancellationToken = default) => _inner.AddAliasAsync(name, handle, cancellationToken);
-
-			/// <inheritdoc/>
-			public Task RemoveAliasAsync(Utf8String name, NodeHandle handle, CancellationToken cancellationToken = default) => _inner.RemoveAliasAsync(name, handle, cancellationToken);
-
-			/// <inheritdoc/>
-			public IAsyncEnumerable<NodeHandle> FindNodesAsync(Utf8String name, CancellationToken cancellationToken = default) => _inner.FindNodesAsync(name, cancellationToken);
-
-			#endregion
-
-			#region Refs
-
-			/// <inheritdoc/>
-			public Task DeleteRefAsync(RefName name, CancellationToken cancellationToken = default) => _inner.DeleteRefAsync(name, cancellationToken);
-
-			/// <inheritdoc/>
-			public Task<NodeHandle?> TryReadRefTargetAsync(RefName name, DateTime cacheTime = default, CancellationToken cancellationToken = default) => _inner.TryReadRefTargetAsync(name, cacheTime, cancellationToken);
-
-			/// <inheritdoc/>
-			public Task<NodeHandle> WriteRefAsync(RefName name, Bundle bundle, int exportIdx = 0, Utf8String prefix = default, RefOptions? options = null, CancellationToken cancellationToken = default) => _inner.WriteRefAsync(name, bundle, exportIdx, prefix, options, cancellationToken);
-
-			/// <inheritdoc/>
-			public Task WriteRefTargetAsync(RefName name, NodeHandle target, RefOptions? options = null, CancellationToken cancellationToken = default) => _inner.WriteRefTargetAsync(name, target, options, cancellationToken);
-
-			#endregion
-		}
-
-		/// <summary>
-		/// Wraps a <see cref="IStorageClient"/> interface with a type argument
-		/// </summary>
-		/// <param name="blobStore">Regular blob store instance</param>
-		/// <returns></returns>
-		public static IStorageClient<T> ForType<T>(this IStorageClient blobStore) => new TypedStorageClient<T>(blobStore);
-
 		#region Bundles
 
 		/// <summary>
