EpicGames.Perforce: Add cache tests for ManagedWorkspace

Also fix mismatched digests in the fixture.

diff --git a/Engine/Source/Programs/Shared/EpicGames.Perforce.Fixture/PerforceFixture.cs b/Engine/Source/Programs/Shared/EpicGames.Perforce.Fixture/PerforceFixture.cs
--- a/Engine/Source/Programs/Shared/EpicGames.Perforce.Fixture/PerforceFixture.cs
+++ b/Engine/Source/Programs/Shared/EpicGames.Perforce.Fixture/PerforceFixture.cs
@@ -29,7 +29,7 @@ public class DepotFileFixture
 		ClientFile = clientFile;
 		Size = content.Length;
 		Revision = revision;
-		Digest = PerforceFixture.CalcMd5(content);
+		Digest = PerforceFixture.CalcMd5(content).ToUpperInvariant();
 		Content = content;
 	}
 }
@@ -145,9 +145,9 @@ public class ChangelistFixture
 				
 				// Since content of files in fixture are only single lines, the trick below works to workaround
 				// differences in line endings after sync (client vs server).
-				content = content.Replace("\r\n", "", StringComparison.Ordinal);
+				content = content.Replace("\r\n", "\n", StringComparison.Ordinal);
 				long size = content.Length;
-				return (clientFile, size, PerforceFixture.CalcMd5(content));
+				return (clientFile, size, PerforceFixture.CalcMd5(content).ToUpperInvariant());
 			}));
 		
 		return (localFiles, streamFiles);
@@ -167,6 +167,11 @@ public class ChangelistFixture
 		
 		return (localFiles, streamFiles);
 	}
+
+	public DepotFileFixture? GetFile(string depotPath)
+	{
+		return new List<DepotFileFixture>(StreamFiles).Find(x => x.DepotFile == depotPath);
+	}
 }
 
 public class StreamFixture
@@ -208,62 +213,62 @@ public class PerforceFixture
 			new(2, "Initial import",
 				new List<DepotFileFixture>()
 				{
-					new("//Foo/Main/main.cpp", "main.cpp", "This is change main.cpp #1", 1),
-					new("//Foo/Main/main.h", "main.h", "This is change main.h #1", 1),
-					new("//Foo/Main/common.h", "common.h", "This is change common.h #1", 1),
-					new("//Foo/Main/unused.cpp", "unused.cpp", "This is change unused.cpp #1", 1),
-					new("//Foo/Main/Data/data.txt", "Data/data.txt", "This is change data.txt #1", 1),
+					new("//Foo/Main/main.cpp", "main.cpp", "This is change main.cpp #1\n", 1),
+					new("//Foo/Main/main.h", "main.h", "This is change main.h #1\n", 1),
+					new("//Foo/Main/common.h", "common.h", "This is change common.h #1\n", 1),
+					new("//Foo/Main/unused.cpp", "unused.cpp", "This is change unused.cpp #1\n", 1),
+					new("//Foo/Main/Data/data.txt", "Data/data.txt", "This is change data.txt #1\n", 1),
 				}),
 			new(3, "Improvement to main.cpp",
 				new List<DepotFileFixture>()
 				{
-					new("//Foo/Main/main.cpp", "main.cpp", "This is change main.cpp #2", 2),
-					new("//Foo/Main/main.h", "main.h", "This is change main.h #1", 1),
-					new("//Foo/Main/common.h", "common.h", "This is change common.h #1", 1),
-					new("//Foo/Main/unused.cpp", "unused.cpp", "This is change unused.cpp #1", 1),
-					new("//Foo/Main/Data/data.txt", "Data/data.txt", "This is change data.txt #1", 1),
+					new("//Foo/Main/main.cpp", "main.cpp", "This is change main.cpp #2\n", 2),
+					new("//Foo/Main/main.h", "main.h", "This is change main.h #1\n", 1),
+					new("//Foo/Main/common.h", "common.h", "This is change common.h #1\n", 1),
+					new("//Foo/Main/unused.cpp", "unused.cpp", "This is change unused.cpp #1\n", 1),
+					new("//Foo/Main/Data/data.txt", "Data/data.txt", "This is change data.txt #1\n", 1),
 				}),
 			new(4, "Delete an unused file",
 				new List<DepotFileFixture>()
 				{
-					new("//Foo/Main/main.cpp", "main.cpp", "This is change main.cpp #2", 2),
-					new("//Foo/Main/main.h", "main.h", "This is change main.h #1", 1),
-					new("//Foo/Main/common.h", "common.h", "This is change common.h #1", 1),
-					new("//Foo/Main/Data/data.txt", "Data/data.txt", "This is change data.txt #1", 1),
+					new("//Foo/Main/main.cpp", "main.cpp", "This is change main.cpp #2\n", 2),
+					new("//Foo/Main/main.h", "main.h", "This is change main.h #1\n", 1),
+					new("//Foo/Main/common.h", "common.h", "This is change common.h #1\n", 1),
+					new("//Foo/Main/Data/data.txt", "Data/data.txt", "This is change data.txt #1\n", 1),
 				}),
 			new(5, "Rename common.h to shared.h",
 				new List<DepotFileFixture>()
 				{
-					new("//Foo/Main/main.cpp", "main.cpp", "This is change main.cpp #2", 2),
-					new("//Foo/Main/main.h", "main.h", "This is change main.h #1", 1),
-					new("//Foo/Main/shared.h", "shared.h", "This is change common.h #1", 1),
-					new("//Foo/Main/Data/data.txt", "Data/data.txt", "This is change data.txt #1", 1),
+					new("//Foo/Main/main.cpp", "main.cpp", "This is change main.cpp #2\n", 2),
+					new("//Foo/Main/main.h", "main.h", "This is change main.h #1\n", 1),
+					new("//Foo/Main/shared.h", "shared.h", "This is change common.h #1\n", 1),
+					new("//Foo/Main/Data/data.txt", "Data/data.txt", "This is change data.txt #1\n", 1),
 				}),
 			new(6, "Some updates to main",
 				new List<DepotFileFixture>()
 				{
-					new("//Foo/Main/main.cpp", "main.cpp", "This is change main.cpp #3", 3),
-					new("//Foo/Main/main.h", "main.h", "This is change main.h #2", 2),
-					new("//Foo/Main/shared.h", "shared.h", "This is change common.h #1", 1),
-					new("//Foo/Main/Data/data.txt", "Data/data.txt", "This is change data.txt #1", 1),
+					new("//Foo/Main/main.cpp", "main.cpp", "This is change main.cpp #3\n", 3),
+					new("//Foo/Main/main.h", "main.h", "This is change main.h #2\n", 2),
+					new("//Foo/Main/shared.h", "shared.h", "This is change common.h #1\n", 1),
+					new("//Foo/Main/Data/data.txt", "Data/data.txt", "This is change data.txt #1\n", 1),
 				}),
 			new(7, "Add more data",
 				new List<DepotFileFixture>()
 				{
-					new("//Foo/Main/main.cpp", "main.cpp", "This is change main.cpp #3", 3),
-					new("//Foo/Main/main.h", "main.h", "This is change main.h #2", 2),
-					new("//Foo/Main/shared.h", "shared.h", "This is change common.h #1", 1),
-					new("//Foo/Main/Data/data.txt", "Data/data.txt", "This is change data.txt #1", 1),
-					new("//Foo/Main/Data/moredata.txt", "Data/moredata.txt", "This is change moredata.txt #1", 1),
+					new("//Foo/Main/main.cpp", "main.cpp", "This is change main.cpp #3\n", 3),
+					new("//Foo/Main/main.h", "main.h", "This is change main.h #2\n", 2),
+					new("//Foo/Main/shared.h", "shared.h", "This is change common.h #1\n", 1),
+					new("//Foo/Main/Data/data.txt", "Data/data.txt", "This is change data.txt #1\n", 1),
+					new("//Foo/Main/Data/moredata.txt", "Data/moredata.txt", "This is change moredata.txt #1\n", 1),
 				}),
 			new(8, "A shelved CL", // Assumes base CL is 7
 				new List<DepotFileFixture>()
 				{
-					new("//Foo/Main/main.h", "main.h", "This is change main.h #3", 3),
-					new("//Foo/Main/shared.h", "shared.h", "This is change common.h #1", 1),
-					new("//Foo/Main/shelved.cpp", "shelved.cpp", "This is change shelved.cpp #1", 1),
-					new("//Foo/Main/Data/data.txt", "Data/data.txt", "This is change data.txt #1", 1),
-					new("//Foo/Main/Data/moredata.txt", "Data/moredata.txt", "This is change moredata.txt #1", 1),
+					new("//Foo/Main/main.h", "main.h", "This is change main.h #3\n", 3),
+					new("//Foo/Main/shared.h", "shared.h", "This is change common.h #1\n", 1),
+					new("//Foo/Main/shelved.cpp", "shelved.cpp", "This is change shelved.cpp #1\n", 1),
+					new("//Foo/Main/Data/data.txt", "Data/data.txt", "This is change data.txt #1\n", 1),
+					new("//Foo/Main/Data/moredata.txt", "Data/moredata.txt", "This is change moredata.txt #1\n", 1),
 				}, true),
 		});
 
@@ -272,4 +277,66 @@ public class PerforceFixture
 		byte[] data = Encoding.ASCII.GetBytes(content);
 		return Md5Hash.Compute(data).ToString();
 	}
+	
+	/// <summary>
+	/// Assert cache contains exactly the provided digests
+	/// </summary>
+	/// <param name="cacheDir">Directory path to cache</param>
+	/// <param name="digests">Digest expected to be in the cache</param>
+	public static void AssertCacheEquals(string cacheDir, params string[] digests)
+	{
+		HashSet<(string cacheId, string digest, long fileSize)> cachedFiles = GetCachedFiles(cacheDir);
+
+		HashSet<string> actual = new (cachedFiles.Select(x => x.digest));
+		HashSet<string> expected = new (digests);
+
+		if (!expected.SetEquals(actual))
+		{
+			List<string> expectedList = new(expected);
+			List<string> actualList = new(actual);
+			
+			expectedList.Sort();
+			actualList.Sort();
+			
+			Console.WriteLine("Expected ------------------------------------------------------");
+			foreach (string digest in expectedList)
+			{
+				Console.WriteLine($"{digest}");
+			}
+			
+			Console.WriteLine("");
+			Console.WriteLine("Actual --------------------------------------------------------");
+			foreach (string digest in actual)
+			{
+				Console.WriteLine($"{digest}");
+			}
+		
+			Assert.Fail("Digests in cache does not match given list");
+		}
+	}
+	
+	private static HashSet<(string cacheId, string digest, long fileSize)> GetCachedFiles(string cacheDir)
+	{
+		EnumerationOptions options = new () { RecurseSubdirectories = true };
+
+		return new(Directory.EnumerateFiles(cacheDir, "*", options)
+			.Select(x => Path.GetRelativePath(cacheDir, x))
+			.Select(clientFile =>
+			{
+				string absPath = Path.Join(cacheDir, clientFile);
+				string content = File.ReadAllText(absPath);
+				
+				// Since content of files in fixture are only single lines, the trick below works to workaround
+				// differences in line endings after sync (client vs server).
+				content = content.Replace("\r\n", "\n", StringComparison.Ordinal);
+				long size = content.Length;
+				string fileName = Path.GetFileName(clientFile);
+				string digest = CalcMd5(content).ToUpperInvariant();
+				
+				// This is a simplification. The fileName/cache ID can theoretically differ if they collide.
+				// See cache ID implementation in ManagedWorkspace
+				Assert.AreEqual(fileName, digest[..16]);
+				return (fileName, digest, size);
+			}));
+	}
 }
\ No newline at end of file
diff --git a/Engine/Source/Programs/Shared/EpicGames.Perforce.Managed.Tests/ManagedWorkspaceTest.cs b/Engine/Source/Programs/Shared/EpicGames.Perforce.Managed.Tests/ManagedWorkspaceTest.cs
--- a/Engine/Source/Programs/Shared/EpicGames.Perforce.Managed.Tests/ManagedWorkspaceTest.cs
+++ b/Engine/Source/Programs/Shared/EpicGames.Perforce.Managed.Tests/ManagedWorkspaceTest.cs
@@ -14,9 +14,11 @@ using Microsoft.VisualStudio.TestTools.UnitTesting;
 namespace EpicGames.Perforce.Managed.Tests;
 
 [TestClass]
+[System.Diagnostics.CodeAnalysis.SuppressMessage("Naming", "CA1707:Identifiers should not contain underscores")]
 public class ManagedWorkspaceTest : BasePerforceFixtureTest
 {
 	private string SyncDir => Path.Join(TempDir.FullName, "Sync");
+	private string CacheDir => Path.Join(TempDir.FullName, "Cache");
 	private string StreamName => Fixture.StreamFooMain.Root;
 	private StreamFixture Stream => Fixture.StreamFooMain;
 
@@ -193,6 +195,79 @@ public class ManagedWorkspaceTest : BasePerforceFixtureTest
 		await Stream.GetChangelist(7).AssertHaveTableAsync(PerforceConnection, useHaveTable);
 	}
 	
+	[TestMethod]
+	[DataRow(true, DisplayName = "With have-table")]
+	[DataRow(false, DisplayName = "Without have-table")]
+	public async Task Caching_SyncingNewChangelist_UnusedFilesMovedToCache(bool useHaveTable)
+	{
+		// Arrange
+		ManagedWorkspace ws = await GetManagedWorkspace(useHaveTable);
+		await SyncAsync(ws, 7);
+		
+		// Act
+		await SyncAsync(ws, 1);
+
+		// Assert - all files from CL 7 are moved to the cache
+		PerforceFixture.AssertCacheEquals(CacheDir, Stream.GetChangelist(7).StreamFiles.Select(x => x.Digest).ToArray());
+	}
+	
+	[TestMethod]
+	[DataRow(true, DisplayName = "With have-table")]
+	[DataRow(false, DisplayName = "Without have-table")]
+	public async Task Caching_SyncingWithCachedData_FilesPopulatedFromCache(bool useHaveTable)
+	{
+		// Arrange
+		ManagedWorkspace ws = await GetManagedWorkspace(useHaveTable);
+		await SyncAsync(ws, 7);
+		await SyncAsync(ws, 1);
+		
+		// Act
+		await SyncAsync(ws, 7);
+
+		// Assert - cache becomes empty
+		PerforceFixture.AssertCacheEquals(CacheDir, Array.Empty<string>());
+	}
+
+	[TestMethod]
+	public async Task FixtureMetadataMatchesServerMetadata()
+	{
+		FStatOptions options = FStatOptions.IncludeFileSizes;
+
+		string fileSpec = "//Foo/Main/...";
+		List<ChangesRecord> submittedChanges = await PerforceConnection.GetChangesAsync(ChangesOptions.None, -1, ChangeStatus.Submitted, fileSpec);
+		List<ChangesRecord> shelvedChanges = await PerforceConnection.GetChangesAsync(ChangesOptions.None, -1, ChangeStatus.Shelved, fileSpec);
+		foreach (ChangesRecord cr in submittedChanges.Concat(shelvedChanges).OrderBy(x => x.Number))
+		{
+			HashSet<(string clientFile, int rev, long size, string digest)> depotFiles = new();
+			List<FStatRecord> fstatRecords = await PerforceConnection.FStatAsync(options, fileSpec + "@" + cr.Number).ToListAsync();
+			foreach (FStatRecord fsr in fstatRecords)
+			{
+				if (fsr.HeadAction is FileAction.Add or FileAction.MoveAdd or FileAction.Edit)
+				{
+					depotFiles.Add((fsr.DepotFile!, fsr.HeadRevision, fsr.FileSize, fsr.Digest!));	
+				}
+			}
+			
+			HashSet<(string clientFile, int rev, long size, string digest)> fixtureFiles = new();
+			ChangelistFixture changelist = Stream.GetChangelist(cr.Number);
+			if (changelist.IsShelved) continue;
+			
+			foreach (DepotFileFixture depotFile in changelist.StreamFiles)
+			{
+				fixtureFiles.Add((depotFile.DepotFile, depotFile.Revision, depotFile.Size, depotFile.Digest));
+			}
+			
+			Assert.AreEqual(depotFiles.Count, fixtureFiles.Count);
+			foreach ((string clientFile, int rev, long size, string digest) tuple in fixtureFiles)
+			{
+				if (!depotFiles.Contains(tuple))
+				{
+					Assert.Fail("File in fixtures does not exist in depot: " + tuple);	
+				}
+			}
+		}
+	}
+	
 	[TestMethod]
 	public async Task ReusePerforceClientWithoutHaveTable()
 	{
