Horde: More refactoring of storage API to hide bundle implementation details.

- Rename StorageClientBase to BundleStorageClient. The current implementation assumes that bundles will be used for storage.
- Use ReadAsync() methods on NodeHandle to expand content, rather than requiring a TreeReader to be passed around.

diff --git a/Engine/Source/Programs/Horde/Horde.Agent.Tests/JsonRpcLoggerTests.cs b/Engine/Source/Programs/Horde/Horde.Agent.Tests/JsonRpcLoggerTests.cs
--- a/Engine/Source/Programs/Horde/Horde.Agent.Tests/JsonRpcLoggerTests.cs
+++ b/Engine/Source/Programs/Horde/Horde.Agent.Tests/JsonRpcLoggerTests.cs
@@ -166,10 +166,10 @@ namespace Horde.Agent.Tests
 			// Check the index text
 			List<Utf8String> extractedIndexText = new List<Utf8String>();
 
-			LogIndexNode index = await file.IndexRef.ExpandAsync(reader);
+			LogIndexNode index = await file.IndexRef.ExpandAsync();
 			foreach (LogChunkRef block in index.PlainTextChunkRefs)
 			{
-				LogChunkNode text = await block.ExpandAsync(reader);
+				LogChunkNode text = await block.ExpandAsync();
 				extractedIndexText.AddRange(text.Lines);
 			}
 
@@ -184,7 +184,7 @@ namespace Horde.Agent.Tests
 
 			foreach (LogChunkRef blockRef in file.TextChunkRefs)
 			{
-				LogChunkNode blockText = await blockRef.ExpandAsync(reader);
+				LogChunkNode blockText = await blockRef.ExpandAsync();
 				foreach (Utf8String line in blockText.Lines)
 				{
 					LogEvent logEvent = LogEvent.Read(line.Span);
diff --git a/Engine/Source/Programs/Horde/Horde.Agent/Execution/JobExecutor.cs b/Engine/Source/Programs/Horde/Horde.Agent/Execution/JobExecutor.cs
--- a/Engine/Source/Programs/Horde/Horde.Agent/Execution/JobExecutor.cs
+++ b/Engine/Source/Programs/Horde/Horde.Agent/Execution/JobExecutor.cs
@@ -811,7 +811,7 @@ namespace Horde.Agent.Execution
 				string nodeName = input.Substring(0, slashIdx);
 				string tagName = input.Substring(slashIdx + 1);
 
-				TempStorageTagManifest fileList = await TempStorage.RetrieveTagAsync(reader, _storagePrefix, nodeName, tagName, manifestDir, logger, cancellationToken);
+				TempStorageTagManifest fileList = await TempStorage.RetrieveTagAsync(storage, _storagePrefix, nodeName, tagName, manifestDir, logger, cancellationToken);
 				tagNameToFileSet[tagName] = fileList.ToFileSet(workspaceDir);
 				inputStorageBlocks.UnionWith(fileList.Blocks);
 			}
@@ -824,7 +824,7 @@ namespace Horde.Agent.Execution
 				scope.Span.SetTag("blocks", inputStorageBlocks.Count);
 				foreach (TempStorageBlockRef inputStorageBlock in inputStorageBlocks)
 				{
-					TempStorageBlockManifest manifest = await TempStorage.RetrieveBlockAsync(reader, _storagePrefix, inputStorageBlock.NodeName, inputStorageBlock.OutputName, workspaceDir, manifestDir, logger, cancellationToken);
+					TempStorageBlockManifest manifest = await TempStorage.RetrieveBlockAsync(storage, _storagePrefix, inputStorageBlock.NodeName, inputStorageBlock.OutputName, workspaceDir, manifestDir, logger, cancellationToken);
 					inputManifests[inputStorageBlock] = manifest;
 				}
 				scope.Span.SetTag("size", inputManifests.Sum(x => x.Value.GetTotalSize()));
diff --git a/Engine/Source/Programs/Horde/Horde.Agent/Utility/TempStorage.cs b/Engine/Source/Programs/Horde/Horde.Agent/Utility/TempStorage.cs
--- a/Engine/Source/Programs/Horde/Horde.Agent/Utility/TempStorage.cs
+++ b/Engine/Source/Programs/Horde/Horde.Agent/Utility/TempStorage.cs
@@ -549,7 +549,7 @@ namespace Horde.Storage.Utility
 		/// <summary>
 		/// Reads a set of tagged files from disk
 		/// </summary>
-		/// <param name="reader">Reader for node data</param>
+		/// <param name="storageClient">Reader for node data</param>
 		/// <param name="refPrefix">Prefix for ref names</param>
 		/// <param name="nodeName">Name of the node which produced the tag set</param>
 		/// <param name="tagName">Name of the tag, with a '#' prefix</param>
@@ -557,7 +557,7 @@ namespace Horde.Storage.Utility
 		/// <param name="logger">Logger for output</param>
 		/// <param name="cancellationToken"></param>
 		/// <returns>The set of files</returns>
-		public static async Task<TempStorageTagManifest> RetrieveTagAsync(TreeReader reader, string refPrefix, string nodeName, string tagName, DirectoryReference manifestDir, ILogger logger, CancellationToken cancellationToken)
+		public static async Task<TempStorageTagManifest> RetrieveTagAsync(IStorageClient storageClient, string refPrefix, string nodeName, string tagName, DirectoryReference manifestDir, ILogger logger, CancellationToken cancellationToken)
 		{
 			// Try to read the tag set from the local directory
 			FileReference localFileListLocation = GetTagManifestLocation(manifestDir, nodeName, tagName);
@@ -570,11 +570,11 @@ namespace Horde.Storage.Utility
 				RefName refName = GetRefNameForNode(refPrefix, nodeName);
 				logger.LogInformation("Reading node \"{NodeName}\" tag \"{TagName}\" from temp storage (ref: {RefName}, localFile: {LocalFile})", nodeName, tagName, refName, localFileListLocation);
 
-				DirectoryNode node = await reader.ReadNodeAsync<DirectoryNode>(refName, cancellationToken: cancellationToken);
+				DirectoryNode node = await storageClient.ReadNodeAsync<DirectoryNode>(refName, cancellationToken: cancellationToken);
 
 				FileEntry fileEntry = node.GetFileEntry(localFileListLocation.GetFileName());
 				DirectoryReference.CreateDirectory(localFileListLocation.Directory);
-				await fileEntry.CopyToFileAsync(reader, localFileListLocation.ToFileInfo(), cancellationToken);
+				await fileEntry.CopyToFileAsync(localFileListLocation.ToFileInfo(), cancellationToken);
 			}
 			return TempStorageTagManifest.Load(localFileListLocation);
 		}
@@ -658,7 +658,7 @@ namespace Horde.Storage.Utility
 		/// <summary>
 		/// Retrieve an output of the given node. Fetches and decompresses the files from shared storage if necessary, or validates the local files.
 		/// </summary>
-		/// <param name="reader">Store to read data from</param>
+		/// <param name="storageClient">Store to read data from</param>
 		/// <param name="refPrefix">Prefix for ref names</param>
 		/// <param name="nodeName">The node which created the storage block</param>
 		/// <param name="blockName">Name of the block to retrieve.</param>
@@ -667,7 +667,7 @@ namespace Horde.Storage.Utility
 		/// <param name="logger">Logger for output</param>
 		/// <param name="cancellationToken"></param>
 		/// <returns>Manifest of the files retrieved</returns>
-		public static async Task<TempStorageBlockManifest> RetrieveBlockAsync(TreeReader reader, string refPrefix, string nodeName, string blockName, DirectoryReference rootDir, DirectoryReference manifestDir, ILogger logger, CancellationToken cancellationToken)
+		public static async Task<TempStorageBlockManifest> RetrieveBlockAsync(IStorageClient storageClient, string refPrefix, string nodeName, string blockName, DirectoryReference rootDir, DirectoryReference manifestDir, ILogger logger, CancellationToken cancellationToken)
 		{
 			// Get the path to the local manifest
 			FileReference localManifestFile = GetBlockManifestLocation(manifestDir, nodeName, blockName);
@@ -688,7 +688,7 @@ namespace Horde.Storage.Utility
 				RefName refName = GetRefNameForNode(refPrefix, nodeName);
 				logger.LogInformation("Reading node \"{NodeName}\" block \"{BlockName}\" from temp storage (ref: {RefName}, local: {LocalFile}, blockdir: {BlockDir})", nodeName, blockName, refName, localManifestFile, blockDirectoryName);
 
-				DirectoryNode node = await reader.ReadNodeAsync<DirectoryNode>(refName, cancellationToken: cancellationToken);
+				DirectoryNode node = await storageClient.ReadNodeAsync<DirectoryNode>(refName, cancellationToken: cancellationToken);
 
 				DirectoryEntry? rootDirEntry;
 				if (!node.TryGetDirectoryEntry(blockDirectoryName, out rootDirEntry))
@@ -697,8 +697,8 @@ namespace Horde.Storage.Utility
 				}
 
 				// Add all the files and flush the ref
-				DirectoryNode rootDirNode = await rootDirEntry.ExpandAsync(reader, cancellationToken);
-				await rootDirNode.CopyToDirectoryAsync(reader, rootDir.ToDirectoryInfo(), logger, cancellationToken);
+				DirectoryNode rootDirNode = await rootDirEntry.ExpandAsync(cancellationToken);
+				await rootDirNode.CopyToDirectoryAsync(rootDir.ToDirectoryInfo(), logger, cancellationToken);
 
 				// Read the manifest in
 				manifest = TempStorageBlockManifest.Load(localManifestFile);
diff --git a/Engine/Source/Programs/Horde/Horde.Server/Artifacts/ArtifactsController.cs b/Engine/Source/Programs/Horde/Horde.Server/Artifacts/ArtifactsController.cs
--- a/Engine/Source/Programs/Horde/Horde.Server/Artifacts/ArtifactsController.cs
+++ b/Engine/Source/Programs/Horde/Horde.Server/Artifacts/ArtifactsController.cs
@@ -238,12 +238,12 @@ namespace Horde.Server.Artifacts
 				return Forbid(ArtifactAclAction.ReadArtifact, artifact.AclScope);
 			}
 
-			TreeReader reader = await _storageService.GetReaderAsync(artifact.NamespaceId, cancellationToken);
+			IStorageClient storageClient = await _storageService.GetClientAsync(artifact.NamespaceId, cancellationToken);
 
 			DirectoryNode directoryNode;
 			try
 			{
-				directoryNode = await reader.ReadNodeAsync<DirectoryNode>(artifact.RefName, DateTime.UtcNow.AddHours(1.0), cancellationToken);
+				directoryNode = await storageClient.ReadNodeAsync<DirectoryNode>(artifact.RefName, DateTime.UtcNow.AddHours(1.0), cancellationToken);
 			}
 			catch (RefNameNotFoundException)
 			{
@@ -259,44 +259,44 @@ namespace Horde.Server.Artifacts
 					{
 						return NotFound();
 					}
-					directoryNode = await nextDirectoryEntry.ExpandAsync(reader, cancellationToken);
+					directoryNode = await nextDirectoryEntry.ExpandAsync(cancellationToken);
 				}
 			}
 
 			GetArtifactDirectoryResponse response = new GetArtifactDirectoryResponse();
-			await ExpandDirectoriesAsync(reader, directoryNode, 0, response, cancellationToken);
+			await ExpandDirectoriesAsync(directoryNode, 0, response, cancellationToken);
 			return PropertyFilter.Apply(response, filter);
 		}
 
-		async Task ExpandDirectoriesAsync(TreeReader reader, DirectoryNode directoryNode, int depth, GetArtifactDirectoryResponse response, CancellationToken cancellationToken)
+		async Task ExpandDirectoriesAsync(DirectoryNode directoryNode, int depth, GetArtifactDirectoryResponse response, CancellationToken cancellationToken)
 		{
 			if (directoryNode.Directories.Count > 0)
 			{
 				response.Directories = new List<GetArtifactDirectoryEntryResponse>();
 				foreach (DirectoryEntry subDirectoryEntry in directoryNode.Directories)
 				{
-					DirectoryNode subDirectoryNode = await subDirectoryEntry.ExpandAsync(reader, cancellationToken);
+					DirectoryNode subDirectoryNode = await subDirectoryEntry.ExpandAsync(cancellationToken);
 
 					GetArtifactDirectoryEntryResponse subDirectoryEntryResponse = new GetArtifactDirectoryEntryResponse(subDirectoryEntry);
 					if (depth == 0)
 					{
 						if (subDirectoryNode.Directories.Count + subDirectoryNode.Files.Count < 16)
 						{
-							await ExpandDirectoriesAsync(reader, subDirectoryNode, depth + 1, subDirectoryEntryResponse, cancellationToken);
+							await ExpandDirectoriesAsync(subDirectoryNode, depth + 1, subDirectoryEntryResponse, cancellationToken);
 						}
 					}
 					else if (depth == 1)
 					{
 						if (subDirectoryNode.Directories.Count + subDirectoryNode.Files.Count < 8)
 						{
-							await ExpandDirectoriesAsync(reader, subDirectoryNode, depth + 1, subDirectoryEntryResponse, cancellationToken);
+							await ExpandDirectoriesAsync(subDirectoryNode, depth + 1, subDirectoryEntryResponse, cancellationToken);
 						}
 					}
 					else if (depth < 10)
 					{
 						if (subDirectoryNode.Directories.Count == 1 && subDirectoryNode.Files.Count == 0)
 						{
-							await ExpandDirectoriesAsync(reader, subDirectoryNode, depth + 1, subDirectoryEntryResponse, cancellationToken);
+							await ExpandDirectoriesAsync(subDirectoryNode, depth + 1, subDirectoryEntryResponse, cancellationToken);
 						}
 					}
 					response.Directories.Add(subDirectoryEntryResponse);
@@ -331,10 +331,10 @@ namespace Horde.Server.Artifacts
 				return Forbid(ArtifactAclAction.ReadArtifact, artifact.AclScope);
 			}
 
-			TreeReader reader = await _storageService.GetReaderAsync(artifact.NamespaceId, cancellationToken);
-			DirectoryNode directory = await reader.ReadNodeAsync<DirectoryNode>(artifact.RefName, DateTime.UtcNow.AddHours(1.0), cancellationToken);
+			IStorageClient storageClient = await _storageService.GetClientAsync(artifact.NamespaceId, cancellationToken);
+			DirectoryNode directory = await storageClient.ReadNodeAsync<DirectoryNode>(artifact.RefName, DateTime.UtcNow.AddHours(1.0), cancellationToken);
 
-			FileEntry? fileEntry = await directory.GetFileEntryByPathAsync(reader, path, cancellationToken);
+			FileEntry? fileEntry = await directory.GetFileEntryByPathAsync(path, cancellationToken);
 			if (fileEntry == null)
 			{
 				return NotFound($"Unable to find file {path}");
@@ -346,7 +346,7 @@ namespace Horde.Server.Artifacts
 				contentType = "application/octet-stream";
 			}
 
-			Stream stream = fileEntry.AsStream(reader);
+			Stream stream = fileEntry.AsStream();
 			if (inline)
 			{
 				return new InlineFileStreamResult(stream, contentType, Path.GetFileName(path));
@@ -403,10 +403,10 @@ namespace Horde.Server.Artifacts
 				filter = new FileFilter(fileFilter);
 			}
 
-			TreeReader reader = await _storageService.GetReaderAsync(artifact.NamespaceId, cancellationToken);
-			DirectoryNode directory = await reader.ReadNodeAsync<DirectoryNode>(artifact.RefName, DateTime.UtcNow.AddHours(1.0), cancellationToken);
+			IStorageClient storageClient = await _storageService.GetClientAsync(artifact.NamespaceId, cancellationToken);
+			DirectoryNode directory = await storageClient.ReadNodeAsync<DirectoryNode>(artifact.RefName, DateTime.UtcNow.AddHours(1.0), cancellationToken);
 
-			Stream stream = directory.AsZipStream(reader, filter);
+			Stream stream = directory.AsZipStream(filter);
 			return new FileStreamResult(stream, "application/zip") { FileDownloadName = $"{artifact.RefName}.zip" };
 		}
 
diff --git a/Engine/Source/Programs/Horde/Horde.Server/Commands/Bundles/ExtractCommand.cs b/Engine/Source/Programs/Horde/Horde.Server/Commands/Bundles/ExtractCommand.cs
--- a/Engine/Source/Programs/Horde/Horde.Server/Commands/Bundles/ExtractCommand.cs
+++ b/Engine/Source/Programs/Horde/Horde.Server/Commands/Bundles/ExtractCommand.cs
@@ -43,12 +43,11 @@ namespace Horde.Server.Commands.Bundles
 
 			IStorageClient store = await storageService.GetClientAsync(NamespaceId, default);
 
-			TreeReader reader = new TreeReader(store, serviceProvider.GetRequiredService<IMemoryCache>(), serviceProvider.GetRequiredService<ILogger<ExtractCommand>>());
-			CommitNode commit = await reader.ReadNodeAsync<CommitNode>(RefName);
+			CommitNode commit = await store.ReadNodeAsync<CommitNode>(RefName);
 			logger.LogInformation("Extracting {Number}: {Description} to {OutputDir}", commit.Number, (commit.Message ?? String.Empty).Replace("\n", "\\n", StringComparison.Ordinal), OutputDir);
 
-			DirectoryNode contents = await commit.Contents.ExpandAsync(reader);
-			await contents.CopyToDirectoryAsync(reader, OutputDir.ToDirectoryInfo(), logger, CancellationToken.None);
+			DirectoryNode contents = await commit.Contents.ExpandAsync();
+			await contents.CopyToDirectoryAsync(OutputDir.ToDirectoryInfo(), logger, CancellationToken.None);
 
 			return 0;
 		}
diff --git a/Engine/Source/Programs/Horde/Horde.Server/Logs/LogFileService.cs b/Engine/Source/Programs/Horde/Horde.Server/Logs/LogFileService.cs
--- a/Engine/Source/Programs/Horde/Horde.Server/Logs/LogFileService.cs
+++ b/Engine/Source/Programs/Horde/Horde.Server/Logs/LogFileService.cs
@@ -412,7 +412,6 @@ namespace Horde.Server.Logs
 		/// </summary>
 		class NewLoggerResponseStream : Stream
 		{
-			readonly TreeReader _reader;
 			readonly LogNode _rootNode;
 
 			/// <summary>
@@ -453,9 +452,8 @@ namespace Horde.Server.Logs
 			/// <summary>
 			/// Constructor
 			/// </summary>
-			public NewLoggerResponseStream(TreeReader reader, LogNode rootNode, long offset, long length)
+			public NewLoggerResponseStream(LogNode rootNode, long offset, long length)
 			{
-				_reader = reader;
 				_rootNode = rootNode;
 
 				_responseOffset = offset;
@@ -528,7 +526,7 @@ namespace Horde.Server.Logs
 
 						// Get the chunk data
 						LogChunkRef chunk = _rootNode.TextChunkRefs[_chunkIdx];
-						LogChunkNode chunkNode = await chunk.ExpandCopyAsync(_reader, cancellationToken);
+						LogChunkNode chunkNode = await chunk.ExpandCopyAsync(cancellationToken);
 
 						// Get the source data
 						_sourceBuffer = chunkNode.Data;
@@ -650,18 +648,18 @@ namespace Horde.Server.Logs
 
 			if (logFile.UseNewStorageBackend)
 			{
-				TreeReader reader = await GetTreeReaderAsync(cancellationToken);
+				IStorageClient storageClient = await GetStorageClientAsync(cancellationToken);
 
 				int maxIndex = index + count;
 
-				LogNode? root = await reader.TryReadNodeAsync<LogNode>(logFile.RefName, cancellationToken: cancellationToken);
+				LogNode? root = await storageClient.TryReadNodeAsync<LogNode>(logFile.RefName, cancellationToken: cancellationToken);
 				if (root != null)
 				{
 					int chunkIdx = root.TextChunkRefs.GetChunkForLine(index);
 					for (; index < maxIndex && chunkIdx < root.TextChunkRefs.Count; chunkIdx++)
 					{
 						LogChunkRef chunk = root.TextChunkRefs[chunkIdx];
-						LogChunkNode chunkData = await chunk.ExpandAsync(reader, cancellationToken);
+						LogChunkNode chunkData = await chunk.ExpandAsync(cancellationToken);
 
 						for (; index < maxIndex && index < chunk.LineIndex; index++)
 						{
@@ -1078,9 +1076,9 @@ namespace Horde.Server.Logs
 		{
 			if (logFile.UseNewStorageBackend)
 			{
-				TreeReader reader = await GetTreeReaderAsync(cancellationToken);
+				IStorageClient storageClient = await GetStorageClientAsync(cancellationToken);
 
-				LogNode? root = await reader.TryReadNodeAsync<LogNode>(logFile.RefName, cancellationToken: cancellationToken);
+				LogNode? root = await storageClient.TryReadNodeAsync<LogNode>(logFile.RefName, cancellationToken: cancellationToken);
 				if (root == null || root.TextChunkRefs.Count == 0)
 				{
 					return new MemoryStream(Array.Empty<byte>(), false);
@@ -1096,14 +1094,14 @@ namespace Horde.Server.Logs
 						long lastChunkLength = lastChunk.Length;
 						if (lastChunkLength <= 0)
 						{
-							LogChunkNode lastChunkNode = await lastChunk.ExpandAsync(reader, cancellationToken);
+							LogChunkNode lastChunkNode = await lastChunk.ExpandAsync(cancellationToken);
 							lastChunkLength = lastChunkNode.Length;
 						}
 						length = Math.Min(length, (lastChunk.Offset + lastChunkLength) - offset);
 					}
 
 					// Create the new stream
-					return new NewLoggerResponseStream(reader, root, offset, length);
+					return new NewLoggerResponseStream(root, offset, length);
 				}
 			}
 			else
@@ -1217,10 +1215,9 @@ namespace Horde.Server.Logs
 			}
 		}
 
-		async Task<TreeReader> GetTreeReaderAsync(CancellationToken cancellationToken)
+		async Task<IStorageClient> GetStorageClientAsync(CancellationToken cancellationToken)
 		{
-			IStorageClient store = await _storageService.GetClientAsync(Namespace.Logs, cancellationToken);
-			return new TreeReader(store, _logFileCache, _logger);
+			return await _storageService.GetClientAsync(Namespace.Logs, cancellationToken);
 		}
 
 		/// <inheritdoc/>
@@ -1228,9 +1225,9 @@ namespace Horde.Server.Logs
 		{
 			if (logFile.UseNewStorageBackend)
 			{
-				TreeReader reader = await GetTreeReaderAsync(cancellationToken);
+				IStorageClient storageClient = await GetStorageClientAsync(cancellationToken);
 
-				LogNode? root = await reader.TryReadNodeAsync<LogNode>(logFile.RefName, cancellationToken: cancellationToken);
+				LogNode? root = await storageClient.TryReadNodeAsync<LogNode>(logFile.RefName, cancellationToken: cancellationToken);
 				if (root == null)
 				{
 					return (0, 0);
@@ -1238,7 +1235,7 @@ namespace Horde.Server.Logs
 
 				int chunkIdx = root.TextChunkRefs.GetChunkForLine(lineIdx);
 				LogChunkRef chunk = root.TextChunkRefs[chunkIdx];
-				LogChunkNode chunkData = await chunk.ExpandAsync(reader, cancellationToken);
+				LogChunkNode chunkData = await chunk.ExpandAsync(cancellationToken);
 
 				if (lineIdx < chunk.LineIndex)
 				{
@@ -1889,16 +1886,16 @@ namespace Horde.Server.Logs
 		async IAsyncEnumerable<int> SearchLogDataInternalNewAsync(ILogFile logFile, string text, int firstLine, SearchStats searchStats, [EnumeratorCancellation] CancellationToken cancellationToken)
 		{
 			SearchTerm searchText = new SearchTerm(text);
-			TreeReader reader = await GetTreeReaderAsync(cancellationToken);
+			IStorageClient storageClient = await GetStorageClientAsync(cancellationToken);
 
 			// Search the index
 			if (logFile.LineCount > 0)
 			{
-				LogNode? root = await reader.ReadNodeAsync<LogNode>(logFile.RefName, cancellationToken: cancellationToken);
+				LogNode? root = await storageClient.ReadNodeAsync<LogNode>(logFile.RefName, cancellationToken: cancellationToken);
 				if(root != null)
 				{
-					LogIndexNode index = await root.IndexRef.ExpandAsync(reader, cancellationToken);
-					await foreach (int lineIdx in index.Search(reader, firstLine, searchText, searchStats, cancellationToken))
+					LogIndexNode index = await root.IndexRef.ExpandAsync(cancellationToken);
+					await foreach (int lineIdx in index.Search(firstLine, searchText, searchStats, cancellationToken))
 					{
 						yield return lineIdx;
 					}
diff --git a/Engine/Source/Programs/Horde/Horde.Server/Perforce/PerforceReplicator.cs b/Engine/Source/Programs/Horde/Horde.Server/Perforce/PerforceReplicator.cs
--- a/Engine/Source/Programs/Horde/Horde.Server/Perforce/PerforceReplicator.cs
+++ b/Engine/Source/Programs/Horde/Horde.Server/Perforce/PerforceReplicator.cs
@@ -257,22 +257,20 @@ namespace Horde.Server.Perforce
 		{
 			IStorageClient store = await _storageService.GetClientAsync(Namespace.Perforce, cancellationToken);
 
-			TreeReader reader = new TreeReader(store, _memoryCache, _logger);
-
 			// Find the parent node
 			RefName refName = GetRefName(streamConfig.Id);
-			CommitNode? parent = await FindParentAsync(reader, refName, change, cancellationToken);
+			CommitNode? parent = await FindParentAsync(store, refName, change, cancellationToken);
 			NodeRef<CommitNode>? parentRef = (parent == null) ? null : new NodeRef<CommitNode>(parent);
 			int parentChange = parent?.Number ?? 0;
 
 			// Read the current incremental state or create a new node to track the incremental state
 			RefName incRefName = GetIncrementalRefName(streamConfig.Id);
-			SyncNode? syncNode = await reader.TryReadNodeAsync<SyncNode>(incRefName, cancellationToken: cancellationToken);
+			SyncNode? syncNode = await store.TryReadNodeAsync<SyncNode>(incRefName, cancellationToken: cancellationToken);
 			if (syncNode == null || syncNode.Change != change || syncNode.ParentChange != parentChange)
 			{
 				syncNode = new SyncNode(change, parentChange, parent?.Contents ?? new NodeRef<DirectoryNode>(new DirectoryNode()));
 			}
-			DirectoryNode root = await syncNode.Contents.ExpandAsync(reader, cancellationToken);
+			DirectoryNode root = await syncNode.Contents.ExpandAsync(cancellationToken);
 
 			// Create a client to replicate from this stream
 			ReplicationClient clientInfo = await FindOrAddReplicationClientAsync(streamConfig);
@@ -456,7 +454,7 @@ namespace Horde.Server.Perforce
 						else if (io.Command == PerforceIoCommand.Close)
 						{
 							FileInfo info = await fileWriter.CloseAsync(io.File, cancellationToken);
-							FileEntry entry = await root.AddFileByPathAsync(reader, info.Path, info.Flags, info.Length, info.Node, cancellationToken);
+							FileEntry entry = await root.AddFileByPathAsync(info.Path, info.Flags, info.Length, info.Node, cancellationToken);
 							entry.CustomData = info.Md5;
 							await writer.WriteAsync(entry, cancellationToken);
 						}
@@ -464,7 +462,7 @@ namespace Horde.Server.Perforce
 						{
 							UnpackUnlinkPayload(io.Payload, out Utf8String path);
 							Utf8String file = GetClientRelativePath(path, clientInfo.Client.Root);
-							await root.DeleteFileByPathAsync(reader, file, cancellationToken);
+							await root.DeleteFileByPathAsync(file, cancellationToken);
 						}
 						else
 						{
@@ -516,14 +514,14 @@ namespace Horde.Server.Perforce
 			await writer.WriteAsync(refName, commitNode, options.RefOptions, cancellationToken); 
 		}
 
-		static async Task<CommitNode?> FindParentAsync(TreeReader reader, RefName refName, int change, CancellationToken cancellationToken)
+		static async Task<CommitNode?> FindParentAsync(IStorageClient storageClient, RefName refName, int change, CancellationToken cancellationToken)
 		{
-			CommitNode? commitNode = await reader.TryReadNodeAsync<CommitNode>(refName, cancellationToken: cancellationToken);
+			CommitNode? commitNode = await storageClient.TryReadNodeAsync<CommitNode>(refName, cancellationToken: cancellationToken);
 			if (commitNode != null)
 			{
 				while (commitNode.Parent != null && commitNode.Number >= change)
 				{
-					commitNode = await commitNode.Parent.ExpandAsync(reader, cancellationToken);
+					commitNode = await commitNode.Parent.ExpandAsync(cancellationToken);
 				}
 			}
 			return commitNode;
diff --git a/Engine/Source/Programs/Horde/Horde.Server/Perforce/ReplicationService.cs b/Engine/Source/Programs/Horde/Horde.Server/Perforce/ReplicationService.cs
--- a/Engine/Source/Programs/Horde/Horde.Server/Perforce/ReplicationService.cs
+++ b/Engine/Source/Programs/Horde/Horde.Server/Perforce/ReplicationService.cs
@@ -139,9 +139,8 @@ namespace Horde.Server.Perforce
 			RefName refName = new RefName(streamConfig.Id.ToString());
 
 			IStorageClientImpl store = await _storageService.GetClientAsync(Namespace.Perforce, cancellationToken);
-			TreeReader reader = new TreeReader(store, _memoryCache, _logger);
 
-			CommitNode? lastCommitNode = await reader.TryReadNodeAsync<CommitNode>(refName, cancellationToken: cancellationToken);
+			CommitNode? lastCommitNode = await store.TryReadNodeAsync<CommitNode>(refName, cancellationToken: cancellationToken);
 			ICommitCollection commits = _perforceService.GetCommits(streamConfig);
 
 			PerforceReplicationOptions options = new PerforceReplicationOptions();
diff --git a/Engine/Source/Programs/Horde/Horde.Server/Server/ExceptionController.cs b/Engine/Source/Programs/Horde/Horde.Server/Server/ExceptionController.cs
new file mode 100644
--- /dev/null
+++ b/Engine/Source/Programs/Horde/Horde.Server/Server/ExceptionController.cs
@@ -0,0 +1,56 @@
+// Copyright Epic Games, Inc. All Rights Reserved.
+
+using System;
+using System.Net;
+using EpicGames.Core;
+using Horde.Server.Utilities;
+using Microsoft.AspNetCore.Diagnostics;
+using Microsoft.AspNetCore.Mvc;
+using Microsoft.Extensions.Logging;
+
+namespace Horde.Server.Server
+{
+	/// <summary>
+	/// Controller managing account status
+	/// </summary>
+	[ApiController]
+	[Route("[controller]")]
+	public class ExceptionController : HordeControllerBase
+	{
+		/// <summary>
+		/// Outputs a diagnostic error response for an exception
+		/// </summary>
+		[Route("/api/v1/exception")]
+		[ApiExplorerSettings(IgnoreApi = true)]
+		public ActionResult Exception()
+		{
+			IExceptionHandlerPathFeature? feature = HttpContext.Features.Get<IExceptionHandlerPathFeature>();
+
+			int statusCode = (int)HttpStatusCode.InternalServerError;
+
+			LogEvent logEvent;
+			if (feature == null)
+			{
+				logEvent = LogEvent.Create(LogLevel.Error, "Exception handler path feature is missing.");
+			}
+			else if (feature.Error == null)
+			{
+				logEvent = LogEvent.Create(LogLevel.Information, "No error code.");
+			}
+			else if (feature.Error is StructuredHttpException structuredHttpEx)
+			{
+				(logEvent, statusCode) = (structuredHttpEx.ToLogEvent(), structuredHttpEx.StatusCode);
+			}
+			else if (feature.Error is StructuredException structuredEx)
+			{
+				logEvent = structuredEx.ToLogEvent();
+			}
+			else
+			{
+				logEvent = LogEvent.Create(LogLevel.Error, default, feature.Error, "Unhandled exception: {Message}", feature.Error.Message);
+			}
+
+			return new ObjectResult(logEvent) { StatusCode = statusCode };
+		}
+	}
+}
diff --git a/Engine/Source/Programs/Horde/Horde.Server/Storage/StorageService.cs b/Engine/Source/Programs/Horde/Horde.Server/Storage/StorageService.cs
--- a/Engine/Source/Programs/Horde/Horde.Server/Storage/StorageService.cs
+++ b/Engine/Source/Programs/Horde/Horde.Server/Storage/StorageService.cs
@@ -94,7 +94,7 @@ namespace Horde.Server.Storage
 	/// </summary>
 	public sealed class StorageService : IHostedService, IDisposable, IStorageClientFactory
 	{
-		sealed class StorageClient : StorageClientBase, IStorageClientImpl
+		sealed class StorageClient : BundleStorageClient, IStorageClientImpl
 		{
 			readonly StorageService _outer;
 			readonly string _prefix;
diff --git a/Engine/Source/Programs/Horde/Horde.Server/Tools/ToolCollection.cs b/Engine/Source/Programs/Horde/Horde.Server/Tools/ToolCollection.cs
--- a/Engine/Source/Programs/Horde/Horde.Server/Tools/ToolCollection.cs
+++ b/Engine/Source/Programs/Horde/Horde.Server/Tools/ToolCollection.cs
@@ -433,10 +433,9 @@ namespace Horde.Server.Tools
 		{
 			IStorageClient client = await GetStorageClientAsync(tool, cancellationToken);
 
-			TreeReader reader = new TreeReader(client, _cache, _logger);
-			DirectoryNode node = await reader.ReadNodeAsync<DirectoryNode>(deployment.RefName, DateTime.UtcNow - TimeSpan.FromDays(2.0), cancellationToken);
+			DirectoryNode node = await client.ReadNodeAsync<DirectoryNode>(deployment.RefName, DateTime.UtcNow - TimeSpan.FromDays(2.0), cancellationToken);
 
-			return node.AsZipStream(reader);
+			return node.AsZipStream();
 		}
 	}
 }
diff --git a/Engine/Source/Programs/Horde/Horde.Server/Tools/ToolsController.cs b/Engine/Source/Programs/Horde/Horde.Server/Tools/ToolsController.cs
--- a/Engine/Source/Programs/Horde/Horde.Server/Tools/ToolsController.cs
+++ b/Engine/Source/Programs/Horde/Horde.Server/Tools/ToolsController.cs
@@ -503,8 +503,7 @@ namespace Horde.Server.Tools
 
 			IStorageClient client = await _toolCollection.GetStorageClientAsync(tool, cancellationToken);
 
-			TreeReader reader = new TreeReader(client, _cache, _logger);
-			DirectoryNode node = await reader.ReadNodeAsync<DirectoryNode>(deployment.RefName, DateTime.UtcNow - TimeSpan.FromDays(2.0), cancellationToken);
+			DirectoryNode node = await client.ReadNodeAsync<DirectoryNode>(deployment.RefName, DateTime.UtcNow - TimeSpan.FromDays(2.0), cancellationToken);
 
 			if (node.Directories.Count == 0 && node.Files.Count == 1 && action != GetToolAction.Zip)
 			{
@@ -516,10 +515,10 @@ namespace Horde.Server.Tools
 					contentType = "application/octet-stream";
 				}
 
-				return new FileStreamResult(entry.AsStream(reader), contentType) { FileDownloadName = entry.Name.ToString() };
+				return new FileStreamResult(entry.AsStream(), contentType) { FileDownloadName = entry.Name.ToString() };
 			}
 
-			Stream stream = node.AsZipStream(reader);
+			Stream stream = node.AsZipStream();
 			return new FileStreamResult(stream, "application/zip") { FileDownloadName = $"{tool.Id}-{deployment.Version}.zip" };
 		}
 
diff --git a/Engine/Source/Programs/Horde/Horde/Commands/Bundles/BundleExtract.cs b/Engine/Source/Programs/Horde/Horde/Commands/Bundles/BundleExtract.cs
--- a/Engine/Source/Programs/Horde/Horde/Commands/Bundles/BundleExtract.cs
+++ b/Engine/Source/Programs/Horde/Horde/Commands/Bundles/BundleExtract.cs
@@ -47,7 +47,7 @@ namespace Horde.Commands.Bundles
 			Stopwatch timer = Stopwatch.StartNew();
 
 			DirectoryNode node = await reader.ReadNodeAsync<DirectoryNode>(handle.Locator);
-			await node.CopyToDirectoryAsync(reader, OutputDir.ToDirectoryInfo(), logger, CancellationToken.None);
+			await node.CopyToDirectoryAsync(OutputDir.ToDirectoryInfo(), logger, CancellationToken.None);
 
 			logger.LogInformation("Elapsed: {Time}s", timer.Elapsed.TotalSeconds);
 			return 0;
diff --git a/Engine/Source/Programs/Horde/Horde/Commands/Vcs/VcsBase.cs b/Engine/Source/Programs/Horde/Horde/Commands/Vcs/VcsBase.cs
--- a/Engine/Source/Programs/Horde/Horde/Commands/Vcs/VcsBase.cs
+++ b/Engine/Source/Programs/Horde/Horde/Commands/Vcs/VcsBase.cs
@@ -346,9 +346,9 @@ namespace Horde.Commands.Vcs
 			}
 		}
 
-		protected static async Task<CommitNode?> GetCommitAsync(TreeReader reader, RefName branchName, int change = 0)
+		protected static async Task<CommitNode?> GetCommitAsync(IStorageClient storageClient, RefName branchName, int change = 0)
 		{
-			CommitNode tip = await reader.ReadNodeAsync<CommitNode>(branchName);
+			CommitNode tip = await storageClient.ReadNodeAsync<CommitNode>(branchName);
 			if (change != 0)
 			{
 				while (tip.Number != change)
@@ -357,7 +357,7 @@ namespace Horde.Commands.Vcs
 					{
 						return null;
 					}
-					tip = await tip.Parent.ExpandAsync(reader);
+					tip = await tip.Parent.ExpandAsync();
 				}
 			}
 			return tip;
diff --git a/Engine/Source/Programs/Horde/Horde/Commands/Vcs/VcsCheckout.cs b/Engine/Source/Programs/Horde/Horde/Commands/Vcs/VcsCheckout.cs
--- a/Engine/Source/Programs/Horde/Horde/Commands/Vcs/VcsCheckout.cs
+++ b/Engine/Source/Programs/Horde/Horde/Commands/Vcs/VcsCheckout.cs
@@ -57,30 +57,27 @@ namespace Horde.Commands.Vcs
 
 			IStorageClient store = await GetStorageClientAsync();
 
-			using MemoryCache cache = new MemoryCache(new MemoryCacheOptions());
-			TreeReader reader = new TreeReader(store, cache, logger);
-
-			CommitNode? tip = await GetCommitAsync(reader, branchName, Change);
+			CommitNode? tip = await GetCommitAsync(store, branchName, Change);
 			if (tip == null)
 			{
 				logger.LogError("Unable to find change {Change}", Change);
 				return 1;
 			}
 
-			workspaceState.Tree = await RealizeAsync(reader, tip.Contents, rootDir, newState, Clean, logger);
+			workspaceState.Tree = await RealizeAsync(tip.Contents, rootDir, newState, Clean, logger);
 			await WriteStateAsync(rootDir, workspaceState);
 
 			logger.LogInformation("Updated workspace to change {Number}", tip.Number);
 			return 0;
 		}
 
-		protected static async Task<DirectoryState> RealizeAsync(TreeReader reader, DirectoryNodeRef directoryRef, DirectoryReference dirPath, DirectoryState? directoryState, bool clean, ILogger logger)
+		protected static async Task<DirectoryState> RealizeAsync(DirectoryNodeRef directoryRef, DirectoryReference dirPath, DirectoryState? directoryState, bool clean, ILogger logger)
 		{
 			DirectoryReference.CreateDirectory(dirPath);
 
 			DirectoryState newState = new DirectoryState();
 
-			DirectoryNode directoryNode = await directoryRef.ExpandAsync(reader);
+			DirectoryNode directoryNode = await directoryRef.ExpandAsync();
 			foreach ((Utf8String name, DirectoryEntry? subDirEntry, DirectoryState? subDirState) in EnumerableExtensions.Zip(directoryNode.NameToDirectory, directoryState?.Directories))
 			{
 				DirectoryReference subDirPath = DirectoryReference.Combine(dirPath, name.ToString());
@@ -93,7 +90,7 @@ namespace Horde.Commands.Vcs
 				}
 				else
 				{
-					newState.Directories[name] = await RealizeAsync(reader, subDirEntry, subDirPath, subDirState, clean, logger);
+					newState.Directories[name] = await RealizeAsync(subDirEntry, subDirPath, subDirState, clean, logger);
 				}
 			}
 
@@ -109,7 +106,7 @@ namespace Horde.Commands.Vcs
 				}
 				else if (fileState == null || fileState.Hash != fileEntry.Hash)
 				{
-					newState.Files[name] = await CheckoutFileAsync(reader, fileEntry, filePath.ToFileInfo(), logger);
+					newState.Files[name] = await CheckoutFileAsync(fileEntry, filePath.ToFileInfo(), logger);
 				}
 				else
 				{
@@ -121,11 +118,11 @@ namespace Horde.Commands.Vcs
 			return newState;
 		}
 
-		static async Task<FileState> CheckoutFileAsync(TreeReader reader, FileEntry fileRef, FileInfo fileInfo, ILogger logger)
+		static async Task<FileState> CheckoutFileAsync(FileEntry fileRef, FileInfo fileInfo, ILogger logger)
 		{
 			logger.LogInformation("Updating {File} to {Hash}", fileInfo, fileRef.Hash);
-			ChunkedDataNode fileNode = await fileRef.ExpandAsync(reader);
-			await fileNode.CopyToFileAsync(reader, fileInfo, CancellationToken.None);
+			ChunkedDataNode fileNode = await fileRef.ExpandAsync();
+			await fileNode.CopyToFileAsync(fileInfo, CancellationToken.None);
 			fileInfo.Refresh();
 			return new FileState(fileInfo, fileRef.Hash);
 		}
diff --git a/Engine/Source/Programs/Horde/Horde/Commands/Vcs/VcsClone.cs b/Engine/Source/Programs/Horde/Horde/Commands/Vcs/VcsClone.cs
--- a/Engine/Source/Programs/Horde/Horde/Commands/Vcs/VcsClone.cs
+++ b/Engine/Source/Programs/Horde/Horde/Commands/Vcs/VcsClone.cs
@@ -27,19 +27,16 @@ namespace Horde.Commands.Vcs
 
 			RefName branchName = new RefName(Branch ?? "ue5-main");
 
-			IStorageClient store = await GetStorageClientAsync();
+			IStorageClient storageClient = await GetStorageClientAsync();
 
-			using MemoryCache cache = new MemoryCache(new MemoryCacheOptions());
-			TreeReader reader = new TreeReader(store, cache, logger);
-
-			CommitNode? tip = await GetCommitAsync(reader, branchName, Change);
+			CommitNode? tip = await GetCommitAsync(storageClient, branchName, Change);
 			if (tip == null)
 			{
 				logger.LogError("Unable to find change {Change}", Change);
 				return 1;
 			}
 
-			workspaceState.Tree = await RealizeAsync(reader, tip.Contents, BaseDir, null, true, logger);
+			workspaceState.Tree = await RealizeAsync(tip.Contents, BaseDir, null, true, logger);
 			await WriteStateAsync(BaseDir, workspaceState);
 
 			return 0;
diff --git a/Engine/Source/Programs/Horde/Horde/Commands/Vcs/VcsCommit.cs b/Engine/Source/Programs/Horde/Horde/Commands/Vcs/VcsCommit.cs
--- a/Engine/Source/Programs/Horde/Horde/Commands/Vcs/VcsCommit.cs
+++ b/Engine/Source/Programs/Horde/Horde/Commands/Vcs/VcsCommit.cs
@@ -52,10 +52,7 @@ namespace Horde.Commands.Vcs
 
 			IStorageClient store = await GetStorageClientAsync();
 
-			using MemoryCache cache = new MemoryCache(new MemoryCacheOptions());
-			TreeReader reader = new TreeReader(store, cache, logger);
-
-			CommitNode? tip = await reader.TryReadNodeAsync<CommitNode>(workspaceState.Branch);
+			CommitNode? tip = await store.TryReadNodeAsync<CommitNode>(workspaceState.Branch);
 			NodeRef<CommitNode>? tipRef = (tip == null) ? null : new NodeRef<CommitNode>(tip);
 
 			DirectoryNode rootNode;
@@ -65,14 +62,14 @@ namespace Horde.Commands.Vcs
 			}
 			else
 			{
-				rootNode = await tip.Contents.ExpandCopyAsync(reader);
+				rootNode = await tip.Contents.ExpandCopyAsync();
 			}
 
 			DirectoryNodeRef rootRef = new DirectoryNodeRef(rootNode);
 
 			List<(DirectoryNode, FileInfo, FileState)> files = new List<(DirectoryNode, FileInfo, FileState)>();
 			List<(DirectoryNodeRef, DirectoryState)> directories = new List<(DirectoryNodeRef, DirectoryState)>();
-			await UpdateTreeAsync(reader, rootRef, rootDir, oldState, newState, files, directories);
+			await UpdateTreeAsync(rootRef, rootDir, oldState, newState, files, directories);
 
 			using IStorageWriter writer = store.CreateWriter();
 			await DirectoryNode.CopyFromDirectoryAsync(files.ConvertAll(x => (x.Item1, x.Item2)), new ChunkingOptions(), writer, null, CancellationToken.None);
@@ -107,11 +104,11 @@ namespace Horde.Commands.Vcs
 			return 0;
 		}
 
-		private async Task UpdateTreeAsync(TreeReader reader, DirectoryNodeRef rootRef, DirectoryReference rootDir, DirectoryState? oldState, DirectoryState newState, List<(DirectoryNode, FileInfo, FileState)> files, List<(DirectoryNodeRef, DirectoryState)> directories)
+		private async Task UpdateTreeAsync(DirectoryNodeRef rootRef, DirectoryReference rootDir, DirectoryState? oldState, DirectoryState newState, List<(DirectoryNode, FileInfo, FileState)> files, List<(DirectoryNodeRef, DirectoryState)> directories)
 		{
 			directories.Add((rootRef, newState));
 
-			DirectoryNode root = await rootRef.ExpandAsync(reader);
+			DirectoryNode root = await rootRef.ExpandAsync();
 			foreach ((Utf8String name, DirectoryState? oldSubDirState, DirectoryState? newSubDirState) in EnumerableExtensions.Zip(oldState?.Directories, newState.Directories))
 			{
 				if (newSubDirState == null)
@@ -120,7 +117,7 @@ namespace Horde.Commands.Vcs
 				}
 				else if (oldSubDirState != newSubDirState)
 				{
-					await UpdateTreeAsync(reader, root.FindOrAddDirectoryEntry(name), DirectoryReference.Combine(rootDir, name.ToString()), oldSubDirState, newSubDirState, files, directories);
+					await UpdateTreeAsync(root.FindOrAddDirectoryEntry(name), DirectoryReference.Combine(rootDir, name.ToString()), oldSubDirState, newSubDirState, files, directories);
 				}
 			}
 
diff --git a/Engine/Source/Programs/Horde/Horde/Commands/Vcs/VcsLog.cs b/Engine/Source/Programs/Horde/Horde/Commands/Vcs/VcsLog.cs
--- a/Engine/Source/Programs/Horde/Horde/Commands/Vcs/VcsLog.cs
+++ b/Engine/Source/Programs/Horde/Horde/Commands/Vcs/VcsLog.cs
@@ -34,7 +34,7 @@ namespace Horde.Commands.Vcs
 
 			List<CommitNode> commits = new List<CommitNode>();
 
-			CommitNode? tip = await reader.TryReadNodeAsync<CommitNode>(workspaceState.Branch);
+			CommitNode? tip = await store.TryReadNodeAsync<CommitNode>(workspaceState.Branch);
 			if (tip != null)
 			{
 				for (int idx = 0; idx < Count; idx++)
@@ -46,7 +46,7 @@ namespace Horde.Commands.Vcs
 						break;
 					}
 
-					tip = await tip.Parent.ExpandAsync(reader);
+					tip = await tip.Parent.ExpandAsync();
 				}
 			}
 
diff --git a/Engine/Source/Programs/Shared/EpicGames.Horde.Tests/BundleStoreTests.cs b/Engine/Source/Programs/Shared/EpicGames.Horde.Tests/BundleStoreTests.cs
--- a/Engine/Source/Programs/Shared/EpicGames.Horde.Tests/BundleStoreTests.cs
+++ b/Engine/Source/Programs/Shared/EpicGames.Horde.Tests/BundleStoreTests.cs
@@ -155,18 +155,17 @@ namespace EpicGames.Horde.Tests
 				await writer.WriteAsync(new RefName("test"), root);
 
 				TreeReader reader = new TreeReader(store, null, NullLogger.Instance);
-				await CheckTree(reader, root);
+				await CheckTree(root);
 			}
 
 			// Check we can read it back in
 			{
-				TreeReader reader = new TreeReader(store, null, NullLogger.Instance);
-				SimpleNode root = await reader.ReadNodeAsync<SimpleNode>(new RefName("test"));
-				await CheckTree(reader, root);
+				SimpleNode root = await store.ReadNodeAsync<SimpleNode>(new RefName("test"));
+				await CheckTree(root);
 			}
 		}
 
-		static async Task CheckTree(TreeReader reader, SimpleNode root)
+		static async Task CheckTree(SimpleNode root)
 		{
 			SimpleNode node5 = root;
 			byte[] data5 = node5.Data.ToArray();
@@ -174,25 +173,25 @@ namespace EpicGames.Horde.Tests
 			IReadOnlyList<NodeRef<SimpleNode>> refs5 = node5.Refs;
 			Assert.AreEqual(2, refs5.Count);
 
-			SimpleNode node4 = await refs5[0].ExpandAsync(reader);
+			SimpleNode node4 = await refs5[0].ExpandAsync();
 			byte[] data4 = node4.Data.ToArray();
 			Assert.IsTrue(data4.SequenceEqual(new byte[] { 4 }));
 			IReadOnlyList<NodeRef<SimpleNode>> refs4 = node4.Refs;
 			Assert.AreEqual(0, refs4.Count);
 
-			SimpleNode node3 = await refs5[1].ExpandAsync(reader);
+			SimpleNode node3 = await refs5[1].ExpandAsync();
 			byte[] data3 = node3.Data.ToArray();
 			Assert.IsTrue(data3.SequenceEqual(new byte[] { 3 }));
 			IReadOnlyList<NodeRef<SimpleNode>> refs3 = node3.Refs;
 			Assert.AreEqual(1, refs3.Count);
 
-			SimpleNode node2 = await refs3[0].ExpandAsync(reader);
+			SimpleNode node2 = await refs3[0].ExpandAsync();
 			byte[] data2 = node2.Data.ToArray();
 			Assert.IsTrue(data2.SequenceEqual(new byte[] { 2 }));
 			IReadOnlyList<NodeRef<SimpleNode>> refs2 = node2.Refs;
 			Assert.AreEqual(1, refs2.Count);
 
-			SimpleNode node1 = await refs2[0].ExpandAsync(reader);
+			SimpleNode node1 = await refs2[0].ExpandAsync();
 			byte[] data1 = node1.Data.ToArray();
 			Assert.IsTrue(data1.SequenceEqual(new byte[] { 1 }));
 			IReadOnlyList<NodeRef<SimpleNode>> refs1 = node1.Refs;
@@ -207,8 +206,7 @@ namespace EpicGames.Horde.Tests
 			RefName refName = new RefName("test");
 			await store.WriteNodeAsync(refName, new SimpleNode(new ReadOnlySequence<byte>(new byte[] { (byte)123 }), Array.Empty<NodeRef<SimpleNode>>()));
 
-			TreeReader reader = new TreeReader(store, null, NullLogger.Instance);
-			SimpleNode node = await reader.ReadNodeAsync<SimpleNode>(refName);
+			SimpleNode node = await store.ReadNodeAsync<SimpleNode>(refName);
 
 			Assert.AreEqual(123, node.Data.FirstSpan[0]);
 		}
@@ -217,7 +215,6 @@ namespace EpicGames.Horde.Tests
 		public async Task DirectoryNodesAsync()
 		{
 			MemoryStorageClient store = new MemoryStorageClient();
-			TreeReader reader = new TreeReader(store, null, NullLogger.Instance);
 
 			// Generate a tree
 			{
@@ -226,26 +223,26 @@ namespace EpicGames.Horde.Tests
 				DirectoryNode world = hello.AddDirectory("world");
 				await store.WriteNodeAsync(new RefName("test"), root);
 
-				await CheckDirectoryTreeAsync(reader, root);
+				await CheckDirectoryTreeAsync(root);
 			}
 
 			// Check we can read it back in
 			{
-				DirectoryNode root = await reader.ReadNodeAsync<DirectoryNode>(new RefName("test"));
-				await CheckDirectoryTreeAsync(reader, root);
+				DirectoryNode root = await store.ReadNodeAsync<DirectoryNode>(new RefName("test"));
+				await CheckDirectoryTreeAsync(root);
 			}
 		}
 
-		static async Task CheckDirectoryTreeAsync(TreeReader reader, DirectoryNode root)
+		static async Task CheckDirectoryTreeAsync(DirectoryNode root)
 		{
 			Assert.AreEqual(1, root.Directories.Count);
 			Assert.AreEqual("hello", root.Directories.First().Name);
 
-			DirectoryNode hello = await root.Directories.First().ExpandAsync(reader, CancellationToken.None);
+			DirectoryNode hello = await root.Directories.First().ExpandAsync(CancellationToken.None);
 			Assert.AreEqual(1, hello.Directories.Count);
 			Assert.AreEqual("world", hello.Directories.First().Name);
 
-			DirectoryNode world = await hello.Directories.First().ExpandAsync(reader, CancellationToken.None);
+			DirectoryNode world = await hello.Directories.First().ExpandAsync(CancellationToken.None);
 			Assert.AreEqual(0, world.Directories.Count);
 		}
 
@@ -269,36 +266,35 @@ namespace EpicGames.Horde.Tests
 
 				await writer.WriteAsync(new RefName("test"), root);
 
-				await CheckFileTreeAsync(reader, root);
+				await CheckFileTreeAsync(root);
 			}
 
 			// Check we can read it back in
 			{
-				DirectoryNode root = await reader.ReadNodeAsync<DirectoryNode>(new RefName("test"));
-				await CheckFileTreeAsync(reader, root);
+				DirectoryNode root = await store.ReadNodeAsync<DirectoryNode>(new RefName("test"));
+				await CheckFileTreeAsync(root);
 			}
 		}
 
-		static async Task CheckFileTreeAsync(TreeReader reader, DirectoryNode root)
+		static async Task CheckFileTreeAsync(DirectoryNode root)
 		{
 			Assert.AreEqual(1, root.Directories.Count);
 			Assert.AreEqual("hello", root.Directories.First().Name);
 
-			DirectoryNode hello = await root.Directories.First().ExpandAsync(reader);
+			DirectoryNode hello = await root.Directories.First().ExpandAsync();
 			Assert.AreEqual(0, hello.Directories.Count);
 			Assert.AreEqual(1, hello.Files.Count);
 			Assert.AreEqual("world", hello.Files.First().Name);
 
-			ChunkedDataNode world = await hello.Files.First().ExpandAsync(reader);
+			ChunkedDataNode world = await hello.Files.First().ExpandAsync();
 
-			byte[] worldData = await GetFileDataAsync(reader, world);
+			byte[] worldData = await GetFileDataAsync(world);
 			Assert.IsTrue(worldData.SequenceEqual(Encoding.UTF8.GetBytes("world")));
 		}
 
 		[TestMethod]
 		public async Task StreamTestAsync()
 		{
-			using IMemoryCache cache = new MemoryCache(new MemoryCacheOptions());
 			MemoryStorageClient store = new MemoryStorageClient();
 
 			const int length = 4096;
@@ -307,7 +303,7 @@ namespace EpicGames.Horde.Tests
 			new Random(0).NextBytes(chunk);
 
 			// Generate a tree
-			NodeLocator locator;
+			NodeRef<ChunkedDataNode> nodeRef;
 			{
 				using IStorageWriter writer = store.CreateWriter(options: new TreeOptions { MaxBlobSize = 1024 });
 
@@ -323,17 +319,15 @@ namespace EpicGames.Horde.Tests
 				}
 
 				HashedNodeHandle handle = await fileWriter.FlushAsync(CancellationToken.None);
-				locator = handle.Handle.Locator;
+				nodeRef = new NodeRef<ChunkedDataNode>(handle);
 			}
 
 			// Check we can read it back in
 			{
-				TreeReader reader = new TreeReader(store, cache, NullLogger.Instance);
-
-				ChunkedDataNode newRoot = await reader.ReadNodeAsync<ChunkedDataNode>(locator);
+				ChunkedDataNode newRoot = await nodeRef.ExpandAsync();
 
 				using MemoryStream stream = new MemoryStream();
-				await newRoot.CopyToStreamAsync(reader, stream, CancellationToken.None);
+				await newRoot.CopyToStreamAsync(stream, CancellationToken.None);
 
 				byte[] output = stream.ToArray();
 				Assert.IsTrue(chunk.SequenceEqual(output));
@@ -376,17 +370,14 @@ namespace EpicGames.Horde.Tests
 
 				await writer.WriteAsync(new RefName("test"), root);
 
-				TreeReader reader = new TreeReader(store, null, NullLogger.Instance);
-				await CheckLargeFileTreeAsync(reader, root, data);
+				await CheckLargeFileTreeAsync(root, data);
 			}
 
 			// Check we can read it back in
 			{
-				TreeReader reader = new TreeReader(store, cache, NullLogger.Instance);
-
-				DirectoryNode newRoot = await reader.ReadNodeAsync<DirectoryNode>(new RefName("test"));
-				await CompareTrees(reader, root, newRoot);
-				await CheckLargeFileTreeAsync(reader, root, data);
+				DirectoryNode newRoot = await store.ReadNodeAsync<DirectoryNode>(new RefName("test"));
+				await CompareTrees(root, newRoot);
+				await CheckLargeFileTreeAsync(root, data);
 
 				NodeRef<ChunkedDataNode> file = root.GetFileEntry("test");
 
@@ -395,7 +386,7 @@ namespace EpicGames.Horde.Tests
 			}
 		}
 
-		static async Task CompareTrees(TreeReader reader, DirectoryNode oldNode, DirectoryNode newNode)
+		static async Task CompareTrees(DirectoryNode oldNode, DirectoryNode newNode)
 		{
 			Assert.AreEqual(oldNode.Length, newNode.Length);
 			Assert.AreEqual(oldNode.Files.Count, newNode.Files.Count);
@@ -403,13 +394,13 @@ namespace EpicGames.Horde.Tests
 
 			foreach ((FileEntry oldFileEntry, FileEntry newFileEntry) in oldNode.Files.Zip(newNode.Files))
 			{
-				ChunkedDataNode oldFile = await oldFileEntry.ExpandAsync(reader);
-				ChunkedDataNode newFile = await newFileEntry.ExpandAsync(reader);
-				await CompareTrees(reader, oldFile, newFile);
+				ChunkedDataNode oldFile = await oldFileEntry.ExpandAsync();
+				ChunkedDataNode newFile = await newFileEntry.ExpandAsync();
+				await CompareTrees(oldFile, newFile);
 			}
 		}
 
-		static async Task CompareTrees(TreeReader reader, ChunkedDataNode oldNode, ChunkedDataNode newNode)
+		static async Task CompareTrees(ChunkedDataNode oldNode, ChunkedDataNode newNode)
 		{
 			if (oldNode is InteriorChunkedDataNode oldInteriorNode)
 			{
@@ -419,9 +410,9 @@ namespace EpicGames.Horde.Tests
 				int index = 0;
 				foreach ((NodeRef<ChunkedDataNode> oldFileRef, NodeRef<ChunkedDataNode> newFileRef) in oldInteriorNode.Children.Zip(newInteriorNode.Children))
 				{
-					ChunkedDataNode oldFile = await oldFileRef.ExpandAsync(reader);
-					ChunkedDataNode newFile = await newFileRef.ExpandAsync(reader);
-					await CompareTrees(reader, oldFile, newFile);
+					ChunkedDataNode oldFile = await oldFileRef.ExpandAsync();
+					ChunkedDataNode newFile = await newFileRef.ExpandAsync();
+					await CompareTrees(oldFile, newFile);
 					index++;
 				}
 			}
@@ -436,18 +427,18 @@ namespace EpicGames.Horde.Tests
 			}
 		}
 
-		static async Task CheckLargeFileTreeAsync(TreeReader reader, DirectoryNode root, byte[] data)
+		static async Task CheckLargeFileTreeAsync(DirectoryNode root, byte[] data)
 		{
 			Assert.AreEqual(0, root.Directories.Count);
 			Assert.AreEqual(1, root.Files.Count);
 
-			ChunkedDataNode world = await root.Files.First().ExpandAsync(reader, CancellationToken.None);
+			ChunkedDataNode world = await root.Files.First().ExpandAsync(CancellationToken.None);
 
-			int length = await CheckFileDataAsync(reader, world, data);
+			int length = await CheckFileDataAsync(world, data);
 			Assert.AreEqual(data.Length, length);
 		}
 
-		static async Task<int> CheckFileDataAsync(TreeReader reader, ChunkedDataNode fileNode, ReadOnlyMemory<byte> data)
+		static async Task<int> CheckFileDataAsync(ChunkedDataNode fileNode, ReadOnlyMemory<byte> data)
 		{
 			int offset = 0;
 			if (fileNode is LeafChunkedDataNode leafNode)
@@ -459,8 +450,8 @@ namespace EpicGames.Horde.Tests
 			{
 				foreach (NodeRef<ChunkedDataNode> childRef in interiorFileNode.Children)
 				{
-					ChunkedDataNode child = await childRef.ExpandAsync(reader);
-					offset += await CheckFileDataAsync(reader, child, data.Slice(offset));
+					ChunkedDataNode child = await childRef.ExpandAsync();
+					offset += await CheckFileDataAsync(child, data.Slice(offset));
 				}
 			}
 			else
@@ -470,11 +461,11 @@ namespace EpicGames.Horde.Tests
 			return offset;
 		}
 
-		static async Task<byte[]> GetFileDataAsync(TreeReader reader, ChunkedDataNode fileNode)
+		static async Task<byte[]> GetFileDataAsync(ChunkedDataNode fileNode)
 		{
 			using (MemoryStream stream = new MemoryStream())
 			{
-				await fileNode.CopyToStreamAsync(reader, stream, CancellationToken.None);
+				await fileNode.CopyToStreamAsync(stream, CancellationToken.None);
 				return stream.ToArray();
 			}
 		}
diff --git a/Engine/Source/Programs/Shared/EpicGames.Horde.Tests/BundleTests.cs b/Engine/Source/Programs/Shared/EpicGames.Horde.Tests/BundleTests.cs
--- a/Engine/Source/Programs/Shared/EpicGames.Horde.Tests/BundleTests.cs
+++ b/Engine/Source/Programs/Shared/EpicGames.Horde.Tests/BundleTests.cs
@@ -98,12 +98,12 @@ namespace EpicGames.Horde.Tests
 			Assert.AreEqual(2, ((InteriorChunkedDataNode)node).Children.Count);
 			await TestBufferlessReadsAsync(reader, handle.Handle.Locator, data.AsMemory(0, 9));
 
-			ChunkedDataNode? childNode1 = await ((InteriorChunkedDataNode)node).Children[0].ExpandAsync(reader);
+			ChunkedDataNode? childNode1 = await ((InteriorChunkedDataNode)node).Children[0].ExpandAsync();
 			Assert.IsNotNull(childNode1);
 			Assert.IsTrue(childNode1 is LeafChunkedDataNode);
 			Assert.AreEqual(8, ((LeafChunkedDataNode)childNode1!).Data.Length);
 
-			ChunkedDataNode? childNode2 = await ((InteriorChunkedDataNode)node).Children[1].ExpandAsync(reader);
+			ChunkedDataNode? childNode2 = await ((InteriorChunkedDataNode)node).Children[1].ExpandAsync();
 			Assert.IsNotNull(childNode2);
 			Assert.IsTrue(childNode2 is LeafChunkedDataNode);
 			Assert.AreEqual(1, ((LeafChunkedDataNode)childNode2!).Data.Length);
@@ -190,7 +190,6 @@ namespace EpicGames.Horde.Tests
 		public async Task BasicTestDirectory()
 		{
 			MemoryStorageClient store = _storage;
-			TreeReader reader = new TreeReader(store, null, NullLogger.Instance);
 
 			DirectoryNode root = new DirectoryNode(DirectoryFlags.None);
 			DirectoryNode node = root.AddDirectory("hello");
@@ -210,16 +209,16 @@ namespace EpicGames.Horde.Tests
 			Assert.AreEqual(3, bundle.Header.Exports.Count);
 
 			// Create a new bundle and read it back in again
-			DirectoryNode newRoot = await reader.ReadNodeAsync<DirectoryNode>(refName);
+			DirectoryNode newRoot = await store.ReadNodeAsync<DirectoryNode>(refName);
 
 			Assert.AreEqual(0, newRoot.Files.Count);
 			Assert.AreEqual(1, newRoot.Directories.Count);
-			DirectoryNode? outputNode = await newRoot.FindDirectoryAsync(reader, "hello", CancellationToken.None);
+			DirectoryNode? outputNode = await newRoot.FindDirectoryAsync("hello", CancellationToken.None);
 			Assert.IsNotNull(outputNode);
 
 			Assert.AreEqual(0, outputNode!.Files.Count);
 			Assert.AreEqual(1, outputNode!.Directories.Count);
-			DirectoryNode? outputNode2 = await outputNode.FindDirectoryAsync(reader, "world", CancellationToken.None);
+			DirectoryNode? outputNode2 = await outputNode.FindDirectoryAsync("world", CancellationToken.None);
 			Assert.IsNotNull(outputNode2);
 
 			Assert.AreEqual(0, outputNode2!.Files.Count);
@@ -270,20 +269,18 @@ namespace EpicGames.Horde.Tests
 			}
 
 			{
-				TreeReader reader = new TreeReader(_storage, null, NullLogger.Instance);
+				DirectoryNode root = await _storage.ReadNodeAsync<DirectoryNode>(refName);
 
-				DirectoryNode root = await reader.ReadNodeAsync<DirectoryNode>(refName);
-
-				DirectoryNode? newNode1 = await root.FindDirectoryAsync(reader, "node1", CancellationToken.None);
+				DirectoryNode? newNode1 = await root.FindDirectoryAsync("node1", CancellationToken.None);
 				Assert.IsNotNull(newNode1);
 
-				DirectoryNode? newNode2 = await newNode1!.FindDirectoryAsync(reader, "node2", CancellationToken.None);
+				DirectoryNode? newNode2 = await newNode1!.FindDirectoryAsync("node2", CancellationToken.None);
 				Assert.IsNotNull(newNode2);
 
-				DirectoryNode? newNode3 = await newNode2!.FindDirectoryAsync(reader, "node3", CancellationToken.None);
+				DirectoryNode? newNode3 = await newNode2!.FindDirectoryAsync("node3", CancellationToken.None);
 				Assert.IsNotNull(newNode3);
 
-				DirectoryNode? newNode4 = await newNode3!.FindDirectoryAsync(reader, "node4", CancellationToken.None);
+				DirectoryNode? newNode4 = await newNode3!.FindDirectoryAsync("node4", CancellationToken.None);
 				Assert.IsNotNull(newNode4);
 			}
 		}
diff --git a/Engine/Source/Programs/Shared/EpicGames.Horde.Tests/ChunkingTests.cs b/Engine/Source/Programs/Shared/EpicGames.Horde.Tests/ChunkingTests.cs
--- a/Engine/Source/Programs/Shared/EpicGames.Horde.Tests/ChunkingTests.cs
+++ b/Engine/Source/Programs/Shared/EpicGames.Horde.Tests/ChunkingTests.cs
@@ -102,7 +102,7 @@ namespace EpicGames.Horde.Tests
 			byte[] result;
 			using (MemoryStream stream = new MemoryStream())
 			{
-				await root.CopyToStreamAsync(reader, stream, CancellationToken.None);
+				await root.CopyToStreamAsync(stream, CancellationToken.None);
 				result = stream.ToArray();
 			}
 
@@ -134,7 +134,7 @@ namespace EpicGames.Horde.Tests
 				int childCount = interiorNode.Children.Count;
 				for (int idx = 0; idx < childCount; idx++)
 				{
-					ChunkedDataNode childNode = await interiorNode.Children[idx].ExpandAsync(reader, CancellationToken.None);
+					ChunkedDataNode childNode = await interiorNode.Children[idx].ExpandAsync(CancellationToken.None);
 					await CheckSizes(reader, childNode, options, idx == childCount - 1);
 				}
 			}
diff --git a/Engine/Source/Programs/Shared/EpicGames.Horde.Tests/LogTests.cs b/Engine/Source/Programs/Shared/EpicGames.Horde.Tests/LogTests.cs
--- a/Engine/Source/Programs/Shared/EpicGames.Horde.Tests/LogTests.cs
+++ b/Engine/Source/Programs/Shared/EpicGames.Horde.Tests/LogTests.cs
@@ -20,24 +20,15 @@ using Microsoft.VisualStudio.TestTools.UnitTesting;
 namespace EpicGames.Horde.Tests
 {
 	[TestClass]
-	public sealed class LogTests : IDisposable
+	public sealed class LogTests
 	{
-		readonly IMemoryCache _cache;
 		readonly MemoryStorageClient _store;
-		readonly TreeReader _reader;
 
 		private readonly byte[] _data = Resources.TextFile;
 
 		public LogTests()
 		{
-			_cache = new MemoryCache(new MemoryCacheOptions());
 			_store = new MemoryStorageClient();
-			_reader = new TreeReader(_store, _cache, NullLogger.Instance);
-		}
-
-		public void Dispose()
-		{
-			_cache.Dispose();
 		}
 
 		[TestMethod]
@@ -87,18 +78,19 @@ namespace EpicGames.Horde.Tests
 			}
 
 			// Flush it to storage, and read the finished log node
-			HashedNodeHandle target;
+			NodeRef<LogNode> logRef;
 			using (IStorageWriter writer = _store.CreateWriter())
 			{
-				target = await builder.FlushAsync(writer, true, CancellationToken.None);
+				HashedNodeHandle target = await builder.FlushAsync(writer, true, CancellationToken.None);
+				logRef = new NodeRef<LogNode>(target);
 			}
-			LogNode log = await _reader.ReadNodeAsync<LogNode>(target.Handle.Locator);
+			LogNode log = await logRef.ExpandAsync();
 
 			// Read the data back out and check it's the same
 			byte[] readData = new byte[_data.Length];
 
 			int writeOffset = 0;
-			await foreach (ReadOnlyMemory<byte> line in log.ReadLogLinesAsync(_reader, 0, CancellationToken.None))
+			await foreach (ReadOnlyMemory<byte> line in log.ReadLogLinesAsync(0, CancellationToken.None))
 			{
 				line.CopyTo(readData.AsMemory(writeOffset));
 				writeOffset += line.Length;
@@ -112,7 +104,7 @@ namespace EpicGames.Horde.Tests
 			Assert.AreEqual(equalSize, readOffset);
 
 			// Test some searches
-			LogIndexNode index = await log.IndexRef.ExpandAsync(_reader);
+			LogIndexNode index = await log.IndexRef.ExpandAsync();
 			await SearchLogDataTestAsync(index);
 		}
 
@@ -134,15 +126,16 @@ namespace EpicGames.Horde.Tests
 				builder.WriteData(Encoding.UTF8.GetBytes(lines[lineIdx]));
 			}
 
-			HashedNodeHandle rootNodeHandle;
+			NodeRef<LogNode> rootNodeRef;
 			using (IStorageWriter writer = _store.CreateWriter())
 			{
-				rootNodeHandle = await builder.FlushAsync(writer, true, CancellationToken.None);
+				HashedNodeHandle rootNodeHandle = await builder.FlushAsync(writer, true, CancellationToken.None);
+				rootNodeRef = new NodeRef<LogNode>(rootNodeHandle);
 			}
 
 			// Read it back in and test the index
-			LogNode rootNode = await _reader.ReadNodeAsync<LogNode>(rootNodeHandle.Handle.Locator);
-			LogIndexNode index = await rootNode.IndexRef.ExpandAsync(_reader);
+			LogNode rootNode = await rootNodeRef.ExpandAsync();
+			LogIndexNode index = await rootNode.IndexRef.ExpandAsync();
 		
 			for (int lineIdx = 0; lineIdx < lines.Length; lineIdx++)
 			{
@@ -153,7 +146,7 @@ namespace EpicGames.Horde.Tests
 						string str = lines[lineIdx].Substring(strOfs, strLen);
 
 						SearchStats stats = new SearchStats();
-						List<int> results = await index.Search(_reader, 0, new SearchTerm(str), stats, CancellationToken.None).ToListAsync();
+						List<int> results = await index.Search(0, new SearchTerm(str), stats, CancellationToken.None).ToListAsync();
 						Assert.AreEqual(1, results.Count);
 						Assert.AreEqual(lineIdx, results[0]);
 
@@ -209,7 +202,7 @@ namespace EpicGames.Horde.Tests
 			}
 		}
 */
-		async Task SearchLogDataTestAsync(LogIndexNode index)
+		static async Task SearchLogDataTestAsync(LogIndexNode index)
 		{
 			await SearchLogDataTestAsync(index, "HISPANIOLA", 0, 4, new[] { 1503, 1520, 1525, 1595 });
 			await SearchLogDataTestAsync(index, "Hispaniola", 0, 4, new[] { 1503, 1520, 1525, 1595 });
@@ -218,10 +211,10 @@ namespace EpicGames.Horde.Tests
 			await SearchLogDataTestAsync(index, "NEWSLETTER", 0, 100, new[] { 7886 });
 		}
 
-		async Task SearchLogDataTestAsync(LogIndexNode index, string text, int firstLine, int count, int[] expectedLines)
+		static async Task SearchLogDataTestAsync(LogIndexNode index, string text, int firstLine, int count, int[] expectedLines)
 		{
 			SearchStats stats = new SearchStats();
-			List<int> lines = await index.Search(_reader, firstLine, new SearchTerm(text), stats, CancellationToken.None).Take(count).ToListAsync();
+			List<int> lines = await index.Search(firstLine, new SearchTerm(text), stats, CancellationToken.None).Take(count).ToListAsync();
 			Assert.IsTrue(lines.SequenceEqual(expectedLines));
 		}
 	}
diff --git a/Engine/Source/Programs/Shared/EpicGames.Horde/Compute/AgentMessageHandler.cs b/Engine/Source/Programs/Shared/EpicGames.Horde/Compute/AgentMessageHandler.cs
--- a/Engine/Source/Programs/Shared/EpicGames.Horde/Compute/AgentMessageHandler.cs
+++ b/Engine/Source/Programs/Shared/EpicGames.Horde/Compute/AgentMessageHandler.cs
@@ -129,7 +129,7 @@ namespace EpicGames.Horde.Compute
 				throw new InvalidOperationException("Cannot write files outside sandbox");
 			}
 
-			await directoryNode.CopyToDirectoryAsync(reader, outputDir.ToDirectoryInfo(), _logger, cancellationToken);
+			await directoryNode.CopyToDirectoryAsync(outputDir.ToDirectoryInfo(), _logger, cancellationToken);
 
 			using (IAgentMessageBuilder message = await channel.CreateMessageAsync(AgentMessageType.WriteFilesResponse, cancellationToken))
 			{
diff --git a/Engine/Source/Programs/Shared/EpicGames.Horde/Compute/AgentStorageClient.cs b/Engine/Source/Programs/Shared/EpicGames.Horde/Compute/AgentStorageClient.cs
--- a/Engine/Source/Programs/Shared/EpicGames.Horde/Compute/AgentStorageClient.cs
+++ b/Engine/Source/Programs/Shared/EpicGames.Horde/Compute/AgentStorageClient.cs
@@ -14,7 +14,7 @@ namespace EpicGames.Horde.Compute
 	/// <summary>
 	/// Storage client which can read bundles over a compute channel
 	/// </summary>
-	public sealed class AgentStorageClient : StorageClientBase, IDisposable
+	public sealed class AgentStorageClient : BundleStorageClient, IDisposable
 	{
 		readonly AgentMessageChannel _channel;
 		readonly SemaphoreSlim _semaphore;
diff --git a/Engine/Source/Programs/Shared/EpicGames.Horde/Logs/LogIndexNode.cs b/Engine/Source/Programs/Shared/EpicGames.Horde/Logs/LogIndexNode.cs
--- a/Engine/Source/Programs/Shared/EpicGames.Horde/Logs/LogIndexNode.cs
+++ b/Engine/Source/Programs/Shared/EpicGames.Horde/Logs/LogIndexNode.cs
@@ -174,13 +174,12 @@ namespace EpicGames.Horde.Logs
 		/// <summary>
 		/// Search for the given text in the index
 		/// </summary>
-		/// <param name="reader">Reader for node data</param>
 		/// <param name="firstLineIndex">First line index to search from</param>
 		/// <param name="text">Text to search for</param>
 		/// <param name="stats">Receives stats for the search</param>
 		/// <param name="cancellationToken">Cancellation token for the operation</param>
 		/// <returns>List of line numbers for the text</returns>
-		public async IAsyncEnumerable<int> Search(TreeReader reader, int firstLineIndex, SearchTerm text, SearchStats stats, [EnumeratorCancellation] CancellationToken cancellationToken)
+		public async IAsyncEnumerable<int> Search(int firstLineIndex, SearchTerm text, SearchStats stats, [EnumeratorCancellation] CancellationToken cancellationToken)
 		{
 			int lastBlockCount = 0;
 			foreach (int blockIdx in EnumeratePossibleChunks(text.Bytes, firstLineIndex))
@@ -194,7 +193,7 @@ namespace EpicGames.Horde.Logs
 				lastBlockCount = blockIdx + 1;
 
 				// Decompress the text
-				LogChunkNode chunk = await indexChunk.ExpandAsync(reader, cancellationToken);
+				LogChunkNode chunk = await indexChunk.ExpandAsync(cancellationToken);
 
 				// Find the initial offset within this block
 				int offset = 0;
diff --git a/Engine/Source/Programs/Shared/EpicGames.Horde/Logs/LogNode.cs b/Engine/Source/Programs/Shared/EpicGames.Horde/Logs/LogNode.cs
--- a/Engine/Source/Programs/Shared/EpicGames.Horde/Logs/LogNode.cs
+++ b/Engine/Source/Programs/Shared/EpicGames.Horde/Logs/LogNode.cs
@@ -330,7 +330,7 @@ namespace EpicGames.Horde.Logs
 		{
 			foreach (LogChunkRef textChunkRef in logNode.TextChunkRefs)
 			{
-				LogChunkNode textChunk = await textChunkRef.ExpandAsync(reader, cancellationToken);
+				LogChunkNode textChunk = await textChunkRef.ExpandAsync(cancellationToken);
 				yield return textChunk.Data;
 			}
 		}
@@ -339,18 +339,17 @@ namespace EpicGames.Horde.Logs
 		/// Reads lines from a line
 		/// </summary>
 		/// <param name="logNode">Log to read from</param>
-		/// <param name="reader">Reader to pull nodes from</param>
 		/// <param name="index">Zero-based index of the first line to read from</param>
 		/// <param name="cancellationToken">Cancellation token</param>
 		/// <returns>Sequence of line buffers</returns>
-		public static async IAsyncEnumerable<ReadOnlyMemory<byte>> ReadLogLinesAsync(this LogNode logNode, TreeReader reader, int index, [EnumeratorCancellation] CancellationToken cancellationToken)
+		public static async IAsyncEnumerable<ReadOnlyMemory<byte>> ReadLogLinesAsync(this LogNode logNode, int index, [EnumeratorCancellation] CancellationToken cancellationToken)
 		{
 			foreach (LogChunkRef textChunkRef in logNode.TextChunkRefs)
 			{
 				int lineIdx = Math.Max(index - textChunkRef.LineIndex, 0);
 				if (lineIdx < textChunkRef.LineCount)
 				{
-					LogChunkNode textChunk = await textChunkRef.ExpandAsync(reader, cancellationToken);
+					LogChunkNode textChunk = await textChunkRef.ExpandAsync(cancellationToken);
 
 					int offset = textChunk.LineOffsets[lineIdx];
 					for (; lineIdx < textChunk.LineCount; lineIdx++)
diff --git a/Engine/Source/Programs/Shared/EpicGames.Horde/Storage/Backends/FileStorageClient.cs b/Engine/Source/Programs/Shared/EpicGames.Horde/Storage/Backends/FileStorageClient.cs
--- a/Engine/Source/Programs/Shared/EpicGames.Horde/Storage/Backends/FileStorageClient.cs
+++ b/Engine/Source/Programs/Shared/EpicGames.Horde/Storage/Backends/FileStorageClient.cs
@@ -17,7 +17,7 @@ namespace EpicGames.Horde.Storage.Backends
 	/// <summary>
 	/// Implementation of <see cref="IStorageClient"/> which writes data to files on disk.
 	/// </summary>
-	public class FileStorageClient : StorageClientBase
+	public class FileStorageClient : BundleStorageClient
 	{
 		readonly DirectoryReference _rootDir;
 		readonly ILogger _logger;
diff --git a/Engine/Source/Programs/Shared/EpicGames.Horde/Storage/Backends/HttpStorageClient.cs b/Engine/Source/Programs/Shared/EpicGames.Horde/Storage/Backends/HttpStorageClient.cs
--- a/Engine/Source/Programs/Shared/EpicGames.Horde/Storage/Backends/HttpStorageClient.cs
+++ b/Engine/Source/Programs/Shared/EpicGames.Horde/Storage/Backends/HttpStorageClient.cs
@@ -20,7 +20,7 @@ namespace EpicGames.Horde.Storage.Backends
 	/// <summary>
 	/// Implementation of <see cref="IStorageClient"/> which communicates with an upstream Horde instance via HTTP.
 	/// </summary>
-	public class HttpStorageClient : StorageClientBase
+	public class HttpStorageClient : BundleStorageClient
 	{
 		/// <summary>
 		/// Name of clients created from the http client factory
diff --git a/Engine/Source/Programs/Shared/EpicGames.Horde/Storage/Backends/MemoryStorageClient.cs b/Engine/Source/Programs/Shared/EpicGames.Horde/Storage/Backends/MemoryStorageClient.cs
--- a/Engine/Source/Programs/Shared/EpicGames.Horde/Storage/Backends/MemoryStorageClient.cs
+++ b/Engine/Source/Programs/Shared/EpicGames.Horde/Storage/Backends/MemoryStorageClient.cs
@@ -18,7 +18,7 @@ namespace EpicGames.Horde.Storage.Backends
 	/// <summary>
 	/// Implementation of <see cref="IStorageClient"/> which stores data in memory. Not intended for production use.
 	/// </summary>
-	public class MemoryStorageClient : StorageClientBase
+	public class MemoryStorageClient : BundleStorageClient
 	{
 		record class ExportEntry(NodeHandle Handle, ExportEntry? Next);
 
diff --git a/Engine/Source/Programs/Shared/EpicGames.Horde/Storage/BundleStorageClient.cs b/Engine/Source/Programs/Shared/EpicGames.Horde/Storage/BundleStorageClient.cs
new file mode 100644
--- /dev/null
+++ b/Engine/Source/Programs/Shared/EpicGames.Horde/Storage/BundleStorageClient.cs
@@ -0,0 +1,90 @@
+// Copyright Epic Games, Inc. All Rights Reserved.
+
+using System;
+using System.Collections.Generic;
+using System.IO;
+using System.Threading;
+using System.Threading.Tasks;
+using EpicGames.Core;
+using Microsoft.Extensions.Caching.Memory;
+using Microsoft.Extensions.Logging;
+
+namespace EpicGames.Horde.Storage
+{
+	/// <summary>
+	/// Base class for an implementation of <see cref="IStorageClient"/>, providing implementations for some common functionality using bundles.
+	/// </summary>
+	public abstract class BundleStorageClient : IStorageClient
+	{
+		/// <summary>
+		/// Reader for node data
+		/// </summary>
+		protected TreeReader TreeReader { get; }
+
+		/// <summary>
+		/// Constructor
+		/// </summary>
+		protected BundleStorageClient(IMemoryCache? memoryCache, ILogger logger)
+		{
+			TreeReader = new TreeReader(this, memoryCache, logger);
+		}
+
+		#region Blobs
+
+		/// <inheritdoc/>
+		public abstract Task<Stream> ReadBlobAsync(BlobLocator locator, CancellationToken cancellationToken = default);
+
+		/// <inheritdoc/>
+		public abstract Task<Stream> ReadBlobRangeAsync(BlobLocator locator, int offset, int length, CancellationToken cancellationToken = default);
+
+		/// <inheritdoc/>
+		public abstract Task<BlobLocator> WriteBlobAsync(Stream stream, Utf8String prefix = default, CancellationToken cancellationToken = default);
+
+		#endregion
+
+		#region Nodes
+
+		/// <inheritdoc/>
+		public IStorageWriter CreateWriter(RefName refName = default, TreeOptions? options = null)
+		{
+			return new TreeWriter(this, TreeReader, refName, options);
+		}
+
+		#endregion
+
+		#region Aliases
+
+		/// <inheritdoc/>
+		public abstract Task AddAliasAsync(Utf8String name, NodeHandle locator, CancellationToken cancellationToken = default);
+
+		/// <inheritdoc/>
+		public abstract Task RemoveAliasAsync(Utf8String name, NodeHandle locator, CancellationToken cancellationToken = default);
+
+		/// <inheritdoc/>
+		public abstract IAsyncEnumerable<NodeHandle> FindNodesAsync(Utf8String name, CancellationToken cancellationToken = default);
+
+		#endregion
+
+		#region Refs
+
+		/// <inheritdoc/>
+		public abstract Task DeleteRefAsync(RefName name, CancellationToken cancellationToken = default);
+
+		/// <inheritdoc/>
+		public abstract Task<NodeHandle?> TryReadRefTargetAsync(RefName name, DateTime cacheTime = default, CancellationToken cancellationToken = default);
+
+		/// <inheritdoc/>
+		public virtual async Task<NodeHandle> WriteRefAsync(RefName name, Bundle bundle, int exportIdx = 0, Utf8String prefix = default, RefOptions? options = null, CancellationToken cancellationToken = default)
+		{
+			BlobLocator locator = await this.WriteBundleAsync(bundle, prefix, cancellationToken);
+			NodeHandle target = new NodeHandle(TreeReader, new NodeLocator(locator, exportIdx));
+			await WriteRefTargetAsync(name, target, options, cancellationToken);
+			return target;
+		}
+
+		/// <inheritdoc/>
+		public abstract Task WriteRefTargetAsync(RefName name, NodeHandle target, RefOptions? options = null, CancellationToken cancellationToken = default);
+
+		#endregion
+	}
+}
diff --git a/Engine/Source/Programs/Shared/EpicGames.Horde/Storage/Node.cs b/Engine/Source/Programs/Shared/EpicGames.Horde/Storage/Node.cs
--- a/Engine/Source/Programs/Shared/EpicGames.Horde/Storage/Node.cs
+++ b/Engine/Source/Programs/Shared/EpicGames.Horde/Storage/Node.cs
@@ -160,6 +160,13 @@ namespace EpicGames.Horde.Storage
 			return s_guidToDeserializer[reader.Type.Guid](reader);
 		}
 
+		/// <summary>
+		/// Deserialize a node from the given reader
+		/// </summary>
+		/// <param name="nodeData">Data to deserialize from</param>
+		/// <returns>New node instance</returns>
+		public static TNode Deserialize<TNode>(NodeData nodeData) where TNode : Node => (TNode)Deserialize(nodeData);
+
 		/// <summary>
 		/// Static constructor. Registers all the types in the current assembly.
 		/// </summary>
@@ -262,6 +269,19 @@ namespace EpicGames.Horde.Storage
 	/// </summary>
 	public static class NodeExtensions
 	{
+		/// <summary>
+		/// Reads and deserializes a node from storage
+		/// </summary>
+		/// <typeparam name="TNode"></typeparam>
+		/// <param name="handle"></param>
+		/// <param name="cancellationToken"></param>
+		/// <returns></returns>
+		public static async ValueTask<TNode> ReadAsync<TNode>(this NodeHandle handle, CancellationToken cancellationToken = default) where TNode : Node
+		{
+			NodeData nodeData = await handle.ReadAsync(cancellationToken);
+			return Node.Deserialize<TNode>(nodeData);
+		}
+
 		/// <summary>
 		/// Read an untyped ref from the reader
 		/// </summary>
@@ -360,5 +380,43 @@ namespace EpicGames.Horde.Storage
 			using IStorageWriter writer = store.CreateWriter(name, options);
 			return await writer.WriteAsync(name, node, refOptions, cancellationToken);
 		}
+
+		/// <summary>
+		/// Reads data for a ref from the store, along with the node's contents.
+		/// </summary>
+		/// <param name="store">Store instance to write to</param>
+		/// <param name="name">The ref name</param>
+		/// <param name="cacheTime">Minimum coherency for any cached value to be returned</param>
+		/// <param name="cancellationToken">Cancellation token for the operation</param>
+		/// <returns>Node for the given ref, or null if it does not exist</returns>
+		public static async Task<TNode?> TryReadNodeAsync<TNode>(this IStorageClient store, RefName name, DateTime cacheTime = default, CancellationToken cancellationToken = default) where TNode : Node
+		{
+			NodeHandle? refTarget = await store.TryReadRefTargetAsync(name, cacheTime, cancellationToken);
+			if (refTarget == null)
+			{
+				return null;
+			}
+
+			NodeData nodeData = await refTarget.ReadAsync(cancellationToken);
+			return Node.Deserialize<TNode>(nodeData);
+		}
+
+		/// <summary>
+		/// Reads a ref from the store, throwing an exception if it does not exist
+		/// </summary>
+		/// <param name="store">Store instance to write to</param>
+		/// <param name="name">Id for the ref</param>
+		/// <param name="cacheTime">Minimum coherency of any cached result</param>
+		/// <param name="cancellationToken">Cancellation token for the operation</param>
+		/// <returns>The blob instance</returns>
+		public static async Task<TNode> ReadNodeAsync<TNode>(this IStorageClient store, RefName name, DateTime cacheTime = default, CancellationToken cancellationToken = default) where TNode : Node
+		{
+			TNode? refValue = await store.TryReadNodeAsync<TNode>(name, cacheTime, cancellationToken);
+			if (refValue == null)
+			{
+				throw new RefNameNotFoundException(name);
+			}
+			return refValue;
+		}
 	}
 }
diff --git a/Engine/Source/Programs/Shared/EpicGames.Horde/Storage/NodeRef.cs b/Engine/Source/Programs/Shared/EpicGames.Horde/Storage/NodeRef.cs
--- a/Engine/Source/Programs/Shared/EpicGames.Horde/Storage/NodeRef.cs
+++ b/Engine/Source/Programs/Shared/EpicGames.Horde/Storage/NodeRef.cs
@@ -34,7 +34,7 @@ namespace EpicGames.Horde.Storage
 	/// Multiple refs MAY point to the same target object.
 	/// 
 	/// To read an untracked object that can be added to a new ref, call <see cref="TreeReader.ReadNodeDataAsync(NodeLocator, CancellationToken)"/> 
-	/// directly, or use <see cref="NodeRef{T}.ExpandCopyAsync(TreeReader, CancellationToken)"/>.
+	/// directly, or use <see cref="NodeRef{T}.ExpandCopyAsync(CancellationToken)"/>.
 	/// </summary>
 	public class NodeRef
 	{
@@ -127,14 +127,13 @@ namespace EpicGames.Horde.Storage
 		/// <summary>
 		/// Resolve this reference to a concrete node
 		/// </summary>
-		/// <param name="reader">Reader to use for expanding this ref</param>
 		/// <param name="cancellationToken">Cancellation token for the operation</param>
 		/// <returns></returns>
-		public async ValueTask<Node> ExpandAsync(TreeReader reader, CancellationToken cancellationToken = default)
+		public async ValueTask<Node> ExpandAsync(CancellationToken cancellationToken = default)
 		{
 			if (Target == null)
 			{
-				NodeData nodeData = await reader.ReadNodeDataAsync(Handle!.Handle.Locator, cancellationToken);
+				NodeData nodeData = await Handle!.Handle.ReadAsync(cancellationToken);
 				Target = Node.Deserialize(nodeData);
 				OnExpand();
 			}
@@ -206,27 +205,27 @@ namespace EpicGames.Horde.Storage
 		/// <summary>
 		/// Resolve this reference to a concrete node
 		/// </summary>
-		/// <param name="reader">Reader to use for expanding this ref</param>
 		/// <param name="cancellationToken">Cancellation token for the operation</param>
 		/// <returns></returns>
-		public new async ValueTask<T> ExpandAsync(TreeReader reader, CancellationToken cancellationToken = default)
+		public new async ValueTask<T> ExpandAsync(CancellationToken cancellationToken = default)
 		{
-			return (T)await base.ExpandAsync(reader, cancellationToken);
+			return (T)await base.ExpandAsync(cancellationToken);
 		}
 
 		/// <summary>
 		/// Resolve this reference to a concrete node
 		/// </summary>
-		/// <param name="reader">Reader to use for expanding this ref</param>
 		/// <param name="cancellationToken">Cancellation token for the operation</param>
 		/// <returns></returns>
-		public async ValueTask<T> ExpandCopyAsync(TreeReader reader, CancellationToken cancellationToken = default)
+		public async ValueTask<T> ExpandCopyAsync(CancellationToken cancellationToken = default)
 		{
 			if (Handle == null)
 			{
 				throw new InvalidOperationException("TreeNodeRef has not been serialized to storage");
 			}
-			return await reader.ReadNodeAsync<T>(Handle.Handle.Locator, cancellationToken);
+
+			NodeData nodeData = await Handle.Handle.ReadAsync(cancellationToken);
+			return Node.Deserialize<T>(nodeData);
 		}
 	}
 
diff --git a/Engine/Source/Programs/Shared/EpicGames.Horde/Storage/Nodes/ChunkedDataNode.cs b/Engine/Source/Programs/Shared/EpicGames.Horde/Storage/Nodes/ChunkedDataNode.cs
--- a/Engine/Source/Programs/Shared/EpicGames.Horde/Storage/Nodes/ChunkedDataNode.cs
+++ b/Engine/Source/Programs/Shared/EpicGames.Horde/Storage/Nodes/ChunkedDataNode.cs
@@ -20,10 +20,9 @@ namespace EpicGames.Horde.Storage.Nodes
 		/// <summary>
 		/// Copies the contents of this node and its children to the given output stream
 		/// </summary>
-		/// <param name="reader">Reader for nodes in the tree</param>
 		/// <param name="outputStream">The output stream to receive the data</param>
 		/// <param name="cancellationToken">Cancellation token for the operation</param>
-		public abstract Task CopyToStreamAsync(TreeReader reader, Stream outputStream, CancellationToken cancellationToken);
+		public abstract Task CopyToStreamAsync(Stream outputStream, CancellationToken cancellationToken);
 
 		static readonly Guid s_leafNodeGuid = Node.GetNodeType<LeafChunkedDataNode>().Guid;
 		static readonly Guid s_interiorNodeGuid = Node.GetNodeType<InteriorChunkedDataNode>().Guid;
@@ -31,20 +30,19 @@ namespace EpicGames.Horde.Storage.Nodes
 		/// <summary>
 		/// Copy the contents of the node to the output stream without creating the intermediate FileNodes
 		/// </summary>
-		/// <param name="reader">Reader for nodes in the tree</param>
-		/// <param name="locator">File node to be copied</param>
+		/// <param name="handle">Handle to the data to read</param>
 		/// <param name="outputStream">The output stream to receive the data</param>
 		/// <param name="cancellationToken">Cancellation token for the operation</param>
-		public static async Task CopyToStreamAsync(TreeReader reader, NodeLocator locator, Stream outputStream, CancellationToken cancellationToken)
+		public static async Task CopyToStreamAsync(NodeHandle handle, Stream outputStream, CancellationToken cancellationToken)
 		{
-			NodeData nodeData = await reader.ReadNodeDataAsync(locator, cancellationToken);
+			NodeData nodeData = await handle.ReadAsync(cancellationToken);
 			if (nodeData.Type.Guid == s_leafNodeGuid)
 			{
 				await LeafChunkedDataNode.CopyToStreamAsync(nodeData, outputStream, cancellationToken);
 			}
 			else if (nodeData.Type.Guid == s_interiorNodeGuid)
 			{
-				await InteriorChunkedDataNode.CopyToStreamAsync(reader, nodeData, outputStream, cancellationToken);
+				await InteriorChunkedDataNode.CopyToStreamAsync(nodeData, outputStream, cancellationToken);
 			}
 			else
 			{
@@ -55,11 +53,10 @@ namespace EpicGames.Horde.Storage.Nodes
 		/// <summary>
 		/// Extracts the contents of this node to a file
 		/// </summary>
-		/// <param name="reader">Reader for nodes in the tree</param>
 		/// <param name="file">File to write with the contents of this node</param>
 		/// <param name="cancellationToken">Cancellation token for the operation</param>
 		/// <returns></returns>
-		public async Task CopyToFileAsync(TreeReader reader, FileInfo file, CancellationToken cancellationToken)
+		public async Task CopyToFileAsync(FileInfo file, CancellationToken cancellationToken)
 		{
 			if(file.Exists && (file.Attributes & FileAttributes.ReadOnly) != 0)
 			{
@@ -67,21 +64,20 @@ namespace EpicGames.Horde.Storage.Nodes
 			}
 			using (FileStream stream = file.Open(FileMode.Create, FileAccess.Write, FileShare.Read))
 			{
-				await CopyToStreamAsync(reader, stream, cancellationToken);
+				await CopyToStreamAsync(stream, cancellationToken);
 			}
 		}
 
 		/// <summary>
 		/// Serialize this node and its children into a byte array
 		/// </summary>
-		/// <param name="reader">Reader for nodes in the tree</param>
 		/// <param name="cancellationToken"></param>
 		/// <returns>Array of data stored by the tree</returns>
-		public async Task<byte[]> ToByteArrayAsync(TreeReader reader, CancellationToken cancellationToken)
+		public async Task<byte[]> ToByteArrayAsync(CancellationToken cancellationToken)
 		{
 			using (MemoryStream stream = new MemoryStream())
 			{
-				await CopyToStreamAsync(reader, stream, cancellationToken);
+				await CopyToStreamAsync(stream, cancellationToken);
 				return stream.ToArray();
 			}
 		}
@@ -124,7 +120,7 @@ namespace EpicGames.Horde.Storage.Nodes
 		public override IEnumerable<NodeRef> EnumerateRefs() => Enumerable.Empty<NodeRef>();
 
 		/// <inheritdoc/>
-		public override async Task CopyToStreamAsync(TreeReader reader, Stream outputStream, CancellationToken cancellationToken)
+		public override async Task CopyToStreamAsync(Stream outputStream, CancellationToken cancellationToken)
 		{
 			await outputStream.WriteAsync(Data, cancellationToken);
 		}
@@ -189,29 +185,28 @@ namespace EpicGames.Horde.Storage.Nodes
 		public override IEnumerable<NodeRef> EnumerateRefs() => Children;
 
 		/// <inheritdoc/>
-		public override async Task CopyToStreamAsync(TreeReader reader, Stream outputStream, CancellationToken cancellationToken)
+		public override async Task CopyToStreamAsync(Stream outputStream, CancellationToken cancellationToken)
 		{
 			foreach (NodeRef<ChunkedDataNode> childNodeRef in Children)
 			{
-				ChunkedDataNode childNode = await childNodeRef.ExpandAsync(reader, cancellationToken);
-				await childNode.CopyToStreamAsync(reader, outputStream, cancellationToken);
+				ChunkedDataNode childNode = await childNodeRef.ExpandAsync(cancellationToken);
+				await childNode.CopyToStreamAsync(outputStream, cancellationToken);
 			}
 		}
 
 		/// <summary>
 		/// Copy the contents of the node to the output stream without creating the intermediate FileNodes
 		/// </summary>
-		/// <param name="reader">Reader for nodes in the tree</param>
 		/// <param name="nodeData">Source data</param>
 		/// <param name="outputStream">The output stream to receive the data</param>
 		/// <param name="cancellationToken">Cancellation token for the operation</param>
-		public static async Task CopyToStreamAsync(TreeReader reader, NodeData nodeData, Stream outputStream, CancellationToken cancellationToken)
+		public static async Task CopyToStreamAsync(NodeData nodeData, Stream outputStream, CancellationToken cancellationToken)
 		{
 			NodeReader nodeReader = new NodeReader(nodeData);
 			while (nodeReader.GetMemory(0).Length > 0)
 			{
 				NodeRef nodeRef = nodeReader.ReadRef();
-				await ChunkedDataNode.CopyToStreamAsync(reader, nodeRef.Handle!.Handle.Locator, outputStream, cancellationToken);
+				await ChunkedDataNode.CopyToStreamAsync(nodeRef.Handle!.Handle, outputStream, cancellationToken);
 			}
 		}
 	}
diff --git a/Engine/Source/Programs/Shared/EpicGames.Horde/Storage/Nodes/DirectoryNode.cs b/Engine/Source/Programs/Shared/EpicGames.Horde/Storage/Nodes/DirectoryNode.cs
--- a/Engine/Source/Programs/Shared/EpicGames.Horde/Storage/Nodes/DirectoryNode.cs
+++ b/Engine/Source/Programs/Shared/EpicGames.Horde/Storage/Nodes/DirectoryNode.cs
@@ -225,14 +225,13 @@ namespace EpicGames.Horde.Storage.Nodes
 		/// <summary>
 		/// Finds or adds a file with the given path
 		/// </summary>
-		/// <param name="reader">Reader for node data</param>
 		/// <param name="path">Path to the file</param>
 		/// <param name="flags">Flags for the new file</param>
 		/// <param name="handle">The file node</param>
 		/// <param name="length">Length of the node</param>
 		/// <param name="cancellationToken">Cancellation token for the operation</param>
 		/// <returns>The new directory object</returns>
-		public async ValueTask<FileEntry> AddFileByPathAsync(TreeReader reader, Utf8String path, FileEntryFlags flags, long length, HashedNodeHandle handle, CancellationToken cancellationToken = default)
+		public async ValueTask<FileEntry> AddFileByPathAsync(Utf8String path, FileEntryFlags flags, long length, HashedNodeHandle handle, CancellationToken cancellationToken = default)
 		{
 			DirectoryNode directory = this;
 
@@ -261,7 +260,7 @@ namespace EpicGames.Horde.Storage.Nodes
 
 				if (pathLength > 0)
 				{
-					directory = await directory.FindOrAddDirectoryAsync(reader, remainingPath.Slice(0, pathLength), cancellationToken);
+					directory = await directory.FindOrAddDirectoryAsync(remainingPath.Slice(0, pathLength), cancellationToken);
 				}
 				remainingPath = remainingPath.Slice(pathLength + 1);
 			}
@@ -300,11 +299,10 @@ namespace EpicGames.Horde.Storage.Nodes
 		/// <summary>
 		/// Attempts to get a file entry from a path
 		/// </summary>
-		/// <param name="reader">Reader for node data</param>
 		/// <param name="path">Path to the directory</param>
 		/// <param name="cancellationToken">Cancellation token</param>
 		/// <returns>The directory with the given path, or null if it was not found</returns>
-		public async ValueTask<FileEntry?> GetFileEntryByPathAsync(TreeReader reader, Utf8String path, CancellationToken cancellationToken)
+		public async ValueTask<FileEntry?> GetFileEntryByPathAsync(Utf8String path, CancellationToken cancellationToken)
 		{
 			FileEntry? fileEntry;
 
@@ -318,7 +316,7 @@ namespace EpicGames.Horde.Storage.Nodes
 			}
 			else
 			{
-				DirectoryNode? directoryNode = await GetDirectoryByPathAsync(reader, path.Slice(0, slashIdx), cancellationToken);
+				DirectoryNode? directoryNode = await GetDirectoryByPathAsync(path.Slice(0, slashIdx), cancellationToken);
 				if (directoryNode == null)
 				{
 					return null;
@@ -335,13 +333,12 @@ namespace EpicGames.Horde.Storage.Nodes
 		/// <summary>
 		/// Attempts to get a directory entry from a path
 		/// </summary>
-		/// <param name="reader">Reader for node data</param>
 		/// <param name="path">Path to the directory</param>
 		/// <param name="cancellationToken">Cancellation token</param>
 		/// <returns>The directory with the given path, or null if it was not found</returns>
-		public ValueTask<DirectoryNode?> GetDirectoryByPathAsync(TreeReader reader, Utf8String path, CancellationToken cancellationToken) => GetDirectoryByPathAsync(this, reader, path, cancellationToken);
+		public ValueTask<DirectoryNode?> GetDirectoryByPathAsync(Utf8String path, CancellationToken cancellationToken) => GetDirectoryByPathAsync(this, path, cancellationToken);
 
-		static async ValueTask<DirectoryNode?> GetDirectoryByPathAsync(DirectoryNode directoryNode, TreeReader reader, Utf8String path, CancellationToken cancellationToken)
+		static async ValueTask<DirectoryNode?> GetDirectoryByPathAsync(DirectoryNode directoryNode, Utf8String path, CancellationToken cancellationToken)
 		{
 			while(path.Length > 0)
 			{
@@ -365,7 +362,7 @@ namespace EpicGames.Horde.Storage.Nodes
 					return null;
 				}
 
-				directoryNode = await directoryEntry.ExpandAsync(reader, cancellationToken);
+				directoryNode = await directoryEntry.ExpandAsync(cancellationToken);
 			}
 			return directoryNode;
 		}
@@ -373,11 +370,10 @@ namespace EpicGames.Horde.Storage.Nodes
 		/// <summary>
 		/// Deletes a file with the given path
 		/// </summary>
-		/// <param name="reader">Reader for existing node data</param>
 		/// <param name="path"></param>
 		/// <param name="cancellationToken"></param>
 		/// <returns></returns>
-		public async ValueTask<bool> DeleteFileByPathAsync(TreeReader reader, Utf8String path, CancellationToken cancellationToken)
+		public async ValueTask<bool> DeleteFileByPathAsync(Utf8String path, CancellationToken cancellationToken)
 		{
 			Utf8String remainingPath = path;
 			for (DirectoryNode? directory = this; directory != null;)
@@ -389,7 +385,7 @@ namespace EpicGames.Horde.Storage.Nodes
 				}
 				if (length > 0)
 				{
-					directory = await directory.FindDirectoryAsync(reader, remainingPath.Slice(0, length), cancellationToken);
+					directory = await directory.FindDirectoryAsync(remainingPath.Slice(0, length), cancellationToken);
 				}
 				remainingPath = remainingPath.Slice(length + 1);
 			}
@@ -462,15 +458,14 @@ namespace EpicGames.Horde.Storage.Nodes
 		/// <summary>
 		/// Tries to get a directory with the given name
 		/// </summary>
-		/// <param name="reader">Reader for node data</param>
 		/// <param name="name">Name of the new directory</param>
 		/// <param name="cancellationToken">Cancellation token for the operation</param>
 		/// <returns>The new directory object</returns>
-		public async ValueTask<DirectoryNode?> FindDirectoryAsync(TreeReader reader, Utf8String name, CancellationToken cancellationToken)
+		public async ValueTask<DirectoryNode?> FindDirectoryAsync(Utf8String name, CancellationToken cancellationToken)
 		{
 			if (TryGetDirectoryEntry(name, out DirectoryEntry? entry))
 			{
-				return await entry.ExpandAsync(reader, cancellationToken);
+				return await entry.ExpandAsync(cancellationToken);
 			}
 			else
 			{
@@ -481,13 +476,12 @@ namespace EpicGames.Horde.Storage.Nodes
 		/// <summary>
 		/// Tries to get a directory with the given name
 		/// </summary>
-		/// <param name="reader">Reader for node data</param>
 		/// <param name="name">Name of the new directory</param>
 		/// <param name="cancellationToken">Cancellation token for the operation</param>
 		/// <returns>The new directory object</returns>
-		public async ValueTask<DirectoryNode> FindOrAddDirectoryAsync(TreeReader reader, Utf8String name, CancellationToken cancellationToken)
+		public async ValueTask<DirectoryNode> FindOrAddDirectoryAsync(Utf8String name, CancellationToken cancellationToken)
 		{
-			DirectoryNode? directory = await FindDirectoryAsync(reader, name, cancellationToken);
+			DirectoryNode? directory = await FindDirectoryAsync(name, cancellationToken);
 			if (directory == null)
 			{
 				directory = AddDirectory(name);
@@ -635,7 +629,7 @@ namespace EpicGames.Horde.Storage.Nodes
 						}
 						else
 						{
-							directory = await FindOrAddDirectoryAsync(null!, fullName.Substring(0, slashIdx), cancellationToken);
+							directory = await FindOrAddDirectoryAsync(fullName.Substring(0, slashIdx), cancellationToken);
 						}
 
 						FileEntryFlags flags = FileEntryFlags.None;
@@ -776,12 +770,11 @@ namespace EpicGames.Horde.Storage.Nodes
 		/// <summary>
 		/// Utility function to allow extracting a packed directory to disk
 		/// </summary>
-		/// <param name="reader">Reader to retrieve data from</param>
 		/// <param name="directoryInfo"></param>
 		/// <param name="logger"></param>
 		/// <param name="cancellationToken"></param>
 		/// <returns></returns>
-		public async Task CopyToDirectoryAsync(TreeReader reader, DirectoryInfo directoryInfo, ILogger logger, CancellationToken cancellationToken)
+		public async Task CopyToDirectoryAsync(DirectoryInfo directoryInfo, ILogger logger, CancellationToken cancellationToken)
 		{
 			directoryInfo.Create();
 
@@ -789,14 +782,14 @@ namespace EpicGames.Horde.Storage.Nodes
 			foreach (FileEntry fileEntry in _nameToFileEntry.Values)
 			{
 				FileInfo fileInfo = new FileInfo(Path.Combine(directoryInfo.FullName, fileEntry.Name.ToString()));
-				ChunkedDataNode fileNode = await fileEntry.ExpandAsync(reader, cancellationToken);
-				tasks.Add(Task.Run(() => fileNode.CopyToFileAsync(reader, fileInfo, cancellationToken), cancellationToken));
+				ChunkedDataNode fileNode = await fileEntry.ExpandAsync(cancellationToken);
+				tasks.Add(Task.Run(() => fileNode.CopyToFileAsync(fileInfo, cancellationToken), cancellationToken));
 			}
 			foreach (DirectoryEntry directoryEntry in _nameToDirectoryEntry.Values)
 			{
 				DirectoryInfo subDirectoryInfo = directoryInfo.CreateSubdirectory(directoryEntry.Name.ToString());
-				DirectoryNode subDirectoryNode = await directoryEntry.ExpandAsync(reader, cancellationToken);
-				tasks.Add(Task.Run(() => subDirectoryNode.CopyToDirectoryAsync(reader, subDirectoryInfo, logger, cancellationToken), cancellationToken));
+				DirectoryNode subDirectoryNode = await directoryEntry.ExpandAsync(cancellationToken);
+				tasks.Add(Task.Run(() => subDirectoryNode.CopyToDirectoryAsync(subDirectoryInfo, logger, cancellationToken), cancellationToken));
 			}
 
 			await Task.WhenAll(tasks);
@@ -805,10 +798,9 @@ namespace EpicGames.Horde.Storage.Nodes
 		/// <summary>
 		/// Returns a stream containing the zipped contents of this directory
 		/// </summary>
-		/// <param name="reader">Reader for other nodes</param>
 		/// <param name="filter">Filter for files to include in the zip</param>
 		/// <returns>Stream containing zipped archive data</returns>
-		public Stream AsZipStream(TreeReader reader, FileFilter? filter = null) => new DirectoryNodeZipStream(reader, this, filter);
+		public Stream AsZipStream(FileFilter? filter = null) => new DirectoryNodeZipStream(this, filter);
 	}
 
 	/// <summary>
@@ -840,13 +832,12 @@ namespace EpicGames.Horde.Storage.Nodes
 		/// <summary>
 		/// Constructor
 		/// </summary>
-		/// <param name="reader">Reader to read nodes through</param>
 		/// <param name="node">Root node to copy from</param>
 		/// <param name="filter">Filter for files to include in the zip</param>
-		public DirectoryNodeZipStream(TreeReader reader, DirectoryNode node, FileFilter? filter)
+		public DirectoryNodeZipStream(DirectoryNode node, FileFilter? filter)
 		{
 			_pipe = new Pipe();
-			_backgroundTask = BackgroundTask.StartNew(ctx => CopyToPipeAsync(reader, node, filter, _pipe.Writer, ctx));
+			_backgroundTask = BackgroundTask.StartNew(ctx => CopyToPipeAsync(node, filter, _pipe.Writer, ctx));
 		}
 
 		/// <inheritdoc/>
@@ -890,22 +881,22 @@ namespace EpicGames.Horde.Storage.Nodes
 			return length;
 		}
 
-		static async Task CopyToPipeAsync(TreeReader reader, DirectoryNode node, FileFilter? filter, PipeWriter writer, CancellationToken cancellationToken)
+		static async Task CopyToPipeAsync(DirectoryNode node, FileFilter? filter, PipeWriter writer, CancellationToken cancellationToken)
 		{
 			using Stream outputStream = writer.AsStream();
 			using ZipArchive archive = new ZipArchive(outputStream, ZipArchiveMode.Create);
-			await CopyFilesAsync(reader, node, "", filter, archive, cancellationToken);
+			await CopyFilesAsync(node, "", filter, archive, cancellationToken);
 		}
 
-		static async Task CopyFilesAsync(TreeReader reader, DirectoryNode directory, string prefix, FileFilter? filter, ZipArchive archive, CancellationToken cancellationToken)
+		static async Task CopyFilesAsync(DirectoryNode directory, string prefix, FileFilter? filter, ZipArchive archive, CancellationToken cancellationToken)
 		{
 			foreach (DirectoryEntry directoryEntry in directory.Directories)
 			{
 				string directoryPath = $"{prefix}{directoryEntry.Name}/";
 				if (filter == null || filter.PossiblyMatches(directoryPath))
 				{
-					DirectoryNode node = await directoryEntry.ExpandAsync(reader, cancellationToken);
-					await CopyFilesAsync(reader, node, directoryPath, filter, archive, cancellationToken);
+					DirectoryNode node = await directoryEntry.ExpandAsync(cancellationToken);
+					await CopyFilesAsync(node, directoryPath, filter, archive, cancellationToken);
 				}
 			}
 
@@ -926,7 +917,7 @@ namespace EpicGames.Horde.Storage.Nodes
 					}
 
 					using Stream entryStream = entry.Open();
-					await fileEntry.CopyToStreamAsync(reader, entryStream, cancellationToken);
+					await fileEntry.CopyToStreamAsync(entryStream, cancellationToken);
 				}
 			}
 		}
diff --git a/Engine/Source/Programs/Shared/EpicGames.Horde/Storage/Nodes/FileEntry.cs b/Engine/Source/Programs/Shared/EpicGames.Horde/Storage/Nodes/FileEntry.cs
--- a/Engine/Source/Programs/Shared/EpicGames.Horde/Storage/Nodes/FileEntry.cs
+++ b/Engine/Source/Programs/Shared/EpicGames.Horde/Storage/Nodes/FileEntry.cs
@@ -131,33 +131,30 @@ namespace EpicGames.Horde.Storage.Nodes
 		/// <summary>
 		/// Creates a stream that returns the contents of this file
 		/// </summary>
-		/// <param name="reader">Reader for node data</param>
 		/// <returns>The content stream</returns>
-		public Stream AsStream(TreeReader reader) => new FileEntryContentStream(reader, this);
+		public Stream AsStream() => new FileEntryContentStream(this);
 
 		/// <summary>
 		/// Copies the contents of this node and its children to the given output stream
 		/// </summary>
-		/// <param name="reader">Reader for retrieving existing node data</param>
 		/// <param name="outputStream">The output stream to receive the data</param>
 		/// <param name="cancellationToken">Cancellation token for the operation</param>
-		public async Task CopyToStreamAsync(TreeReader reader, Stream outputStream, CancellationToken cancellationToken)
+		public async Task CopyToStreamAsync(Stream outputStream, CancellationToken cancellationToken)
 		{
-			ChunkedDataNode node = await ExpandAsync(reader, cancellationToken);
-			await node.CopyToStreamAsync(reader, outputStream, cancellationToken);
+			ChunkedDataNode node = await ExpandAsync(cancellationToken);
+			await node.CopyToStreamAsync(outputStream, cancellationToken);
 		}
 
 		/// <summary>
 		/// Extracts the contents of this node to a file
 		/// </summary>
-		/// <param name="reader">Reader for retrieving existing node data</param>
 		/// <param name="file">File to write with the contents of this node</param>
 		/// <param name="cancellationToken">Cancellation token for the operation</param>
 		/// <returns></returns>
-		public async Task CopyToFileAsync(TreeReader reader, FileInfo file, CancellationToken cancellationToken)
+		public async Task CopyToFileAsync(FileInfo file, CancellationToken cancellationToken)
 		{
-			ChunkedDataNode node = await ExpandAsync(reader, cancellationToken);
-			await node.CopyToFileAsync(reader, file, cancellationToken);
+			ChunkedDataNode node = await ExpandAsync(cancellationToken);
+			await node.CopyToFileAsync(file, cancellationToken);
 		}
 
 		/// <inheritdoc/>
@@ -191,13 +188,12 @@ namespace EpicGames.Horde.Storage.Nodes
 		/// <summary>
 		/// Constructor
 		/// </summary>
-		/// <param name="reader">Reader to read nodes through</param>
 		/// <param name="entry">The file entry to copy from</param>
-		public FileEntryContentStream(TreeReader reader, FileEntry entry)
+		public FileEntryContentStream(FileEntry entry)
 		{
 			_pipe = new Pipe();
 			_readStream = _pipe.Reader.AsStream();
-			_writeTask = BackgroundTask.StartNew(ctx => WriteDataAsync(reader, entry, _pipe.Writer, ctx));
+			_writeTask = BackgroundTask.StartNew(ctx => WriteDataAsync(entry, _pipe.Writer, ctx));
 
 			Length = entry.Length;
 		}
@@ -212,10 +208,10 @@ namespace EpicGames.Horde.Storage.Nodes
 			_readStream.Dispose();
 		}
 
-		static async Task WriteDataAsync(TreeReader reader, FileEntry entry, PipeWriter pipeWriter, CancellationToken cancellationToken)
+		static async Task WriteDataAsync(FileEntry entry, PipeWriter pipeWriter, CancellationToken cancellationToken)
 		{
 			using Stream stream = pipeWriter.AsStream();
-			await entry.CopyToStreamAsync(reader, stream, cancellationToken);
+			await entry.CopyToStreamAsync(stream, cancellationToken);
 		}
 
 		/// <inheritdoc/>
diff --git a/Engine/Source/Programs/Shared/EpicGames.Horde/Storage/StorageClientBase.cs b/Engine/Source/Programs/Shared/EpicGames.Horde/Storage/StorageClientBase.cs
deleted file mode 100644
diff --git a/Engine/Source/Programs/Shared/EpicGames.Horde/Storage/TreeReader.cs b/Engine/Source/Programs/Shared/EpicGames.Horde/Storage/TreeReader.cs
--- a/Engine/Source/Programs/Shared/EpicGames.Horde/Storage/TreeReader.cs
+++ b/Engine/Source/Programs/Shared/EpicGames.Horde/Storage/TreeReader.cs
@@ -647,39 +647,5 @@ namespace EpicGames.Horde.Storage
 		/// <param name="cancellationToken">Cancellation token for the operation</param>
 		/// <returns>Node data read from the given bundle</returns>
 		public async ValueTask<TNode> ReadNodeAsync<TNode>(NodeLocator locator, CancellationToken cancellationToken = default) where TNode : Node => (TNode)await ReadNodeAsync(locator, cancellationToken);
-
-		/// <summary>
-		/// Reads data for a ref from the store, along with the node's contents.
-		/// </summary>
-		/// <param name="name">The ref name</param>
-		/// <param name="cacheTime">Minimum coherency for any cached value to be returned</param>
-		/// <param name="cancellationToken">Cancellation token for the operation</param>
-		/// <returns>Node for the given ref, or null if it does not exist</returns>
-		public async Task<TNode?> TryReadNodeAsync<TNode>(RefName name, DateTime cacheTime = default, CancellationToken cancellationToken = default) where TNode : Node
-		{
-			NodeHandle? refTarget = await _store.TryReadRefTargetAsync(name, cacheTime, cancellationToken);
-			if (refTarget == null)
-			{
-				return null;
-			}
-			return await ReadNodeAsync<TNode>(refTarget.Locator, cancellationToken);
-		}
-
-		/// <summary>
-		/// Reads a ref from the store, throwing an exception if it does not exist
-		/// </summary>
-		/// <param name="name">Id for the ref</param>
-		/// <param name="cacheTime">Minimum coherency of any cached result</param>
-		/// <param name="cancellationToken">Cancellation token for the operation</param>
-		/// <returns>The blob instance</returns>
-		public async Task<TNode> ReadNodeAsync<TNode>(RefName name, DateTime cacheTime = default, CancellationToken cancellationToken = default) where TNode : Node
-		{
-			TNode? refValue = await TryReadNodeAsync<TNode>(name, cacheTime, cancellationToken);
-			if (refValue == null)
-			{
-				throw new RefNameNotFoundException(name);
-			}
-			return refValue;
-		}
 	}
 }
diff --git a/Engine/Source/Programs/UnrealBuildTool/Artifacts/HordeStorageArtifactCache.cs b/Engine/Source/Programs/UnrealBuildTool/Artifacts/HordeStorageArtifactCache.cs
--- a/Engine/Source/Programs/UnrealBuildTool/Artifacts/HordeStorageArtifactCache.cs
+++ b/Engine/Source/Programs/UnrealBuildTool/Artifacts/HordeStorageArtifactCache.cs
@@ -289,14 +289,13 @@ namespace UnrealBuildTool.Artifacts
 			await _semaphore.WaitAsync(cancellationToken);
 			try
 			{
-				TreeReader reader = CreateTreeReader();
 				foreach (IoHash key in partialKeys)
 				{
 					lock (_pendingMappings)
 					{
 						mappings.AddRange(_pendingMappings.Where(x => x.Key == key));
 					}
-					ArtifactMappingCollectionNode? node = await reader.TryReadNodeAsync<ArtifactMappingCollectionNode>(GetRefName(key), default, cancellationToken);
+					ArtifactMappingCollectionNode? node = await _store.TryReadNodeAsync<ArtifactMappingCollectionNode>(GetRefName(key), default, cancellationToken);
 					if (node != null)
 					{
 						foreach (HordeArtifactMapping mapping in node.Mappings.Values)
@@ -321,7 +320,6 @@ namespace UnrealBuildTool.Artifacts
 				return null;
 			}
 
-			TreeReader reader = CreateTreeReader();
 			bool[] output = new bool[mappings.Length];
 			Array.Fill(output, false);
 
@@ -329,7 +327,7 @@ namespace UnrealBuildTool.Artifacts
 			{
 				output[index] = false;
 				ArtifactMapping mapping = mappings[index];
-				ArtifactMappingCollectionNode? node = await reader.TryReadNodeAsync<ArtifactMappingCollectionNode>(GetRefName(mapping.Key), default, cancellationToken);
+				ArtifactMappingCollectionNode? node = await _store.TryReadNodeAsync<ArtifactMappingCollectionNode>(GetRefName(mapping.Key), default, cancellationToken);
 				if (node != null)
 				{
 					if (node.Mappings.TryGetValue(mapping.MappingKey, out HordeArtifactMapping hordeArtifactMapping))
@@ -348,7 +346,7 @@ namespace UnrealBuildTool.Artifacts
 							{
 								string outputName = hordeArtifactMapping.ArtifactMapping.Outputs[refIndex++].GetFullPath(mapping.DirectoryMapping);
 								using FileStream stream = new(outputName, FileMode.Create, FileAccess.Write, FileShare.ReadWrite);
-								await ChunkedDataNode.CopyToStreamAsync(reader, artifactRef.Handle.Handle.Locator, stream, cancellationToken);
+								await ChunkedDataNode.CopyToStreamAsync(artifactRef.Handle.Handle, stream, cancellationToken);
 							}
 							catch (Exception)
 							{
@@ -486,8 +484,6 @@ namespace UnrealBuildTool.Artifacts
 				return tasks;
 			}
 
-			TreeReader reader = CreateTreeReader();
-
 			// Loop through the mappings
 			foreach (ArtifactMapping mapping in mappings)
 			{
@@ -498,7 +494,7 @@ namespace UnrealBuildTool.Artifacts
 					RefName refName = GetRefName(mapping.Key);
 
 					// Locate the destination collection for this key
-					ArtifactMappingCollectionNode? node = reader.TryReadNodeAsync<ArtifactMappingCollectionNode>(refName, default, cancellationToken).Result;
+					ArtifactMappingCollectionNode? node = _store!.TryReadNodeAsync<ArtifactMappingCollectionNode>(refName, default, cancellationToken).Result;
 					node ??= new ArtifactMappingCollectionNode();
 
 					// Update the mapping collection
@@ -561,14 +557,5 @@ namespace UnrealBuildTool.Artifacts
 		{
 			return new RefName($"action_artifact_v2_{key}");
 		}
-
-		/// <summary>
-		/// Create a new tree reader
-		/// </summary>
-		/// <returns>TreeReader</returns>
-		private TreeReader CreateTreeReader()
-		{
-			return new(_store!, null, NullLogger.Instance);
-		}
 	}
 }
