Horde: Build fix.

diff --git a/Engine/Source/Programs/Shared/EpicGames.Horde.Tests/BundleTests.cs b/Engine/Source/Programs/Shared/EpicGames.Horde.Tests/BundleTests.cs
--- a/Engine/Source/Programs/Shared/EpicGames.Horde.Tests/BundleTests.cs
+++ b/Engine/Source/Programs/Shared/EpicGames.Horde.Tests/BundleTests.cs
@@ -84,19 +84,19 @@ namespace EpicGames.Horde.Tests
 			node = await reader.ReadNodeAsync<ChunkedDataNode>(handle.Handle.Locator);
 			Assert.IsTrue(node is LeafChunkedDataNode);
 			Assert.AreEqual(7, ((LeafChunkedDataNode)node).Data.Length);
-			await TestBufferlessReadsAsync(reader, handle.Handle.Locator, data.AsMemory(0, 7));
+			await TestBufferlessReadsAsync(handle.Handle, data.AsMemory(0, 7));
 
 			handle = await fileNodeWriter.CreateAsync(data.AsMemory(0, 8), CancellationToken.None);
 			node = await reader.ReadNodeAsync<ChunkedDataNode>(handle.Handle.Locator);
 			Assert.IsTrue(node is LeafChunkedDataNode);
 			Assert.AreEqual(8, ((LeafChunkedDataNode)node).Data.Length);
-			await TestBufferlessReadsAsync(reader, handle.Handle.Locator, data.AsMemory(0, 8));
+			await TestBufferlessReadsAsync(handle.Handle, data.AsMemory(0, 8));
 
 			handle = await fileNodeWriter.CreateAsync(data.AsMemory(0, 9), CancellationToken.None);
 			node = await reader.ReadNodeAsync<ChunkedDataNode>(handle.Handle.Locator);
 			Assert.IsTrue(node is InteriorChunkedDataNode);
 			Assert.AreEqual(2, ((InteriorChunkedDataNode)node).Children.Count);
-			await TestBufferlessReadsAsync(reader, handle.Handle.Locator, data.AsMemory(0, 9));
+			await TestBufferlessReadsAsync(handle.Handle, data.AsMemory(0, 9));
 
 			ChunkedDataNode? childNode1 = await ((InteriorChunkedDataNode)node).Children[0].ExpandAsync();
 			Assert.IsNotNull(childNode1);
@@ -111,7 +111,7 @@ namespace EpicGames.Horde.Tests
 			handle = await fileNodeWriter.CreateAsync(data, CancellationToken.None);
 			node = await reader.ReadNodeAsync<ChunkedDataNode>(handle.Handle.Locator);
 			Assert.IsTrue(node is InteriorChunkedDataNode);
-			await TestBufferlessReadsAsync(reader, handle.Handle.Locator, data);
+			await TestBufferlessReadsAsync(handle.Handle, data);
 		}
 
 		private static byte[] CreateBuffer(int length)
@@ -124,10 +124,10 @@ namespace EpicGames.Horde.Tests
 			return output;
 		}
 
-		private static async Task TestBufferlessReadsAsync(TreeReader reader, NodeLocator locator, ReadOnlyMemory<byte> expected)
+		private static async Task TestBufferlessReadsAsync(NodeHandle handle, ReadOnlyMemory<byte> expected)
 		{
 			using MemoryStream memoryStream = new MemoryStream();
-			await ChunkedDataNode.CopyToStreamAsync(reader, locator, memoryStream, default);
+			await ChunkedDataNode.CopyToStreamAsync(handle, memoryStream, default);
 			ReadOnlyMemory<byte> read = memoryStream.ToArray().AsMemory();
 			Assert.IsTrue(read.Span.SequenceEqual(expected.Span));			
 		}
