Horde: Refresh file metadata after moving from the cache and updating timestamps.

diff --git a/Engine/Source/Programs/Shared/EpicGames.Perforce.Managed/ManagedWorkspace.cs b/Engine/Source/Programs/Shared/EpicGames.Perforce.Managed/ManagedWorkspace.cs
--- a/Engine/Source/Programs/Shared/EpicGames.Perforce.Managed/ManagedWorkspace.cs
+++ b/Engine/Source/Programs/Shared/EpicGames.Perforce.Managed/ManagedWorkspace.cs
@@ -2075,7 +2075,8 @@ namespace EpicGames.Perforce.Managed
 				FileReference.Move(fileToMove._trackedFile.GetLocation(), targetFile);
 				try
 				{ 
-					FileReference.SetLastWriteTimeUtc(targetFile, DateTime.UtcNow); 
+					FileReference.SetLastWriteTimeUtc(targetFile, DateTime.UtcNow);
+					fileToMove._workspaceFile.UpdateMetadata();
 				} 
 				catch 
 				{
