EpicGames.Perforce: Speculatively add altsync and noaltsync values to client options enum.

diff --git a/Engine/Source/Programs/Shared/EpicGames.Perforce/Records/ClientOptions.cs b/Engine/Source/Programs/Shared/EpicGames.Perforce/Records/ClientOptions.cs
--- a/Engine/Source/Programs/Shared/EpicGames.Perforce/Records/ClientOptions.cs
+++ b/Engine/Source/Programs/Shared/EpicGames.Perforce/Records/ClientOptions.cs
@@ -86,5 +86,17 @@ namespace EpicGames.Perforce
 		/// </summary>
 		[PerforceEnum("normdir")]
 		NoRmdir = 2048,
+
+		/// <summary>
+		/// If set, and a Helix Core Virtual File Service (P4VFS) application is installed to make use of P4ALTSYNC, the client can create, delete, or update placeholders of the file metadata instead of always taking time to download the full file content. 
+		/// </summary>
+		[PerforceEnum("altsync")]
+		AltSync = 4096,
+
+		/// <summary>
+		/// Opposite of AltSync.
+		/// </summary>
+		[PerforceEnum("noaltsync")]
+		NoAltSync = 8192,
 	}
 }
