Horde: Rework directory packing algorithm for better data locality.

* Leaf node data is always packed at the start of the output stream, followed by chunk interior nodes, followed by directory nodes.
* Interior nodes now use the IoHash of the last serialized child node to determine slicing points. Previous implementation used a BuzHash from the start of the node, which resulted in not using a sliding window.

diff --git a/Engine/Source/Programs/Shared/EpicGames.Horde.Tests/BundleStoreTests.cs b/Engine/Source/Programs/Shared/EpicGames.Horde.Tests/BundleStoreTests.cs
--- a/Engine/Source/Programs/Shared/EpicGames.Horde.Tests/BundleStoreTests.cs
+++ b/Engine/Source/Programs/Shared/EpicGames.Horde.Tests/BundleStoreTests.cs
@@ -308,9 +308,7 @@ namespace EpicGames.Horde.Tests
 				using IStorageWriter writer = store.CreateWriter(options: new TreeOptions { MaxBlobSize = 1024 });
 
 				ChunkingOptions options = new ChunkingOptions();
-				options.LeafOptions.MinSize = 128;
-				options.LeafOptions.TargetSize = 256;
-				options.LeafOptions.MaxSize = 64 * 1024;
+				options.LeafOptions = new LeafChunkedDataNodeOptions(128, 256, 64 * 1024);
 
 				ChunkedDataWriter fileWriter = new ChunkedDataWriter(writer, options);
 				for (int idx = 0; idx < chunk.Length / 16; idx++)
@@ -360,9 +358,7 @@ namespace EpicGames.Horde.Tests
 				root = new DirectoryNode(DirectoryFlags.None);
 
 				ChunkingOptions options = new ChunkingOptions();
-				options.LeafOptions.MinSize = 128;
-				options.LeafOptions.TargetSize = 256;
-				options.LeafOptions.MaxSize = 64 * 1024;
+				options.LeafOptions = new LeafChunkedDataNodeOptions(128, 256, 64 * 1024);
 
 				ChunkedDataWriter fileWriter = new ChunkedDataWriter(writer, options);
 				NodeHandle handle = await fileWriter.CreateAsync(data, CancellationToken.None);
diff --git a/Engine/Source/Programs/Shared/EpicGames.Horde.Tests/BundleTests.cs b/Engine/Source/Programs/Shared/EpicGames.Horde.Tests/BundleTests.cs
--- a/Engine/Source/Programs/Shared/EpicGames.Horde.Tests/BundleTests.cs
+++ b/Engine/Source/Programs/Shared/EpicGames.Horde.Tests/BundleTests.cs
@@ -72,7 +72,7 @@ namespace EpicGames.Horde.Tests
 			using IStorageWriter writer = _storage.CreateWriter(refName);
 
 			ChunkingOptions options = new ChunkingOptions();
-			options.LeafOptions = new ChunkingOptionsForNodeType(8, 8, 8);
+			options.LeafOptions = new LeafChunkedDataNodeOptions(8, 8, 8);
 
 			ChunkedDataWriter fileNodeWriter = new ChunkedDataWriter(writer, options);
 
diff --git a/Engine/Source/Programs/Shared/EpicGames.Horde.Tests/ChunkingTests.cs b/Engine/Source/Programs/Shared/EpicGames.Horde.Tests/ChunkingTests.cs
--- a/Engine/Source/Programs/Shared/EpicGames.Horde.Tests/ChunkingTests.cs
+++ b/Engine/Source/Programs/Shared/EpicGames.Horde.Tests/ChunkingTests.cs
@@ -51,8 +51,8 @@ namespace EpicGames.Horde.Tests
 		public async Task FixedSizeChunkingTests()
 		{
 			ChunkingOptions options = new ChunkingOptions();
-			options.LeafOptions = new ChunkingOptionsForNodeType(64, 64, 64);
-			options.InteriorOptions = new ChunkingOptionsForNodeType(IoHash.NumBytes * 4, IoHash.NumBytes * 4, IoHash.NumBytes * 4);
+			options.LeafOptions = new LeafChunkedDataNodeOptions(64, 64, 64);
+			options.InteriorOptions = new InteriorChunkedDataNodeOptions(4, 4, 4);
 
 			await TestChunkingAsync(options);
 		}
@@ -61,8 +61,8 @@ namespace EpicGames.Horde.Tests
 		public async Task VariableSizeChunkingTests()
 		{
 			ChunkingOptions options = new ChunkingOptions();
-			options.LeafOptions = new ChunkingOptionsForNodeType(32, 64, 96);
-			options.InteriorOptions = new ChunkingOptionsForNodeType(IoHash.NumBytes * 1, IoHash.NumBytes * 4, IoHash.NumBytes * 12);
+			options.LeafOptions = new LeafChunkedDataNodeOptions(32, 64, 96);
+			options.InteriorOptions = new InteriorChunkedDataNodeOptions(1, 4, 12);
 
 			await TestChunkingAsync(options);
 		}
@@ -128,8 +128,8 @@ namespace EpicGames.Horde.Tests
 			{
 				InteriorChunkedDataNode interiorNode = (InteriorChunkedDataNode)node;
 
-				Assert.IsTrue(rightmost || interiorNode.Children.Count * IoHash.NumBytes >= options.InteriorOptions.MinSize);
-				Assert.IsTrue(interiorNode.Children.Count <= options.InteriorOptions.MaxSize);
+				Assert.IsTrue(rightmost || interiorNode.Children.Count >= options.InteriorOptions.MinChildCount);
+				Assert.IsTrue(interiorNode.Children.Count <= options.InteriorOptions.MaxChildCount);
 
 				int childCount = interiorNode.Children.Count;
 				for (int idx = 0; idx < childCount; idx++)
diff --git a/Engine/Source/Programs/Shared/EpicGames.Horde/Storage/Node.cs b/Engine/Source/Programs/Shared/EpicGames.Horde/Storage/Node.cs
--- a/Engine/Source/Programs/Shared/EpicGames.Horde/Storage/Node.cs
+++ b/Engine/Source/Programs/Shared/EpicGames.Horde/Storage/Node.cs
@@ -364,6 +364,19 @@ namespace EpicGames.Horde.Storage
 			}
 		}
 
+		/// <summary>
+		/// Serializes a single node
+		/// </summary>
+		/// <param name="writer">Writer for the node data</param>
+		/// <param name="node">The node to write</param>
+		/// <param name="cancellationToken">Cancellation token for the operation</param>
+		/// <returns>Handle to the written node</returns>
+		public static async Task<NodeHandle> WriteNodeAsync(this IStorageWriter writer, Node node, CancellationToken cancellationToken = default)
+		{
+			NodeRef nodeRef = new NodeRef(node);
+			return await writer.WriteAsync(nodeRef, cancellationToken);
+		}
+
 		/// <summary>
 		/// Writes a node to storage
 		/// </summary>
diff --git a/Engine/Source/Programs/Shared/EpicGames.Horde/Storage/Nodes/ChunkedDataNode.cs b/Engine/Source/Programs/Shared/EpicGames.Horde/Storage/Nodes/ChunkedDataNode.cs
--- a/Engine/Source/Programs/Shared/EpicGames.Horde/Storage/Nodes/ChunkedDataNode.cs
+++ b/Engine/Source/Programs/Shared/EpicGames.Horde/Storage/Nodes/ChunkedDataNode.cs
@@ -1,7 +1,10 @@
 // Copyright Epic Games, Inc. All Rights Reserved.
 
 using System;
+using System.Buffers;
+using System.Buffers.Binary;
 using System.Collections.Generic;
+using System.Diagnostics;
 using System.IO;
 using System.Linq;
 using System.Threading;
@@ -136,6 +139,145 @@ namespace EpicGames.Horde.Storage.Nodes
 			// Keep this code in sync with the constructor
 			await outputStream.WriteAsync(nodeData.Data, cancellationToken);
 		}
+
+		/// <summary>
+		/// Creates nodes from the given file
+		/// </summary>
+		/// <param name="writer"></param>
+		/// <param name="file"></param>
+		/// <param name="options"></param>
+		/// <param name="cancellationToken"></param>
+		/// <returns></returns>
+		public static async Task<List<NodeHandle>> CreateFromFileAsync(IStorageWriter writer, FileReference file, LeafChunkedDataNodeOptions options, CancellationToken cancellationToken)
+		{
+			using (FileStream stream = FileReference.Open(file, FileMode.Open, FileAccess.Read, FileShare.Read))
+			{
+				return await CreateFromStreamAsync(writer, stream, options, cancellationToken);
+			}
+		}
+
+		/// <summary>
+		/// Creates nodes from the given file
+		/// </summary>
+		/// <param name="writer"></param>
+		/// <param name="file"></param>
+		/// <param name="options"></param>
+		/// <param name="cancellationToken"></param>
+		/// <returns></returns>
+		public static async Task<List<NodeHandle>> CreateFromFileAsync(IStorageWriter writer, FileInfo file, LeafChunkedDataNodeOptions options, CancellationToken cancellationToken)
+		{
+			using (FileStream stream = file.Open(FileMode.Open, FileAccess.Read, FileShare.Read))
+			{
+				return await CreateFromStreamAsync(writer, stream, options, cancellationToken);
+			}
+		}
+
+		/// <summary>
+		/// Creates nodes from the given file
+		/// </summary>
+		/// <param name="writer"></param>
+		/// <param name="stream"></param>
+		/// <param name="options"></param>
+		/// <param name="cancellationToken"></param>
+		/// <returns></returns>
+		public static async Task<List<NodeHandle>> CreateFromStreamAsync(IStorageWriter writer, Stream stream, LeafChunkedDataNodeOptions options, CancellationToken cancellationToken)
+		{
+			List<NodeHandle> handles = new List<NodeHandle>();
+
+			using IMemoryOwner<byte> readBuffer = MemoryPool<byte>.Shared.Rent(options.MaxSize);
+
+			int size = 0;
+			for (; ; )
+			{
+				size += await stream.ReadGreedyAsync(readBuffer.Memory.Slice(size), cancellationToken);
+				if (size == 0)
+				{
+					break;
+				}
+
+				int nextLength = GetChunkLength(readBuffer.Memory.Span.Slice(0, size), options);
+
+				Memory<byte> outputBuffer = writer.GetOutputBuffer(0, nextLength);
+				readBuffer.Memory.Slice(0, nextLength).CopyTo(outputBuffer);
+
+				NodeHandle handle = await writer.WriteNodeAsync(nextLength, Array.Empty<NodeHandle>(), GetNodeType<LeafChunkedDataNode>(), cancellationToken);
+				handles.Add(handle);
+
+				readBuffer.Memory.Slice(nextLength, size - nextLength).CopyTo(readBuffer.Memory);
+				size -= nextLength;
+			}
+
+			return handles;
+		}
+
+		/// <summary>
+		/// Determines how much data to append to an existing leaf node
+		/// </summary>
+		/// <param name="inputData">Data to be appended</param>
+		/// <param name="options">Options for chunking the data</param>
+		/// <returns>The number of bytes to append</returns>
+		public static int GetChunkLength(ReadOnlySpan<byte> inputData, LeafChunkedDataNodeOptions options)
+		{
+			// If the target option sizes are fixed, just chunk the data along fixed boundaries
+			if (options.MinSize == options.TargetSize && options.MaxSize == options.TargetSize)
+			{
+				return Math.Min(inputData.Length, options.MaxSize);
+			}
+
+			// Cap the append data span to the maximum amount we can add
+			int maxLength = options.MaxSize;
+			if (maxLength < inputData.Length)
+			{
+				inputData = inputData.Slice(0, maxLength);
+			}
+
+			int windowSize = options.MinSize;
+
+			// Fast path for appending data to the buffer up to the chunk window size
+			int length = Math.Min(windowSize, inputData.Length);
+			uint rollingHash = BuzHash.Add(0, inputData.Slice(0, length));
+
+			// Get the threshold for the rolling hash to split the output
+			uint rollingHashThreshold = (uint)((1L << 32) / options.TargetSize);
+
+			// Step through the rest of the data which is completely contained in appendData.
+			if (length < inputData.Length)
+			{
+				Debug.Assert(length >= windowSize);
+
+				ReadOnlySpan<byte> tailSpan = inputData.Slice(length - windowSize, inputData.Length - windowSize);
+				ReadOnlySpan<byte> headSpan = inputData.Slice(length);
+
+				int count = BuzHash.Update(tailSpan, headSpan, rollingHashThreshold, ref rollingHash);
+				if (count != -1)
+				{
+					length += count;
+					return length;
+				}
+
+				length += headSpan.Length;
+			}
+
+			return length;
+		}
+	}
+
+	/// <summary>
+	/// Options for creating interior nodes
+	/// </summary>
+	/// <param name="MinChildCount">Minimum number of children in each node</param>
+	/// <param name="TargetChildCount">Target number of children in each node</param>
+	/// <param name="MaxChildCount">Maximum number of children in each node</param>
+	/// <param name="SliceThreshold">Threshold hash value for splitting interior nodes</param>
+	public record class InteriorChunkedDataNodeOptions(int MinChildCount, int TargetChildCount, int MaxChildCount, uint SliceThreshold)
+	{
+		/// <summary>
+		/// Constructor
+		/// </summary>
+		public InteriorChunkedDataNodeOptions(int minChildCount, int targetChildCount, int maxChildCount)
+			: this(minChildCount, targetChildCount, maxChildCount, (uint)((1UL << 32) / (uint)targetChildCount))
+		{
+		}
 	}
 
 	/// <summary>
@@ -184,6 +326,71 @@ namespace EpicGames.Horde.Storage.Nodes
 		/// <inheritdoc/>
 		public override IEnumerable<NodeRef> EnumerateRefs() => Children;
 
+		/// <summary>
+		/// Create a tree of nodes from the given list of handles, splitting nodes in each layer based on the hash of the last node.
+		/// </summary>
+		/// <param name="handles">List of leaf handles</param>
+		/// <param name="options">Options for splitting the tree</param>
+		/// <param name="writer">Output writer for new interior nodes</param>
+		/// <param name="cancellationToken">Cancellation token for the operation</param>
+		/// <returns>Handle to the root node of the tree</returns>
+		public static async Task<NodeHandle> CreateTreeAsync(List<NodeHandle> handles, InteriorChunkedDataNodeOptions options, IStorageWriter writer, CancellationToken cancellationToken)
+		{
+			List<NodeHandle> handleBuffer = new List<NodeHandle>();
+
+			List<InteriorChunkedDataNode> interiorNodes = new List<InteriorChunkedDataNode>();
+			while (handles.Count > 1)
+			{
+				interiorNodes.Clear();
+				CreateTreeLayer(handles, options, interiorNodes);
+
+				handleBuffer.Clear();
+				foreach (InteriorChunkedDataNode interiorNode in interiorNodes)
+				{
+					NodeHandle handle = await writer.WriteNodeAsync(interiorNode, cancellationToken);
+					handleBuffer.Add(handle);
+				}
+
+				handles = handleBuffer;
+			}
+
+			return handles[0];
+		}
+		
+		/// <summary>
+		/// Split a list of handles into a layer of interior nodes
+		/// </summary>
+		static void CreateTreeLayer(List<NodeHandle> handles, InteriorChunkedDataNodeOptions options, List<InteriorChunkedDataNode> interiorNodes)
+		{
+			Span<byte> buffer = stackalloc byte[IoHash.NumBytes];
+
+			for (int index = 0; index < handles.Count; )
+			{
+				int minIndex = index;
+				int maxIndex = Math.Min(minIndex + options.MaxChildCount, handles.Count);
+
+				index = Math.Min(index + options.MinChildCount, handles.Count);
+				for (; index < maxIndex; index++)
+				{
+					handles[index].Hash.CopyTo(buffer);
+
+					uint value = BinaryPrimitives.ReadUInt32LittleEndian(buffer);
+					if (value < options.SliceThreshold)
+					{
+						break;
+					}
+				}
+
+				NodeRef<ChunkedDataNode>[] children = new NodeRef<ChunkedDataNode>[index - minIndex];
+				for (int childIndex = minIndex; childIndex < index; childIndex++)
+				{
+					children[childIndex - minIndex] = new NodeRef<ChunkedDataNode>(handles[childIndex]);
+				}
+
+				interiorNodes.Add(new InteriorChunkedDataNode(children));
+			}
+		}
+
 		/// <inheritdoc/>
 		public override async Task CopyToStreamAsync(Stream outputStream, CancellationToken cancellationToken)
 		{
diff --git a/Engine/Source/Programs/Shared/EpicGames.Horde/Storage/Nodes/ChunkedDataWriter.cs b/Engine/Source/Programs/Shared/EpicGames.Horde/Storage/Nodes/ChunkedDataWriter.cs
--- a/Engine/Source/Programs/Shared/EpicGames.Horde/Storage/Nodes/ChunkedDataWriter.cs
+++ b/Engine/Source/Programs/Shared/EpicGames.Horde/Storage/Nodes/ChunkedDataWriter.cs
@@ -1,6 +1,7 @@
 // Copyright Epic Games, Inc. All Rights Reserved.
 
 using System;
+using System.Buffers.Binary;
 using System.Collections.Generic;
 using System.Diagnostics;
 using System.IO;
@@ -18,61 +19,39 @@ namespace EpicGames.Horde.Storage.Nodes
 		/// <summary>
 		/// Options for creating leaf nodes
 		/// </summary>
-		public ChunkingOptionsForNodeType LeafOptions { get; set; }
+		public LeafChunkedDataNodeOptions LeafOptions { get; set; }
 
 		/// <summary>
 		/// Options for creating interior nodes
 		/// </summary>
-		public ChunkingOptionsForNodeType InteriorOptions { get; set; }
+		public InteriorChunkedDataNodeOptions InteriorOptions { get; set; }
 
 		/// <summary>
 		/// Constructor
 		/// </summary>
 		public ChunkingOptions()
 		{
-			LeafOptions = new ChunkingOptionsForNodeType(32 * 1024, 256 * 1024, 64 * 1024);
-			InteriorOptions = new ChunkingOptionsForNodeType(32 * 1024, 256 * 1024, 64 * 1024);
+			LeafOptions = new LeafChunkedDataNodeOptions(32 * 1024, 256 * 1024, 64 * 1024);
+			InteriorOptions = new InteriorChunkedDataNodeOptions(1, 5, 10);
 		}
 	}
 
 	/// <summary>
 	/// Options for creating a specific type of file nodes
 	/// </summary>
-	public class ChunkingOptionsForNodeType
+	/// <param name="MinSize">Minimum chunk size</param>
+	/// <param name="MaxSize">Maximum chunk size. Chunks will be split on this boundary if another match is not found.</param>
+	/// <param name="TargetSize">Target chunk size for content-slicing</param>
+	public record class LeafChunkedDataNodeOptions(int MinSize, int MaxSize, int TargetSize)
 	{
-		/// <summary>
-		/// Minimum chunk size
-		/// </summary>
-		public int MinSize { get; set; }
-
-		/// <summary>
-		/// Maximum chunk size. Chunks will be split on this boundary if another match is not found.
-		/// </summary>
-		public int MaxSize { get; set; }
-
-		/// <summary>
-		/// Target chunk size for content-slicing
-		/// </summary>
-		public int TargetSize { get; set; }
-
 		/// <summary>
 		/// Constructor
 		/// </summary>
 		/// <param name="size">Fixed size chunks to use</param>
-		public ChunkingOptionsForNodeType(int size)
+		public LeafChunkedDataNodeOptions(int size)
 			: this(size, size, size)
 		{
 		}
-
-		/// <summary>
-		/// Default constructor
-		/// </summary>
-		public ChunkingOptionsForNodeType(int minSize, int maxSize, int targetSize)
-		{
-			MinSize = minSize;
-			MaxSize = maxSize;
-			TargetSize = targetSize;
-		}
 	}
 
 	/// <summary>
@@ -85,43 +64,14 @@ namespace EpicGames.Horde.Storage.Nodes
 		/// </summary>
 		public const int DefaultBufferLength = 32 * 1024;
 
-		class InteriorNodeState
-		{
-			public InteriorNodeState? _parent;
-			public readonly ArrayMemoryWriter _data;
-			public readonly List<NodeHandle> _children = new List<NodeHandle>();
-
-			public uint _rollingHash;
-
-			public InteriorNodeState(int maxSize)
-			{
-				_data = new ArrayMemoryWriter(maxSize);
-			}
-
-			public void Reset()
-			{
-				_rollingHash = 0;
-				_data.Clear();
-				_children.Clear();
-			}
-
-			public void Write(NodeHandle handle)
-			{
-				_children.Add(handle);
-				_data.WriteIoHash(handle.Hash);
-			}
-		}
-
 		static readonly NodeType s_leafNodeType = Node.GetNodeType<LeafChunkedDataNode>();
-		static readonly NodeType s_interiorNodeType = Node.GetNodeType<InteriorChunkedDataNode>();
 
 		readonly IStorageWriter _writer;
 		readonly ChunkingOptions _options;
 
 		// Tree state
-		InteriorNodeState? _topInteriorNode = null;
 		long _totalLength;
-		readonly Stack<InteriorNodeState> _freeInteriorNodes = new Stack<InteriorNodeState>();
+		readonly List<NodeHandle> _leafHandles = new List<NodeHandle>();
 
 		// Leaf node state
 		uint _leafHash;
@@ -148,7 +98,7 @@ namespace EpicGames.Horde.Storage.Nodes
 		/// </summary>
 		public void Reset()
 		{
-			FreeInteriorNodes();
+			_leafHandles.Clear();
 			ResetLeafState();
 			_totalLength = 0;
 		}
@@ -162,34 +112,6 @@ namespace EpicGames.Horde.Storage.Nodes
 			_leafLength = 0;
 		}
 
-		/// <summary>
-		/// Creates a new interior node state
-		/// </summary>
-		/// <returns>State object</returns>
-		InteriorNodeState CreateInteriorNode()
-		{
-			InteriorNodeState? result;
-			if (!_freeInteriorNodes.TryPop(out result))
-			{
-				result = new InteriorNodeState(_options.InteriorOptions.MaxSize);
-			}
-			return result;
-		}
-
-		/// <summary>
-		/// Free all the current interior nodes
-		/// </summary>
-		void FreeInteriorNodes()
-		{
-			while (_topInteriorNode != null)
-			{
-				InteriorNodeState current = _topInteriorNode;
-				_topInteriorNode = _topInteriorNode._parent;
-				current.Reset();
-				_freeInteriorNodes.Push(current);
-			}
-		}
-
 		/// <summary>
 		/// Creates data for the given file
 		/// </summary>
@@ -280,7 +202,6 @@ namespace EpicGames.Horde.Storage.Nodes
 			Memory<byte> buffer = _writer.GetOutputBuffer(_leafLength, _leafLength);
 			for (; ; )
 			{
-				// Append data to the current leaf node
 				int appendLength = AppendToLeafNode(buffer.Span.Slice(0, _leafLength), data.Span, ref _leafHash, _options.LeafOptions);
 
 				buffer = _writer.GetOutputBuffer(_leafLength, _leafLength + appendLength);
@@ -296,10 +217,7 @@ namespace EpicGames.Horde.Storage.Nodes
 					break;
 				}
 
-				// Flush the leaf node and any interior nodes that are full
-				NodeHandle handle = await WriteLeafNodeAsync(cancellationToken);
-				ResetLeafState();
-				await AddToInteriorNodeAsync(handle, cancellationToken);
+				await FlushLeafNodeAsync(cancellationToken);
 			}
 		}
 
@@ -311,7 +229,7 @@ namespace EpicGames.Horde.Storage.Nodes
 		/// <param name="rollingHash">Current BuzHash of the data</param>
 		/// <param name="options">Options for chunking the data</param>
 		/// <returns>The number of bytes to append</returns>
-		static int AppendToLeafNode(ReadOnlySpan<byte> currentData, ReadOnlySpan<byte> appendData, ref uint rollingHash, ChunkingOptionsForNodeType options)
+		static int AppendToLeafNode(ReadOnlySpan<byte> currentData, ReadOnlySpan<byte> appendData, ref uint rollingHash, LeafChunkedDataNodeOptions options)
 		{
 			// If the target option sizes are fixed, just chunk the data along fixed boundaries
 			if (options.MinSize == options.TargetSize && options.MaxSize == options.TargetSize)
@@ -380,77 +298,6 @@ namespace EpicGames.Horde.Storage.Nodes
 			return appendLength;
 		}
 
-		async Task AddToInteriorNodeAsync(NodeHandle handle, CancellationToken cancellationToken)
-		{
-			_topInteriorNode ??= CreateInteriorNode();
-			await AddToInteriorNodeAsync(_topInteriorNode, handle, cancellationToken);
-		}
-
-		async Task AddToInteriorNodeAsync(InteriorNodeState interiorNode, NodeHandle handle, CancellationToken cancellationToken)
-		{
-			// If the node is already full, flush it
-			if (IsInteriorNodeComplete(interiorNode._data.WrittenSpan, interiorNode._rollingHash, _options.InteriorOptions))
-			{
-				NodeHandle interiorNodeHandle = await WriteInteriorNodeAndResetAsync(interiorNode, cancellationToken);
-				interiorNode._parent ??= CreateInteriorNode();
-				await AddToInteriorNodeAsync(interiorNode._parent, interiorNodeHandle, cancellationToken);
-			}
-
-			// Add this handle
-			AppendToInteriorNode(interiorNode._data.WrittenSpan, handle.Hash, ref interiorNode._rollingHash, _options.InteriorOptions);
-			interiorNode.Write(handle);
-		}
-
-
-		/// <summary>
-		/// Test whether the current node is complete
-		/// </summary>
-		/// <param name="currentData"></param>
-		/// <param name="rollingHash"></param>
-		/// <param name="options"></param>
-		/// <returns></returns>
-		static bool IsInteriorNodeComplete(ReadOnlySpan<byte> currentData, uint rollingHash, ChunkingOptionsForNodeType options)
-		{
-			if (currentData.Length + IoHash.NumBytes > options.MaxSize)
-			{
-				return true;
-			}
-
-			if (currentData.Length >= options.MinSize)
-			{
-				uint rollingHashThreshold = BuzHash.GetThreshold(options.TargetSize);
-				if (rollingHash < rollingHashThreshold)
-				{
-					return true;
-				}
-			}
-
-			return false;
-		}
-
-		/// <summary>
-		/// Append a new hash to this interior node
-		/// </summary>
-		/// <param name="currentData">Current data for the node</param>
-		/// <param name="hash">Hash of the child node</param>
-		/// <param name="rollingHash">Current rolling hash for the node</param>
-		/// <param name="options">Options for chunking the node</param>
-		/// <returns>True if the hash could be appended, false otherwise</returns>
-		static void AppendToInteriorNode(ReadOnlySpan<byte> currentData, IoHash hash, ref uint rollingHash, ChunkingOptionsForNodeType options)
-		{
-			Span<byte> hashData = stackalloc byte[IoHash.NumBytes];
-			hash.CopyTo(hashData);
-
-			rollingHash = BuzHash.Add(rollingHash, hashData);
-
-			int windowSize = options.MinSize - (options.MinSize % IoHash.NumBytes);
-			if (currentData.Length > windowSize)
-			{
-				ReadOnlySpan<byte> removeData = currentData.Slice(currentData.Length - windowSize, IoHash.NumBytes);
-				rollingHash = BuzHash.Sub(rollingHash, removeData, windowSize + IoHash.NumBytes);
-			}
-		}
-
 		/// <summary>
 		/// Complete the current file, and write all open nodes to the underlying writer
 		/// </summary>
@@ -458,17 +305,9 @@ namespace EpicGames.Horde.Storage.Nodes
 		/// <returns>Handle to the root node</returns>
 		public async Task<NodeHandle> CompleteAsync(CancellationToken cancellationToken)
 		{
-			NodeHandle handle = await WriteLeafNodeAsync(cancellationToken);
-			ResetLeafState();
-
-			for (InteriorNodeState? state = _topInteriorNode; state != null; state = state._parent)
-			{
-				await AddToInteriorNodeAsync(state, handle, cancellationToken);
-				handle = await WriteInteriorNodeAndResetAsync(state, cancellationToken);
-			}
-
-			FreeInteriorNodes();
-			return handle;
+			await FlushLeafNodeAsync(cancellationToken);
+			NodeHandle rootHandle = await InteriorChunkedDataNode.CreateTreeAsync(_leafHandles, _options.InteriorOptions, _writer, cancellationToken);
+			return rootHandle;
 		}
 
 		/// <summary>
@@ -483,31 +322,16 @@ namespace EpicGames.Horde.Storage.Nodes
 			return handle;
 		}
 
-		/// <summary>
-		/// Writes the state of the given interior node to storage
-		/// </summary>
-		/// <param name="state"></param>
-		/// <param name="cancellationToken"></param>
-		/// <returns></returns>
-		async ValueTask<NodeHandle> WriteInteriorNodeAndResetAsync(InteriorNodeState state, CancellationToken cancellationToken)
-		{
-			Memory<byte> buffer = _writer.GetOutputBuffer(0, state._data.Length);
-			state._data.WrittenMemory.CopyTo(buffer);
-
-			NodeHandle handle = await _writer.WriteNodeAsync(state._data.Length, state._children, s_interiorNodeType, cancellationToken);
-			state.Reset();
-
-			return handle;
-		}
-
 		/// <summary>
 		/// Writes the contents of the current leaf node to storage
 		/// </summary>
 		/// <param name="cancellationToken">Cancellation token for the operation</param>
 		/// <returns>Handle to the written leaf node</returns>
-		async ValueTask<NodeHandle> WriteLeafNodeAsync(CancellationToken cancellationToken)
+		async ValueTask FlushLeafNodeAsync(CancellationToken cancellationToken)
 		{
-			return await _writer.WriteNodeAsync(_leafLength, Array.Empty<NodeHandle>(), s_leafNodeType, cancellationToken);
+			NodeHandle handle = await _writer.WriteNodeAsync(_leafLength, Array.Empty<NodeHandle>(), s_leafNodeType, cancellationToken);
+			_leafHandles.Add(handle);
+			ResetLeafState();
 		}
 	}
 }
diff --git a/Engine/Source/Programs/Shared/EpicGames.Horde/Storage/Nodes/DirectoryNode.cs b/Engine/Source/Programs/Shared/EpicGames.Horde/Storage/Nodes/DirectoryNode.cs
--- a/Engine/Source/Programs/Shared/EpicGames.Horde/Storage/Nodes/DirectoryNode.cs
+++ b/Engine/Source/Programs/Shared/EpicGames.Horde/Storage/Nodes/DirectoryNode.cs
@@ -1,9 +1,11 @@
 // Copyright Epic Games, Inc. All Rights Reserved.
 
 using EpicGames.Core;
+using Microsoft.CodeAnalysis.Options;
 using Microsoft.Extensions.Logging;
 using System;
 using System.Buffers;
+using System.Buffers.Binary;
 using System.Collections.Generic;
 using System.Diagnostics;
 using System.Diagnostics.CodeAnalysis;
@@ -340,7 +342,7 @@ namespace EpicGames.Horde.Storage.Nodes
 
 		static async ValueTask<DirectoryNode?> GetDirectoryByPathAsync(DirectoryNode directoryNode, Utf8String path, CancellationToken cancellationToken)
 		{
-			while(path.Length > 0)
+			while (path.Length > 0)
 			{
 				Utf8String directoryName;
 
@@ -551,6 +553,203 @@ namespace EpicGames.Horde.Storage.Nodes
 			}
 		}
 
+		class InputFile
+		{
+			public FileInfo FileInfo { get; }
+			public List<NodeHandle> Handles { get; } = new List<NodeHandle>();
+
+			public InputFile(FileInfo fileInfo) => FileInfo = fileInfo;
+		}
+
+		class InputTree
+		{
+			public Dictionary<string, InputTree> Directories { get; } = new Dictionary<string, InputTree>();
+			public Dictionary<string, InputFile> Files { get; } = new Dictionary<string, InputFile>();
+
+			public void AddFiles(DirectoryReference baseDir, IEnumerable<FileInfo> files)
+			{
+				IEnumerator<FileInfo> enumerator = files.GetEnumerator();
+				if (enumerator.MoveNext())
+				{
+					if (AddFiles(baseDir.FullName, enumerator))
+					{
+						throw new ArgumentException($"File {enumerator.Current} is not under {baseDir}", nameof(baseDir));
+					}
+				}
+			}
+
+			bool AddFiles(ReadOnlySpan<char> baseDir, IEnumerator<FileInfo> files)
+			{
+				for (; ; )
+				{
+					FileInfo fileInfo = files.Current;
+
+					int dirNameLength = baseDir.Length;
+					if (fileInfo.FullName.Length < dirNameLength || !fileInfo.FullName.AsSpan().StartsWith(baseDir, FileReference.Comparison) || fileInfo.FullName[dirNameLength] != Path.DirectorySeparatorChar)
+					{
+						return true;
+					}
+
+					int nextDirLength = fileInfo.FullName.AsSpan(baseDir.Length + 1).IndexOf(Path.DirectorySeparatorChar);
+					if (nextDirLength == -1)
+					{
+						string fileName = fileInfo.FullName.Substring(dirNameLength + 1);
+						Files[fileName] = new InputFile(fileInfo);
+
+						if (!files.MoveNext())
+						{
+							return false;
+						}
+					}
+					else
+					{
+						string dirName = fileInfo.FullName.Substring(dirNameLength + 1, nextDirLength);
+
+						InputTree? nextTree;
+						if (!Directories.TryGetValue(dirName, out nextTree))
+						{
+							nextTree = new InputTree();
+							Directories.Add(dirName, nextTree);
+						}
+
+						ReadOnlySpan<char> nextBaseDir = fileInfo.FullName.AsSpan(0, dirNameLength + 1 + nextDirLength);
+						if (!nextTree.AddFiles(nextBaseDir, files))
+						{
+							return false;
+						}
+					}
+				}
+			}
+		}
+
+		/// <summary>
+		/// Create a directory tree from a set of files
+		/// </summary>
+		/// <param name="baseDir">Base directory to root files from</param>
+		/// <param name="files">Files to include in the data</param>
+		/// <param name="options">Options for splitting files</param>
+		/// <param name="writer">Writer for output data</param>
+		/// <param name="cancellationToken">Cancellation token for the operation</param>
+		/// <returns>Root directory node containing the files</returns>
+		public static async Task<DirectoryNode> CreateAsync(DirectoryReference baseDir, IEnumerable<FileInfo> files, ChunkingOptions options, IStorageWriter writer, CancellationToken cancellationToken)
+		{
+			// Create a tree out of all the files
+			InputTree tree = new InputTree();
+			tree.AddFiles(baseDir, files);
+
+			// Create leaf nodes
+			List<List<InputFile>> filePartitions = PartitionInputFiles(tree);
+			await Parallel.ForEachAsync(filePartitions, cancellationToken, (filePartition, ctx) => CreateLeafChunkNodesAsync(writer, filePartition, options, ctx));
+
+			await CreateInteriorChunkNodesAsync(tree, options.InteriorOptions, writer, cancellationToken);
+			return await CreateDirectoryNodesAsync(tree, writer, options, cancellationToken);
+		}
+
+		static List<List<InputFile>> PartitionInputFiles(InputTree tree)
+		{
+			// Create the partitions
+			List<InputFile> sortedFiles = new List<InputFile>();
+			long totalSize = GetInputFiles(tree, sortedFiles);
+
+			// Maximum number of streams to write in parallel
+			const int MaxWriters = 8;
+
+			// Minimum size of output payload for each writer
+			const long MinSizePerWriter = 1024 * 1024;
+
+			// Calculate the number of writers
+			int numWriters = 1 + (int)Math.Min(totalSize / MinSizePerWriter, MaxWriters);
+
+			// Create the partitions
+			List<List<InputFile>> partitions = new List<List<InputFile>>();
+
+			long writtenSize = 0;
+			long remainingSizeForPartition = 0;
+			foreach (InputFile sortedFile in sortedFiles)
+			{
+				if (sortedFile.FileInfo.Length >= remainingSizeForPartition && partitions.Count < numWriters)
+				{
+					partitions.Add(new List<InputFile>());
+					remainingSizeForPartition = (totalSize - writtenSize) / (numWriters - partitions.Count);
+				}
+
+				partitions[^1].Add(sortedFile);
+				writtenSize += sortedFile.FileInfo.Length;
+
+				remainingSizeForPartition -= sortedFile.FileInfo.Length;
+			}
+
+			return partitions;
+		}
+
+		static long GetInputFiles(InputTree tree, List<InputFile> files)
+		{
+			long totalSize = 0;
+			foreach (InputTree directory in tree.Directories.Values)
+			{
+				totalSize += GetInputFiles(directory, files);
+			}
+
+			foreach (InputFile inputFile in tree.Files.Values)
+			{
+				totalSize += inputFile.FileInfo.Length;
+				files.Add(inputFile);
+			}
+			return totalSize;
+		}
+
+		static async ValueTask CreateLeafChunkNodesAsync(IStorageWriter writer, List<InputFile> files, ChunkingOptions options, CancellationToken cancellationToken)
+		{
+			using IStorageWriter writerFork = writer.Fork();
+			foreach (InputFile file in files)
+			{
+				file.Handles.AddRange(await LeafChunkedDataNode.CreateFromFileAsync(writer, file.FileInfo, options.LeafOptions, cancellationToken));
+			}
+		}
+
+		static async Task CreateLeafChunkNodesRecursiveAsync(InputTree tree, IStorageWriter writer, ChunkingOptions options, CancellationToken cancellationToken)
+		{
+			foreach (InputTree directory in tree.Directories.Values)
+			{
+				await CreateLeafChunkNodesRecursiveAsync(directory, writer, options, cancellationToken);
+			}
+			foreach (InputFile file in tree.Files.Values)
+			{
+				file.Handles.AddRange(await LeafChunkedDataNode.CreateFromFileAsync(writer, file.FileInfo, options.LeafOptions, cancellationToken));
+			}
+		}
+
+		static async Task CreateInteriorChunkNodesAsync(InputTree tree, InteriorChunkedDataNodeOptions options, IStorageWriter writer, CancellationToken cancellationToken)
+		{
+			foreach (InputTree directory in tree.Directories.Values)
+			{
+				await CreateInteriorChunkNodesAsync(directory, options, writer, cancellationToken);
+			}
+
+			foreach (InputFile file in tree.Files.Values)
+			{
+				NodeHandle handle = await InteriorChunkedDataNode.CreateTreeAsync(file.Handles, options, writer, cancellationToken);
+				file.Handles.Clear();
+				file.Handles.Add(handle);
+			}
+		}
+
+		static async Task<DirectoryNode> CreateDirectoryNodesAsync(InputTree tree, IStorageWriter writer, ChunkingOptions options, CancellationToken cancellationToken)
+		{
+			DirectoryNode directoryNode = new DirectoryNode();
+			foreach ((string name, InputTree directory) in tree.Directories)
+			{
+				DirectoryNode childNode = await CreateDirectoryNodesAsync(directory, writer, options, cancellationToken);
+				directoryNode.AddDirectory(new DirectoryEntry(name, childNode));
+			}
+			foreach ((string name, InputFile file) in tree.Files)
+			{
+				FileEntry fileEntry = new FileEntry(name, FileEntryFlags.None, file.FileInfo.Length, file.Handles[0]);
+				directoryNode.AddFile(fileEntry);
+			}
+			return directoryNode;
+		}
+
 		/// <summary>
 		/// Adds files from a flat list of paths
 		/// </summary>
