Horde: Remove legacy DirectoryNode update methods. DirectoryUpdate class now allows building a structured representation of an update to perform in batch.

diff --git a/Engine/Source/Programs/Horde/Horde/Commands/Bundles/BundleCreate.cs b/Engine/Source/Programs/Horde/Horde/Commands/Bundles/BundleCreate.cs
--- a/Engine/Source/Programs/Horde/Horde/Commands/Bundles/BundleCreate.cs
+++ b/Engine/Source/Programs/Horde/Horde/Commands/Bundles/BundleCreate.cs
@@ -68,10 +68,8 @@ namespace Horde.Commands.Bundles
 
 			using (IStorageWriter writer = store.CreateWriter())
 			{
-				DirectoryNode node = new DirectoryNode(DirectoryFlags.None);
-
 				ChunkingOptions options = new ChunkingOptions();
-				await node.CopyFilesAsync(baseDir, files, options, writer, null, CancellationToken.None);
+				DirectoryNode node = await DirectoryNode.CreateAsync(baseDir, files, options, writer, null, CancellationToken.None);
 
 				NodeHandle handle = await writer.FlushAsync(node);
 				if (File != null)
diff --git a/Engine/Source/Programs/Horde/Horde/Commands/Vcs/VcsCommit.cs b/Engine/Source/Programs/Horde/Horde/Commands/Vcs/VcsCommit.cs
--- a/Engine/Source/Programs/Horde/Horde/Commands/Vcs/VcsCommit.cs
+++ b/Engine/Source/Programs/Horde/Horde/Commands/Vcs/VcsCommit.cs
@@ -8,7 +8,6 @@ using System.Threading.Tasks;
 using EpicGames.Core;
 using EpicGames.Horde.Storage;
 using EpicGames.Horde.Storage.Nodes;
-using Microsoft.Extensions.Caching.Memory;
 using Microsoft.Extensions.Logging;
 
 namespace Horde.Commands.Vcs
@@ -27,8 +26,10 @@ namespace Horde.Commands.Vcs
 		{
 		}
 
-		public override async Task<int> ExecuteAsync(ILogger logger)
+		public override Task<int> ExecuteAsync(ILogger logger)
 		{
+			throw new NotImplementedException();
+			/*
 			DirectoryReference rootDir = FindRootDir();
 
 			WorkspaceState workspaceState = await ReadStateAsync(rootDir);
@@ -101,7 +102,7 @@ namespace Horde.Commands.Vcs
 			await WriteStateAsync(rootDir, workspaceState);
 
 			logger.LogInformation("Commited in change {Number}", newTip.Number);
-			return 0;
+			return 0;*/
 		}
 
 		private async Task UpdateTreeAsync(DirectoryNodeRef rootRef, DirectoryReference rootDir, DirectoryState? oldState, DirectoryState newState, List<(DirectoryNode, FileInfo, FileState)> files, List<(DirectoryNodeRef, DirectoryState)> directories)
diff --git a/Engine/Source/Programs/Shared/EpicGames.Horde/Storage/Nodes/DirectoryNode.cs b/Engine/Source/Programs/Shared/EpicGames.Horde/Storage/Nodes/DirectoryNode.cs
--- a/Engine/Source/Programs/Shared/EpicGames.Horde/Storage/Nodes/DirectoryNode.cs
+++ b/Engine/Source/Programs/Shared/EpicGames.Horde/Storage/Nodes/DirectoryNode.cs
@@ -553,104 +553,24 @@ namespace EpicGames.Horde.Storage.Nodes
 			}
 		}
 
-		class InputFile
-		{
-			public FileInfo FileInfo { get; }
-			public List<NodeHandle> Handles { get; } = new List<NodeHandle>();
-
-			public InputFile(FileInfo fileInfo) => FileInfo = fileInfo;
-		}
-
-		class InputTree
-		{
-			public Dictionary<string, InputTree> Directories { get; } = new Dictionary<string, InputTree>();
-			public Dictionary<string, InputFile> Files { get; } = new Dictionary<string, InputFile>();
-
-			public void AddFiles(DirectoryReference baseDir, IEnumerable<FileInfo> files)
-			{
-				IEnumerator<FileInfo> enumerator = files.GetEnumerator();
-				if (enumerator.MoveNext())
-				{
-					if (AddFiles(baseDir.FullName, enumerator))
-					{
-						throw new ArgumentException($"File {enumerator.Current} is not under {baseDir}", nameof(baseDir));
-					}
-				}
-			}
-
-			bool AddFiles(ReadOnlySpan<char> baseDir, IEnumerator<FileInfo> files)
-			{
-				for (; ; )
-				{
-					FileInfo fileInfo = files.Current;
-
-					int dirNameLength = baseDir.Length;
-					if (fileInfo.FullName.Length < dirNameLength || !fileInfo.FullName.AsSpan().StartsWith(baseDir, FileReference.Comparison) || fileInfo.FullName[dirNameLength] != Path.DirectorySeparatorChar)
-					{
-						return true;
-					}
-
-					int nextDirLength = fileInfo.FullName.AsSpan(baseDir.Length + 1).IndexOf(Path.DirectorySeparatorChar);
-					if (nextDirLength == -1)
-					{
-						string fileName = fileInfo.FullName.Substring(dirNameLength + 1);
-						Files[fileName] = new InputFile(fileInfo);
-
-						if (!files.MoveNext())
-						{
-							return false;
-						}
-					}
-					else
-					{
-						string dirName = fileInfo.FullName.Substring(dirNameLength + 1, nextDirLength);
-
-						InputTree? nextTree;
-						if (!Directories.TryGetValue(dirName, out nextTree))
-						{
-							nextTree = new InputTree();
-							Directories.Add(dirName, nextTree);
-						}
-
-						ReadOnlySpan<char> nextBaseDir = fileInfo.FullName.AsSpan(0, dirNameLength + 1 + nextDirLength);
-						if (!nextTree.AddFiles(nextBaseDir, files))
-						{
-							return false;
-						}
-					}
-				}
-			}
-		}
-
 		/// <summary>
-		/// Create a directory tree from a set of files
+		/// Adds files from a flat list of paths
 		/// </summary>
-		/// <param name="baseDir">Base directory to root files from</param>
-		/// <param name="files">Files to include in the data</param>
-		/// <param name="options">Options for splitting files</param>
-		/// <param name="writer">Writer for output data</param>
+		/// <param name="baseDir">Base directory to base paths relative to</param>
+		/// <param name="files">Files to add</param>
+		/// <param name="options">Options for chunking file content</param>
+		/// <param name="writer">Writer for new node data</param>
+		/// <param name="progress">Feedback interface for progress updates</param>
 		/// <param name="cancellationToken">Cancellation token for the operation</param>
-		/// <returns>Root directory node containing the files</returns>
-		public static async Task<DirectoryNode> CreateAsync(DirectoryReference baseDir, IEnumerable<FileInfo> files, ChunkingOptions options, IStorageWriter writer, CancellationToken cancellationToken)
+		public static async Task<DirectoryNode> CreateAsync(DirectoryReference baseDir, IEnumerable<FileReference> files, ChunkingOptions options, IStorageWriter writer, IProgress<ICopyStats>? progress, CancellationToken cancellationToken)
 		{
-			// Create a tree out of all the files
-			InputTree tree = new InputTree();
-			tree.AddFiles(baseDir, files);
-
-			// Create leaf nodes
-			List<List<InputFile>> filePartitions = PartitionInputFiles(tree);
-			await Parallel.ForEachAsync(filePartitions, cancellationToken, (filePartition, ctx) => CreateLeafChunkNodesAsync(writer, filePartition, options, ctx));
-
-			await CreateInteriorChunkNodesAsync(tree, options.InteriorOptions, writer, cancellationToken);
-			return await CreateDirectoryNodesAsync(tree, writer, options, cancellationToken);
+			DirectoryNode directoryNode = new DirectoryNode();
+			await directoryNode.CopyFilesAsync(baseDir, files, options, writer, progress, cancellationToken);
+			return directoryNode;
 		}
 
-		static List<List<InputFile>> PartitionInputFiles(InputTree tree)
+		static List<List<FileContent>> PartitionInputFiles(List<FileContent> sortedFiles, long totalSize)
 		{
-			// Create the partitions
-			List<InputFile> sortedFiles = new List<InputFile>();
-			long totalSize = GetInputFiles(tree, sortedFiles);
-
 			// Maximum number of streams to write in parallel
 			const int MaxWriters = 8;
 
@@ -661,140 +581,142 @@ namespace EpicGames.Horde.Storage.Nodes
 			int numWriters = 1 + (int)Math.Min(totalSize / MinSizePerWriter, MaxWriters);
 
 			// Create the partitions
-			List<List<InputFile>> partitions = new List<List<InputFile>>();
+			List<List<FileContent>> partitions = new List<List<FileContent>>();
 
 			long writtenSize = 0;
 			long remainingSizeForPartition = 0;
-			foreach (InputFile sortedFile in sortedFiles)
+			foreach (FileContent sortedFile in sortedFiles)
 			{
-				if (sortedFile.FileInfo.Length >= remainingSizeForPartition && partitions.Count < numWriters)
+				if (sortedFile.Length >= remainingSizeForPartition && partitions.Count < numWriters)
 				{
-					partitions.Add(new List<InputFile>());
+					partitions.Add(new List<FileContent>());
 					remainingSizeForPartition = (totalSize - writtenSize) / (numWriters - partitions.Count);
 				}
 
 				partitions[^1].Add(sortedFile);
-				writtenSize += sortedFile.FileInfo.Length;
+				writtenSize += sortedFile.Length;
 
-				remainingSizeForPartition -= sortedFile.FileInfo.Length;
+				remainingSizeForPartition -= sortedFile.Length;
 			}
 
 			return partitions;
 		}
 
-		static long GetInputFiles(InputTree tree, List<InputFile> files)
+		static long GetInputFiles(DirectoryUpdate tree, List<FileContent> files)
 		{
 			long totalSize = 0;
-			foreach (InputTree directory in tree.Directories.Values)
+			foreach (DirectoryUpdate? directory in tree.Directories.Values)
 			{
-				totalSize += GetInputFiles(directory, files);
+				if (directory != null)
+				{
+					totalSize += GetInputFiles(directory, files);
+				}
 			}
 
-			foreach (InputFile inputFile in tree.Files.Values)
+			foreach (FileContent? inputFile in tree.Files.Values)
 			{
-				totalSize += inputFile.FileInfo.Length;
-				files.Add(inputFile);
+				if (inputFile != null)
+				{
+					totalSize += inputFile.Length;
+					files.Add(inputFile);
+				}
 			}
 			return totalSize;
 		}
 
-		static async ValueTask CreateLeafChunkNodesAsync(IStorageWriter writer, List<InputFile> files, ChunkingOptions options, CancellationToken cancellationToken)
+		static async ValueTask CreateLeafChunkNodesAsync(IStorageWriter writer, List<FileContent> files, CopyStats? copyStats, ChunkingOptions options, CancellationToken cancellationToken)
 		{
 			using IStorageWriter writerFork = writer.Fork();
-			foreach (InputFile file in files)
+			foreach (FileContent file in files)
 			{
-				file.Handles.AddRange(await LeafChunkedDataNode.CreateFromFileAsync(writer, file.FileInfo, options.LeafOptions, cancellationToken));
+				using (Stream stream = file.OpenRead())
+				{
+					file.Handles.Clear();
+					file.Handles.AddRange(await LeafChunkedDataNode.CreateFromStreamAsync(writer, stream, options.LeafOptions, cancellationToken));
+				}
+				copyStats?.Update(1, file.Length);
 			}
 		}
 
-		static async Task CreateLeafChunkNodesRecursiveAsync(InputTree tree, IStorageWriter writer, ChunkingOptions options, CancellationToken cancellationToken)
+		async Task CreateDirectoryNodesAsync(DirectoryUpdate tree, IStorageWriter writer, ChunkingOptions options, CancellationToken cancellationToken)
 		{
-			foreach (InputTree directory in tree.Directories.Values)
+			foreach ((string name, DirectoryUpdate? directory) in tree.Directories)
 			{
-				await CreateLeafChunkNodesRecursiveAsync(directory, writer, options, cancellationToken);
+				if (directory == null)
+				{
+					_nameToDirectoryEntry.Remove(name);
+				}
+				else
+				{
+					DirectoryNode childNode = await FindOrAddDirectoryAsync(name, cancellationToken);
+					await childNode.CreateDirectoryNodesAsync(directory, writer, options, cancellationToken);
+				}
 			}
-			foreach (InputFile file in tree.Files.Values)
+			foreach ((string name, FileContent? file) in tree.Files)
 			{
-				file.Handles.AddRange(await LeafChunkedDataNode.CreateFromFileAsync(writer, file.FileInfo, options.LeafOptions, cancellationToken));
+				if (file == null)
+				{
+					_nameToFileEntry.Remove(name);
+				}
+				else
+				{
+					FileEntry fileEntry = new FileEntry(name, FileEntryFlags.None, file.Length, file.Handles[0]);
+					_nameToFileEntry[name] = fileEntry;
+				}
 			}
 		}
 
-		static async Task CreateInteriorChunkNodesAsync(InputTree tree, InteriorChunkedDataNodeOptions options, IStorageWriter writer, CancellationToken cancellationToken)
+		/// <summary>
+		/// Adds files from a flat list of paths
+		/// </summary>
+		/// <param name="baseDir">Base directory to base paths relative to</param>
+		/// <param name="files">Files to add</param>
+		/// <param name="options">Options for chunking file content</param>
+		/// <param name="writer">Writer for new node data</param>
+		/// <param name="progress">Feedback interface for progress updates</param>
+		/// <param name="cancellationToken">Cancellation token for the operation</param>
+		public async Task CopyFilesAsync(DirectoryReference baseDir, IEnumerable<FileReference> files, ChunkingOptions options, IStorageWriter writer, IProgress<ICopyStats>? progress, CancellationToken cancellationToken)
 		{
-			foreach (InputTree directory in tree.Directories.Values)
-			{
-				await CreateInteriorChunkNodesAsync(directory, options, writer, cancellationToken);
-			}
+			DirectoryUpdate tree = new DirectoryUpdate();
+			tree.AddFiles(baseDir, files.Select(x => x.ToFileInfo()));
 
-			foreach (InputFile file in tree.Files.Values)
-			{
-				NodeHandle handle = await InteriorChunkedDataNode.CreateTreeAsync(file.Handles, options, writer, cancellationToken);
-				file.Handles.Clear();
-				file.Handles.Add(handle);
-			}
-		}
-
-		static async Task<DirectoryNode> CreateDirectoryNodesAsync(InputTree tree, IStorageWriter writer, ChunkingOptions options, CancellationToken cancellationToken)
-		{
-			DirectoryNode directoryNode = new DirectoryNode();
-			foreach ((string name, InputTree directory) in tree.Directories)
-			{
-				DirectoryNode childNode = await CreateDirectoryNodesAsync(directory, writer, options, cancellationToken);
-				directoryNode.AddDirectory(new DirectoryEntry(name, childNode));
-			}
-			foreach ((string name, InputFile file) in tree.Files)
-			{
-				FileEntry fileEntry = new FileEntry(name, FileEntryFlags.None, file.FileInfo.Length, file.Handles[0]);
-				directoryNode.AddFile(fileEntry);
-			}
-			return directoryNode;
+			await CopyFilesAsync(tree, options, writer, progress, cancellationToken);
 		}
 
 		/// <summary>
 		/// Adds files from a flat list of paths
 		/// </summary>
-		/// <param name="baseDir">Base directory to base paths relative to</param>
-		/// <param name="files">Files to add</param>
+		/// <param name="tree">Files to add</param>
 		/// <param name="options">Options for chunking file content</param>
 		/// <param name="writer">Writer for new node data</param>
 		/// <param name="progress">Feedback interface for progress updates</param>
 		/// <param name="cancellationToken">Cancellation token for the operation</param>
-		public async Task CopyFilesAsync(DirectoryReference baseDir, IEnumerable<FileReference> files, ChunkingOptions options, IStorageWriter writer, IProgress<ICopyStats>? progress, CancellationToken cancellationToken)
+		async Task CopyFilesAsync(DirectoryUpdate tree, ChunkingOptions options, IStorageWriter writer, IProgress<ICopyStats>? progress, CancellationToken cancellationToken)
 		{
-			Dictionary<DirectoryReference, DirectoryNode> dirToNode = new Dictionary<DirectoryReference, DirectoryNode>();
-			dirToNode.Add(baseDir, this);
+			// Create all the leaf chunks
+			List<FileContent> sortedFiles = new List<FileContent>();
+			long totalSize = GetInputFiles(tree, sortedFiles);
+
+			List<List<FileContent>> filePartitions = PartitionInputFiles(sortedFiles, totalSize);
 
-			List<(DirectoryNode, FileInfo)> groupedFiles = new List<(DirectoryNode, FileInfo)>();
-			foreach (FileReference file in files.OrderBy(x => x))
+			CopyStats? copyStats = null;
+			if (progress != null)
 			{
-				DirectoryNode? node = FindOrAddDirectory(file.Directory, baseDir, dirToNode);
-				if (node == null)
-				{
-					throw new InvalidOperationException($"File {file} is not under base directory {baseDir}");
-				}
-				groupedFiles.Add((node, file.ToFileInfo()));
+				copyStats = new CopyStats(sortedFiles.Count, totalSize, progress);
 			}
 
-			await CopyFromDirectoryAsync(groupedFiles, options, writer, progress, cancellationToken);
-		}
+			await Parallel.ForEachAsync(filePartitions, cancellationToken, (filePartition, ctx) => CreateLeafChunkNodesAsync(writer, filePartition, copyStats, options, ctx));
 
-		static DirectoryNode? FindOrAddDirectory(DirectoryReference dir, DirectoryReference baseDir, Dictionary<DirectoryReference, DirectoryNode> dirToNode)
-		{
-			DirectoryNode? node;
-			if (!dirToNode.TryGetValue(dir, out node))
+			// Create all the interior chunks
+			foreach (FileContent file in sortedFiles)
 			{
-				DirectoryReference? parentDir = dir.ParentDirectory;
-				if (parentDir != null)
-				{
-					DirectoryNode? parentNode = FindOrAddDirectory(parentDir, baseDir, dirToNode);
-					if (parentNode != null)
-					{
-						node = parentNode.AddDirectory(dir.GetDirectoryName());
-						dirToNode.Add(dir, node);
-					}
-				}
+				NodeHandle handle = await InteriorChunkedDataNode.CreateTreeAsync(file.Handles, options.InteriorOptions, writer, cancellationToken);
+				file.Handles.Clear();
+				file.Handles.Add(handle);
 			}
-			return node;
+
+			// Update all the directories
+			await CreateDirectoryNodesAsync(tree, writer, options, cancellationToken);
 		}
 
 		/// <summary>
@@ -854,116 +776,10 @@ namespace EpicGames.Horde.Storage.Nodes
 		/// <param name="cancellationToken">Cancellation token for the operation</param>
 		public async Task CopyFromDirectoryAsync(DirectoryInfo directoryInfo, ChunkingOptions options, IStorageWriter writer, IProgress<ICopyStats>? progress, CancellationToken cancellationToken = default)
 		{
-			// Enumerate all the files below this directory
-			List<(DirectoryNode DirectoryNode, FileInfo FileInfo)> files = new List<(DirectoryNode, FileInfo)>();
-			FindFilesToCopy(directoryInfo, files);
-			await CopyFromDirectoryAsync(files, options, writer, progress, cancellationToken);
-		}
-
-		/// <summary>
-		/// Adds files from a directory on disk
-		/// </summary>
-		/// <param name="files"></param>
-		/// <param name="options">Options for chunking file content</param>
-		/// <param name="writer">Writer for new node data</param>
-		/// <param name="progress">Progress notification object</param>
-		/// <param name="cancellationToken">Cancellation token for the operation</param>
-		/// <returns></returns>
-		public static async Task CopyFromDirectoryAsync(List<(DirectoryNode DirectoryNode, FileInfo FileInfo)> files, ChunkingOptions options, IStorageWriter writer, IProgress<ICopyStats>? progress, CancellationToken cancellationToken = default)
-		{
-			const int MaxWriters = 32;
-			const long MinSizePerWriter = 1024 * 1024;
-
-			// Compute the total size
-			long totalSize = files.Sum(x => x.Item2.Length);
-
-			// Create the progress reporting object
-			CopyStats? copyStats = null;
-			if (progress != null)
-			{
-				copyStats = new CopyStats(files.Count, totalSize, progress);
-			}
-
-			List<Task> tasks = new List<Task>();
-			FileEntry[] fileEntries = new FileEntry[files.Count];
-
-			// Split it into separate writers
-			long remainingSize = totalSize;
-			for (int minIdx = 0; minIdx < files.Count; )
-			{
-				long chunkSize = Math.Max(MinSizePerWriter, remainingSize / Math.Max(1, MaxWriters - tasks.Count));
-
-				int maxIdx = minIdx + 1;
-				long currentSize = files[minIdx].FileInfo.Length;
-				while (maxIdx < files.Count && currentSize <= chunkSize)
-				{
-					currentSize += files[maxIdx].FileInfo.Length;
-					maxIdx++;
-				}
-
-				int minIdxCopy = minIdx;
-				tasks.Add(Task.Run(() => CopyFilesAsync(files, minIdxCopy, maxIdx, fileEntries, options, writer, copyStats, cancellationToken), cancellationToken));
-
-				remainingSize -= currentSize;
-				minIdx = maxIdx;
-			}
-
-			// Wait for them all to finish
-			await Task.WhenAll(tasks);
-
-			// Update the directory with all the output entries
-			for (int idx = 0; idx < files.Count; idx++)
-			{
-				files[idx].DirectoryNode.AddFile(fileEntries[idx]);
-			}
-
-			// Write the final stats
-			copyStats?.Flush();
-		}
-
-		void FindFilesToCopy(DirectoryInfo directoryInfo, List<(DirectoryNode, FileInfo)> files)
-		{
-			foreach (DirectoryInfo subDirectoryInfo in directoryInfo.EnumerateDirectories())
-			{
-				AddDirectory(subDirectoryInfo.Name).FindFilesToCopy(subDirectoryInfo, files);
-			}
-			foreach (FileInfo fileInfo in directoryInfo.EnumerateFiles())
-			{
-				files.Add((this, fileInfo));
-			}
-		}
-
-		static async Task CopyFilesAsync(List<(DirectoryNode DirectoryNode, FileInfo FileInfo)> files, int minIdx, int maxIdx, FileEntry[] entries, ChunkingOptions options, IStorageWriter baseWriter, CopyStats? copyStats, CancellationToken cancellationToken)
-		{
-			IStorageWriter writer = baseWriter;
-			try
-			{
-				if (minIdx != 0)
-				{
-					writer = baseWriter.Fork();
-				}
+			DirectoryUpdate tree = new DirectoryUpdate();
+			tree.AddFiles(directoryInfo);
 
-				ChunkedDataWriter fileNodeWriter = new ChunkedDataWriter(writer, options);
-				for (int idx = minIdx; idx < maxIdx; idx++)
-				{
-					FileInfo fileInfo = files[idx].FileInfo;
-					NodeHandle handle = await fileNodeWriter.CreateAsync(fileInfo, cancellationToken);
-					entries[idx] = new FileEntry(fileInfo.Name, FileEntryFlags.None, fileNodeWriter.Length, handle);
-					copyStats?.Update(1, fileNodeWriter.Length);
-				}
-
-				if (minIdx != 0)
-				{
-					await writer.FlushAsync(cancellationToken);
-				}
-			}
-			finally
-			{
-				if (minIdx != 0)
-				{
-					writer.Dispose();
-				}
-			}
+			await CopyFilesAsync(tree, options, writer, progress, cancellationToken);
 		}
 
 		/// <summary>
@@ -1002,6 +818,125 @@ namespace EpicGames.Horde.Storage.Nodes
 		public Stream AsZipStream(FileFilter? filter = null) => new DirectoryNodeZipStream(this, filter);
 	}
 
+	/// <summary>
+	/// Information about the content of a file
+	/// </summary>
+	public abstract class FileContent
+	{
+		internal List<NodeHandle> Handles { get; } = new List<NodeHandle>();
+
+		/// <summary>
+		/// Length of the file
+		/// </summary>
+		public abstract long Length { get; }
+
+		/// <summary>
+		/// Opens the stream for reading
+		/// </summary>
+		public abstract Stream OpenRead();
+	}
+
+	/// <summary>
+	/// Information about the content of a file
+	/// </summary>
+	public class FileInfoFileContent : FileContent
+	{
+		readonly FileInfo _fileInfo;
+
+		/// <inheritdoc/>
+		public override long Length => _fileInfo.Length;
+
+		/// <summary>
+		/// Constructor
+		/// </summary>
+		public FileInfoFileContent(FileInfo fileInfo) => _fileInfo = fileInfo;
+
+		/// <inheritdoc/>
+		public override Stream OpenRead() => _fileInfo.OpenRead();
+	}
+
+	/// <summary>
+	/// Describes an update to a directory node
+	/// </summary>
+	public class DirectoryUpdate
+	{
+		/// <summary>
+		/// Directories to be updated
+		/// </summary>
+		public Dictionary<string, DirectoryUpdate?> Directories { get; } = new Dictionary<string, DirectoryUpdate?>();
+
+		/// <summary>
+		/// Files to be updated
+		/// </summary>
+		public Dictionary<string, FileContent?> Files { get; } = new Dictionary<string, FileContent?>();
+
+		/// <summary>
+		/// Adds new files from a directory on disk
+		/// </summary>
+		/// <param name="directoryInfo"></param>
+		public void AddFiles(DirectoryInfo directoryInfo) => AddFiles(new DirectoryReference(directoryInfo), directoryInfo.EnumerateFiles("*", SearchOption.AllDirectories));
+
+		/// <summary>
+		/// Adds a filtered list of files from disk
+		/// </summary>
+		/// <param name="baseDir">Base directory to root file paths from</param>
+		/// <param name="files">Files to add</param>
+		public void AddFiles(DirectoryReference baseDir, IEnumerable<FileInfo> files)
+		{
+			IEnumerator<FileInfo> enumerator = files.GetEnumerator();
+			if (enumerator.MoveNext())
+			{
+				if (AddFiles(baseDir.FullName, enumerator))
+				{
+					throw new ArgumentException($"File {enumerator.Current} is not under {baseDir}", nameof(baseDir));
+				}
+			}
+		}
+
+		bool AddFiles(ReadOnlySpan<char> baseDir, IEnumerator<FileInfo> files)
+		{
+			for (; ; )
+			{
+				FileInfo fileInfo = files.Current;
+
+				int dirNameLength = baseDir.Length;
+				if (fileInfo.FullName.Length < dirNameLength || !fileInfo.FullName.AsSpan().StartsWith(baseDir, FileReference.Comparison) || fileInfo.FullName[dirNameLength] != Path.DirectorySeparatorChar)
+				{
+					return true;
+				}
+
+				int nextDirLength = fileInfo.FullName.AsSpan(baseDir.Length + 1).IndexOf(Path.DirectorySeparatorChar);
+				if (nextDirLength == -1)
+				{
+					string fileName = fileInfo.FullName.Substring(dirNameLength + 1);
+					Files[fileName] = new FileInfoFileContent(fileInfo);
+
+					if (!files.MoveNext())
+					{
+						return false;
+					}
+				}
+				else
+				{
+					string dirName = fileInfo.FullName.Substring(dirNameLength + 1, nextDirLength);
+
+					DirectoryUpdate? nextTree;
+					if (!Directories.TryGetValue(dirName, out nextTree) || nextTree == null)
+					{
+						nextTree = new DirectoryUpdate();
+						Directories[dirName] = nextTree;
+					}
+
+					ReadOnlySpan<char> nextBaseDir = fileInfo.FullName.AsSpan(0, dirNameLength + 1 + nextDirLength);
+					if (!nextTree.AddFiles(nextBaseDir, files))
+					{
+						return false;
+					}
+				}
+			}
+		}
+	}
+
 	/// <summary>
 	/// Stream which zips a directory node tree dynamically
 	/// </summary>
