Horde: Fix parallel partitioning of input files into bundles, due to contents not being disposed on flush.

(Also add an exception for to any pending async tasks on dispose, ensuring that these situations are clearly flagged in future).

diff --git a/Engine/Source/Programs/Shared/EpicGames.Horde/Storage/Nodes/DirectoryNode.cs b/Engine/Source/Programs/Shared/EpicGames.Horde/Storage/Nodes/DirectoryNode.cs
--- a/Engine/Source/Programs/Shared/EpicGames.Horde/Storage/Nodes/DirectoryNode.cs
+++ b/Engine/Source/Programs/Shared/EpicGames.Horde/Storage/Nodes/DirectoryNode.cs
@@ -632,10 +632,11 @@ namespace EpicGames.Horde.Storage.Nodes
 				using (Stream stream = file.OpenRead())
 				{
 					file.Handles.Clear();
-					file.Handles.AddRange(await LeafChunkedDataNode.CreateFromStreamAsync(writer, stream, options.LeafOptions, cancellationToken));
+					file.Handles.AddRange(await LeafChunkedDataNode.CreateFromStreamAsync(writerFork, stream, options.LeafOptions, cancellationToken));
 				}
 				copyStats?.Update(1, file.Length);
 			}
+			await writerFork.FlushAsync(cancellationToken);
 		}
 
 		async Task CreateDirectoryNodesAsync(DirectoryUpdate tree, IStorageWriter writer, ChunkingOptions options, CancellationToken cancellationToken)
@@ -696,7 +697,6 @@ namespace EpicGames.Horde.Storage.Nodes
 			// Create all the leaf chunks
 			List<FileContent> sortedFiles = new List<FileContent>();
 			long totalSize = GetInputFiles(tree, sortedFiles);
-
 			List<List<FileContent>> filePartitions = PartitionInputFiles(sortedFiles, totalSize);
 
 			CopyStats? copyStats = null;
diff --git a/Engine/Source/Programs/Shared/EpicGames.Horde/Storage/TreeWriter.cs b/Engine/Source/Programs/Shared/EpicGames.Horde/Storage/TreeWriter.cs
--- a/Engine/Source/Programs/Shared/EpicGames.Horde/Storage/TreeWriter.cs
+++ b/Engine/Source/Programs/Shared/EpicGames.Horde/Storage/TreeWriter.cs
@@ -300,6 +300,10 @@ namespace EpicGames.Horde.Storage
 			/// <inheritdoc/>
 			public void Dispose()
 			{
+				if (!CompleteTask.IsCompleted)
+				{
+					_completeEvent.TrySetException(new ObjectDisposedException($"Bundle {BundleId} has been disposed. Make sure to call FlushAsync() on forked writers that may have dependencies before disposing.", $"Bundle {BundleId}"));
+				}
 				if (_currentPacket != null)
 				{
 					_currentPacket.Dispose();
@@ -315,11 +319,14 @@ namespace EpicGames.Horde.Storage
 			public bool CanComplete() => !IsReadOnly && _dependencies.Count == 0;
 
 			// Add a dependency onto another bundle
-			public void AddDependencyOn(PendingBundle bundle)
+			public void AddDependencyOn(PendingBundle bundle, ILogger? traceLogger)
 			{
 				if (bundle != this)
 				{
-					_dependencies.Add(bundle.CompleteTask);
+					if (_dependencies.Add(bundle.CompleteTask))
+					{
+						traceLogger?.LogInformation("Added dependency from bundle {BundleId} on {OtherBundleId}", BundleId, bundle.BundleId);
+					}
 				}
 			}
 
@@ -446,6 +453,16 @@ namespace EpicGames.Horde.Storage
 				try
 				{
 					await prevWriteTask;
+					if (traceLogger != null)
+					{
+						foreach (Task dependency in _dependencies)
+						{
+							if (!dependency.IsCompleted)
+							{
+								traceLogger.LogInformation("Bundle {BundleId} is stalling waiting for dependencies to flush first", BundleId);
+							}
+						}
+					}
 					await Task.WhenAll(_dependencies);
 
 					Bundle bundle = CreateBundle();
@@ -724,7 +741,7 @@ namespace EpicGames.Horde.Storage
 			{
 				if (reference.PendingBundle != null)
 				{
-					currentBundle.AddDependencyOn(reference.PendingBundle);
+					currentBundle.AddDependencyOn(reference.PendingBundle, _traceLogger);
 				}
 			}
 
