Horde: Fix race condition with starting recv task for compute work before any buffer is attached.

diff --git a/Engine/Source/Programs/Shared/EpicGames.Horde/Compute/ComputeSocket.cs b/Engine/Source/Programs/Shared/EpicGames.Horde/Compute/ComputeSocket.cs
--- a/Engine/Source/Programs/Shared/EpicGames.Horde/Compute/ComputeSocket.cs
+++ b/Engine/Source/Programs/Shared/EpicGames.Horde/Compute/ComputeSocket.cs
@@ -48,7 +48,7 @@ namespace EpicGames.Horde.Compute
 
 		readonly CancellationTokenSource _cancellationSource = new CancellationTokenSource();
 
-		readonly BackgroundTask _recvTask;
+		BackgroundTask? _recvTask;
 		readonly Dictionary<int, IComputeBufferWriter> _recvBufferWriters = new Dictionary<int, IComputeBufferWriter>();
 
 		readonly SemaphoreSlim _sendSemaphore = new SemaphoreSlim(1, 1);
@@ -68,7 +68,6 @@ namespace EpicGames.Horde.Compute
 			_transport = transport;
 			_endpoint = endpoint;
 			_logger = logger;
-			_recvTask = BackgroundTask.StartNew(ctx => RunRecvTaskAsync(transport, ctx));
 		}
 
 		/// <summary>
@@ -95,8 +94,11 @@ namespace EpicGames.Horde.Compute
 			await _transport.MarkCompleteAsync(cancellationToken);
 
 			// Wait for the reader to stop
-			await _recvTask.Task.WaitAsync(TimeSpan.FromSeconds(5), cancellationToken);
-			await _recvTask.StopAsync();
+			if (_recvTask != null)
+			{
+				await _recvTask.Task.WaitAsync(TimeSpan.FromSeconds(5), cancellationToken);
+				await _recvTask.StopAsync();
+			}
 		}
 
 		/// <inheritdoc/>
@@ -104,7 +106,10 @@ namespace EpicGames.Horde.Compute
 		{
 			_cancellationSource.Cancel();
 			await Task.WhenAll(_sendTasks.Values);
-			await _recvTask.DisposeAsync();
+			if (_recvTask != null)
+			{
+				await _recvTask.DisposeAsync();
+			}
 			_sendSemaphore.Dispose();
 			_cancellationSource.Dispose();
 			GC.SuppressFinalize(this);
@@ -266,6 +271,7 @@ namespace EpicGames.Horde.Compute
 				if (!complete)
 				{
 					_recvBufferWriters.Add(channelId, recvBufferWriter.AddRef());
+					_recvTask ??= BackgroundTask.StartNew(ctx => RunRecvTaskAsync(_transport, ctx));
 				}
 			}
 
